Multi-GPU:  strip out Alternate Frame Rendering (AFR).  Michal Valient requested this for 5.2.  AFR has heavily decayed in 5.0, not supporting Lumen, Nanite, VSMs. and likely other newer features, and it's impossible even in principle to get most of those features to ever work in a performant way, due to their temporal nature.  Removing this lowers maintenance and support costs for MGPU going forward.

A lot of files touched, but generally it's a mechanical matter of removing the global variable GNumAlternateFrameRenderingGroups, and treating all code using it as if it's a fixed constant of one.  If a conditional becomes always false, the code block is removed.  Certain utility functions only called from dead stripped AFR code are then removed (e.g. RHIBroadcastTemporalEffect).  On the D3D11 side, RHIBeginUpdateMultiFrameResource / RHIEndUpdateMultiFrameResource become NOPs (return at the top of the function) when GNumAlternateFrameRenderingGroups is 1, so those are removed across the board.

diff --git a/Engine/Plugins/FX/Niagara/Source/Niagara/Private/NiagaraGpuComputeDispatch.cpp b/Engine/Plugins/FX/Niagara/Source/Niagara/Private/NiagaraGpuComputeDispatch.cpp
--- a/Engine/Plugins/FX/Niagara/Source/Niagara/Private/NiagaraGpuComputeDispatch.cpp
+++ b/Engine/Plugins/FX/Niagara/Source/Niagara/Private/NiagaraGpuComputeDispatch.cpp
@@ -86,11 +86,6 @@ const FName FNiagaraGpuComputeDispatch::Name(TEXT("FNiagaraGpuComputeDispatch"))
 
 namespace FNiagaraGpuComputeDispatchLocal
 {
-#if WITH_MGPU
-	const FName TemporalEffectBuffersName("FNiagaraGpuComputeDispatch_Buffers");
-	const FName TemporalEffectTexturesName("FNiagaraGpuComputeDispatch_Textures");
-#endif // WITH_MGPU
-
 	int32 GTickFlushMaxQueuedFrames = 3;
 	static FAutoConsoleVariableRef CVarNiagaraTickFlushMaxQueuedFrames(
 		TEXT("fx.Niagara.Batcher.TickFlush.MaxQueuedFrames"),
@@ -1240,16 +1235,6 @@ void FNiagaraGpuComputeDispatch::ExecuteTicks(FRDGBuilder& GraphBuilder, TConstA
 				ComputeContext->SetDataToRender(CurrentData);
 
 #if WITH_MGPU
-				if (bAFREnabled)
-				{
-					AddAFRBuffer(CurrentData->GetGPUBufferFloat().Buffer);
-					AddAFRBuffer(CurrentData->GetGPUBufferHalf().Buffer);
-					AddAFRBuffer(CurrentData->GetGPUBufferInt().Buffer);
-					if (ComputeContext->MainDataSet->RequiresPersistentIDs())
-					{
-						AddAFRBuffer(ComputeContext->MainDataSet->GetGPUFreeIDs().Buffer);
-					}
-				}
 				if (bCrossGPUTransferEnabled)
 				{
 					AddCrossGPUTransfer(GraphBuilder.RHICmdList, CurrentData->GetGPUBufferFloat().Buffer);
@@ -1716,8 +1701,7 @@ void FNiagaraGpuComputeDispatch::PreInitViews(FRDGBuilder& GraphBuilder, bool bA
 	bRaisedWarningThisFrame = false;
 #endif
 #if WITH_MGPU
-	bAFREnabled = GNumAlternateFrameRenderingGroups > 1;
-	bCrossGPUTransferEnabled = !bAFREnabled && (GNumExplicitGPUsForRendering > 1);
+	bCrossGPUTransferEnabled = GNumExplicitGPUsForRendering > 1;
 	StageToTransferGPUBuffers = ENiagaraGpuComputeTickStage::Last;
 	StageToWaitForGPUTransfers = ENiagaraGpuComputeTickStage::First;
 #endif
@@ -2295,10 +2279,6 @@ void FNiagaraGpuComputeDispatch::MultiGPUResourceModified(FRDGBuilder& GraphBuil
 
 void FNiagaraGpuComputeDispatch::MultiGPUResourceModified(FRHICommandList& RHICmdList, FRHIBuffer* Buffer, bool bRequiredForSimulation, bool bRequiredForRendering) const
 {
-	if (bAFREnabled && bRequiredForSimulation)
-	{
-		const_cast<FNiagaraGpuComputeDispatch*>(this)->AddAFRBuffer(Buffer);
-	}
 	if (bCrossGPUTransferEnabled && bRequiredForRendering)
 	{
 		const_cast<FNiagaraGpuComputeDispatch*>(this)->AddCrossGPUTransfer(RHICmdList, Buffer);
@@ -2307,13 +2287,6 @@ void FNiagaraGpuComputeDispatch::MultiGPUResourceModified(FRHICommandList& RHICm
 
 void FNiagaraGpuComputeDispatch::MultiGPUResourceModified(FRHICommandList& RHICmdList, FRHITexture* Texture, bool bRequiredForSimulation, bool bRequiredForRendering) const
 {
-	if (bAFREnabled && bRequiredForSimulation)
-	{
-		if (Texture)
-		{
-			const_cast<FNiagaraGpuComputeDispatch*>(this)->AFRTextures.Add(Texture);
-		}
-	}
 	if (bCrossGPUTransferEnabled && bRequiredForRendering)
 	{
 		const bool bPullData = false;
@@ -2330,15 +2303,6 @@ void FNiagaraGpuComputeDispatch::MultiGPUResourceModified(FRHICommandList& RHICm
 	}
 }
 
-void FNiagaraGpuComputeDispatch::AddAFRBuffer(FRHIBuffer* Buffer)
-{
-	check(bAFREnabled);
-	if (Buffer)
-	{
-		AFRBuffers.Add(Buffer);
-	}
-}
-
 void FNiagaraGpuComputeDispatch::AddCrossGPUTransfer(FRHICommandList& RHICmdList, FRHIBuffer* Buffer)
 {
 	check(bCrossGPUTransferEnabled);
@@ -2386,18 +2350,6 @@ void FNiagaraGpuComputeDispatch::TransferMultiGPUBufers(FRHICommandList& RHICmdL
 		return;
 	}
 
-	// Transfer buffers for AFR rendering
-	if (AFRBuffers.Num())
-	{
-		RHICmdList.BroadcastTemporalEffect(FNiagaraGpuComputeDispatchLocal::TemporalEffectBuffersName, AFRBuffers);
-		AFRBuffers.Reset();
-	}
-	if (AFRTextures.Num())
-	{
-		RHICmdList.BroadcastTemporalEffect(FNiagaraGpuComputeDispatchLocal::TemporalEffectTexturesName, AFRTextures);
-		AFRTextures.Reset();
-	}
-
 	// Transfer buffers for cross GPU rendering
 	if (CrossGPUTransferBuffers.Num())
 	{
@@ -2410,8 +2362,7 @@ void FNiagaraGpuComputeDispatch::WaitForMultiGPUBuffers(FRHICommandList& RHICmdL
 {
 	if (StageToWaitForGPUTransfers == TickStage)
 	{
-		RHICmdList.WaitForTemporalEffect(FNiagaraGpuComputeDispatchLocal::TemporalEffectBuffersName);
-		RHICmdList.WaitForTemporalEffect(FNiagaraGpuComputeDispatchLocal::TemporalEffectTexturesName);
+		// TODO:  implement delayed fence wait here for cross GPU transfers issued above
 	}
 }
 #endif // WITH_MGPU
diff --git a/Engine/Plugins/FX/Niagara/Source/Niagara/Private/NiagaraGpuComputeDispatch.h b/Engine/Plugins/FX/Niagara/Source/Niagara/Private/NiagaraGpuComputeDispatch.h
--- a/Engine/Plugins/FX/Niagara/Source/Niagara/Private/NiagaraGpuComputeDispatch.h
+++ b/Engine/Plugins/FX/Niagara/Source/Niagara/Private/NiagaraGpuComputeDispatch.h
@@ -255,14 +255,9 @@ private:
 	ENiagaraGpuComputeTickStage::Type StageToTransferGPUBuffers = ENiagaraGpuComputeTickStage::First;
 	ENiagaraGpuComputeTickStage::Type StageToWaitForGPUTransfers = ENiagaraGpuComputeTickStage::First;
 
-	bool bAFREnabled = false;
-	TArray<FRHIBuffer*> AFRBuffers;
-	TArray<FRHITexture*> AFRTextures;
-
 	bool bCrossGPUTransferEnabled = false;
 	TArray<FTransferResourceParams> CrossGPUTransferBuffers;
 
-	void AddAFRBuffer(FRHIBuffer* Buffer);
 	void AddCrossGPUTransfer(FRHICommandList& RHICmdList, FRHIBuffer* Buffer);
 
 	void CalculateCrossGPUTransferLocation();
diff --git a/Engine/Source/Runtime/Apple/MetalRHI/Private/MetalProfiler.cpp b/Engine/Source/Runtime/Apple/MetalRHI/Private/MetalProfiler.cpp
--- a/Engine/Source/Runtime/Apple/MetalRHI/Private/MetalProfiler.cpp
+++ b/Engine/Source/Runtime/Apple/MetalRHI/Private/MetalProfiler.cpp
@@ -190,11 +190,6 @@ void FMetalGPUProfiler::BeginFrame()
 		CurrentEventNodeFrame = new FMetalEventNodeFrame(Context, GTriggerGPUProfile);
 		CurrentEventNodeFrame->StartFrame();
 		
-		if(GNumAlternateFrameRenderingGroups > 1)
-		{
-			GTriggerGPUProfile = false;
-		}
-
 		if(GTriggerGPUProfile)
 		{
 			bTrackingEvents = true;
diff --git a/Engine/Source/Runtime/D3D12RHI/Private/D3D12Adapter.cpp b/Engine/Source/Runtime/D3D12RHI/Private/D3D12Adapter.cpp
--- a/Engine/Source/Runtime/D3D12RHI/Private/D3D12Adapter.cpp
+++ b/Engine/Source/Runtime/D3D12RHI/Private/D3D12Adapter.cpp
@@ -258,14 +258,8 @@ FD3D12Adapter::FD3D12Adapter(FD3D12AdapterDesc& DescIn)
 
 	uint32 MaxGPUCount = 1; // By default, multi-gpu is disabled.
 #if WITH_MGPU
-	if (!FParse::Value(FCommandLine::Get(), TEXT("MaxGPUCount="), MaxGPUCount))
-	{
-		// If there is a mode token in the command line, enable multi-gpu.
-		if (FParse::Param(FCommandLine::Get(), TEXT("AFR")))
-		{
-			MaxGPUCount = MAX_NUM_GPUS;
-		}
-	}
+	FParse::Value(FCommandLine::Get(), TEXT("MaxGPUCount="), MaxGPUCount);
+
 	if (FParse::Param(FCommandLine::Get(), TEXT("VMGPU")))
 	{
 		GVirtualMGPU = 1;
@@ -726,40 +720,6 @@ void FD3D12Adapter::CreateRootDevice(bool bWithDebug)
 		GNumExplicitGPUsForRendering = Desc.NumDeviceNodes;
 		UE_LOG(LogD3D12RHI, Log, TEXT("Enabling multi-GPU with %d nodes"), Desc.NumDeviceNodes);
 	}
-
-	// Viewport ignores AFR if PresentGPU is specified.
-	int32 Dummy;
-	if (!FParse::Value(FCommandLine::Get(), TEXT("PresentGPU="), Dummy))
-	{
-		bool bWantsAFR = false;
-		if (FParse::Value(FCommandLine::Get(), TEXT("NumAFRGroups="), GNumAlternateFrameRenderingGroups))
-		{
-			bWantsAFR = true;
-		}
-		else if (FParse::Param(FCommandLine::Get(), TEXT("AFR")))
-		{
-			bWantsAFR = true;
-			GNumAlternateFrameRenderingGroups = GNumExplicitGPUsForRendering;
-		}
-
-		if (bWantsAFR)
-		{
-			if (GNumAlternateFrameRenderingGroups <= 1 || GNumAlternateFrameRenderingGroups > GNumExplicitGPUsForRendering)
-			{
-				UE_LOG(LogD3D12RHI, Error, TEXT("Cannot enable alternate frame rendering because NumAFRGroups (%u) must be > 1 and <= MaxGPUCount (%u)"), GNumAlternateFrameRenderingGroups, GNumExplicitGPUsForRendering);
-				GNumAlternateFrameRenderingGroups = 1;
-			}
-			else if (GNumExplicitGPUsForRendering % GNumAlternateFrameRenderingGroups != 0)
-			{
-				UE_LOG(LogD3D12RHI, Error, TEXT("Cannot enable alternate frame rendering because MaxGPUCount (%u) must be evenly divisible by NumAFRGroups (%u)"), GNumExplicitGPUsForRendering, GNumAlternateFrameRenderingGroups);
-				GNumAlternateFrameRenderingGroups = 1;
-			}
-			else
-			{
-				UE_LOG(LogD3D12RHI, Log, TEXT("Enabling alternate frame rendering with %u AFR groups"), GNumAlternateFrameRenderingGroups);
-			}
-		}
-	}
 #endif
 }
 
@@ -1513,22 +1473,6 @@ void FD3D12Adapter::EndFrame()
 #endif
 }
 
-#if WITH_MGPU
-FD3D12Adapter::FTemporalEffect& FD3D12Adapter::GetTemporalEffect(const FName& EffectName)
-{
-	FTemporalEffect* Effect = TemporalEffectMap.Find(EffectName);
-
-	if (Effect == nullptr)
-	{
-		Effect = &TemporalEffectMap.Emplace(EffectName, FTemporalEffect());
-		Effect->AddDefaulted(FRHIGPUMask::All().GetNumActive());
-	}
-
-	check(Effect);
-	return *Effect;
-}
-#endif // WITH_MGPU
-
 FD3D12FastConstantAllocator& FD3D12Adapter::GetTransientUniformBufferAllocator()
 {
 	// Multi-GPU support : is using device 0 always appropriate here?
diff --git a/Engine/Source/Runtime/D3D12RHI/Private/D3D12Adapter.h b/Engine/Source/Runtime/D3D12RHI/Private/D3D12Adapter.h
--- a/Engine/Source/Runtime/D3D12RHI/Private/D3D12Adapter.h
+++ b/Engine/Source/Runtime/D3D12RHI/Private/D3D12Adapter.h
@@ -374,9 +374,6 @@ public:
 
 	FD3D12TransientHeapCache& GetOrCreateTransientHeapCache();
 
-	typedef TArray<FD3D12SyncPointRef, TInlineAllocator<MAX_NUM_GPUS>> FTemporalEffect;
-	FTemporalEffect& GetTemporalEffect(const FName& EffectName);
-
 	FD3D12FastConstantAllocator& GetTransientUniformBufferAllocator();
 	void ReleaseTransientUniformBufferAllocator(FTransientUniformBufferAllocator* InAllocator);
 
@@ -574,10 +571,6 @@ protected:
 
 	FD3D12MemoryInfo MemoryInfo;
 
-#if WITH_MGPU
-	TMap<FName, FTemporalEffect> TemporalEffectMap;
-#endif
-
 	TArray<FTransientUniformBufferAllocator*> TransientUniformBufferAllocators;
 	FCriticalSection TransientUniformBufferAllocatorsCS;
 
diff --git a/Engine/Source/Runtime/D3D12RHI/Private/D3D12Allocation.cpp b/Engine/Source/Runtime/D3D12RHI/Private/D3D12Allocation.cpp
--- a/Engine/Source/Runtime/D3D12RHI/Private/D3D12Allocation.cpp
+++ b/Engine/Source/Runtime/D3D12RHI/Private/D3D12Allocation.cpp
@@ -2136,6 +2136,7 @@ void FD3D12FastAllocatorPage::UpdateFence()
 
 void FD3D12FastAllocatorPagePool::ReturnFastAllocatorPage(FD3D12FastAllocatorPage* Page)
 {
+	// TODO:  AFR has been removed, but I don't understand the comment, so I'm leaving it.  What did this code have to do with AFR?
 	// Extend the lifetime of these resources when in AFR as other nodes might be relying on this
 	Page->UpdateFence();
 	Pool.Add(Page);
diff --git a/Engine/Source/Runtime/D3D12RHI/Private/D3D12CommandContext.h b/Engine/Source/Runtime/D3D12RHI/Private/D3D12CommandContext.h
--- a/Engine/Source/Runtime/D3D12RHI/Private/D3D12CommandContext.h
+++ b/Engine/Source/Runtime/D3D12RHI/Private/D3D12CommandContext.h
@@ -5,15 +5,6 @@ D3D12CommandContext.h: D3D12 Command Context Interfaces
 =============================================================================*/
 
 #pragma once
-#define AFR_ENGINE_CHANGES_PRESENT WITH_MGPU
-
-// TODO: Because the upper engine is yet to implement these interfaces we can't 'override' something that doesn't exist.
-//       Remove when upper engine is ready
-#if AFR_ENGINE_CHANGES_PRESENT
-#define AFR_API_OVERRIDE override
-#else
-#define AFR_API_OVERRIDE
-#endif
 
 #include "D3D12RHIPrivate.h"
 #include "Windows/AllowWindowsPlatformTypes.h"
@@ -625,16 +616,6 @@ public:
 
 	void ResolveTexture(UE::RHICore::FResolveTextureInfo Info);
 
-	// When using Alternate Frame Rendering some temporal effects i.e. effects which consume GPU work from previous frames must synchronize their resources
-	// to prevent visual corruption.
-
-	// This should be called right before the effect consumes it's temporal resources.
-	virtual void RHIWaitForTemporalEffect(const FName& InEffectName) final AFR_API_OVERRIDE;
-
-	// This should be called right after the effect generates the resources which will be used in subsequent frame(s).
-	virtual void RHIBroadcastTemporalEffect(const FName& InEffectName, const TArrayView<FRHITexture*> InTextures) final AFR_API_OVERRIDE;
-	virtual void RHIBroadcastTemporalEffect(const FName& InEffectName, const TArrayView<FRHIBuffer*> InBuffers) final AFR_API_OVERRIDE;
-
 #if D3D12_RHI_RAYTRACING
 	virtual void RHIBindAccelerationStructureMemory(FRHIRayTracingScene* Scene, FRHIBuffer* Buffer, uint32 BufferOffset) final override;
 	virtual void BuildAccelerationStructuresInternal(const TArrayView<D3D12_BUILD_RAYTRACING_ACCELERATION_STRUCTURE_DESC> BuildDesc);
@@ -717,11 +698,6 @@ protected:
 
 private:
 
-#if WITH_MGPU
-	template <typename TD3D12Resource, typename TCopyFunction>
-	void BroadcastTemporalEffect(const FName& InEffectName, const TArrayView<TD3D12Resource*> InResources, const TCopyFunction& InCopyFunction);
-#endif
-
 	static void ClearUAV(TRHICommandList_RecursiveHazardous<FD3D12CommandContext>& RHICmdList, FD3D12UnorderedAccessView* UAV, const void* ClearValues, bool bFloat);
 
 	template <typename TRHIShader>
@@ -741,11 +717,8 @@ private:
 	TArray<FRHIUniformBuffer*> StaticUniformBuffers;
 };
 
-// This class is a shim to get AFR working. Currently the upper engine only queries for the 'Immediate Context'
-// once. However when in AFR we need to switch which context is active every frame so we return an instance of this class
-// as the default context so that we can control when to swap which device we talk to.
-// Because IRHICommandContext is pure virtual we can return the normal FD3D12CommandContext when not using mGPU thus there
-// is no additional overhead for the common case i.e. 1 GPU.
+// Version of command context to handle multi-GPU.  Because IRHICommandContext is pure virtual we can return the normal
+// FD3D12CommandContext when not using mGPU, thus there is no additional overhead for the common case i.e. 1 GPU.
 class FD3D12CommandContextRedirector final : public FD3D12CommandContextBase
 {
 public:
@@ -981,21 +954,6 @@ public:
 		ContextRedirect(RHISetShadingRate(ShadingRate, Combiner));
 	}
 
-	FORCEINLINE virtual void RHIWaitForTemporalEffect(const FName& InEffectName) final AFR_API_OVERRIDE
-	{
-		ContextRedirect(RHIWaitForTemporalEffect(InEffectName));
-	}
-
-	FORCEINLINE virtual void RHIBroadcastTemporalEffect(const FName& InEffectName, const TArrayView<FRHITexture*> InTextures) final AFR_API_OVERRIDE
-	{
-		ContextRedirect(RHIBroadcastTemporalEffect(InEffectName, InTextures));
-	}
-
-	FORCEINLINE virtual void RHIBroadcastTemporalEffect(const FName& InEffectName, const TArrayView<FRHIBuffer*> InBuffers) final AFR_API_OVERRIDE
-	{
-		ContextRedirect(RHIBroadcastTemporalEffect(InEffectName, InBuffers));
-	}
-
 	FORCEINLINE virtual void RHIBeginFrame() final override
 	{
 		ContextRedirect(RHIBeginFrame());
diff --git a/Engine/Source/Runtime/D3D12RHI/Private/D3D12Commands.cpp b/Engine/Source/Runtime/D3D12RHI/Private/D3D12Commands.cpp
--- a/Engine/Source/Runtime/D3D12RHI/Private/D3D12Commands.cpp
+++ b/Engine/Source/Runtime/D3D12RHI/Private/D3D12Commands.cpp
@@ -19,14 +19,6 @@ D3D12Commands.cpp: D3D RHI commands implementation.
 #include "GlobalRenderResources.h"
 #include "RHIShaderParametersShared.h"
 
-int32 AFRSyncTemporalResources = 1;
-static FAutoConsoleVariableRef CVarSyncTemporalResources(
-	TEXT("D3D12.AFRSyncTemporalResources"),
-	AFRSyncTemporalResources,
-	TEXT("Synchronize inter-frame dependencies between GPUs"),
-	ECVF_RenderThreadSafe
-	);
-
 static int32 GD3D12TransientAllocatorFullAliasingBarrier = 0;
 static FAutoConsoleVariableRef CVarD3D12TransientAllocatorFullAliasingBarrier(
 	TEXT("d3d12.TransientAllocator.FullAliasingBarrier"),
@@ -2013,139 +2005,3 @@ void FD3D12CommandContext::RHISubmitCommandsHint()
 {
 	// Nothing to do
 }
-
-#if WITH_MGPU
-
-template <typename TD3D12Resource, typename TCopyFunction>
-void FD3D12CommandContext::BroadcastTemporalEffect(const FName& InEffectName, const TArrayView<TD3D12Resource*> InResources, const TCopyFunction& InCopyFunction)
-{
-	check(IsDefaultContext());
-
-	if (GNumAlternateFrameRenderingGroups == 1 || !AFRSyncTemporalResources)
-	{
-		return;
-	}
-
-	const uint32 SiblingGPUIndex = AFRUtils::GetNextSiblingGPUIndex(GetGPUIndex());
-
-	for (TD3D12Resource* SrcResource : InResources)
-	{
-		// Resources must be in the COMMON state before using on the copy queue.
-		TD3D12Resource* DstResource = SrcResource->GetLinkedObject(SiblingGPUIndex);
-		TransitionResource(SrcResource->GetResource(), D3D12_RESOURCE_STATE_TBD, D3D12_RESOURCE_STATE_COMMON, D3D12_RESOURCE_BARRIER_ALL_SUBRESOURCES);
-		TransitionResource(DstResource->GetResource(), D3D12_RESOURCE_STATE_TBD, D3D12_RESOURCE_STATE_COMMON, D3D12_RESOURCE_BARRIER_ALL_SUBRESOURCES);
-	}
-	FlushResourceBarriers();
-
-	// Insert a sync point on the current context.
-	// This will be used to fence the copy from starting until all prior work on this queue is complete.
-	// @todo mgpu - RHIBroadcastTemporalEffect does not take async compute into consideration.
-	FD3D12SyncPointRef WaitForGraphics_SyncPoint = FD3D12SyncPoint::Create(ED3D12SyncPointType::GPUOnly);
-	SignalSyncPoint(WaitForGraphics_SyncPoint);
-
-	// Start the copy work
-	FD3D12SyncPointRef& CopySyncPoint = Device->GetParentAdapter()->GetTemporalEffect(InEffectName)[SiblingGPUIndex];
-	{
-		// Tell the copy queue to wait for the current queue to finish rendering before starting the copy.
-		FD3D12CopyScope CopyScope(Device, ED3D12SyncPointType::GPUOnly, WaitForGraphics_SyncPoint);
-
-		for (TD3D12Resource* SrcResource : InResources)
-		{
-			TD3D12Resource* DstResource = SrcResource->GetLinkedObject(SiblingGPUIndex);
-			InCopyFunction(CopyScope.Context, DstResource, SrcResource);
-		}
-
-		// Put the sync point in the Effect object to await later
-		checkf(CopySyncPoint == nullptr, TEXT("A previously issued sync point was not awaited by the sibling GPU."));
-		CopySyncPoint = CopyScope.GetSyncPoint();
-	}
-}
-
-/*
-* When using AFR certain inter-frame dependecies need to be synchronized across all GPUs.
-* For example a rendering technique that relies on results from the previous frame (which occured on the other GPU).
-*/
-void FD3D12CommandContext::RHIBroadcastTemporalEffect(const FName& InEffectName, const TArrayView<FRHITexture*> InTextures)
-{
-	FMemMark Mark(FMemStack::Get());
-
-	TArray<FD3D12Texture*, TMemStackAllocator<>> Resources;
-	Resources.Reserve(InTextures.Num());
-
-	for (FRHITexture* Texture : InTextures)
-	{
-		Resources.Emplace(RetrieveTexture(Texture));
-	}
-
-	BroadcastTemporalEffect(InEffectName, MakeArrayView(Resources),
-		[](FD3D12ContextCommon& Context, FD3D12Texture* Dst, FD3D12Texture* Src)
-		{
-			FD3D12Resource* DstResource = Dst->GetResource();
-			FD3D12Resource* SrcResource = Src->GetResource();
-
-			Context.UpdateResidency(DstResource);
-			Context.UpdateResidency(SrcResource);
-
-			Context.CopyCommandList()->CopyResource(
-				DstResource->GetResource(),
-				SrcResource->GetResource()
-			);
-		});
-}
-
-void FD3D12CommandContext::RHIBroadcastTemporalEffect(const FName& InEffectName, const TArrayView<FRHIBuffer*> InBuffers)
-{
-	FMemMark Mark(FMemStack::Get());
-
-	TArray<FD3D12Buffer*, TMemStackAllocator<>> Resources;
-	Resources.Reserve(InBuffers.Num());
-
-	for (auto* Buffer : InBuffers)
-	{
-		Resources.Emplace(RetrieveObject<FD3D12Buffer>(Buffer));
-	}
-
-	BroadcastTemporalEffect(InEffectName, MakeArrayView(Resources),
-		[](FD3D12ContextCommon& Context, FD3D12Buffer* Dst, FD3D12Buffer* Src)
-		{
-			FD3D12Resource* DstResource = Dst->GetResource();
-			FD3D12Resource* SrcResource = Src->GetResource();
-
-			Context.UpdateResidency(DstResource);
-			Context.UpdateResidency(SrcResource);
-
-			Context.CopyCommandList()->CopyBufferRegion(
-				DstResource->GetResource(),
-				Dst->ResourceLocation.GetOffsetFromBaseOfResource(),
-				SrcResource->GetResource(),
-				Src->ResourceLocation.GetOffsetFromBaseOfResource(),
-				Src->ResourceLocation.GetSize()
-			);
-		});
-}
-
-void FD3D12CommandContext::RHIWaitForTemporalEffect(const FName& InEffectName)
-{
-	check(IsDefaultContext());
-
-	if (GNumAlternateFrameRenderingGroups == 1 || !AFRSyncTemporalResources)
-	{
-		return;
-	}
-
-	FD3D12SyncPointRef& SyncPoint = GetParentAdapter()->GetTemporalEffect(InEffectName)[GetGPUIndex()];
-
-	if (SyncPoint)
-	{
-		WaitSyncPoint(SyncPoint);
-		SyncPoint = nullptr;
-	}
-}
-
-#else
-
-void FD3D12CommandContext::RHIBroadcastTemporalEffect(const FName& InEffectName, const TArrayView<FRHITexture*> InTextures) { }
-void FD3D12CommandContext::RHIBroadcastTemporalEffect(const FName& InEffectName, const TArrayView<FRHIBuffer* > InBuffers ) { }
-void FD3D12CommandContext::RHIWaitForTemporalEffect  (const FName& InEffectName) { }
-
-#endif // WITH_MGPU
diff --git a/Engine/Source/Runtime/D3D12RHI/Private/D3D12Texture.cpp b/Engine/Source/Runtime/D3D12RHI/Private/D3D12Texture.cpp
--- a/Engine/Source/Runtime/D3D12RHI/Private/D3D12Texture.cpp
+++ b/Engine/Source/Runtime/D3D12RHI/Private/D3D12Texture.cpp
@@ -1727,7 +1727,7 @@ void* FD3D12Texture::Lock(class FRHICommandListImmediate* RHICmdList, uint32 Mip
 	{
 		LockedResource->bLockedForReadOnly = true;
 
-		//TODO: Make this work for AFR (it's probably a very rare occurance though)
+		//TODO: Make this work for multi-GPU (it's probably a very rare occurrence though)
 		ensure(GNumExplicitGPUsForRendering == 1);
 
 		// If we're reading from the texture, we create a staging resource, copy the texture contents to it, and map it.
diff --git a/Engine/Source/Runtime/D3D12RHI/Private/D3D12Viewport.cpp b/Engine/Source/Runtime/D3D12RHI/Private/D3D12Viewport.cpp
--- a/Engine/Source/Runtime/D3D12RHI/Private/D3D12Viewport.cpp
+++ b/Engine/Source/Runtime/D3D12RHI/Private/D3D12Viewport.cpp
@@ -17,14 +17,6 @@ namespace D3D12RHI
 	 */
 	namespace RHIConsoleVariables
 	{
-		int32 AFRUseFramePacing = 0;
-		static FAutoConsoleVariableRef CVarUseFramePacing(
-			TEXT("D3D12.AFRUseFramePacing"),
-			AFRUseFramePacing,
-			TEXT("Control when frames are presented when using mGPU and Alternate Frame Rendering."),
-			ECVF_RenderThreadSafe
-			);
-
 #if !UE_BUILD_SHIPPING
 #if LOG_VIEWPORT_EVENTS
 		int32 LogViewportEvents = 1;
@@ -113,14 +105,20 @@ void FD3D12FramePacing::PrePresentQueued(ID3D12CommandQueue* Queue)
 	const float Alpha = FMath::Clamp(Delta / 1000.0f / FramePacingAvgTimePeriod, 0.0f, 1.0f);
 
 	/** Number of milliseconds the GPU was busy last frame. */
-	// Multi-GPU support : Should be updated to use GPUIndex for AFR.
+	/**
+	 * TODO:  Proper Multi-GPU support for measuring GPU frame cycles would involve doing something more complicated.
+	 * This could involve either using the frame cycles for whichever GPU was idle the least, or identifying time spans
+	 * where all GPUs are idle.  The latter seems more appropriate for the purpose of frame pacing, but would require
+	 * adding detailed idle time span tracking across all GPUs, which is some work.  The call below just queries cycles
+	 * for GPU 0 only.
+	 */
 	const uint32 GPUCycles = RHIGetGPUFrameCycles();
 	const float GPUMsForFrame = FPlatformTime::ToMilliseconds(GPUCycles);
 
 	AvgFrameTimeMs = (Alpha * GPUMsForFrame) + ((1.0f - Alpha) * AvgFrameTimeMs);
 	LastFrameTimeMs = CurrTimeMs;
 
-	const float TargetFrameTime = AvgFrameTimeMs * FramePacingPercentage / GNumAlternateFrameRenderingGroups;
+	const float TargetFrameTime = AvgFrameTimeMs * FramePacingPercentage;
 
 	const uint32 WriteIndex = NextIndex % MaxFrames;
 	SleepTimes[WriteIndex] = (uint32)TargetFrameTime;
@@ -369,11 +367,6 @@ void FD3D12Viewport::CalculateSwapChainDepth(int32 DefaultSwapChainDepth)
 		{
 			BackbufferMultiGPUBinding = FMath::Clamp<int32>(BackbufferMultiGPUBinding, INDEX_NONE, (int32)GNumExplicitGPUsForRendering - 1) ;
 		}
-		else if (GNumAlternateFrameRenderingGroups > 1)
-		{
-			BackbufferMultiGPUBinding = INDEX_NONE;
-			NumBackBuffers = GNumExplicitGPUsForRendering > 2 ? GNumExplicitGPUsForRendering : 4;
-		}
 	}
 #endif // WITH_MGPU
 
@@ -692,35 +685,6 @@ bool FD3D12Viewport::Present(bool bLockToVsync)
 
 	FD3D12CommandContext* PresentContext = nullptr;
 #if WITH_MGPU
-	if (GNumAlternateFrameRenderingGroups > 1)
-	{
-		const uint32 PresentGPUIndex = BackBufferGPUIndices[CurrentBackBufferIndex_RHIThread];
-		PresentContext = &Adapter->GetDevice(PresentGPUIndex)->GetDefaultCommandContext();
-
-		// In AFR it's possible that the current frame will complete faster than the frame
-		// already in progress so we need to add synchronization to ensure that our Present
-		// occurs after the previous frame's Present. Otherwise we can put frames in the
-		// system present queue out of order.
-		if (LastFrameSyncPoint)
-		{
-			PresentContext->WaitSyncPoint(LastFrameSyncPoint);
-		}
-	}
-
-#if 0 // Multi-GPU support : figure out what kind of synchronization is needed.
-	// When using an alternating frame rendering technique with multiple GPUs the time of frame
-	// delivery must be paced in order to provide a nice experience.
-	if (Adapter->GetMultiGPUMode() == MGPU_AFR && RHIConsoleVariables::AFRUseFramePacing && !bLockToVsync)
-	{
-		if (FramePacerRunnable == nullptr)
-		{
-			FramePacerRunnable = new FD3D12FramePacing(Adapter);
-		}
-	
-		FramePacerRunnable->PrePresentQueued(Device->GetCommandListManager().GetD3DCommandQueue());
-	}
-	else
-#endif
 	if (FramePacerRunnable)
 	{
 		delete FramePacerRunnable;
@@ -934,7 +898,7 @@ void FD3D12CommandContextBase::RHIEndDrawingViewport(FRHIViewport* ViewportRHI,
 
 	bool bNativelyPresented = Viewport->Present(bLockToVsync);
 	
-	// Multi-GPU support : here each GPU wait's for it's own frame completion. Note that even in AFR, each GPU renders an (empty) frame.
+	// Multi-GPU support : here each GPU wait's for it's own frame completion.
 	if (bNativelyPresented)
 	{
 		static const auto CFinishFrameVar = IConsoleManager::Get().FindTConsoleVariableDataInt(TEXT("r.FinishCurrentFrame"));
diff --git a/Engine/Source/Runtime/D3D12RHI/Private/Windows/WindowsD3D12Viewport.cpp b/Engine/Source/Runtime/D3D12RHI/Private/Windows/WindowsD3D12Viewport.cpp
--- a/Engine/Source/Runtime/D3D12RHI/Private/Windows/WindowsD3D12Viewport.cpp
+++ b/Engine/Source/Runtime/D3D12RHI/Private/Windows/WindowsD3D12Viewport.cpp
@@ -342,12 +342,12 @@ void FD3D12Viewport::ResizeInternal()
 
 		for (uint32 i = 0; i < NumBackBuffers; ++i)
 		{
-			// When BackbufferMultiGPUBinding == INDEX_NONE, cycle through each GPU (for AFR or debugging).
-			const uint32 BackBufferGPUIndex = BackbufferMultiGPUBinding >= 0 ? (uint32)BackbufferMultiGPUBinding : (i % GNumAlternateFrameRenderingGroups);
+			// When BackbufferMultiGPUBinding == INDEX_NONE, cycle through each GPU.
+			const uint32 BackBufferGPUIndex = BackbufferMultiGPUBinding == INDEX_NONE ? i % GNumExplicitGPUsForRendering : BackbufferMultiGPUBinding;
 			BackBufferGPUIndices.Add(BackBufferGPUIndex);
 		}
 
-		// Interleave the swapchains between the AFR devices
+		// Select the GPU for each element in the swapchain 
 		for (uint32 i = 0; i < NumBackBuffers; ++i)
 		{
 			const uint32 GPUIndex = BackBufferGPUIndices[i];
diff --git a/Engine/Source/Runtime/D3D12RHI/Public/D3D12Viewport.h b/Engine/Source/Runtime/D3D12RHI/Public/D3D12Viewport.h
--- a/Engine/Source/Runtime/D3D12RHI/Public/D3D12Viewport.h
+++ b/Engine/Source/Runtime/D3D12RHI/Public/D3D12Viewport.h
@@ -238,7 +238,7 @@ private:
 	/** A fence value used to track the GPU's progress. */
 	TArray<FD3D12SyncPointRef> FrameSyncPoints;
 
-	// Determine how deep the swapchain should be (based on AFR or not)
+	// Determine how deep the swapchain should be
 	void CalculateSwapChainDepth(int32 DefaultSwapChainDepth);
 
 	FCustomPresentRHIRef CustomPresent;
diff --git a/Engine/Source/Runtime/Engine/Private/Particles/FXSystemPrivate.h b/Engine/Source/Runtime/Engine/Private/Particles/FXSystemPrivate.h
--- a/Engine/Source/Runtime/Engine/Private/Particles/FXSystemPrivate.h
+++ b/Engine/Source/Runtime/Engine/Private/Particles/FXSystemPrivate.h
@@ -392,8 +392,8 @@ private:
 #endif // #if WITH_EDITOR
 
 #if WITH_MGPU
-	EParticleSimulatePhase::Type PhaseToWaitForTemporalEffect = EParticleSimulatePhase::First;
-	EParticleSimulatePhase::Type PhaseToBroadcastTemporalEffect = EParticleSimulatePhase::First;
+	EParticleSimulatePhase::Type PhaseToWaitForResourceTransfer = EParticleSimulatePhase::First;
+	EParticleSimulatePhase::Type PhaseToBroadcastResourceTransfer = EParticleSimulatePhase::First;
 #endif
 
 	TUniformBufferRef<FSceneTextureUniformParameters> SceneTexturesUniformParams;
diff --git a/Engine/Source/Runtime/Engine/Private/Particles/ParticleCurveTexture.cpp b/Engine/Source/Runtime/Engine/Private/Particles/ParticleCurveTexture.cpp
--- a/Engine/Source/Runtime/Engine/Private/Particles/ParticleCurveTexture.cpp
+++ b/Engine/Source/Runtime/Engine/Private/Particles/ParticleCurveTexture.cpp
@@ -175,8 +175,6 @@ static void InjectCurves(
 
 	SCOPED_DRAW_EVENT(RHICmdList, InjectParticleCurves);
 
-	RHICmdList.BeginUpdateMultiFrameResource(CurveTextureRHI);
-
 	ERenderTargetLoadAction LoadAction = ERenderTargetLoadAction::ELoad;
 
 	if (bFirstCall)
@@ -280,7 +278,6 @@ static void InjectCurves(
 	}
 	RHICmdList.EndRenderPass();
 	RHICmdList.Transition(FRHITransitionInfo(CurveTextureRHI, ERHIAccess::RTV, ERHIAccess::SRVMask));
-	RHICmdList.EndUpdateMultiFrameResource(CurveTextureRHI);
 }
 
 /*------------------------------------------------------------------------------
diff --git a/Engine/Source/Runtime/Engine/Private/Particles/ParticleGpuSimulation.cpp b/Engine/Source/Runtime/Engine/Private/Particles/ParticleGpuSimulation.cpp
--- a/Engine/Source/Runtime/Engine/Private/Particles/ParticleGpuSimulation.cpp
+++ b/Engine/Source/Runtime/Engine/Private/Particles/ParticleGpuSimulation.cpp
@@ -45,10 +45,6 @@
 DECLARE_CYCLE_STAT(TEXT("GPUSpriteEmitterInstance Init GT"), STAT_GPUSpriteEmitterInstance_Init, STATGROUP_Particles);
 DECLARE_GPU_STAT_NAMED(ParticleSimulation, TEXT("Particle Simulation"));
 
-#if WITH_MGPU
-DECLARE_GPU_STAT(AFRWaitForParticleSimulation);
-#endif
-
 /*------------------------------------------------------------------------------
 	Constants to tune memory and performance for GPU particle simulation.
 ------------------------------------------------------------------------------*/
@@ -111,14 +107,6 @@ static TAutoConsoleVariable<int32> CVarGPUParticleMaxNumIterations(TEXT("r.GPUPa
 
 static TAutoConsoleVariable<int32> CVarSimulateGPUParticles(TEXT("r.GPUParticle.Simulate"), 1, TEXT("Enable or disable GPU particle simulation"));
 
-static TAutoConsoleVariable<int32> CVarGPUParticleAFRReinject(
-	TEXT("r.GPUParticle.AFRReinject"),
-	1,
-	TEXT("Toggle optimization when running in AFR to re-inject particle injections on the next GPU rather than doing a slow GPU->GPU transfer of the texture data\n")	
-	TEXT("  0: Reinjection off\n")
-	TEXT("  1: Reinjection on"),
-	ECVF_ReadOnly);
-
 /*-----------------------------------------------------------------------------
 	Allocators used to manage GPU particle resources.
 -----------------------------------------------------------------------------*/
@@ -308,8 +296,6 @@ public:
 		const int32 SizeX = GParticleSimulationTextureSizeX;
 		const int32 SizeY = GParticleSimulationTextureSizeY;
 
-		const ETextureCreateFlags ExtraFlags = CVarGPUParticleAFRReinject.GetValueOnRenderThread() == 1 ? TexCreate_AFRManual : TexCreate_None;
-
 #define ATTRIBUTES_TEXTURE_NAME	TEXT("ParticleAttributes")
 
 		const FRHITextureCreateDesc Desc =
@@ -317,7 +303,7 @@ public:
 			.SetExtent(SizeX, SizeY)
 			.SetFormat(PF_B8G8R8A8)
 			.SetClearValue(FClearValueBinding::Transparent)
-			.SetFlags(ETextureCreateFlags::RenderTargetable | ETextureCreateFlags::ShaderResource | ETextureCreateFlags::NoFastClear | ExtraFlags)
+			.SetFlags(ETextureCreateFlags::RenderTargetable | ETextureCreateFlags::ShaderResource | ETextureCreateFlags::NoFastClear)
 			.SetInitialState(ERHIAccess::RTV);
 
 		TextureRHI = RHICreateTexture(Desc);
@@ -4413,16 +4399,11 @@ void FFXSystem::PrepareGPUSimulation(FRHICommandListImmediate& RHICmdList)
 	// Grab resources.
 	FParticleStateTextures& CurrentStateTextures = ParticleSimulationResources->GetCurrentStateTextures();
 
-	// We delay this transition in AFR to coincide with the first time the simulation
-	// texture is actually used.
-	if (GNumAlternateFrameRenderingGroups == 1)
-	{
-		// Setup render states.
-		RHICmdList.Transition({
-			FRHITransitionInfo(CurrentStateTextures.PositionTextureRHI, ERHIAccess::SRVMask, ERHIAccess::RTV),
-			FRHITransitionInfo(CurrentStateTextures.VelocityTextureRHI, ERHIAccess::SRVMask, ERHIAccess::RTV)
-		});
-	}
+	// Setup render states.
+	RHICmdList.Transition({
+		FRHITransitionInfo(CurrentStateTextures.PositionTextureRHI, ERHIAccess::SRVMask, ERHIAccess::RTV),
+		FRHITransitionInfo(CurrentStateTextures.VelocityTextureRHI, ERHIAccess::SRVMask, ERHIAccess::RTV)
+	});
 }
 
 void FFXSystem::FinalizeGPUSimulation(FRHICommandListImmediate& RHICmdList)
@@ -4452,28 +4433,8 @@ void FFXSystem::SimulateGPUParticles(
 
 
 #if WITH_MGPU
-	static const FName TemporalEffectName("SimulateGPUParticles");
-	TArray<FRHITexture*, TFixedAllocator<4>> TemporalEffectTextures;
-	if (GNumAlternateFrameRenderingGroups > 1)
-	{
-		if (Phase == PhaseToWaitForTemporalEffect)
-		{
-			SCOPED_GPU_STAT(RHICmdList, AFRWaitForParticleSimulation);
-			RHICmdList.WaitForTemporalEffect(TemporalEffectName);
-
-			// Only the previous state textures are actually being copied, but due to the data
-			// race mentioned in the BroadcastTemporalEffect block below we need to delay
-			// transitioning the current textures to writable state as well.
-			RHICmdList.Transition({
-				FRHITransitionInfo(CurrentStateRenderTargets[0], ERHIAccess::Unknown, ERHIAccess::RTV),
-				FRHITransitionInfo(CurrentStateRenderTargets[1], ERHIAccess::Unknown, ERHIAccess::RTV)});
-		}
-		TemporalEffectTextures.Add(CurrentStateRenderTargets[0]);
-		TemporalEffectTextures.Add(CurrentStateRenderTargets[1]);
-	}
-
 	TArray<FTransferResourceParams, TFixedAllocator<4>> CrossGPUTransferResources;
-	const bool bCrossTransferEnabled = GNumAlternateFrameRenderingGroups == 1 && GNumExplicitGPUsForRendering > 1;
+	const bool bCrossTransferEnabled = GNumExplicitGPUsForRendering > 1;
 	auto AddCrossGPUTransferResource =
 		[&](FRHITexture* TextureToTransfer)
 		{
@@ -4502,44 +4463,28 @@ void FFXSystem::SimulateGPUParticles(
 		if ( !CurrentStateTextures.bTexturesCleared )
 		{
 			
-			RHICmdList.BeginUpdateMultiFrameResource(CurrentStateRenderTargets[0]);
-			RHICmdList.BeginUpdateMultiFrameResource(CurrentStateRenderTargets[1]);
-
-			{
-				FRHIRenderPassInfo RenderPassInfo(2, CurrentStateRenderTargets, ERenderTargetActions::Clear_Store);
-				RHICmdList.BeginRenderPass(RenderPassInfo, TEXT("GPUParticlesClearStateTextures"));
-				RHICmdList.EndRenderPass();
-			}
+			FRHIRenderPassInfo RenderPassInfo(2, CurrentStateRenderTargets, ERenderTargetActions::Clear_Store);
+			RHICmdList.BeginRenderPass(RenderPassInfo, TEXT("GPUParticlesClearStateTextures"));
+			RHICmdList.EndRenderPass();
 			
 			CurrentStateTextures.bTexturesCleared = true;
-			
-			RHICmdList.EndUpdateMultiFrameResource(CurrentStateRenderTargets[0]);
-			RHICmdList.EndUpdateMultiFrameResource(CurrentStateRenderTargets[1]);
 		}
 		
 		if ( !PrevStateTextures.bTexturesCleared )
 		{
-			RHICmdList.BeginUpdateMultiFrameResource(PreviousStateRenderTargets[0]);
-			RHICmdList.BeginUpdateMultiFrameResource(PreviousStateRenderTargets[1]);
-
-			{
-				RHICmdList.Transition({
-					FRHITransitionInfo(PreviousStateRenderTargets[0], ERHIAccess::SRVMask, ERHIAccess::RTV),
-					FRHITransitionInfo(PreviousStateRenderTargets[1], ERHIAccess::SRVMask, ERHIAccess::RTV)});
+			RHICmdList.Transition({
+				FRHITransitionInfo(PreviousStateRenderTargets[0], ERHIAccess::SRVMask, ERHIAccess::RTV),
+				FRHITransitionInfo(PreviousStateRenderTargets[1], ERHIAccess::SRVMask, ERHIAccess::RTV)});
 
-				FRHIRenderPassInfo RPInfo(2, PreviousStateRenderTargets, ERenderTargetActions::Clear_Store);
-				RHICmdList.BeginRenderPass(RPInfo, TEXT("GPUParticlesClearPreviousStateTextures"));
-				RHICmdList.EndRenderPass();
+			FRHIRenderPassInfo RPInfo(2, PreviousStateRenderTargets, ERenderTargetActions::Clear_Store);
+			RHICmdList.BeginRenderPass(RPInfo, TEXT("GPUParticlesClearPreviousStateTextures"));
+			RHICmdList.EndRenderPass();
 
-				RHICmdList.Transition({
-					FRHITransitionInfo(PreviousStateRenderTargets[0], ERHIAccess::RTV, ERHIAccess::SRVMask),
-					FRHITransitionInfo(PreviousStateRenderTargets[1], ERHIAccess::RTV, ERHIAccess::SRVMask) });
-			}
+			RHICmdList.Transition({
+				FRHITransitionInfo(PreviousStateRenderTargets[0], ERHIAccess::RTV, ERHIAccess::SRVMask),
+				FRHITransitionInfo(PreviousStateRenderTargets[1], ERHIAccess::RTV, ERHIAccess::SRVMask) });
 			
 			PrevStateTextures.bTexturesCleared = true;
-			
-			RHICmdList.EndUpdateMultiFrameResource(PreviousStateRenderTargets[0]);
-			RHICmdList.EndUpdateMultiFrameResource(PreviousStateRenderTargets[1]);
 		}
 	}
 	
@@ -4694,9 +4639,6 @@ void FFXSystem::SimulateGPUParticles(
 		}
 	}
 
-	RHICmdList.BeginUpdateMultiFrameResource(CurrentStateRenderTargets[0]);
-	RHICmdList.BeginUpdateMultiFrameResource(CurrentStateRenderTargets[1]);
-	
 	if ( SimulationCommands.Num() || TilesToClear.Num())
 	{
 		FRHIRenderPassInfo RPInfo(2, CurrentStateRenderTargets, ERenderTargetActions::Load_Store);
@@ -4757,8 +4699,6 @@ void FFXSystem::SimulateGPUParticles(
 			ParticleSimulationResources->RenderAttributesTexture.TextureRHI,
 			ParticleSimulationResources->SimulationAttributesTexture.TextureRHI
 		};
-		RHICmdList.BeginUpdateMultiFrameResource(ParticleSimulationResources->RenderAttributesTexture.TextureRHI);
-		RHICmdList.BeginUpdateMultiFrameResource(ParticleSimulationResources->SimulationAttributesTexture.TextureRHI);
 
 		FRHIRenderPassInfo RPInfo(4, InjectRenderTargets, ERenderTargetActions::Load_Store);
 		{
@@ -4789,23 +4729,6 @@ void FFXSystem::SimulateGPUParticles(
 			});
 		}
 
-		if (GNumAlternateFrameRenderingGroups > 1)
-		{
-			if (CVarGPUParticleAFRReinject.GetValueOnRenderThread() == 1)
-			{
-				ensureMsgf(GNumAlternateFrameRenderingGroups == 2, TEXT("GPU Particles running on an AFR depth > 2 not supported.  Currently: %i"), GNumAlternateFrameRenderingGroups);
-
-				// Place these particles into the multi-gpu update queue
-				LastFrameNewParticles.Append(NewParticles);
-			}
-			else
-			{
-#if WITH_MGPU
-				TemporalEffectTextures.Add(ParticleSimulationResources->RenderAttributesTexture.TextureRHI);
-				TemporalEffectTextures.Add(ParticleSimulationResources->SimulationAttributesTexture.TextureRHI);
-#endif
-			}
-		}
 #if WITH_MGPU
 		if (bCrossTransferEnabled)
 		{
@@ -4813,8 +4736,6 @@ void FFXSystem::SimulateGPUParticles(
 			AddCrossGPUTransferResource(ParticleSimulationResources->SimulationAttributesTexture.TextureRHI);
 		}
 #endif
-		RHICmdList.EndUpdateMultiFrameResource(ParticleSimulationResources->RenderAttributesTexture.TextureRHI);
-		RHICmdList.EndUpdateMultiFrameResource(ParticleSimulationResources->SimulationAttributesTexture.TextureRHI);
 	}
 	
 	// finish current state render
@@ -4822,13 +4743,10 @@ void FFXSystem::SimulateGPUParticles(
 		FRHITransitionInfo(CurrentStateRenderTargets[0], ERHIAccess::RTV, ERHIAccess::SRVMask),
 		FRHITransitionInfo(CurrentStateRenderTargets[1], ERHIAccess::RTV, ERHIAccess::SRVMask)});
 
-	RHICmdList.EndUpdateMultiFrameResource(CurrentStateRenderTargets[0]);
-	RHICmdList.EndUpdateMultiFrameResource(CurrentStateRenderTargets[1]);
-
 	if (SimulationCommands.Num() && FixDeltaSeconds > 0)
 	{
 		//the fixed timestep works in two stages.  A first stage which simulates the fixed timestep and this second stage which simulates any remaining time from the actual delta time.  e.g.  fixed timestep of 16ms and actual dt of 23ms
-		//will make this second step simulate an interpolated extra 7ms.  This second interpolated step is what we render on THIS frame, but it is NOT fed into the next frame's simulation.  Thus we do not need to transfer it between GPUs in AFR mode.
+		//will make this second step simulate an interpolated extra 7ms.  This second interpolated step is what we render on THIS frame, but it is NOT fed into the next frame's simulation.
 		FParticleStateTextures& VisualizeStateTextures = ParticleSimulationResources->GetPreviousStateTextures();
 		
 		RHICmdList.Transition({ 
@@ -4870,27 +4788,11 @@ void FFXSystem::SimulateGPUParticles(
 	}
 
 #if WITH_MGPU
-	// Previously, this broadcast was before the above extra stage simulate work. This
-	// is because that work is temporary, only applied visually to the current frame
-	// rather than fed back through the simulation for the next frame. However, this
-	// led to a possible data race in AFR under certain circumstances: Let's say GPU0
-	// writes to StateTextures[0][GPU0] on this frame, then initiates the copy from
-	// StateTextures[0][GPU0] to StateTextures[0][GPU1] and signals GPU1 that it can
-	// start particle work. GPU0 then starts writing to StateTextures[1][GPU0] as a
-	// temporary for extra work on this frame. GPU1 Completes its work on
-	// StateTextures[1][GPU1], and initiates the copy from StateTextures[1][GPU1] to
-	// StateTextures[1][GPU0] at the end of its frame. However, GPU0 may still be
-	// writing to StateTextures[1][GPU0], and this leads to a data race. For now, we're
-	// moving the broadcast to after any extra work to block the next AFR group from
-	// using the textures until we're done, however for a potential AFR performance
-	// boost, we can put the broadcast back to where it was, and use a third buffer for
-	// temporary extra particle simulation work.
-	if (Phase == PhaseToBroadcastTemporalEffect)
-	{
-		if (GNumAlternateFrameRenderingGroups > 1)
-		{
-			RHICmdList.BroadcastTemporalEffect(TemporalEffectName, TemporalEffectTextures);
-		}
+	// TODO:  This transfer was done here due to AFR considerations, but AFR support has been removed.
+	// Should investigate whether it still needs to be done here, or would be more optimal to move
+	// somewhere else?
+	if (Phase == PhaseToBroadcastResourceTransfer)
+	{
 		if (CrossGPUTransferResources.Num() > 0)
 		{
 			RHICmdList.TransferResources(CrossGPUTransferResources);
@@ -4907,6 +4809,7 @@ void FFXSystem::SimulateGPUParticles(
 
 void FFXSystem::OnSimulationPhaseChanged(const FParticleSimulationGPU* GPUSimulation, EParticleSimulatePhase::Type PrevPhase)
 {
+	// TODO:  Could this be simplified now that AFR has been removed?
 	// We keep track of the number of simulations of each phase type to more
 	// efficiently synchronize temporal effects. We want to avoid having any long
 	// stretches of time where the AFR frames can't run in parallel.
@@ -4957,16 +4860,16 @@ void FFXSystem::UpdateMultiGPUResources(FRHICommandListImmediate& RHICmdList)
 	LastFrameNewParticles.Reset();
 
 #if WITH_MGPU
-	PhaseToBroadcastTemporalEffect = EParticleSimulatePhase::Last;
-	while (PhaseToBroadcastTemporalEffect > EParticleSimulatePhase::First && NumGPUSimulations[PhaseToBroadcastTemporalEffect] == 0)
+	PhaseToBroadcastResourceTransfer = EParticleSimulatePhase::Last;
+	while (PhaseToBroadcastResourceTransfer > EParticleSimulatePhase::First && NumGPUSimulations[PhaseToBroadcastResourceTransfer] == 0)
 	{
-		PhaseToBroadcastTemporalEffect = static_cast<EParticleSimulatePhase::Type>(PhaseToBroadcastTemporalEffect - 1);
+		PhaseToBroadcastResourceTransfer = static_cast<EParticleSimulatePhase::Type>(PhaseToBroadcastResourceTransfer - 1);
 	}
 
-	PhaseToWaitForTemporalEffect = EParticleSimulatePhase::First;
-	while (PhaseToWaitForTemporalEffect < PhaseToBroadcastTemporalEffect && NumGPUSimulations[PhaseToWaitForTemporalEffect] == 0)
+	PhaseToWaitForResourceTransfer = EParticleSimulatePhase::First;
+	while (PhaseToWaitForResourceTransfer < PhaseToBroadcastResourceTransfer && NumGPUSimulations[PhaseToWaitForResourceTransfer] == 0)
 	{
-		PhaseToWaitForTemporalEffect = static_cast<EParticleSimulatePhase::Type>(PhaseToWaitForTemporalEffect + 1);
+		PhaseToWaitForResourceTransfer = static_cast<EParticleSimulatePhase::Type>(PhaseToWaitForResourceTransfer + 1);
 	}
 #endif
 }
diff --git a/Engine/Source/Runtime/Engine/Private/UnrealEngine.cpp b/Engine/Source/Runtime/Engine/Private/UnrealEngine.cpp
--- a/Engine/Source/Runtime/Engine/Private/UnrealEngine.cpp
+++ b/Engine/Source/Runtime/Engine/Private/UnrealEngine.cpp
@@ -3586,8 +3586,7 @@ public:
 		}
 #if WITH_MGPU
 		// We have to do all the work for secondary views that are on a different GPU than
-		// the primary view. NB: This assumes that the primary view is assigned to the
-		// first GPU of the AFR group. See FSceneRenderer::ComputeViewGPUMasks.
+		// the primary view.
 		else if (GNumExplicitGPUsForRendering > 1)
 		{
 			return EStereoscopicPass::eSSP_PRIMARY;
diff --git a/Engine/Source/Runtime/Engine/Private/VT/VirtualTextureUploadCache.cpp b/Engine/Source/Runtime/Engine/Private/VT/VirtualTextureUploadCache.cpp
--- a/Engine/Source/Runtime/Engine/Private/VT/VirtualTextureUploadCache.cpp
+++ b/Engine/Source/Runtime/Engine/Private/VT/VirtualTextureUploadCache.cpp
@@ -259,8 +259,8 @@ void FVirtualTextureUploadCache::Finalize(FRDGBuilder& GraphBuilder)
 	SCOPED_DRAW_EVENT(RHICmdList, FVirtualTextureUploadCache_Finalize);
 	SCOPE_CYCLE_COUNTER(STAT_VTP_FlushUpload)
 
-		// Multi-GPU support : May be inefficient for AFR.
-		SCOPED_GPU_MASK(RHICmdList, FRHIGPUMask::All());
+	// Multi-GPU support:
+	SCOPED_GPU_MASK(RHICmdList, FRHIGPUMask::All());
 
 	for (int PoolIndex = 0; PoolIndex < Pools.Num(); ++PoolIndex)
 	{
diff --git a/Engine/Source/Runtime/Engine/Public/PlanarReflectionSceneProxy.h b/Engine/Source/Runtime/Engine/Public/PlanarReflectionSceneProxy.h
--- a/Engine/Source/Runtime/Engine/Public/PlanarReflectionSceneProxy.h
+++ b/Engine/Source/Runtime/Engine/Public/PlanarReflectionSceneProxy.h
@@ -71,7 +71,7 @@ public:
 		return ActiveGPUMask;
 	}
 
-	// Changes the GPUMask used when updating the reflection capture in AFR.
+	// Changes the GPUMask used when updating the reflection capture with multi-GPU.
 	void SetActiveGPUMask(FRHIGPUMask InGPUMask)
 	{
 		check(IsInRenderingThread());
@@ -80,7 +80,7 @@ public:
 
 private:
 
-	FRHIGPUMask ActiveGPUMask; // In AFR we need to change which GPUs are rendered to every frame.
+	FRHIGPUMask ActiveGPUMask; // GPU mask copied from parent render target for multi-GPU
 	FIntPoint Size;
 };
 
diff --git a/Engine/Source/Runtime/Engine/Public/TextureResource.h b/Engine/Source/Runtime/Engine/Public/TextureResource.h
--- a/Engine/Source/Runtime/Engine/Public/TextureResource.h
+++ b/Engine/Source/Runtime/Engine/Public/TextureResource.h
@@ -368,7 +368,7 @@ public:
 		return GPUMask & ActiveGPUMask;
 	}
 
-	// Changes the GPUMask used when updating the texture in AFR.
+	// Changes the GPUMask used when updating the texture with multi-GPU.
 	void SetActiveGPUMask(FRHIGPUMask InGPUMask)
 	{
 		check(IsInRenderingThread());
@@ -385,7 +385,7 @@ protected:
 
 private:
 	FRHIGPUMask GPUMask;
-	FRHIGPUMask ActiveGPUMask; // In AFR we need to change which GPUs are rendered to every frame.
+	FRHIGPUMask ActiveGPUMask;  // GPU mask copied from parent render target for multi-GPU
 };
 
 /**
diff --git a/Engine/Source/Runtime/Launch/Private/LaunchEngineLoop.cpp b/Engine/Source/Runtime/Launch/Private/LaunchEngineLoop.cpp
--- a/Engine/Source/Runtime/Launch/Private/LaunchEngineLoop.cpp
+++ b/Engine/Source/Runtime/Launch/Private/LaunchEngineLoop.cpp
@@ -6678,8 +6678,7 @@ void FEngineLoop::AppPreExit( )
 	// which uses async tasks to wait until resources are safe to be deleted (for example, FMediaTextureResource).
 	// To avoid this, we set the frame number to the maximum possible value here, before waiting for the thread pool to die. It's safe to do so
 	// because FlushRenderingCommands() is called multiple times on exit before reaching this point, so there's no way the render thread has any
-	// more frames in flight. Note that simply incrementing the value doesn't work, because FGenericRHIGPUFence::WriteInternal adds
-	// GNumAlternateFrameRenderingGroups to the current frame number to account for multi-GPU, and we don't want to depend on that RHI export here.
+	// more frames in flight.
 	GFrameNumberRenderThread = MAX_uint32;
 
 	if (GIOThreadPool != nullptr)
diff --git a/Engine/Source/Runtime/OpenGLDrv/Private/OpenGLDrv.cpp b/Engine/Source/Runtime/OpenGLDrv/Private/OpenGLDrv.cpp
--- a/Engine/Source/Runtime/OpenGLDrv/Private/OpenGLDrv.cpp
+++ b/Engine/Source/Runtime/OpenGLDrv/Private/OpenGLDrv.cpp
@@ -205,18 +205,14 @@ void FOpenGLGPUProfiler::BeginFrame(FOpenGLDynamicRHI* InRHI)
 	}
 	bPreviousLatchedGProfilingGPUHitches = bLatchedGProfilingGPUHitches;
 
-	// Skip timing events when using SLI, they will not be accurate anyway
-	if (GNumAlternateFrameRenderingGroups == 1)
+	if (FrameTiming.IsSupported())
 	{
-		if (FrameTiming.IsSupported())
-		{
-			FrameTiming.StartTiming();
-		}
-		if (FOpenGLDisjointTimeStampQuery::IsSupported())
-		{
-			CurrentGPUFrameQueryIndex = (CurrentGPUFrameQueryIndex + 1) % MAX_GPUFRAMEQUERIES;
-			DisjointGPUFrameTimeQuery[CurrentGPUFrameQueryIndex].StartTracking();
-		}
+		FrameTiming.StartTiming();
+	}
+	if (FOpenGLDisjointTimeStampQuery::IsSupported())
+	{
+		CurrentGPUFrameQueryIndex = (CurrentGPUFrameQueryIndex + 1) % MAX_GPUFRAMEQUERIES;
+		DisjointGPUFrameTimeQuery[CurrentGPUFrameQueryIndex].StartTracking();
 	}
 
 	if (GetEmitDrawEvents())
@@ -238,28 +234,22 @@ void FOpenGLGPUProfiler::EndFrame()
 		PopEvent();
 	}
 
-	// Skip timing events when using SLI, they will not be accurate anyway
-	if (GNumAlternateFrameRenderingGroups == 1)
+	if (FrameTiming.IsSupported())
 	{
-		if (FrameTiming.IsSupported())
-		{
-			FrameTiming.EndTiming();
-		}
-		if (FOpenGLDisjointTimeStampQuery::IsSupported())
-		{
-			DisjointGPUFrameTimeQuery[CurrentGPUFrameQueryIndex].EndTracking();
-		}
+		FrameTiming.EndTiming();
+	}
+	if (FOpenGLDisjointTimeStampQuery::IsSupported())
+	{
+		DisjointGPUFrameTimeQuery[CurrentGPUFrameQueryIndex].EndTracking();
 	}
 
-	// Skip timing events when using SLI, as they will block the GPU and we want maximum throughput
-	// Stat unit GPU time is not accurate anyway with SLI
-	if (FrameTiming.IsSupported() && GNumAlternateFrameRenderingGroups == 1)
+	if (FrameTiming.IsSupported())
 	{
 		uint64 GPUTiming = FrameTiming.GetTiming();
 		uint64 GPUFreq = FrameTiming.GetTimingFrequency();
 		GGPUFrameTime = FMath::TruncToInt( double(GPUTiming) / double(GPUFreq) / FPlatformTime::GetSecondsPerCycle() );
 	}
-	else if (FOpenGLDisjointTimeStampQuery::IsSupported() && GNumAlternateFrameRenderingGroups == 1)
+	else if (FOpenGLDisjointTimeStampQuery::IsSupported())
 	{
 		static uint32 GLastGPUFrameTime = 0;
 		uint64 GPUTiming = 0;
diff --git a/Engine/Source/Runtime/OpenGLDrv/Private/OpenGLViewport.cpp b/Engine/Source/Runtime/OpenGLDrv/Private/OpenGLViewport.cpp
--- a/Engine/Source/Runtime/OpenGLDrv/Private/OpenGLViewport.cpp
+++ b/Engine/Source/Runtime/OpenGLDrv/Private/OpenGLViewport.cpp
@@ -153,36 +153,32 @@ void FOpenGLDynamicRHI::RHIEndDrawingViewport(FRHIViewport* ViewportRHI,bool bPr
 
 		DrawingViewport = NULL;
 
-		// Don't wait on the GPU when using SLI, let the driver determine how many frames behind the GPU should be allowed to get
-		if (GNumAlternateFrameRenderingGroups == 1)
+		if (bNeedFinishFrame)
 		{
-			if (bNeedFinishFrame)
+			static const auto CFinishFrameVar = IConsoleManager::Get().FindTConsoleVariableDataInt(TEXT("r.FinishCurrentFrame"));
+			if (!CFinishFrameVar->GetValueOnRenderThread())
 			{
-				static const auto CFinishFrameVar = IConsoleManager::Get().FindTConsoleVariableDataInt(TEXT("r.FinishCurrentFrame"));
-				if (!CFinishFrameVar->GetValueOnRenderThread())
-				{
-					// Wait for the GPU to finish rendering the previous frame before finishing this frame.
-					Viewport->WaitForFrameEventCompletion();
-					Viewport->IssueFrameEvent();
-				}
-				else
-				{
-					// Finish current frame immediately to reduce latency
-					Viewport->IssueFrameEvent();
-					Viewport->WaitForFrameEventCompletion();
-				}
+				// Wait for the GPU to finish rendering the previous frame before finishing this frame.
+				Viewport->WaitForFrameEventCompletion();
+				Viewport->IssueFrameEvent();
 			}
-		
-			// If the input latency timer has been triggered, block until the GPU is completely
-			// finished displaying this frame and calculate the delta time.
-			if ( GInputLatencyTimer.RenderThreadTrigger )
+			else
 			{
+				// Finish current frame immediately to reduce latency
+				Viewport->IssueFrameEvent();
 				Viewport->WaitForFrameEventCompletion();
-				uint32 EndTime = FPlatformTime::Cycles();
-				GInputLatencyTimer.DeltaTime = EndTime - GInputLatencyTimer.StartTime;
-				GInputLatencyTimer.RenderThreadTrigger = false;
 			}
 		}
+		
+		// If the input latency timer has been triggered, block until the GPU is completely
+		// finished displaying this frame and calculate the delta time.
+		if ( GInputLatencyTimer.RenderThreadTrigger )
+		{
+			Viewport->WaitForFrameEventCompletion();
+			uint32 EndTime = FPlatformTime::Cycles();
+			GInputLatencyTimer.DeltaTime = EndTime - GInputLatencyTimer.StartTime;
+			GInputLatencyTimer.RenderThreadTrigger = false;
+		}
 
 		if (bRevertToSharedContextAfterDrawingViewport)
 		{
diff --git a/Engine/Source/Runtime/RHI/Private/DynamicRHI.cpp b/Engine/Source/Runtime/RHI/Private/DynamicRHI.cpp
--- a/Engine/Source/Runtime/RHI/Private/DynamicRHI.cpp
+++ b/Engine/Source/Runtime/RHI/Private/DynamicRHI.cpp
@@ -346,10 +346,6 @@ void RHIInit(bool bHasEditorToken)
 
 				GDynamicRHI->Init();
 
-#if WITH_MGPU
-				AFRUtils::StaticInitialize();
-#endif
-
 				// Validation of contexts.
 				GRHICommandList.GetImmediateCommandList().GetContext();
 				check(GIsRHIInitialized);
diff --git a/Engine/Source/Runtime/RHI/Private/MultiGPU.cpp b/Engine/Source/Runtime/RHI/Private/MultiGPU.cpp
--- a/Engine/Source/Runtime/RHI/Private/MultiGPU.cpp
+++ b/Engine/Source/Runtime/RHI/Private/MultiGPU.cpp
@@ -7,40 +7,6 @@
 #include "MultiGPU.h"
 
 #if WITH_SLI || WITH_MGPU
-uint32 GNumAlternateFrameRenderingGroups = 1;
 uint32 GNumExplicitGPUsForRendering = 1;
 uint32 GVirtualMGPU = 0;
 #endif
-
-#if WITH_MGPU
-
-RHI_API TArray<FRHIGPUMask, TFixedAllocator<MAX_NUM_GPUS>> AFRUtils::GroupMasks;
-RHI_API TArray<FRHIGPUMask, TFixedAllocator<MAX_NUM_GPUS>> AFRUtils::SiblingMasks;
-
-void AFRUtils::StaticInitialize()
-{
-	check(GroupMasks.Num() == 0 && SiblingMasks.Num() == 0);
-	const uint32 NumGPUsPerGroup = GetNumGPUsPerGroup();
-
-	// Set up group masks.
-	for (uint32 GroupIndex = 0; GroupIndex < GNumAlternateFrameRenderingGroups; GroupIndex++)
-	{
-		FRHIGPUMask& GPUMask = GroupMasks.Emplace_GetRef(FRHIGPUMask::FromIndex(GroupIndex));
-		for (uint32 IndexWithinGroup = 1; IndexWithinGroup < NumGPUsPerGroup; IndexWithinGroup++)
-		{
-			GPUMask |= FRHIGPUMask::FromIndex(IndexWithinGroup * GNumAlternateFrameRenderingGroups + GroupIndex);
-		}
-	}
-
-	// Set up sibling masks.
-	for (uint32 IndexWithinGroup = 0; IndexWithinGroup < NumGPUsPerGroup; IndexWithinGroup++)
-	{
-		FRHIGPUMask& GPUMask = SiblingMasks.Emplace_GetRef(FRHIGPUMask::FromIndex(IndexWithinGroup * GNumAlternateFrameRenderingGroups));
-		for (uint32 GroupIndex = 1; GroupIndex < GNumAlternateFrameRenderingGroups; GroupIndex++)
-		{
-			GPUMask |= FRHIGPUMask::FromIndex(IndexWithinGroup * GNumAlternateFrameRenderingGroups + GroupIndex);
-		}
-	}
-}
-
-#endif // WITH_MGPU
diff --git a/Engine/Source/Runtime/RHI/Private/RHIGPUReadback.cpp b/Engine/Source/Runtime/RHI/Private/RHIGPUReadback.cpp
--- a/Engine/Source/Runtime/RHI/Private/RHIGPUReadback.cpp
+++ b/Engine/Source/Runtime/RHI/Private/RHIGPUReadback.cpp
@@ -24,8 +24,8 @@ void FGenericRHIGPUFence::Clear()
 
 void FGenericRHIGPUFence::WriteInternal()
 {
-	// GPU generally overlap the game. This overlap increases when using AFR. In normal mode this can make us appear to be further behind the gpu than we actually are.
-	InsertedFrameNumber = GFrameNumberRenderThread + GNumAlternateFrameRenderingGroups;
+	// GPU generally overlaps the game
+	InsertedFrameNumber = GFrameNumberRenderThread + 1;
 }
 
 bool FGenericRHIGPUFence::Poll() const
diff --git a/Engine/Source/Runtime/RHI/Private/RHIStrings.cpp b/Engine/Source/Runtime/RHI/Private/RHIStrings.cpp
--- a/Engine/Source/Runtime/RHI/Private/RHIStrings.cpp
+++ b/Engine/Source/Runtime/RHI/Private/RHIStrings.cpp
@@ -367,7 +367,6 @@ FString GetTextureCreateFlagsName(ETextureCreateFlags TextureCreateFlags)
 		case TexCreate_DepthStencilResolveTarget:		return TEXT("TexCreate_DepthStencilResolveTarget");
 		case TexCreate_Streamable:						return TEXT("TexCreate_Streamable");
 		case TexCreate_NoFastClearFinalize:				return TEXT("TexCreate_NoFastClearFinalize");
-		case TexCreate_AFRManual:						return TEXT("TexCreate_AFRManual");
 		case TexCreate_ReduceMemoryWithTilingMode:		return TEXT("TexCreate_ReduceMemoryWithTilingMode");
 		case TexCreate_AtomicCompatible:				return TEXT("TexCreate_AtomicCompatible");
 		case TexCreate_External:						return TEXT("TexCreate_External");
@@ -592,10 +591,16 @@ const TCHAR* GetTextureCreateFlagString(ETextureCreateFlags TextureCreateFlag)
 		return TEXT("Streamable");
 	case ETextureCreateFlags::NoFastClearFinalize:
 		return TEXT("NoFastClearFinalize");
-	case ETextureCreateFlags::AFRManual:
-		return TEXT("AFRManual");
+	case ETextureCreateFlags::Atomic64Compatible:
+		return TEXT("Atomic64Compatible");
 	case ETextureCreateFlags::ReduceMemoryWithTilingMode:
 		return TEXT("ReduceMemoryWithTilingMode");
+	case ETextureCreateFlags::AtomicCompatible:
+		return TEXT("AtomicCompatible");
+	case ETextureCreateFlags::External:
+		return TEXT("External");
+	case ETextureCreateFlags::MultiGPUGraphIgnore:
+		return TEXT("MultiGPUGraphIgnore");
 	}
 	return TEXT("");
 }
diff --git a/Engine/Source/Runtime/RHI/Public/MultiGPU.h b/Engine/Source/Runtime/RHI/Public/MultiGPU.h
--- a/Engine/Source/Runtime/RHI/Public/MultiGPU.h
+++ b/Engine/Source/Runtime/RHI/Public/MultiGPU.h
@@ -17,14 +17,12 @@
 #define WITH_MGPU 1	// Explicit MGPU
 #define MAX_NUM_GPUS 8
 extern RHI_API uint32 GNumExplicitGPUsForRendering;
-extern RHI_API uint32 GNumAlternateFrameRenderingGroups;
 extern RHI_API uint32 GVirtualMGPU;
 #else
 #define WITH_SLI 0
 #define WITH_MGPU 0
 #define MAX_NUM_GPUS 1
 #define GNumExplicitGPUsForRendering 1
-#define GNumAlternateFrameRenderingGroups 1
 #define GVirtualMGPU 0
 #endif
 
@@ -168,157 +166,3 @@ public:
 	FORCEINLINE friend FRHIGPUMask::FIterator begin(const FRHIGPUMask& NodeMask) { return FRHIGPUMask::FIterator(NodeMask.GPUMask); }
 	FORCEINLINE friend FRHIGPUMask::FIterator end(const FRHIGPUMask& NodeMask) { return FRHIGPUMask::FIterator(0); }
 };
-
-/**
- * GPU mask utilities to get information about AFR groups and siblings.
- * - An AFR group is the set of GPUs that are working on the same frame together.
- * - AFR siblings are the GPUs in other groups do the same kind of work on subsequent
- *   frames. For example, two GPUs that render the same view on different frames are
- *   AFR siblings.
- *
- * For an 4 GPU setup with 2 AFR groups:
- * - There are 2 GPUs per AFR group. 0b1010 and 0b0101 are the two groups.
- * - Each GPU has 1 sibling. 0b1100 and 0b0011 are siblings.
- */
-struct AFRUtils
-{
-	/**
-	 * Gets the number of GPUs per AFR group.
-	 */
-	static inline uint32 GetNumGPUsPerGroup()
-	{
-		checkSlow(GNumExplicitGPUsForRendering % GNumAlternateFrameRenderingGroups == 0);
-		return GNumExplicitGPUsForRendering / GNumAlternateFrameRenderingGroups;
-	}
-
-	/**
-	 * Gets the AFR group index for a GPU index.
-	 */
-	static inline uint32 GetGroupIndex(uint32 GPUIndex)
-	{
-		return GPUIndex % GNumAlternateFrameRenderingGroups;
-	}
-
-	/**
-	 * Gets the index of a GPU relative its AFR group.
-	 */
-	static inline uint32 GetIndexWithinGroup(uint32 GPUIndex)
-	{
-		return GPUIndex / GNumAlternateFrameRenderingGroups;
-	}
-
-	/**
-	 * Gets the next AFR sibling for a GPU.
-	 */
-	static inline uint32 GetNextSiblingGPUIndex(uint32 GPUIndex)
-	{
-#if WITH_MGPU
-		return GetIndexWithinGroup(GPUIndex) * GNumAlternateFrameRenderingGroups + GetGroupIndex(GPUIndex + 1);
-#else
-		return GPUIndex;
-#endif
-	}
-
-	/**
-	 * Gets a mask containing the next AFR siblings for a GPU mask.
-	 */
-	static inline FRHIGPUMask GetNextSiblingGPUMask(FRHIGPUMask InGPUMask)
-	{
-		FRHIGPUMask::FIterator It(InGPUMask);
-		FRHIGPUMask GPUMask = FRHIGPUMask::FromIndex(GetNextSiblingGPUIndex(*It));
-		while (++It)
-		{
-			GPUMask |= FRHIGPUMask::FromIndex(GetNextSiblingGPUIndex(*It));
-		}
-		return GPUMask;
-	}
-
-	/**
-	 * Gets the previous AFR sibling for a GPU.
-	 */
-	static inline uint32 GetPrevSiblingGPUIndex(uint32 GPUIndex)
-	{
-#if WITH_MGPU
-		return GetIndexWithinGroup(GPUIndex) * GNumAlternateFrameRenderingGroups + GetGroupIndex(GPUIndex + GNumAlternateFrameRenderingGroups - 1);
-#else
-		return GPUIndex;
-#endif
-	}
-
-	/**
-	 * Gets a mask containing the previous AFR siblings for a GPU mask.
-	 */
-	static inline FRHIGPUMask GetPrevSiblingGPUMask(FRHIGPUMask InGPUMask)
-	{
-		FRHIGPUMask::FIterator It(InGPUMask);
-		FRHIGPUMask GPUMask = FRHIGPUMask::FromIndex(GetPrevSiblingGPUIndex(*It));
-		while (++It)
-		{
-			GPUMask |= FRHIGPUMask::FromIndex(GetPrevSiblingGPUIndex(*It));
-		}
-		return GPUMask;
-	}
-
-	/**
-	 * Gets the GPU mask including all GPUs within the same AFR group.
-	 */
-	static inline FRHIGPUMask GetGPUMaskForGroup(uint32 GPUIndex)
-	{
-#if WITH_MGPU
-		return GroupMasks[GetGroupIndex(GPUIndex)];
-#else
-		return FRHIGPUMask::FromIndex(GPUIndex);
-#endif
-	}
-
-	/**
-	 * Gets the GPU mask including all GPUs within the same AFR group(s).
-	 */
-	static inline FRHIGPUMask GetGPUMaskForGroup(FRHIGPUMask InGPUMask)
-	{
-		FRHIGPUMask::FIterator It(InGPUMask);
-		FRHIGPUMask GPUMask = GetGPUMaskForGroup(*It);
-		while (++It)
-		{
-			GPUMask |= GetGPUMaskForGroup(*It);
-		}
-		return GPUMask;
-	}
-
-	/**
-	 * Gets the GPU mask including all siblings across all AFR groups.
-	 */
-	static inline FRHIGPUMask GetGPUMaskWithSiblings(uint32 GPUIndex)
-	{
-#if WITH_MGPU
-		return SiblingMasks[GetIndexWithinGroup(GPUIndex)];
-#else
-		return FRHIGPUMask::FromIndex(GPUIndex);
-#endif
-	}
-
-	/**
-	 * Gets the GPU mask including all siblings across all AFR groups.
-	 */
-	static inline FRHIGPUMask GetGPUMaskWithSiblings(FRHIGPUMask InGPUMask)
-	{
-		FRHIGPUMask::FIterator It(InGPUMask);
-		FRHIGPUMask GPUMask = GetGPUMaskWithSiblings(*It);
-		while (++It)
-		{
-			GPUMask |= GetGPUMaskWithSiblings(*It);
-		}
-		return GPUMask;
-	}
-
-#if WITH_MGPU
-	static void StaticInitialize();
-
-	static inline const TArray<FRHIGPUMask, TFixedAllocator<MAX_NUM_GPUS>>& GetGroupMasks() { return GroupMasks; }
-	static inline const TArray<FRHIGPUMask, TFixedAllocator<MAX_NUM_GPUS>>& GetSiblingMasks() { return SiblingMasks; }
-
-private:
-	static RHI_API TArray<FRHIGPUMask, TFixedAllocator<MAX_NUM_GPUS>> GroupMasks;
-	static RHI_API TArray<FRHIGPUMask, TFixedAllocator<MAX_NUM_GPUS>> SiblingMasks;
-#endif // WITH_MGPU
-};
diff --git a/Engine/Source/Runtime/RHI/Public/RHICommandList.h b/Engine/Source/Runtime/RHI/Public/RHICommandList.h
--- a/Engine/Source/Runtime/RHI/Public/RHICommandList.h
+++ b/Engine/Source/Runtime/RHI/Public/RHICommandList.h
@@ -1182,28 +1182,6 @@ FRHICOMMAND_MACRO(FRHICommandSetGPUMask)
 	RHI_API void Execute(FRHICommandListBase& CmdList);
 };
 
-FRHICOMMAND_MACRO(FRHICommandWaitForTemporalEffect)
-{
-	FName EffectName;
-	FORCEINLINE_DEBUGGABLE FRHICommandWaitForTemporalEffect(const FName& InEffectName)
-		: EffectName(InEffectName)
-	{
-	}
-	RHI_API void Execute(FRHICommandListBase& CmdList);
-};
-
-FRHICOMMAND_MACRO_TPL(TRHIResource, FRHICommandBroadcastTemporalEffect)
-{
-	FName EffectName;
-	const TArrayView<TRHIResource*> Resources;
-	FORCEINLINE_DEBUGGABLE FRHICommandBroadcastTemporalEffect(const FName& InEffectName, const TArrayView<TRHIResource*> InResources)
-		: EffectName(InEffectName)
-		, Resources(InResources)
-	{
-	}
-	RHI_API void Execute(FRHICommandListBase& CmdList);
-};
-
 FRHICOMMAND_MACRO(FRHICommandTransferResources)
 {
 	TArray<FTransferResourceParams, TInlineAllocator<4>> Params;
@@ -3128,40 +3106,14 @@ PRAGMA_ENABLE_DEPRECATION_WARNINGS
 	}
 
 #if WITH_MGPU
-	FORCEINLINE_DEBUGGABLE void WaitForTemporalEffect(const FName& EffectName)
-	{
-		//check(IsOutsideRenderPass());
-		if (Bypass())
-		{
-			GetContext().RHIWaitForTemporalEffect(EffectName);
-			return;
-		}
-		ALLOC_COMMAND(FRHICommandWaitForTemporalEffect)(EffectName);
-	}
-
-	FORCEINLINE_DEBUGGABLE void BroadcastTemporalEffect(const FName& EffectName, const TArrayView<FRHITexture*> Textures)
-	{
-		//check(IsOutsideRenderPass());
-		if (Bypass())
-		{
-			GetContext().RHIBroadcastTemporalEffect(EffectName, Textures);
-			return;
-		}
+	UE_DEPRECATED(5.2, "AFR support has been removed in 5.2, function is a nop and calls should be removed.")
+	FORCEINLINE void WaitForTemporalEffect(const FName& EffectName) {}
 
-		ALLOC_COMMAND(FRHICommandBroadcastTemporalEffect<FRHITexture>)(EffectName, AllocArray(Textures));
-	}
+	UE_DEPRECATED(5.2, "AFR support has been removed in 5.2, function is a nop and calls should be removed.")
+	FORCEINLINE void BroadcastTemporalEffect(const FName& EffectName, const TArrayView<FRHITexture*> Textures) {}
 
-	FORCEINLINE_DEBUGGABLE void BroadcastTemporalEffect(const FName& EffectName, const TArrayView<FRHIBuffer*> Buffers)
-	{
-		//check(IsOutsideRenderPass());
-		if (Bypass())
-		{
-			GetContext().RHIBroadcastTemporalEffect(EffectName, Buffers);
-			return;
-		}
-
-		ALLOC_COMMAND(FRHICommandBroadcastTemporalEffect<FRHIBuffer>)(EffectName, AllocArray(Buffers));
-	}
+	UE_DEPRECATED(5.2, "AFR support has been removed in 5.2, function is a nop and calls should be removed.")
+	FORCEINLINE void BroadcastTemporalEffect(const FName& EffectName, const TArrayView<FRHIBuffer*> Buffers) {}
 #endif // WITH_MGPU
 
 #if RHI_RAYTRACING
@@ -3285,49 +3237,17 @@ public:
 		}
 	}
 
-	FORCEINLINE_DEBUGGABLE void BeginUpdateMultiFrameResource(FRHITexture* Texture)
-	{
-		//check(IsOutsideRenderPass());
-		if (Bypass())
-		{
-			GetContext().RHIBeginUpdateMultiFrameResource( Texture);
-			return;
-		}
-		ALLOC_COMMAND(FRHICommandBeginUpdateMultiFrameResource)(Texture);
-	}
+	UE_DEPRECATED(5.2, "AFR support has been removed in 5.2, function is a nop and calls should be removed.")
+	FORCEINLINE void BeginUpdateMultiFrameResource(FRHITexture* Texture) {}
 
-	FORCEINLINE_DEBUGGABLE void EndUpdateMultiFrameResource(FRHITexture* Texture)
-	{
-		//check(IsOutsideRenderPass());
-		if (Bypass())
-		{
-			GetContext().RHIEndUpdateMultiFrameResource(Texture);
-			return;
-		}
-		ALLOC_COMMAND(FRHICommandEndUpdateMultiFrameResource)(Texture);
-	}
+	UE_DEPRECATED(5.2, "AFR support has been removed in 5.2, function is a nop and calls should be removed.")
+	FORCEINLINE void EndUpdateMultiFrameResource(FRHITexture* Texture) {}
 
-	FORCEINLINE_DEBUGGABLE void BeginUpdateMultiFrameResource(FRHIUnorderedAccessView* UAV)
-	{
-		//check(IsOutsideRenderPass());
-		if (Bypass())
-		{
-			GetContext().RHIBeginUpdateMultiFrameResource(UAV);
-			return;
-		}
-		ALLOC_COMMAND(FRHICommandBeginUpdateMultiFrameUAV)(UAV);
-	}
+	UE_DEPRECATED(5.2, "AFR support has been removed in 5.2, function is a nop and calls should be removed.")
+	FORCEINLINE void BeginUpdateMultiFrameResource(FRHIUnorderedAccessView* UAV) {}
 
-	FORCEINLINE_DEBUGGABLE void EndUpdateMultiFrameResource(FRHIUnorderedAccessView* UAV)
-	{
-		//check(IsOutsideRenderPass());
-		if (Bypass())
-		{
-			GetContext().RHIEndUpdateMultiFrameResource(UAV);
-			return;
-		}
-		ALLOC_COMMAND(FRHICommandEndUpdateMultiFrameUAV)(UAV);
-	}
+	UE_DEPRECATED(5.2, "AFR support has been removed in 5.2, function is a nop and calls should be removed.")
+	FORCEINLINE void EndUpdateMultiFrameResource(FRHIUnorderedAccessView* UAV) {}
 
 	using FRHIComputeCommandList::SetShaderUniformBuffer;
 
diff --git a/Engine/Source/Runtime/RHI/Public/RHICommandListCommandExecutes.inl b/Engine/Source/Runtime/RHI/Public/RHICommandListCommandExecutes.inl
--- a/Engine/Source/Runtime/RHI/Public/RHICommandListCommandExecutes.inl
+++ b/Engine/Source/Runtime/RHI/Public/RHICommandListCommandExecutes.inl
@@ -60,30 +60,6 @@ struct FRHICommandSetRayTracingBindings;
 
 template <typename TRHIShader> struct FRHICommandSetLocalUniformBuffer;
 
-void FRHICommandBeginUpdateMultiFrameResource::Execute(FRHICommandListBase& CmdList)
-{
-	RHISTAT(BeginUpdateMultiFrameResource);
-	INTERNAL_DECORATOR(RHIBeginUpdateMultiFrameResource)(Texture);
-}
-
-void FRHICommandEndUpdateMultiFrameResource::Execute(FRHICommandListBase& CmdList)
-{
-	RHISTAT(EndUpdateMultiFrameResource);
-	INTERNAL_DECORATOR(RHIEndUpdateMultiFrameResource)(Texture);
-}
-
-void FRHICommandBeginUpdateMultiFrameUAV::Execute(FRHICommandListBase& CmdList)
-{
-	RHISTAT(BeginUpdateMultiFrameUAV);
-	INTERNAL_DECORATOR(RHIBeginUpdateMultiFrameResource)(UAV);
-}
-
-void FRHICommandEndUpdateMultiFrameUAV::Execute(FRHICommandListBase& CmdList)
-{
-	RHISTAT(EndUpdateMultiFrameUAV);
-	INTERNAL_DECORATOR(RHIEndUpdateMultiFrameResource)(UAV);
-}
-
 #if WITH_MGPU
 void FRHICommandSetGPUMask::Execute(FRHICommandListBase& CmdList)
 {
@@ -101,22 +77,6 @@ void FRHICommandSetGPUMask::Execute(FRHICommandListBase& CmdList)
 		}
 	}
 }
-void FRHICommandWaitForTemporalEffect::Execute(FRHICommandListBase& CmdList)
-{
-	RHISTAT(WaitForTemporalEffect);
-	INTERNAL_DECORATOR(RHIWaitForTemporalEffect)(EffectName);
-}
-
-template<> RHI_API void FRHICommandBroadcastTemporalEffect<FRHITexture>::Execute(FRHICommandListBase& CmdList)
-{
-	RHISTAT(BroadcastTemporalEffect);
-	INTERNAL_DECORATOR(RHIBroadcastTemporalEffect)(EffectName, Resources);
-}
-template<> RHI_API void FRHICommandBroadcastTemporalEffect<FRHIBuffer>::Execute(FRHICommandListBase& CmdList)
-{
-	RHISTAT(BroadcastTemporalEffect);
-	INTERNAL_DECORATOR(RHIBroadcastTemporalEffect)(EffectName, Resources);
-}
 
 void FRHICommandTransferResources::Execute(FRHICommandListBase& CmdList)
 {
diff --git a/Engine/Source/Runtime/RHI/Public/RHIContext.h b/Engine/Source/Runtime/RHI/Public/RHIContext.h
--- a/Engine/Source/Runtime/RHI/Public/RHIContext.h
+++ b/Engine/Source/Runtime/RHI/Public/RHIContext.h
@@ -362,23 +362,6 @@ public:
 		return FRHIGPUMask::GPU0();
 	}
 
-#if WITH_MGPU
-	virtual void RHIWaitForTemporalEffect(const FName& InEffectName)
-	{
-		/* empty default implementation */
-	}
-
-	virtual void RHIBroadcastTemporalEffect(const FName& InEffectName, const TArrayView<FRHITexture*> InTextures)
-	{
-		/* empty default implementation */
-	}
-
-	virtual void RHIBroadcastTemporalEffect(const FName& InEffectName, const TArrayView<FRHIBuffer*> InBuffers)
-	{
-		/* empty default implementation */
-	}
-#endif // WITH_MGPU
-
 	/**
 	 * Synchronizes the content of a resource between two GPUs using a copy operation.
 	 * @param Params - the parameters for each resource or texture region copied between GPUs.
@@ -624,29 +607,6 @@ public:
 	// This method is queued with an RHIThread, otherwise it will flush after it is queued; without an RHI thread there is no benefit to queuing this frame advance commands
 	virtual void RHIEndScene() = 0;
 
-	/**
-	* Signals the beginning and ending of rendering to a resource to be used in the next frame on a multiGPU system
-	*/
-	virtual void RHIBeginUpdateMultiFrameResource(FRHITexture* Texture)
-	{
-		/* empty default implementation */
-	}
-
-	virtual void RHIEndUpdateMultiFrameResource(FRHITexture* Texture)
-	{
-		/* empty default implementation */
-	}
-
-	virtual void RHIBeginUpdateMultiFrameResource(FRHIUnorderedAccessView* UAV)
-	{
-		/* empty default implementation */
-	}
-
-	virtual void RHIEndUpdateMultiFrameResource(FRHIUnorderedAccessView* UAV)
-	{
-		/* empty default implementation */
-	}
-
 	virtual void RHISetStreamSource(uint32 StreamIndex, FRHIBuffer* VertexBuffer, uint32 Offset) = 0;
 
 	// @param MinX including like Win32 RECT
diff --git a/Engine/Source/Runtime/RHI/Public/RHIDefinitions.h b/Engine/Source/Runtime/RHI/Public/RHIDefinitions.h
--- a/Engine/Source/Runtime/RHI/Public/RHIDefinitions.h
+++ b/Engine/Source/Runtime/RHI/Public/RHIDefinitions.h
@@ -952,8 +952,8 @@ enum class ETextureCreateFlags : uint64
     Streamable                        = 1ull << 28,
     // Render target will not FinalizeFastClear; Caches and meta data will be flushed, but clearing will be skipped (avoids potentially trashing metadata)
     NoFastClearFinalize               = 1ull << 29,
-    // Hint to the driver that this resource is managed properly by the engine for Alternate-Frame-Rendering in mGPU usage.
-    AFRManual                         = 1ull << 30,
+	/** Texture needs to support atomic operations */
+	Atomic64Compatible                = 1ull << 30,
     // Workaround for 128^3 volume textures getting bloated 4x due to tiling mode on some platforms.
     ReduceMemoryWithTilingMode        = 1ull << 31,
     /** Texture needs to support atomic operations */
@@ -962,8 +962,6 @@ enum class ETextureCreateFlags : uint64
 	External                		  = 1ull << 34,
 	/** Don't automatically transfer across GPUs in multi-GPU scenarios.  For example, if you are transferring it yourself manually. */
 	MultiGPUGraphIgnore				  = 1ull << 35,
-	/** Texture needs to support atomic operations */
-    Atomic64Compatible                = 1ull << 36,
 };
 ENUM_CLASS_FLAGS(ETextureCreateFlags);
 
@@ -999,7 +997,6 @@ ENUM_CLASS_FLAGS(ETextureCreateFlags);
 #define TexCreate_DepthStencilResolveTarget      ETextureCreateFlags::DepthStencilResolveTarget
 #define TexCreate_Streamable                     ETextureCreateFlags::Streamable
 #define TexCreate_NoFastClearFinalize            ETextureCreateFlags::NoFastClearFinalize
-#define TexCreate_AFRManual                      ETextureCreateFlags::AFRManual
 #define TexCreate_ReduceMemoryWithTilingMode     ETextureCreateFlags::ReduceMemoryWithTilingMode
 #define TexCreate_Transient                      ETextureCreateFlags::Transient
 #define TexCreate_AtomicCompatible               ETextureCreateFlags::AtomicCompatible
diff --git a/Engine/Source/Runtime/RHI/Public/RHIValidationContext.h b/Engine/Source/Runtime/RHI/Public/RHIValidationContext.h
--- a/Engine/Source/Runtime/RHI/Public/RHIValidationContext.h
+++ b/Engine/Source/Runtime/RHI/Public/RHIValidationContext.h
@@ -595,29 +595,6 @@ public:
 		RHIContext->RHIEndScene();
 	}
 
-	/**
-	* Signals the beginning and ending of rendering to a resource to be used in the next frame on a multiGPU system
-	*/
-	virtual void RHIBeginUpdateMultiFrameResource(FRHITexture* Texture) override final
-	{
-		RHIContext->RHIBeginUpdateMultiFrameResource(Texture);
-	}
-
-	virtual void RHIEndUpdateMultiFrameResource(FRHITexture* Texture) override final
-	{
-		RHIContext->RHIEndUpdateMultiFrameResource(Texture);
-	}
-
-	virtual void RHIBeginUpdateMultiFrameResource(FRHIUnorderedAccessView* UAV) override final
-	{
-		RHIContext->RHIBeginUpdateMultiFrameResource(UAV);
-	}
-
-	virtual void RHIEndUpdateMultiFrameResource(FRHIUnorderedAccessView* UAV) override final
-	{
-		RHIContext->RHIEndUpdateMultiFrameResource(UAV);
-	}
-
 	virtual void RHISetStreamSource(uint32 StreamIndex, FRHIBuffer* VertexBuffer, uint32 Offset) override final
 	{
 		//#todo-rco: Decide if this is needed or not...
diff --git a/Engine/Source/Runtime/RenderCore/Private/RenderGraphBuilder.cpp b/Engine/Source/Runtime/RenderCore/Private/RenderGraphBuilder.cpp
--- a/Engine/Source/Runtime/RenderCore/Private/RenderGraphBuilder.cpp
+++ b/Engine/Source/Runtime/RenderCore/Private/RenderGraphBuilder.cpp
@@ -1934,19 +1934,6 @@ void FRDGBuilder::Execute()
 	RHICmdList.SetStaticUniformBuffers({});
 
 #if WITH_MGPU
-	if (NameForTemporalEffect != NAME_None)
-	{
-		TArray<FRHITexture*> BroadcastTexturesForTemporalEffect;
-		for (const FExtractedTexture& ExtractedTexture : ExtractedTextures)
-		{
-			if (EnumHasAnyFlags(ExtractedTexture.Texture->Flags, ERDGTextureFlags::MultiFrame))
-			{
-				BroadcastTexturesForTemporalEffect.Add(ExtractedTexture.Texture->GetRHIUnchecked());
-			}
-		}
-		RHICmdList.BroadcastTemporalEffect(NameForTemporalEffect, BroadcastTexturesForTemporalEffect);
-	}
-
 	if (bForceCopyCrossGPU)
 	{
 		ForceCopyCrossGPU();
@@ -2424,12 +2411,6 @@ FRDGPass* FRDGBuilder::SetupEmptyPass(FRDGPass* Pass)
 void FRDGBuilder::CompilePassOps(FRDGPass* Pass)
 {
 #if WITH_MGPU
-	if (!bWaitedForTemporalEffect && NameForTemporalEffect != NAME_None && Pass->Pipeline == ERHIPipeline::Graphics)
-	{
-		bWaitedForTemporalEffect = true;
-		Pass->bWaitForTemporalEffect = 1;
-	}
-
 	FRHIGPUMask GPUMask = Pass->GPUMask;
 #else
 	FRHIGPUMask GPUMask = FRHIGPUMask::All();
@@ -2878,13 +2859,6 @@ void FRDGBuilder::ExecutePass(FRDGPass* Pass, FRHIComputeCommandList& RHICmdList
 
 			IF_RDG_ENABLE_DEBUG(ConditionalDebugBreak(RDG_BREAKPOINT_PASS_EXECUTE, BuilderName.GetTCHAR(), Pass->GetName()));
 
-#if WITH_MGPU
-			if (Pass->bWaitForTemporalEffect)
-			{
-				static_cast<FRHICommandList&>(RHICmdListPass).WaitForTemporalEffect(NameForTemporalEffect);
-			}
-#endif
-
 			Pass->GPUScopeOpsPrologue.Execute(RHICmdListPass);
 
 			ExecutePassPrologue(RHICmdListPass, Pass);
diff --git a/Engine/Source/Runtime/RenderCore/Public/RenderGraphBuilder.h b/Engine/Source/Runtime/RenderCore/Public/RenderGraphBuilder.h
--- a/Engine/Source/Runtime/RenderCore/Public/RenderGraphBuilder.h
+++ b/Engine/Source/Runtime/RenderCore/Public/RenderGraphBuilder.h
@@ -210,13 +210,6 @@ public:
 	/** Adds a user-defined dependency between two passes. This can be used to fine-tune async compute overlap by forcing a sync point. */
 	void AddPassDependency(FRDGPass* Producer, FRDGPass* Consumer);
 
-#if WITH_MGPU
-	void SetNameForTemporalEffect(FName InNameForTemporalEffect)
-	{
-		NameForTemporalEffect = InNameForTemporalEffect;
-	}
-#endif
-
 	/** Sets the current command list stat for all subsequent passes. */
 	void SetCommandListStat(TStatId StatId);
 
@@ -749,12 +742,6 @@ private:
 	} AuxiliaryPasses;
 
 #if WITH_MGPU
-	/** Name for the temporal effect used to synchronize multi-frame resources. */
-	FName NameForTemporalEffect;
-
-	/** Whether we performed the wait for the temporal effect yet. */
-	bool bWaitedForTemporalEffect = false;
-
 	/** Copy all cross GPU external resources (not marked MultiGPUGraphIgnore) at the end of execution (bad for perf, but useful for debugging). */
 	bool bForceCopyCrossGPU = false;
 #endif
diff --git a/Engine/Source/Runtime/RenderCore/Public/RenderGraphPass.h b/Engine/Source/Runtime/RenderCore/Public/RenderGraphPass.h
--- a/Engine/Source/Runtime/RenderCore/Public/RenderGraphPass.h
+++ b/Engine/Source/Runtime/RenderCore/Public/RenderGraphPass.h
@@ -401,9 +401,6 @@ protected:
 
 			/** If set, the pass should set its command list stat. */
 			uint32 bSetCommandListStat : 1;
-
-			/** If set, the pass will wait on the assigned mGPU temporal effect. */
-			uint32 bWaitForTemporalEffect : 1;
 		};
 		uint32 PackedBits1 = 0;
 	};
diff --git a/Engine/Source/Runtime/Renderer/Private/DeferredShadingRenderer.cpp b/Engine/Source/Runtime/Renderer/Private/DeferredShadingRenderer.cpp
--- a/Engine/Source/Runtime/Renderer/Private/DeferredShadingRenderer.cpp
+++ b/Engine/Source/Runtime/Renderer/Private/DeferredShadingRenderer.cpp
@@ -2409,13 +2409,6 @@ bool FDeferredShadingSceneRenderer::IsNaniteEnabled() const
 	return UseNanite(ShaderPlatform) && ViewFamily.EngineShowFlags.NaniteMeshes && Nanite::GStreamingManager.HasResourceEntries();
 }
 
-#if WITH_MGPU
-BEGIN_SHADER_PARAMETER_STRUCT(FGBufferTemporalTextureParams, )
-	RDG_TEXTURE_ACCESS(DepthResolve, ERHIAccess::CopySrc)
-	RDG_TEXTURE_ACCESS(GBufferA, ERHIAccess::CopySrc)
-END_SHADER_PARAMETER_STRUCT()
-#endif
-
 void FDeferredShadingSceneRenderer::Render(FRDGBuilder& GraphBuilder)
 {
 	const bool bNaniteEnabled = IsNaniteEnabled();
@@ -3434,35 +3427,6 @@ void FDeferredShadingSceneRenderer::Render(FRDGBuilder& GraphBuilder)
 			AddResolveSceneDepthPass(GraphBuilder, Views, SceneTextures.Depth);
 		}
 
-#if WITH_MGPU
-		if (SceneTextures.Depth.Resolve && SceneTextures.GBufferA)
-		{
-			FGBufferTemporalTextureParams* PassParameters = GraphBuilder.AllocParameters<FGBufferTemporalTextureParams>();
-
-			PassParameters->DepthResolve = SceneTextures.Depth.Resolve;
-			PassParameters->GBufferA = SceneTextures.GBufferA;
-
-			GraphBuilder.AddPass(
-				RDG_EVENT_NAME("GBuffer Temporal Copy"),
-				PassParameters,
-				ERDGPassFlags::Copy | ERDGPassFlags::NeverCull,
-				[PassParameters](FRHIComputeCommandList& RHICmdList)
-				{
-					FName GBufferTemporalEffect("GBufferTemporalCopy");
-
-					RHICmdList.WaitForTemporalEffect(GBufferTemporalEffect);
-
-					const uint32 NumGBufferTemporalTextures = 2;
-					TStaticArray<FRHITexture*, NumGBufferTemporalTextures> GBufferTemporalTexturesRHI;
-					GBufferTemporalTexturesRHI[0] = PassParameters->DepthResolve->GetRHI();
-					GBufferTemporalTexturesRHI[1] = PassParameters->GBufferA->GetRHI();
-
-					RHICmdList.BroadcastTemporalEffect(
-						GBufferTemporalEffect, MakeArrayView(GBufferTemporalTexturesRHI.GetData(), GBufferTemporalTexturesRHI.Num()));
-				});
-		}
-#endif  // WITH_MGPU
-		
 		if (bNaniteEnabled)
 		{
 			if (GNaniteShowStats != 0)
diff --git a/Engine/Source/Runtime/Renderer/Private/DistanceFieldLightingPost.cpp b/Engine/Source/Runtime/Renderer/Private/DistanceFieldLightingPost.cpp
--- a/Engine/Source/Runtime/Renderer/Private/DistanceFieldLightingPost.cpp
+++ b/Engine/Source/Runtime/Renderer/Private/DistanceFieldLightingPost.cpp
@@ -13,10 +13,6 @@
 #include "PipelineStateCache.h"
 #include "ScenePrivate.h"
 
-#if WITH_MGPU
-DECLARE_GPU_STAT(AFRWaitForDistanceFieldAOHistory);
-#endif
-
 int32 GAOUseHistory = 1;
 FAutoConsoleVariableRef CVarAOUseHistory(
 	TEXT("r.AOUseHistory"),
@@ -322,15 +318,6 @@ void UpdateHistory(
 
 	if (BentNormalHistoryState && DistanceFieldAOUseHistory(View))
 	{
-#if WITH_MGPU
-		RDG_GPU_STAT_SCOPE(GraphBuilder, AFRWaitForDistanceFieldAOHistory);
-		AddPass(GraphBuilder, RDG_EVENT_NAME("WaitForTemporalEffect"), [&View](FRHICommandList& RHICmdList)
-		{
-			static const FName NameForTemporalEffect("DistanceFieldAOHistory");
-			RHICmdList.WaitForTemporalEffect(FName(NameForTemporalEffect, View.ViewState->UniqueID));
-		});
-#endif
-
 		FIntPoint BufferSize = GetBufferSizeForAO(View);
 
 		if (*BentNormalHistoryState 
@@ -526,11 +513,6 @@ void UpdateHistory(
 		DistanceFieldAOHistoryViewRect->Min = FIntPoint::ZeroValue;
 		DistanceFieldAOHistoryViewRect->Max.X = View.ViewRect.Size().X / GAODownsampleFactor;
 		DistanceFieldAOHistoryViewRect->Max.Y = View.ViewRect.Size().Y / GAODownsampleFactor;
-
-#if WITH_MGPU && 0 // TODO(RDG)
-		FRHITexture* TexturesToCopyForTemporalEffect[] = { BentNormalHistoryOutput->GetRHI() };
-		RHICmdList.BroadcastTemporalEffect(FName(NameForTemporalEffect, View.ViewState->UniqueID), TexturesToCopyForTemporalEffect);
-#endif
 	}
 	else
 	{
diff --git a/Engine/Source/Runtime/Renderer/Private/DistanceFieldObjectManagement.cpp b/Engine/Source/Runtime/Renderer/Private/DistanceFieldObjectManagement.cpp
--- a/Engine/Source/Runtime/Renderer/Private/DistanceFieldObjectManagement.cpp
+++ b/Engine/Source/Runtime/Renderer/Private/DistanceFieldObjectManagement.cpp
@@ -911,7 +911,7 @@ void FSceneRenderer::PrepareDistanceFieldScene(FRDGBuilder& GraphBuilder, FRDGEx
 			{
 				FViewInfo& View = Views[ViewIndex];
 
-				RDG_GPU_MASK_SCOPE(GraphBuilder, GNumAlternateFrameRenderingGroups > 1 ? FRHIGPUMask::All() : View.GPUMask);
+				RDG_GPU_MASK_SCOPE(GraphBuilder, View.GPUMask);
 
 				float OcclusionMaxDistance = Scene->DefaultMaxDistanceFieldOcclusionDistance;
 
diff --git a/Engine/Source/Runtime/Renderer/Private/GPUScene.cpp b/Engine/Source/Runtime/Renderer/Private/GPUScene.cpp
--- a/Engine/Source/Runtime/Renderer/Private/GPUScene.cpp
+++ b/Engine/Source/Runtime/Renderer/Private/GPUScene.cpp
@@ -748,8 +748,6 @@ void FGPUScene::UpdateBufferState(FRDGBuilder& GraphBuilder, FScene& Scene, cons
 	ensure(bIsEnabled == UseGPUScene(GMaxRHIShaderPlatform, Scene.GetFeatureLevel()));
 	ensure(NumScenePrimitives == Scene.Primitives.Num());
 
-	// Multi-GPU support : Updating on all GPUs is inefficient for AFR. Work is wasted
-	// for any primitives that update on consecutive frames.
 	RDG_GPU_MASK_SCOPE(GraphBuilder, FRHIGPUMask::All());
 
 	constexpr int32 InitialBufferSize = 256;
@@ -939,8 +937,6 @@ void FGPUScene::UploadGeneral(FRDGBuilder& GraphBuilder, FScene& Scene, FRDGExte
 
 	SCOPED_NAMED_EVENT(UpdateGPUScene, FColor::Green);
 
-	// Multi-GPU support : Updating on all GPUs is inefficient for AFR. Work is wasted
-	// for any primitives that update on consecutive frames.
 	RDG_GPU_MASK_SCOPE(GraphBuilder, FRHIGPUMask::All());
 	RDG_EVENT_SCOPE(GraphBuilder, "UpdateGPUScene NumPrimitiveDataUploads %u", NumPrimitiveDataUploads);
 
diff --git a/Engine/Source/Runtime/Renderer/Private/PostProcess/TemporalAA.cpp b/Engine/Source/Runtime/Renderer/Private/PostProcess/TemporalAA.cpp
--- a/Engine/Source/Runtime/Renderer/Private/PostProcess/TemporalAA.cpp
+++ b/Engine/Source/Runtime/Renderer/Private/PostProcess/TemporalAA.cpp
@@ -876,42 +876,16 @@ FTAAOutputs AddTemporalAAPass(
 			bUseHistoryTexture[i] = PassParameters->Common.HistoryBuffer[i] != nullptr;
 		}
 
-		{
-			FIntVector GroupCount = FComputeShaderUtils::GetGroupCount(PracticableDestRect.Size(), GTemporalAATileSizeX);
-			FComputeShaderUtils::ValidateGroupCount(GroupCount);
-
-			const FShaderParametersMetadata* ParametersMetadata = FTemporalAACS::FParameters::FTypeInfo::GetStructMetadata();
-
-			GraphBuilder.AddPass(
-				RDG_EVENT_NAME("TAA(%s Quality=%s) %dx%d -> %dx%d",
-					PassName,
-					kTAAQualityNames[int32(PermutationVector.Get<FTemporalAA::FTAAQualityDim>())],
-					PracticableSrcRect.Width(), PracticableSrcRect.Height(),
-					PracticableDestRect.Width(), PracticableDestRect.Height()),
-				ParametersMetadata,
-				PassParameters,
-				ERDGPassFlags::Compute,
-				[ParametersMetadata, PassParameters, ComputeShader, GroupCount](FRHIComputeCommandList& RHICmdList)
-				{
-#if WITH_MGPU
-					RHICmdList.WaitForTemporalEffect(TAAEffectName);
-#endif  // WITH_MGPU
-
-					FComputeShaderUtils::Dispatch(RHICmdList, ComputeShader, ParametersMetadata, *PassParameters, GroupCount);
-
-#if WITH_MGPU
-					TArray<FRHITexture*, TFixedAllocator<2>> OutputTexturesRHI;
-					OutputTexturesRHI.Add(PassParameters->OutComputeTex[0]->GetParentRHI());
-					if (PassParameters->OutComputeTexDownsampled)
-					{
-						OutputTexturesRHI.Add(PassParameters->OutComputeTexDownsampled->GetParentRHI());
-					}
-
-					RHICmdList.BroadcastTemporalEffect(
-						TAAEffectName, MakeArrayView(OutputTexturesRHI.GetData(), OutputTexturesRHI.Num()));
-#endif  // WITH_MGPU
-				});
-		}
+		FComputeShaderUtils::AddPass(
+			GraphBuilder,
+			RDG_EVENT_NAME("TAA(%s Quality=%s) %dx%d -> %dx%d",
+				PassName,
+				kTAAQualityNames[int32(PermutationVector.Get<FTemporalAA::FTAAQualityDim>())],
+				PracticableSrcRect.Width(), PracticableSrcRect.Height(),
+				PracticableDestRect.Width(), PracticableDestRect.Height()),
+			ComputeShader,
+			PassParameters,
+			FComputeShaderUtils::GetGroupCount(PracticableDestRect.Size(), GTemporalAATileSizeX));
 	}
 	else
 	{
diff --git a/Engine/Source/Runtime/Renderer/Private/SceneCaptureRendering.cpp b/Engine/Source/Runtime/Renderer/Private/SceneCaptureRendering.cpp
--- a/Engine/Source/Runtime/Renderer/Private/SceneCaptureRendering.cpp
+++ b/Engine/Source/Runtime/Renderer/Private/SceneCaptureRendering.cpp
@@ -975,9 +975,8 @@ void FScene::UpdateSceneCaptureContents(USceneCaptureComponent2D* CaptureCompone
 
 
 		// If capturing every frame, only render to the GPUs that are actually being used
-		// this frame. Otherwise we will get poor performance in AFR. We can only determine
-		// this by querying the viewport back buffer on the render thread, so pass that
-		// along if it exists.
+		// this frame. We can only determine this by querying the viewport back buffer on
+		// the render thread, so pass that along if it exists.
 		FRenderTarget* GameViewportRT = nullptr;
 		if (CaptureComponent->bCaptureEveryFrame)
 		{
@@ -1006,7 +1005,7 @@ void FScene::UpdateSceneCaptureContents(USceneCaptureComponent2D* CaptureCompone
 			{
 				if (GameViewportRT != nullptr)
 				{
-					const FRHIGPUMask GPUMask = AFRUtils::GetGPUMaskForGroup(GameViewportRT->GetGPUMask(RHICmdList));
+					const FRHIGPUMask GPUMask = GameViewportRT->GetGPUMask(RHICmdList);
 					TextureRenderTargetResource->SetActiveGPUMask(GPUMask);
 				}
 				else
diff --git a/Engine/Source/Runtime/Renderer/Private/SceneOcclusion.cpp b/Engine/Source/Runtime/Renderer/Private/SceneOcclusion.cpp
--- a/Engine/Source/Runtime/Renderer/Private/SceneOcclusion.cpp
+++ b/Engine/Source/Runtime/Renderer/Private/SceneOcclusion.cpp
@@ -106,9 +106,8 @@ int32 FOcclusionQueryHelpers::GetNumBufferedFrames(ERHIFeatureLevel::Type Featur
 {
 	int32 NumGPUS = 1;
 #if WITH_SLI || WITH_MGPU
-	// If we're running with SLI, assume throughput is more important than latency, and buffer an extra frame
-	ensure(GNumAlternateFrameRenderingGroups <= (int32)FOcclusionQueryHelpers::MaxBufferedOcclusionFrames);
-	return FMath::Min<int32>(GNumAlternateFrameRenderingGroups, (int32)FOcclusionQueryHelpers::MaxBufferedOcclusionFrames);
+	// TODO:  Should this still be differentiated for MGPU?  Originally this logic was here for AFR, which has been removed.
+	return FMath::Min<int32>(1, (int32)FOcclusionQueryHelpers::MaxBufferedOcclusionFrames);
 #endif
 	static const auto NumBufferedQueriesVar = IConsoleManager::Get().FindTConsoleVariableDataInt(TEXT("r.NumBufferedOcclusionQueries"));
 	EShaderPlatform ShaderPlatform = GShaderPlatformForFeatureLevel[FeatureLevel];
@@ -288,8 +287,7 @@ bool FSceneViewState::IsShadowOccluded(FSceneViewState::FProjectedShadowKey Shad
 
 	// Read the occlusion query results.
 	uint64 NumSamples = 0;
-	// Only block on the query if not running SLI
-	const bool bWaitOnQuery = GNumAlternateFrameRenderingGroups == 1;
+	const bool bWaitOnQuery = true;
 
 	if (Query && RHIGetRenderQueryResult(Query->GetQuery(), NumSamples, bWaitOnQuery))
 	{
diff --git a/Engine/Source/Runtime/Renderer/Private/SceneRendering.cpp b/Engine/Source/Runtime/Renderer/Private/SceneRendering.cpp
--- a/Engine/Source/Runtime/Renderer/Private/SceneRendering.cpp
+++ b/Engine/Source/Runtime/Renderer/Private/SceneRendering.cpp
@@ -2952,9 +2952,8 @@ void FSceneRenderer::ComputeGPUMasks(FRHICommandListImmediate* RHICmdList)
 	//
 	// TODO:  we should remove this conditional, and set the GPU mask for the source render targets, but the goal is to have
 	// a minimal scope CL for the 5.1.1 hot fix.  This effectively reverts the change from CL 20540730, just for scene captures.
-	if ((GNumExplicitGPUsForRendering > 1) && ViewFamily.RenderTarget && ((GNumAlternateFrameRenderingGroups > 1) || Views[0].bIsSceneCapture))
+	if ((GNumExplicitGPUsForRendering > 1) && ViewFamily.RenderTarget && Views[0].bIsSceneCapture)
 	{
-		// Get render target GPU mask, taking into account AFR
 		check(RHICmdList);
 		RenderTargetGPUMask = ViewFamily.RenderTarget->GetGPUMask(*RHICmdList);
 	}
@@ -2963,23 +2962,16 @@ void FSceneRenderer::ComputeGPUMasks(FRHICommandListImmediate* RHICmdList)
 	// Otherwise fallback on rendering the whole view family on each relevant GPU using broadcast logic.
 	if (GNumExplicitGPUsForRendering > 1 && CVarEnableMultiGPUForkAndJoin.GetValueOnAnyThread() != 0)
 	{
-		// Check whether this looks like an AFR setup (note that the logic also applies when there is only one AFR group).
-		// Each AFR group uses multiple GPU. AFRGroup(i) = { i, NumAFRGroups + i,  2 * NumAFRGroups + i, ... } up to NumGPUs.
-		// Each view rendered gets assigned to the next GPU in that group. 
-		const FRHIGPUMask UsableGPUMask = AFRUtils::GetGPUMaskForGroup(RenderTargetGPUMask);
-
 		// Start iterating from RenderTargetGPUMask and then wrap around. This avoids an
 		// unnecessary cross-gpu transfer in cases where you only have 1 view and the
 		// render target is located on a GPU other than GPU 0.
-		FRHIGPUMask::FIterator GPUIterator(FRHIGPUMask::FilterGPUsBefore(RenderTargetGPUMask.GetFirstIndex()) & UsableGPUMask);
+		FRHIGPUMask::FIterator GPUIterator(RenderTargetGPUMask);
 		for (FViewInfo& ViewInfo : Views)
 		{
 			// Only handle views that are to be rendered (this excludes instance stereo).
 			if (ViewInfo.ShouldRenderView())
 			{
-				// Multi-GPU support : This is inefficient for AFR if the reflection capture
-				// updates every frame. Work is wasted on the GPUs that are not involved in
-				// rendering the current frame.
+				// TODO:  should reflection captures run on one GPU and transfer, like all other rendering?
 				if (ViewInfo.bIsReflectionCapture)
 				{
 					ViewInfo.GPUMask = FRHIGPUMask::All();
@@ -2997,7 +2989,7 @@ void FSceneRenderer::ComputeGPUMasks(FRHICommandListImmediate* RHICmdList)
 					++GPUIterator;
 					if (!GPUIterator)
 					{
-						GPUIterator = FRHIGPUMask::FIterator(UsableGPUMask);
+						GPUIterator = FRHIGPUMask::FIterator(RenderTargetGPUMask);
 					}
 				}
 			}
@@ -3073,9 +3065,7 @@ static void GetCrossGPUTransfers(FSceneRenderer* SceneRenderer, TArray<FCrossGPU
 void FSceneRenderer::PreallocateCrossGPUFences(const TArray<FSceneRenderer*>& SceneRenderers)
 {
 #if WITH_MGPU
-	// We can only apply cross GPU fence wait transfer optimizations if AFR isn't enabled.  This is because we can't determine
-	// which GPU a given render target is being rendered on where this function is called.
-	if ((SceneRenderers.Num() > 1) && (GNumAlternateFrameRenderingGroups <= 1))
+	if (SceneRenderers.Num() > 1)
 	{
 		// Allocated fences to wait on are placed in the last scene renderer
 		TArray<FTransferResourceFenceData*>& LastRendererFencesWait = SceneRenderers.Last()->CrossGPUTransferFencesWait;
diff --git a/Engine/Source/Runtime/Renderer/Private/ScreenSpaceDenoise.cpp b/Engine/Source/Runtime/Renderer/Private/ScreenSpaceDenoise.cpp
--- a/Engine/Source/Runtime/Renderer/Private/ScreenSpaceDenoise.cpp
+++ b/Engine/Source/Runtime/Renderer/Private/ScreenSpaceDenoise.cpp
@@ -146,35 +146,6 @@ DECLARE_GPU_STAT(ShadowsDenoiser)
 DECLARE_GPU_STAT(AmbientOcclusionDenoiser)
 DECLARE_GPU_STAT(DiffuseIndirectDenoiser)
 
-#if WITH_MGPU
-extern const FName SSDAmbientOcclusionEffectName;
-
-const FName SSDShadowVisibilityMasksEffectName("SSDShadowVisibility");
-const FName SSDMultiPolychromaticPenumbraHarmonicsEffectNames[4] =
-{
-	"SSDPolychromaticPenumbra0",
-	"SSDPolychromaticPenumbra1",
-	"SSDPolychromaticPenumbra2",
-	"SSDPolychromaticPenumbra3",
-};
-const FName SSDMultiPolychromaticPenumbraIntegralEffectName("SSDPolychromaticPenumbraIntegrand");
-const FName SSDMultiPolychromaticPenumbraFinalEffectName("SSDPolychromaticPenumbraFinal");
-const FName SSDReflectionsEffectName("SSDReflections");
-const FName SSDWaterReflectionsEffectName("SSDWaterReflections");
-const FName SSDAmbientOcclusionEffectName("SSDAmbientOcclusion");
-const FName SSDDiffuseIndirectEffectName("SSDDiffuseIndirect");
-const FName SSDSkyLightEffectName("SSDSkyLight");
-const FName SSDReflectedSkyLightEffectName("SSDReflectedSkyLight");
-const FName SSDDiffuseIndirectHarmonicEffectName("SSDDiffuseIndirectHarmonic");
-const FName SSDScreenSpaceDiffuseIndirectEffectName("SSDScreenSpaceDiffuseIndirect");
-const FName SSDIndirectProbeHierarchyEffectName("SSDIndirectProbeHierarchy");
-
-#define DECLARE_FSSD_CONSTANT_PIXEL_DENSITY_SETTINGS(EFFECT_NAME) FSSDConstantPixelDensitySettings Settings(EFFECT_NAME)
-
-#else  // WITH_MGPU
-#define DECLARE_FSSD_CONSTANT_PIXEL_DENSITY_SETTINGS(EFFECT_NAME) FSSDConstantPixelDensitySettings Settings
-#endif  // !WITH_MGPU
-
 namespace
 {
 
@@ -1370,13 +1341,6 @@ struct FSSDConstantPixelDensitySettings
 	TStaticArray<const FLightSceneInfo*, IScreenSpaceDenoiser::kMaxBatchSize> LightSceneInfo;
 	FRDGTextureRef CompressedDepthTexture = nullptr;
 	FRDGTextureRef CompressedShadingModelTexture = nullptr;
-#if WITH_MGPU
-	FName EffectName;
-
-	FSSDConstantPixelDensitySettings(const FName& InEffectName)
-		: EffectName(InEffectName)
-	{}
-#endif  // WITH_MGPU
 };
 
 /** Denoises a signal at constant pixel density across the viewport. */
@@ -2063,42 +2027,6 @@ static void DenoiseSignalAtConstantPixelDensity(
 				bExtractCompressedMetadata[i] = PassParameters->PrevCompressedMetadata[i] != nullptr;
 		}
 
-#if WITH_MGPU
-		{
-			FName EffectName = Settings.EffectName;
-
-			FIntVector GroupCount = FComputeShaderUtils::GetGroupCount(Viewport.Size(), FComputeShaderUtils::kGolden2DGroupSize);
-			FComputeShaderUtils::ValidateGroupCount(GroupCount);
-
-			const FShaderParametersMetadata* ParametersMetadata = FSSDTemporalAccumulationCS::FParameters::FTypeInfo::GetStructMetadata();
-
-			GraphBuilder.AddPass(
-				RDG_EVENT_NAME("SSD TemporalAccumulation%s",
-					(!Settings.bUseTemporalAccumulation || bGlobalCameraCut) ? TEXT("(Disabled)") : TEXT("")),
-				ParametersMetadata,
-				PassParameters,
-				ERDGPassFlags::Compute,
-				[ParametersMetadata, PassParameters, ComputeShader, GroupCount, EffectName](FRHIComputeCommandList& RHICmdList)
-				{
-					RHICmdList.WaitForTemporalEffect(EffectName);
-
-					FComputeShaderUtils::Dispatch(RHICmdList, ComputeShader, ParametersMetadata, *PassParameters, GroupCount);
-
-					TArray<FRHITexture*, TFixedAllocator<kMaxBufferProcessingCount>> SignalOutputTexturesRHI;
-					for (FRDGTextureUAV* TextureUAV : PassParameters->SignalHistoryOutput.UAVs)
-					{
-						if (TextureUAV)
-						{
-							SignalOutputTexturesRHI.Add(TextureUAV->GetParentRHI());
-						}
-					}
-
-					RHICmdList.BroadcastTemporalEffect(
-						EffectName, MakeArrayView(SignalOutputTexturesRHI.GetData(), SignalOutputTexturesRHI.Num()));
-				});
-		}
-#else  // WITH_MGPU
-
 		FComputeShaderUtils::AddPass(
 			GraphBuilder,
 			RDG_EVENT_NAME("SSD TemporalAccumulation%s",
@@ -2107,8 +2035,6 @@ static void DenoiseSignalAtConstantPixelDensity(
 			PassParameters,
 			FComputeShaderUtils::GetGroupCount(Viewport.Size(), FComputeShaderUtils::kGolden2DGroupSize));
 
-#endif  // !WITH_MGPU
-
 		SignalHistory = SignalOutput;
 	} // if (View.ViewState && Settings.bUseTemporalAccumulation)
 	
@@ -2357,7 +2283,7 @@ public:
 
 		FSSDSignalTextures InputSignal;
 
-		DECLARE_FSSD_CONSTANT_PIXEL_DENSITY_SETTINGS(SSDShadowVisibilityMasksEffectName);
+		FSSDConstantPixelDensitySettings Settings;
 		Settings.SignalProcessing = ESignalProcessing::ShadowVisibilityMask;
 		Settings.InputResolutionFraction = 1.0f;
 		Settings.ReconstructionSamples = FMath::Clamp(CVarShadowReconstructionSampleCount.GetValueOnRenderThread(), 1, kStackowiakMaxSampleCountPerSet);
@@ -2481,7 +2407,7 @@ public:
 			FViewInfoPooledRenderTargets ViewInfoPooledRenderTargets;
 			SetupSceneViewInfoPooledRenderTargets(View, &ViewInfoPooledRenderTargets);
 
-			DECLARE_FSSD_CONSTANT_PIXEL_DENSITY_SETTINGS(SSDMultiPolychromaticPenumbraHarmonicsEffectNames[HarmonicId]);
+			FSSDConstantPixelDensitySettings Settings;
 			Settings.FullResViewport = View.ViewRect;
 			Settings.SignalProcessing = ESignalProcessing::PolychromaticPenumbraHarmonic;
 			Settings.HarmonicPeriode = Periode;
@@ -2520,7 +2446,7 @@ public:
 			FViewInfoPooledRenderTargets ViewInfoPooledRenderTargets;
 			SetupSceneViewInfoPooledRenderTargets(View, &ViewInfoPooledRenderTargets);
 
-			DECLARE_FSSD_CONSTANT_PIXEL_DENSITY_SETTINGS(SSDMultiPolychromaticPenumbraIntegralEffectName);
+			FSSDConstantPixelDensitySettings Settings;
 			Settings.FullResViewport = View.ViewRect;
 			Settings.SignalProcessing = ESignalProcessing::PolychromaticPenumbraHarmonic;
 			Settings.HarmonicPeriode = Periode;
@@ -2605,7 +2531,7 @@ public:
 			FViewInfoPooledRenderTargets ViewInfoPooledRenderTargets;
 			SetupSceneViewInfoPooledRenderTargets(View, &ViewInfoPooledRenderTargets);
 
-			DECLARE_FSSD_CONSTANT_PIXEL_DENSITY_SETTINGS(SSDMultiPolychromaticPenumbraFinalEffectName);
+			FSSDConstantPixelDensitySettings Settings;
 			Settings.FullResViewport = View.ViewRect;
 			Settings.SignalProcessing = ESignalProcessing::PolychromaticPenumbraHarmonic;
 			Settings.bEnableReconstruction = false;
@@ -2651,7 +2577,7 @@ public:
 		InputSignal.Textures[0] = ReflectionInputs.Color;
 		InputSignal.Textures[1] = ReflectionInputs.RayHitDistance;
 
-		DECLARE_FSSD_CONSTANT_PIXEL_DENSITY_SETTINGS(SSDReflectionsEffectName);
+		FSSDConstantPixelDensitySettings Settings;
 		Settings.FullResViewport = View.ViewRect;
 		Settings.SignalProcessing = ESignalProcessing::Reflections;
 		Settings.InputResolutionFraction = RayTracingConfig.ResolutionFraction;
@@ -2700,7 +2626,7 @@ public:
 		InputSignal.Textures[0] = ReflectionInputs.Color;
 		InputSignal.Textures[1] = ReflectionInputs.RayHitDistance;
 
-		DECLARE_FSSD_CONSTANT_PIXEL_DENSITY_SETTINGS(SSDWaterReflectionsEffectName);
+		FSSDConstantPixelDensitySettings Settings;
 		Settings.FullResViewport = View.ViewRect;
 		Settings.SignalProcessing = ESignalProcessing::Reflections; // TODO: water reflection to denoise only water pixels
 		Settings.InputResolutionFraction = RayTracingConfig.ResolutionFraction;
@@ -2745,7 +2671,7 @@ public:
 		InputSignal.Textures[0] = ReflectionInputs.Mask;
 		InputSignal.Textures[1] = ReflectionInputs.RayHitDistance;
 		
-		DECLARE_FSSD_CONSTANT_PIXEL_DENSITY_SETTINGS(SSDAmbientOcclusionEffectName);
+		FSSDConstantPixelDensitySettings Settings;
 		Settings.FullResViewport = View.ViewRect;
 		Settings.SignalProcessing = ESignalProcessing::AmbientOcclusion;
 		Settings.InputResolutionFraction = RayTracingConfig.ResolutionFraction;
@@ -2792,7 +2718,7 @@ public:
 		InputSignal.Textures[0] = Inputs.Color;
 		InputSignal.Textures[1] = Inputs.RayHitDistance;
 
-		DECLARE_FSSD_CONSTANT_PIXEL_DENSITY_SETTINGS(SSDDiffuseIndirectEffectName);
+		FSSDConstantPixelDensitySettings Settings;
 		Settings.FullResViewport = View.ViewRect;
 		Settings.SignalProcessing = ESignalProcessing::DiffuseAndAmbientOcclusion;
 		Settings.InputResolutionFraction = Config.ResolutionFraction;
@@ -2836,7 +2762,7 @@ public:
 		InputSignal.Textures[0] = Inputs.Color;
 		InputSignal.Textures[1] = Inputs.RayHitDistance;
 
-		DECLARE_FSSD_CONSTANT_PIXEL_DENSITY_SETTINGS(SSDSkyLightEffectName);
+		FSSDConstantPixelDensitySettings Settings;
 		Settings.FullResViewport = View.ViewRect;
 		Settings.SignalProcessing = ESignalProcessing::DiffuseAndAmbientOcclusion;
 		Settings.InputResolutionFraction = Config.ResolutionFraction;
@@ -2882,7 +2808,7 @@ public:
 		InputSignal.Textures[0] = Inputs.Color;
 		InputSignal.Textures[1] = Inputs.RayHitDistance;
 
-		DECLARE_FSSD_CONSTANT_PIXEL_DENSITY_SETTINGS(SSDReflectedSkyLightEffectName);
+		FSSDConstantPixelDensitySettings Settings;
 		Settings.FullResViewport = View.ViewRect;
 		Settings.SignalProcessing = ESignalProcessing::DiffuseAndAmbientOcclusion;
 		Settings.InputResolutionFraction = Config.ResolutionFraction;
@@ -2928,7 +2854,7 @@ public:
 		for (int32 i = 0; i < IScreenSpaceDenoiser::kSphericalHarmonicTextureCount; i++)
 			InputSignal.Textures[i] = Inputs.SphericalHarmonic[i];
 
-		DECLARE_FSSD_CONSTANT_PIXEL_DENSITY_SETTINGS(SSDDiffuseIndirectHarmonicEffectName);
+		FSSDConstantPixelDensitySettings Settings;
 		Settings.FullResViewport = View.ViewRect;
 		Settings.SignalProcessing = ESignalProcessing::DiffuseSphericalHarmonic;
 		Settings.InputResolutionFraction = 1.0f / float(CommonDiffuseParameters.DownscaleFactor);
@@ -2975,7 +2901,7 @@ public:
 		InputSignal.Textures[0] = Inputs.Color;
 		InputSignal.Textures[1] = Inputs.AmbientOcclusionMask;
 
-		DECLARE_FSSD_CONSTANT_PIXEL_DENSITY_SETTINGS(SSDScreenSpaceDiffuseIndirectEffectName);
+		FSSDConstantPixelDensitySettings Settings;
 		Settings.FullResViewport = View.ViewRect;
 		Settings.SignalProcessing = ESignalProcessing::ScreenSpaceDiffuseIndirect;
 		Settings.InputResolutionFraction = Config.ResolutionFraction;
@@ -3011,7 +2937,7 @@ FSSDSignalTextures IScreenSpaceDenoiser::DenoiseIndirectProbeHierarchy(
 	FRDGTextureRef CompressedDepthTexture,
 	FRDGTextureRef CompressedShadingModelTexture)
 {
-	DECLARE_FSSD_CONSTANT_PIXEL_DENSITY_SETTINGS(SSDIndirectProbeHierarchyEffectName);
+	FSSDConstantPixelDensitySettings Settings;
 	Settings.FullResViewport = View.ViewRect;
 	Settings.SignalProcessing = ESignalProcessing::IndirectProbeHierarchy;
 	Settings.bEnableReconstruction = false;
diff --git a/Engine/Source/Runtime/Renderer/Private/VT/VirtualTextureSpace.cpp b/Engine/Source/Runtime/Renderer/Private/VT/VirtualTextureSpace.cpp
--- a/Engine/Source/Runtime/Renderer/Private/VT/VirtualTextureSpace.cpp
+++ b/Engine/Source/Runtime/Renderer/Private/VT/VirtualTextureSpace.cpp
@@ -363,7 +363,7 @@ void FVirtualTextureSpace::ApplyUpdates(FVirtualTextureSystem* System, FRDGBuild
 		return;
 	}
 
-	// Multi-GPU support : May be ineffecient for AFR.
+	// Multi-GPU support
 	RDG_GPU_MASK_SCOPE(GraphBuilder, FRHIGPUMask::All());
 
 	for (uint32 LayerIndex = 0u; LayerIndex < Description.NumPageTableLayers; ++LayerIndex)
diff --git a/Engine/Source/Runtime/Renderer/Private/VolumetricFog.cpp b/Engine/Source/Runtime/Renderer/Private/VolumetricFog.cpp
--- a/Engine/Source/Runtime/Renderer/Private/VolumetricFog.cpp
+++ b/Engine/Source/Runtime/Renderer/Private/VolumetricFog.cpp
@@ -1008,11 +1008,6 @@ void FDeferredShadingSceneRenderer::ComputeVolumetricFog(FRDGBuilder& GraphBuild
 
 		RDG_EVENT_SCOPE(GraphBuilder, "VolumetricFog");
 
-#if WITH_MGPU
-		static const FName NameForTemporalEffect("ComputeVolumetricFog");
-		GraphBuilder.SetNameForTemporalEffect(FName(NameForTemporalEffect, View.ViewState ? View.ViewState->UniqueID : 0));
-#endif
-
 		FRDGTextureRef ConservativeDepthTexture;
 		// To use a depth target format, and depth tests, we will have to render depth from a PS depth output. Keeping it simple for now with all the tests happening in shader.
 		if (GVolumetricFogConservativeDepth > 0)
diff --git a/Engine/Source/Runtime/Windows/D3D11RHI/Private/D3D11Commands.cpp b/Engine/Source/Runtime/Windows/D3D11RHI/Private/D3D11Commands.cpp
--- a/Engine/Source/Runtime/Windows/D3D11RHI/Private/D3D11Commands.cpp
+++ b/Engine/Source/Runtime/Windows/D3D11RHI/Private/D3D11Commands.cpp
@@ -85,89 +85,6 @@ static TAutoConsoleVariable<int32> CVarAllowUAVFlushExt(
 	ECVF_RenderThreadSafe);
 
 
-#if WITH_NVAPI
-//MultiGPU
-void FD3D11DynamicRHI::RHIBeginUpdateMultiFrameResource(FRHITexture* RHITexture)
-{
-	if (!IsRHIDeviceNVIDIA() || GNumAlternateFrameRenderingGroups == 1) return;
-
-	FD3D11Texture* Texture = ResourceCast(RHITexture);
-
-	if (!Texture)
-	{
-		return;
-	}
-
-	if (!Texture->GetIHVResourceHandle())
-	{
-		// get a resource handle for this texture
-		void* IHVHandle = nullptr;
-		NvAPI_D3D_GetObjectHandleForResource(Direct3DDevice, Texture->GetResource(), (NVDX_ObjectHandle*)&(IHVHandle));
-		Texture->SetIHVResourceHandle(IHVHandle);
-	}
-	
-	RHIPushEvent(TEXT("BeginMFUpdate"), FColor::Black);
-	NvAPI_D3D_BeginResourceRendering(Direct3DDevice, (NVDX_ObjectHandle)Texture->GetIHVResourceHandle(), 0);
-	RHIPopEvent();
-}
-
-void FD3D11DynamicRHI::RHIEndUpdateMultiFrameResource(FRHITexture* RHITexture)
-{
-	if (!IsRHIDeviceNVIDIA() || GNumAlternateFrameRenderingGroups == 1) return;
-
-	FD3D11Texture* Texture = ResourceCast(RHITexture);
-
-	if (!Texture || !Texture->GetIHVResourceHandle())
-	{
-		return;
-	}
-
-	RHIPushEvent(TEXT("EndMFUpdate"), FColor::Black);
-	NvAPI_D3D_EndResourceRendering(Direct3DDevice, (NVDX_ObjectHandle)Texture->GetIHVResourceHandle(), 0);
-	RHIPopEvent();	
-}
-
-void FD3D11DynamicRHI::RHIBeginUpdateMultiFrameResource(FRHIUnorderedAccessView* UAVRHI)
-{
-	if (!IsRHIDeviceNVIDIA() || GNumAlternateFrameRenderingGroups == 1) return;
-
-	FD3D11UnorderedAccessView* UAV = ResourceCast(UAVRHI);
-	
-	if (!UAV)
-	{
-		return;
-	}
-
-	if (!UAV->IHVResourceHandle)
-	{
-		// get a resource handle for this texture		
-		ID3D11Resource* D3DResource = nullptr;
-		UAV->View->GetResource(&D3DResource);
-		NvAPI_D3D_GetObjectHandleForResource(Direct3DDevice, D3DResource, (NVDX_ObjectHandle*)&(UAV->IHVResourceHandle));
-	}
-	
-	RHIPushEvent(TEXT("BeginMFUpdateUAV"), FColor::Black);
-	NvAPI_D3D_BeginResourceRendering(Direct3DDevice, (NVDX_ObjectHandle)UAV->IHVResourceHandle, 0);
-	RHIPopEvent();
-}
-
-void FD3D11DynamicRHI::RHIEndUpdateMultiFrameResource(FRHIUnorderedAccessView* UAVRHI)
-{
-	if (!IsRHIDeviceNVIDIA() || GNumAlternateFrameRenderingGroups == 1) return;
-
-	FD3D11UnorderedAccessView* UAV = ResourceCast(UAVRHI);
-
-	if (!UAV || !UAV->IHVResourceHandle)
-	{
-		return;
-	}
-
-	RHIPushEvent(TEXT("EndMFUpdateUAV"), FColor::Black);
-	NvAPI_D3D_EndResourceRendering(Direct3DDevice, (NVDX_ObjectHandle)UAV->IHVResourceHandle, 0);
-	RHIPopEvent();
-}
-#endif
-
 // Vertex state.
 void FD3D11DynamicRHI::RHISetStreamSource(uint32 StreamIndex, FRHIBuffer* VertexBufferRHI, uint32 Offset)
 {
diff --git a/Engine/Source/Runtime/Windows/D3D11RHI/Private/D3D11RHI.cpp b/Engine/Source/Runtime/Windows/D3D11RHI/Private/D3D11RHI.cpp
--- a/Engine/Source/Runtime/Windows/D3D11RHI/Private/D3D11RHI.cpp
+++ b/Engine/Source/Runtime/Windows/D3D11RHI/Private/D3D11RHI.cpp
@@ -331,11 +331,7 @@ void FD3DGPUProfiler::BeginFrame(FD3D11DynamicRHI* InRHI)
 	}
 	bPreviousLatchedGProfilingGPUHitches = bLatchedGProfilingGPUHitches;
 
-	// Skip timing events when using SLI, they will not be accurate anyway
-	if (GNumAlternateFrameRenderingGroups == 1)
-	{
-		FrameTiming.StartTiming();
-	}
+	FrameTiming.StartTiming();
 
 	if (GetEmitDrawEvents())
 	{
@@ -363,15 +359,9 @@ void FD3DGPUProfiler::EndFrame()
 		PopEvent();
 	}
 
-	// Skip timing events when using SLI, they will not be accurate anyway
-	if (GNumAlternateFrameRenderingGroups == 1)
-	{
-		FrameTiming.EndTiming();
-	}
+	FrameTiming.EndTiming();
 
-	// Skip timing events when using SLI, as they will block the GPU and we want maximum throughput
-	// Stat unit GPU time is not accurate anyway with SLI
-	if (FrameTiming.IsSupported() && GNumAlternateFrameRenderingGroups == 1)
+	if (FrameTiming.IsSupported())
 	{
 		uint64 GPUTiming = FrameTiming.GetTiming();
 		uint64 GPUFreq = FrameTiming.GetTimingFrequency();
diff --git a/Engine/Source/Runtime/Windows/D3D11RHI/Private/D3D11RHIPrivate.h b/Engine/Source/Runtime/Windows/D3D11RHI/Private/D3D11RHIPrivate.h
--- a/Engine/Source/Runtime/Windows/D3D11RHI/Private/D3D11RHIPrivate.h
+++ b/Engine/Source/Runtime/Windows/D3D11RHI/Private/D3D11RHIPrivate.h
@@ -472,14 +472,6 @@ public:
 	 */
 	bool GetQueryData(ID3D11Query* Query, void* Data, SIZE_T DataSize, ERenderQueryType QueryType, bool bWait, bool bStallRHIThread);
 
-#if WITH_NVAPI
-	virtual void RHIBeginUpdateMultiFrameResource(FRHITexture* Texture) override;
-	virtual void RHIEndUpdateMultiFrameResource(FRHITexture* Texture) override;
-
-	virtual void RHIBeginUpdateMultiFrameResource(FRHIUnorderedAccessView* UAV) override;
-	virtual void RHIEndUpdateMultiFrameResource(FRHIUnorderedAccessView* UAV) override;
-#endif
-
 	virtual FSamplerStateRHIRef RHICreateSamplerState(const FSamplerStateInitializerRHI& Initializer) final override;
 	virtual FRasterizerStateRHIRef RHICreateRasterizerState(const FRasterizerStateInitializerRHI& Initializer) final override;
 	virtual FDepthStencilStateRHIRef RHICreateDepthStencilState(const FDepthStencilStateInitializerRHI& Initializer) final override;
diff --git a/Engine/Source/Runtime/Windows/D3D11RHI/Private/D3D11Texture.cpp b/Engine/Source/Runtime/Windows/D3D11RHI/Private/D3D11Texture.cpp
--- a/Engine/Source/Runtime/Windows/D3D11RHI/Private/D3D11Texture.cpp
+++ b/Engine/Source/Runtime/Windows/D3D11RHI/Private/D3D11Texture.cpp
@@ -769,19 +769,6 @@ FD3D11Texture* FD3D11DynamicRHI::CreateD3D11Texture2D(FRHITextureCreateDesc cons
 		DepthStencilViews
 	);
 
-#if WITH_NVAPI
-	if (IsRHIDeviceNVIDIA() && EnumHasAnyFlags(Flags, TexCreate_AFRManual))
-	{
-		// get a resource handle for this texture
-		void* IHVHandle = nullptr;
-		//getobjecthandle not threadsafe
-		NvAPI_D3D_GetObjectHandleForResource(Direct3DDevice, Texture2D->GetResource(), (NVDX_ObjectHandle*)&(IHVHandle));
-		Texture2D->SetIHVResourceHandle(IHVHandle);
-		
-		NvU32 ManualAFR = 1;
-		NvAPI_D3D_SetResourceHint(Direct3DDevice, (NVDX_ObjectHandle)IHVHandle, NVAPI_D3D_SRH_CATEGORY_SLI, NVAPI_D3D_SRH_SLI_APP_CONTROLLED_INTERFRAME_CONTENT_SYNC, &ManualAFR);
-	}
-#endif
 	if (CreateDesc.BulkData)
 	{
 		CreateDesc.BulkData->Discard();
@@ -930,19 +917,6 @@ FD3D11Texture* FD3D11DynamicRHI::CreateD3D11Texture3D(FRHITextureCreateDesc cons
 		{}
 	);
 
-#if WITH_NVAPI
-	if (IsRHIDeviceNVIDIA() && EnumHasAnyFlags(Flags, TexCreate_AFRManual))
-	{
-		// get a resource handle for this texture
-		void* IHVHandle = nullptr;
-		//getobjecthandle not threadsafe
-		NvAPI_D3D_GetObjectHandleForResource(Direct3DDevice, Texture3D->GetResource(), (NVDX_ObjectHandle*)&(IHVHandle));
-		Texture3D->SetIHVResourceHandle(IHVHandle);
-
-		NvU32 ManualAFR = 1;
-		NvAPI_D3D_SetResourceHint(Direct3DDevice, (NVDX_ObjectHandle)IHVHandle, NVAPI_D3D_SRH_CATEGORY_SLI, NVAPI_D3D_SRH_SLI_APP_CONTROLLED_INTERFRAME_CONTENT_SYNC, &ManualAFR);
-	}
-#endif
 	if (CreateDesc.BulkData)
 	{
 		CreateDesc.BulkData->Discard();
diff --git a/Engine/Source/Runtime/Windows/D3D11RHI/Private/D3D11Viewport.cpp b/Engine/Source/Runtime/Windows/D3D11RHI/Private/D3D11Viewport.cpp
--- a/Engine/Source/Runtime/Windows/D3D11RHI/Private/D3D11Viewport.cpp
+++ b/Engine/Source/Runtime/Windows/D3D11RHI/Private/D3D11Viewport.cpp
@@ -779,36 +779,32 @@ void FD3D11DynamicRHI::RHIEndDrawingViewport(FRHIViewport* ViewportRHI,bool bPre
 		bNativelyPresented = Viewport->Present(bLockToVsync);
 	}
 
-	// Don't wait on the GPU when using SLI, let the driver determine how many frames behind the GPU should be allowed to get
-	if (GNumAlternateFrameRenderingGroups == 1)
-	{
-		if (bNativelyPresented)
-		{ 
-			static const auto CFinishFrameVar = IConsoleManager::Get().FindTConsoleVariableDataInt(TEXT("r.FinishCurrentFrame"));
-			if (!CFinishFrameVar->GetValueOnRenderThread())
-			{
-				// Wait for the GPU to finish rendering the previous frame before finishing this frame.
-				Viewport->WaitForFrameEventCompletion();
-				Viewport->IssueFrameEvent();
-			}
-			else
-			{
-				// Finish current frame immediately to reduce latency
-				Viewport->IssueFrameEvent();
-				Viewport->WaitForFrameEventCompletion();
-			}
+	if (bNativelyPresented)
+	{ 
+		static const auto CFinishFrameVar = IConsoleManager::Get().FindTConsoleVariableDataInt(TEXT("r.FinishCurrentFrame"));
+		if (!CFinishFrameVar->GetValueOnRenderThread())
+		{
+			// Wait for the GPU to finish rendering the previous frame before finishing this frame.
+			Viewport->WaitForFrameEventCompletion();
+			Viewport->IssueFrameEvent();
 		}
-
-		// If the input latency timer has been triggered, block until the GPU is completely
-		// finished displaying this frame and calculate the delta time.
-		if ( GInputLatencyTimer.RenderThreadTrigger )
+		else
 		{
+			// Finish current frame immediately to reduce latency
+			Viewport->IssueFrameEvent();
 			Viewport->WaitForFrameEventCompletion();
-			uint32 EndTime = FPlatformTime::Cycles();
-			GInputLatencyTimer.DeltaTime = EndTime - GInputLatencyTimer.StartTime;
-			GInputLatencyTimer.RenderThreadTrigger = false;
 		}
 	}
+
+	// If the input latency timer has been triggered, block until the GPU is completely
+	// finished displaying this frame and calculate the delta time.
+	if ( GInputLatencyTimer.RenderThreadTrigger )
+	{
+		Viewport->WaitForFrameEventCompletion();
+		uint32 EndTime = FPlatformTime::Cycles();
+		GInputLatencyTimer.DeltaTime = EndTime - GInputLatencyTimer.StartTime;
+		GInputLatencyTimer.RenderThreadTrigger = false;
+	}
 }
 
 void FD3D11DynamicRHI::RHIAdvanceFrameForGetViewportBackBuffer(FRHIViewport* Viewport)
diff --git a/Engine/Source/Runtime/Windows/D3D11RHI/Private/Windows/WindowsD3D11Device.cpp b/Engine/Source/Runtime/Windows/D3D11RHI/Private/Windows/WindowsD3D11Device.cpp
--- a/Engine/Source/Runtime/Windows/D3D11RHI/Private/Windows/WindowsD3D11Device.cpp
+++ b/Engine/Source/Runtime/Windows/D3D11RHI/Private/Windows/WindowsD3D11Device.cpp
@@ -107,14 +107,6 @@ static TAutoConsoleVariable<int32> CVarNVidiaTimestampWorkaround(
 	TEXT("If true we disable timestamps on pre-maxwell hardware (workaround for driver bug)\n"),
 	ECVF_Default);
 
-int32 GDX11ForcedGPUs = -1;
-static FAutoConsoleVariableRef CVarDX11NumGPUs(
-	TEXT("r.DX11NumForcedGPUs"),
-	GDX11ForcedGPUs,
-	TEXT("Num Forced GPUs."),
-	ECVF_Default
-	);
-
 /**
  * Console variables used by the D3D11 RHI device.
  */
@@ -2090,38 +2082,6 @@ void FD3D11DynamicRHI::InitD3DDevice()
 		}
 #endif
 
-#if WITH_SLI
-		GNumAlternateFrameRenderingGroups = 1;
-
-#ifdef NVAPI_INTERFACE
-		if (!bRenderDoc && IsRHIDeviceNVIDIA())
-		{
-			NV_GET_CURRENT_SLI_STATE SLICaps;
-			FMemory::Memzero(SLICaps);
-			SLICaps.version = NV_GET_CURRENT_SLI_STATE_VER;
-			NvAPI_Status SLIStatus = NvAPI_D3D_GetCurrentSLIState(Direct3DDevice, &SLICaps);
-			if (SLIStatus == NVAPI_OK)
-			{
-				if (SLICaps.numAFRGroups > 1)
-				{
-					GNumAlternateFrameRenderingGroups = SLICaps.numAFRGroups;
-					UE_LOG(LogD3D11RHI, Log, TEXT("Detected %i SLI GPUs Setting GNumAlternateFrameRenderingGroups to: %i."), SLICaps.numAFRGroups, GNumAlternateFrameRenderingGroups);
-				}
-			}
-			else
-			{
-				UE_LOG(LogD3D11RHI, Log, TEXT("NvAPI_D3D_GetCurrentSLIState failed: 0x%x"), (int32)SLIStatus);
-			}
-		}
-#endif //NVAPI_INTERFACE
-
-		if (GDX11ForcedGPUs > 0)
-		{
-			GNumAlternateFrameRenderingGroups = GDX11ForcedGPUs;
-			UE_LOG(LogD3D11RHI, Log, TEXT("r.DX11NumForcedGPUs forcing GNumAlternateFrameRenderingGroups to: %i "), GDX11ForcedGPUs);
-		}
-#endif // WITH_SLI
-
 		if (IsRHIDeviceNVIDIA())
 		{
 			GSupportsDepthBoundsTest = true;
