Add some platform specific Input Device Property settings to allow for any platform to set maximum values on its trigger effect strength and position data.

#rnx

diff --git a/Engine/Source/Runtime/Engine/Classes/GameFramework/InputSettings.h b/Engine/Source/Runtime/Engine/Classes/GameFramework/InputSettings.h
--- a/Engine/Source/Runtime/Engine/Classes/GameFramework/InputSettings.h
+++ b/Engine/Source/Runtime/Engine/Classes/GameFramework/InputSettings.h
@@ -340,6 +340,8 @@ class ENGINE_API UInputPlatformSettings : public UPlatformSettings
 
 public:
 
+	UInputPlatformSettings();
+	
 	static UInputPlatformSettings* Get();
 
 #if WITH_EDITOR
@@ -363,6 +365,52 @@ public:
 
 	/** Returns an array of all Hardware Device Identifiers known to this platform */
 	const TArray<FHardwareDeviceIdentifier>& GetHardwareDevices() const;
+
+	////////////////////////////////////////////////////
+	// Trigger Feedback
+	
+	/**
+	 * The maximum position that a trigger can be set to
+	 * 
+	 * @see UInputDeviceTriggerFeedbackProperty
+	 */
+	UPROPERTY(config, EditAnywhere, Category = "Device Properties|Trigger Feedback", meta = (UIMin = "0"))
+	int32 MaxTriggerFeedbackPosition;
+
+	/**
+	 * The maximum strength that trigger feedback can be set to
+	 * 
+	 * @see UInputDeviceTriggerFeedbackProperty
+	 */
+	UPROPERTY(config, EditAnywhere, Category = "Device Properties|Trigger Feedback", meta = (UIMin = "0"))
+	int32 MaxTriggerFeedbackStrength;
+
+	////////////////////////////////////////////////////
+	// Trigger Vibrations
+	
+	/**
+	 * The max position that a vibration trigger effect can be set to.
+	 * 
+	 * @see UInputDeviceTriggerVibrationProperty::GetTriggerPositionValue
+	 */
+	UPROPERTY(config, EditAnywhere, Category = "Device Properties|Trigger Vibration", meta = (UIMin = "0"))
+	int32 MaxTriggerVibrationTriggerPosition;
+
+	/**
+	 * The max frequency that a trigger vibration can occur
+	 * 
+	 * @see UInputDeviceTriggerVibrationProperty::GetVibrationFrequencyValue
+	 */
+	UPROPERTY(config, EditAnywhere, Category = "Device Properties|Trigger Vibration", meta = (UIMin = "0"))
+	int32 MaxTriggerVibrationFrequency;
+
+	/**
+	 * The maximum amplitude that can be set on trigger vibrations
+	 * 
+	 * @see UInputDeviceTriggerVibrationProperty::GetVibrationAmplitudeValue
+	 */
+	UPROPERTY(config, EditAnywhere, Category = "Device Properties|Trigger Vibration", meta = (UIMin = "0"))
+	int32 MaxTriggerVibrationAmplitude;
 	
 protected:
 
diff --git a/Engine/Source/Runtime/Engine/Private/GameFramework/InputDeviceProperties.cpp b/Engine/Source/Runtime/Engine/Private/GameFramework/InputDeviceProperties.cpp
--- a/Engine/Source/Runtime/Engine/Private/GameFramework/InputDeviceProperties.cpp
+++ b/Engine/Source/Runtime/Engine/Private/GameFramework/InputDeviceProperties.cpp
@@ -5,6 +5,7 @@
 #include "Curves/CurveLinearColor.h"
 #include "Curves/CurveFloat.h"
 #include "GameFramework/InputDeviceLibrary.h"
+#include "GameFramework/InputSettings.h"
 
 #include UE_INLINE_GENERATED_CPP_BY_NAME(InputDeviceProperties)
 
@@ -217,9 +218,8 @@ int32 UInputDeviceTriggerFeedbackProperty::GetPositionValue(const FDeviceTrigger
 {
 	if (ensure(Data->FeedbackPositionCurve))
 	{
-		// TODO: Make the max Strength a cvar
 		int32 Pos = Data->FeedbackPositionCurve->GetFloatValue(Duration);
-		return FMath::Clamp(Pos, 0, 8);
+		return FMath::Clamp(Pos, 0, UInputPlatformSettings::Get()->MaxTriggerFeedbackPosition);
 	}
 
 	return 0;
@@ -229,9 +229,8 @@ int32 UInputDeviceTriggerFeedbackProperty::GetStrengthValue(const FDeviceTrigger
 {
 	if (ensure(Data->FeedbackStrenghCurve))
 	{
-		// TODO: Make the max Strength a cvar
 		int32 Strength = Data->FeedbackStrenghCurve->GetFloatValue(Duration);
-		return FMath::Clamp(Strength, 0, 8);
+		return FMath::Clamp(Strength, 0, UInputPlatformSettings::Get()->MaxTriggerFeedbackStrength);
 	}
 
 	return 0;
@@ -390,9 +389,8 @@ int32 UInputDeviceTriggerVibrationProperty::GetTriggerPositionValue(const FDevic
 {
 	if (ensure(Data->TriggerPositionCurve))
 	{
-		// TODO: Make the max Strength a cvar
 		int32 Strength = Data->TriggerPositionCurve->GetFloatValue(Duration);
-		return FMath::Clamp(Strength, 0, 9);
+		return FMath::Clamp(Strength, 0, UInputPlatformSettings::Get()->MaxTriggerVibrationTriggerPosition);
 	}
 
 	return 0;
@@ -402,9 +400,8 @@ int32 UInputDeviceTriggerVibrationProperty::GetVibrationFrequencyValue(const FDe
 {
 	if (ensure(Data->VibrationFrequencyCurve))
 	{
-		// TODO: Make the max Frequency a cvar
 		int32 Strength = Data->VibrationFrequencyCurve->GetFloatValue(Duration);
-		return FMath::Clamp(Strength, 0, 255);
+		return FMath::Clamp(Strength, 0, UInputPlatformSettings::Get()->MaxTriggerVibrationFrequency);
 	}
 	
 	return 0;
@@ -414,9 +411,8 @@ int32 UInputDeviceTriggerVibrationProperty::GetVibrationAmplitudeValue(const FDe
 {
 	if (ensure(Data->VibrationAmplitudeCurve))
 	{
-		// TODO: Make the max Amplitude a cvar
 		int32 Strength = Data->VibrationAmplitudeCurve->GetFloatValue(Duration);
-		return FMath::Clamp(Strength, 0, 8);
+		return FMath::Clamp(Strength, 0, UInputPlatformSettings::Get()->MaxTriggerVibrationAmplitude);
 	}
 	
 	return 0;
diff --git a/Engine/Source/Runtime/Engine/Private/UserInterface/InputSettings.cpp b/Engine/Source/Runtime/Engine/Private/UserInterface/InputSettings.cpp
--- a/Engine/Source/Runtime/Engine/Private/UserInterface/InputSettings.cpp
+++ b/Engine/Source/Runtime/Engine/Private/UserInterface/InputSettings.cpp
@@ -495,6 +495,9 @@ void UInputSettings::SetDefaultInputComponentClass(TSubclassOf<UInputComponent>
 	}
 }
 
+/////////////////////////////////////////////////////////////
+// FHardwareDeviceIdentifier
+
 FHardwareDeviceIdentifier FHardwareDeviceIdentifier::Invalid = { NAME_None, NAME_None };
 FHardwareDeviceIdentifier FHardwareDeviceIdentifier::DefaultKeyboardAndMouse = { TEXT("DefaultKeyboardAndMouse"), TEXT("KBM") };
 
@@ -518,6 +521,17 @@ bool FHardwareDeviceIdentifier::IsValid() const
 	return InputClassName.IsValid() && HardwareDeviceIdentifier.IsValid();
 }
 
+//////////////////////////////////////////////////////////////////
+// UInputPlatformSettings
+
+UInputPlatformSettings::UInputPlatformSettings()
+	: MaxTriggerFeedbackPosition(8)
+	, MaxTriggerFeedbackStrength(8)
+	, MaxTriggerVibrationTriggerPosition(9)
+	, MaxTriggerVibrationFrequency(255)
+	, MaxTriggerVibrationAmplitude(8)
+{ }
+
 UInputPlatformSettings* UInputPlatformSettings::Get()
 {
 	return UPlatformSettingsManager::Get().GetSettingsForPlatform<UInputPlatformSettings>();
