Seamless Travel
* Now keeping the network object info of actors that are seamless traveling so we can restore it after we reset the full list.
* Fixes subobject refcounting ensures triggered during SeamlessTravel

diff --git a/Engine/Source/Runtime/Engine/Classes/Engine/NetDriver.h b/Engine/Source/Runtime/Engine/Classes/Engine/NetDriver.h
--- a/Engine/Source/Runtime/Engine/Classes/Engine/NetDriver.h
+++ b/Engine/Source/Runtime/Engine/Classes/Engine/NetDriver.h
@@ -1554,6 +1554,9 @@ public:
 
 	ENGINE_API virtual void NotifyActorTearOff(AActor* Actor);
 
+	/** Called when an actor is about to be carried during a seamless travel */
+	ENGINE_API void NotifyActorIsTraveling(AActor* TravelingActor);
+
 	/** Set whether this actor should swap roles before replicating properties. */
 	ENGINE_API void SetRoleSwapOnReplicate(AActor* Actor, bool bSwapRoles);
 
diff --git a/Engine/Source/Runtime/Engine/Classes/Engine/NetworkObjectList.h b/Engine/Source/Runtime/Engine/Classes/Engine/NetworkObjectList.h
--- a/Engine/Source/Runtime/Engine/Classes/Engine/NetworkObjectList.h
+++ b/Engine/Source/Runtime/Engine/Classes/Engine/NetworkObjectList.h
@@ -259,6 +259,12 @@ public:
 	/** Removes the recently dormant status from the passed in connection */
 	void ClearRecentlyDormantConnection(AActor* const Actor, UNetConnection* const Connection, UNetDriver* NetDriver);
 
+	/** Called when a replicated actor is about to be carried from one world to another */
+	void OnActorIsTraveling(AActor* TravelingAtor);
+
+	/** Called when seamless traveling is almost done just before we initialize the new world */
+	void OnPostSeamlessTravel();
+
 	/** 
 	 *	Does the necessary house keeping when a new connection is added 
 	 *	When a new connection is added, we must add all objects back to the active list so the new connection will process it
@@ -347,6 +353,9 @@ private:
 	FNetworkObjectSet ActiveNetworkObjects;
 	FNetworkObjectSet ObjectsDormantOnAllConnections;
 
+	/** Store the network info of actors that travel to the new world during a seamless travel. */
+	FNetworkObjectSet SeamlessTravelingObjects;
+
 	TMap<TWeakObjectPtr<UNetConnection>, int32 > NumDormantObjectsPerConnection;
 
 	TMap<FName, FNetworkObjectSet> FullyDormantObjectsByLevel;
diff --git a/Engine/Source/Runtime/Engine/Private/NetDriver.cpp b/Engine/Source/Runtime/Engine/Private/NetDriver.cpp
--- a/Engine/Source/Runtime/Engine/Private/NetDriver.cpp
+++ b/Engine/Source/Runtime/Engine/Private/NetDriver.cpp
@@ -3517,6 +3517,11 @@ void UNetDriver::NotifyActorTearOff(AActor* Actor)
 #endif // UE_WITH_IRIS
 }
 
+void UNetDriver::NotifyActorIsTraveling(AActor* TravelingActor)
+{
+	NetworkObjects->OnActorIsTraveling(TravelingActor);
+}
+
 void UNetDriver::ForceNetUpdate(AActor* Actor)
 {
 	// Let Replication Driver handle it if its exists
@@ -6003,6 +6008,8 @@ void UNetDriver::PreSeamlessTravelGarbageCollect()
 void UNetDriver::PostSeamlessTravelGarbageCollect()
 {
 	CleanPackageMaps();
+
+	NetworkObjects->OnPostSeamlessTravel();
 }
 
 void UNetDriver::SetReplicationDriver(UReplicationDriver* NewReplicationDriver)
diff --git a/Engine/Source/Runtime/Engine/Private/NetworkObjectList.cpp b/Engine/Source/Runtime/Engine/Private/NetworkObjectList.cpp
--- a/Engine/Source/Runtime/Engine/Private/NetworkObjectList.cpp
+++ b/Engine/Source/Runtime/Engine/Private/NetworkObjectList.cpp
@@ -406,6 +406,38 @@ void FNetworkObjectList::ClearRecentlyDormantConnection(AActor* const Actor, UNe
 	NetworkObjectInfo->RecentlyDormantConnections.Remove(Connection);
 }
 
+void FNetworkObjectList::OnActorIsTraveling(AActor* TravelingAtor)
+{
+	TSharedPtr<FNetworkObjectInfo>* NetworkObjectInfoPtr = AllNetworkObjects.Find(TravelingAtor);
+	check(NetworkObjectInfoPtr);
+
+	SeamlessTravelingObjects.Add(*NetworkObjectInfoPtr);
+}
+
+void FNetworkObjectList::OnPostSeamlessTravel()
+{
+	for (const TSharedPtr<FNetworkObjectInfo>& SeamlessTravelingActorInfo : SeamlessTravelingObjects)
+	{
+		// Make sure the actor didn't re-register itself during seamless travel
+		if (AllNetworkObjects.Find(SeamlessTravelingActorInfo->Actor) != nullptr)
+		{
+			AActor* TravelingActor = SeamlessTravelingActorInfo->Actor;
+			ensureMsgf(false, TEXT("Seamless traveling actor %s was readded to the netdriver before seamless travel was finished!"), *GetNameSafe(TravelingActor));
+			
+			// Clean our lists of this duplicate info.
+			AllNetworkObjects.Remove(TravelingActor);
+			ActiveNetworkObjects.Remove(TravelingActor);
+			ObjectsDormantOnAllConnections.Remove(TravelingActor);
+		}
+
+		// Add the saved object info back into the active list
+		AllNetworkObjects.Add(SeamlessTravelingActorInfo);
+		ActiveNetworkObjects.Add(SeamlessTravelingActorInfo);
+	}
+
+	SeamlessTravelingObjects.Empty();
+}
+
 void FNetworkObjectList::HandleConnectionAdded()
 {
 	// When a new connection is added, we must add all objects back to the active list so the new connection will process it
diff --git a/Engine/Source/Runtime/Engine/Private/World.cpp b/Engine/Source/Runtime/Engine/Private/World.cpp
--- a/Engine/Source/Runtime/Engine/Private/World.cpp
+++ b/Engine/Source/Runtime/Engine/Private/World.cpp
@@ -7220,6 +7220,11 @@ UWorld* FSeamlessTravelHandler::Tick()
 				// Keep if it's in the current level AND it isn't specifically excluded AND it was either marked as should keep OR we don't own this actor
 				if (bIsInCurrentLevel && !bForceExcludeActor && (bManuallyMarkedKeep || bKeepNonOwnedActor))
 				{
+					if (NetworkObjectInfo)
+					{
+						NetDriver->NotifyActorIsTraveling(TheActor);
+					}
+
 					ActuallyKeptActors.Add(TheActor);
 					return true;
 				}
