Added stat scopes for UChildActorComponent::CreateChildActor, and DestroyChildActor.

diff --git a/Engine/Source/Runtime/Engine/Private/Components/ChildActorComponent.cpp b/Engine/Source/Runtime/Engine/Private/Components/ChildActorComponent.cpp
--- a/Engine/Source/Runtime/Engine/Private/Components/ChildActorComponent.cpp
+++ b/Engine/Source/Runtime/Engine/Private/Components/ChildActorComponent.cpp
@@ -710,6 +710,8 @@ void UChildActorComponent::CreateChildActor(TFunction<void(AActor*)> CustomizerF
 		}
 	}
 
+	QUICK_SCOPE_CYCLE_COUNTER(STAT_CreateChildActor);
+
 	// Kill spawned actor if we have one
 	DestroyChildActor();
 
@@ -855,6 +857,8 @@ void UChildActorComponent::DestroyChildActor()
 
 	if (ChildActor && (ChildActor->HasAuthority() || !IsChildActorReplicated()) && !IsBeingRemovedFromLevel())
 	{
+		QUICK_SCOPE_CYCLE_COUNTER(STAT_DestroyChildActor);
+
 		if (!GExitPurge)
 		{
 			// if still alive, destroy, otherwise just clear the pointer
