Fix crash occuring when retrieving or setting the runtime state of a DataLayerInstance from a DataLayerAsset but there is no instance for the asset existing in the world.

diff --git a/Engine/Source/Runtime/Engine/Private/WorldPartition/DataLayer/DataLayerSubsystem.cpp b/Engine/Source/Runtime/Engine/Private/WorldPartition/DataLayer/DataLayerSubsystem.cpp
--- a/Engine/Source/Runtime/Engine/Private/WorldPartition/DataLayer/DataLayerSubsystem.cpp
+++ b/Engine/Source/Runtime/Engine/Private/WorldPartition/DataLayer/DataLayerSubsystem.cpp
@@ -318,17 +318,36 @@ UDataLayerInstance* UDataLayerSubsystem::GetDataLayerInstanceFromAsset(const UDa
 
 void UDataLayerSubsystem::SetDataLayerInstanceRuntimeState(const UDataLayerAsset* InDataLayerAsset, EDataLayerRuntimeState InState, bool bInIsRecursive)
 {
-	SetDataLayerRuntimeState(GetDataLayerInstanceFromAsset(InDataLayerAsset), InState, bInIsRecursive);
+	if (UDataLayerInstance* DataLayerInstance = GetDataLayerInstanceFromAsset(InDataLayerAsset))
+	{
+		SetDataLayerRuntimeState(DataLayerInstance, InState, bInIsRecursive);
+	}
+	else
+	{
+		UE_LOG(LogWorldPartition, Warning, TEXT("UDataLayerSubsystem::SetDataLayerInstanceRuntimeState called with a Data Layer Asset that does not have a Data Layer Instance in this world : %s"), *GetWorld()->GetName());
+	}
 }
 
 EDataLayerRuntimeState UDataLayerSubsystem::GetDataLayerInstanceRuntimeState(const UDataLayerAsset* InDataLayerAsset) const
 {
-	return GetDataLayerRuntimeState(GetDataLayerInstanceFromAsset(InDataLayerAsset));
+	if (UDataLayerInstance* DataLayerInstance = GetDataLayerInstanceFromAsset(InDataLayerAsset))
+	{
+		return GetDataLayerRuntimeState(DataLayerInstance);
+	}
+	
+	UE_LOG(LogWorldPartition, Warning, TEXT("UDataLayerSubsystem::GetDataLayerInstanceRuntimeState called with a Data Layer Asset that does not have a Data Layer Instance in this world : %s"), *GetWorld()->GetName());
+	return EDataLayerRuntimeState::Unloaded;
 }
 
 EDataLayerRuntimeState UDataLayerSubsystem::GetDataLayerInstanceEffectiveRuntimeState(const UDataLayerAsset* InDataLayerAsset) const
 {
-	return GetDataLayerEffectiveRuntimeState(GetDataLayerInstanceFromAsset(InDataLayerAsset));
+	if (UDataLayerInstance* DataLayerInstance = GetDataLayerInstanceFromAsset(InDataLayerAsset))
+	{
+		return GetDataLayerEffectiveRuntimeState(DataLayerInstance);
+	}
+
+	UE_LOG(LogWorldPartition, Warning, TEXT("UDataLayerSubsystem::GetDataLayerInstanceEffectiveRuntimeState called with a Data Layer Asset that does not have a Data Layer Instance in this world : %s"), *GetWorld()->GetName());
+	return EDataLayerRuntimeState::Unloaded;
 }
 
 void UDataLayerSubsystem::SetDataLayerRuntimeState(const UDataLayerInstance* InDataLayerInstance, EDataLayerRuntimeState InState, bool bInIsRecursive)
