Remove ALLOW_DITHERED_LOD_FOR_INSTANCED_STATIC_MESHES.
This codepath is largely untested, infrequently used, and causes crashes.

diff --git a/Engine/Source/Runtime/Engine/Private/HierarchicalInstancedStaticMesh.cpp b/Engine/Source/Runtime/Engine/Private/HierarchicalInstancedStaticMesh.cpp
--- a/Engine/Source/Runtime/Engine/Private/HierarchicalInstancedStaticMesh.cpp
+++ b/Engine/Source/Runtime/Engine/Private/HierarchicalInstancedStaticMesh.cpp
@@ -1506,7 +1506,7 @@ void FHierarchicalStaticMeshSceneProxy::GetDynamicMeshElements(const TArray<cons
 	QUICK_SCOPE_CYCLE_COUNTER(STAT_HierarchicalInstancedStaticMeshSceneProxy_GetMeshElements);
 	SCOPE_CYCLE_COUNTER(STAT_HISMCGetDynamicMeshElement);
 
-	bool bMultipleSections = ALLOW_DITHERED_LOD_FOR_INSTANCED_STATIC_MESHES && bDitheredLODTransitions && CVarDitheredLOD.GetValueOnRenderThread() > 0;
+	bool bMultipleSections = bDitheredLODTransitions && CVarDitheredLOD.GetValueOnRenderThread() > 0;
 	bool bSingleSections = !bMultipleSections;
 	bool bOverestimate = CVarOverestimateLOD.GetValueOnRenderThread() > 0;
 
diff --git a/Engine/Source/Runtime/Engine/Private/InstancedStaticMesh.cpp b/Engine/Source/Runtime/Engine/Private/InstancedStaticMesh.cpp
--- a/Engine/Source/Runtime/Engine/Private/InstancedStaticMesh.cpp
+++ b/Engine/Source/Runtime/Engine/Private/InstancedStaticMesh.cpp
@@ -824,12 +824,12 @@ void FInstancedStaticMeshVertexFactory::ModifyCompilationEnvironment(const FVert
 
 	if (IsFeatureLevelSupported(Parameters.Platform, ERHIFeatureLevel::SM5))
 	{
-		OutEnvironment.SetDefine(TEXT("USE_DITHERED_LOD_TRANSITION_FOR_INSTANCED"), ALLOW_DITHERED_LOD_FOR_INSTANCED_STATIC_MESHES);
+		OutEnvironment.SetDefine(TEXT("USE_DITHERED_LOD_TRANSITION_FOR_INSTANCED"), true);
 	}
 	else
 	{
 		// On mobile dithered LOD transition has to be explicitly enabled in material and project settings
-		OutEnvironment.SetDefine(TEXT("USE_DITHERED_LOD_TRANSITION_FOR_INSTANCED"), Parameters.MaterialParameters.bIsDitheredLODTransition && ALLOW_DITHERED_LOD_FOR_INSTANCED_STATIC_MESHES);
+		OutEnvironment.SetDefine(TEXT("USE_DITHERED_LOD_TRANSITION_FOR_INSTANCED"), static_cast<bool>(Parameters.MaterialParameters.bIsDitheredLODTransition));
 	}
 
 	FLocalVertexFactory::ModifyCompilationEnvironment(Parameters, OutEnvironment);
@@ -840,8 +840,7 @@ void FInstancedStaticMeshVertexFactory::ModifyCompilationEnvironment(const FVert
  */
 void FInstancedStaticMeshVertexFactory::GetPSOPrecacheVertexFetchElements(EVertexInputStreamType VertexInputStreamType, FVertexDeclarationElementList& Elements)
 {
-	// Fallback to local vertex factory because manual vertex fetch is supported but special case handling might be needed when 
-	// ALLOW_DITHERED_LOD_FOR_INSTANCED_STATIC_MESHES is disabled (does that code path still work?)
+	// Fallback to local vertex factory because manual vertex fetch is supported
 	FLocalVertexFactory::GetPSOPrecacheVertexFetchElements(VertexInputStreamType, Elements);
 }
 
@@ -869,34 +868,6 @@ void FInstancedStaticMeshVertexFactory::InitRHI()
 
 	check(HasValidFeatureLevel());
 
-#if !ALLOW_DITHERED_LOD_FOR_INSTANCED_STATIC_MESHES // position(and normal) only shaders cannot work with dithered LOD
-	// If the vertex buffer containing position is not the same vertex buffer containing the rest of the data,
-	// then initialize PositionStream and PositionDeclaration.
-	if (Data.PositionComponent.VertexBuffer != Data.TangentBasisComponents[0].VertexBuffer)
-	{
-		auto AddDeclaration = [&](EVertexInputStreamType InputStreamType)
-		{
-			FVertexDeclarationElementList StreamElements;
-			StreamElements.Add(AccessStreamComponent(Data.PositionComponent, 0, InputStreamType));
-
-			if (InputStreamType == EVertexInputStreamType::PositionAndNormalOnly && Data.TangentBasisComponents[1].VertexBuffer != NULL)
-			{
-				StreamElements.Add(AccessStreamComponent(Data.TangentBasisComponents[1], 2, InputStreamType));
-			}
-
-			// Add the instanced location streams
-			StreamElements.Add(AccessStreamComponent(Data.InstanceOriginComponent, 8, InputStreamType));
-			StreamElements.Add(AccessStreamComponent(Data.InstanceTransformComponent[0], 9, InputStreamType));
-			StreamElements.Add(AccessStreamComponent(Data.InstanceTransformComponent[1], 10, InputStreamType));
-			StreamElements.Add(AccessStreamComponent(Data.InstanceTransformComponent[2], 11, InputStreamType));
-
-			InitDeclaration(StreamElements, InputStreamType);
-		};
-		AddDeclaration(EVertexInputStreamType::PositionOnly);
-		AddDeclaration(EVertexInputStreamType::PositionAndNormalOnly);
-	}
-#endif
-
 	FVertexDeclarationElementList Elements;
 	if(Data.PositionComponent.VertexBuffer != NULL)
 	{
@@ -1036,13 +1007,12 @@ IMPLEMENT_VERTEX_FACTORY_TYPE(FInstancedStaticMeshVertexFactory,"/Engine/Private
 	| EVertexFactoryFlags::SupportsStaticLighting
 	| EVertexFactoryFlags::SupportsDynamicLighting
 	| EVertexFactoryFlags::SupportsPrecisePrevWorldPos
-	| (!ALLOW_DITHERED_LOD_FOR_INSTANCED_STATIC_MESHES ? EVertexFactoryFlags::SupportsPositionOnly : EVertexFactoryFlags::None)
 	| EVertexFactoryFlags::SupportsCachingMeshDrawCommands
 	| EVertexFactoryFlags::SupportsRayTracing
 	| EVertexFactoryFlags::SupportsRayTracingDynamicGeometry
 	| EVertexFactoryFlags::SupportsLightmapBaking
 	| EVertexFactoryFlags::SupportsPrimitiveIdStream
-	| (ALLOW_DITHERED_LOD_FOR_INSTANCED_STATIC_MESHES ? EVertexFactoryFlags::DoesNotSupportNullPixelShader : EVertexFactoryFlags::None)
+	| EVertexFactoryFlags::DoesNotSupportNullPixelShader
 	| EVertexFactoryFlags::SupportsManualVertexFetch
 	| EVertexFactoryFlags::SupportsPSOPrecaching
 	| EVertexFactoryFlags::SupportsLumenMeshCards
diff --git a/Engine/Source/Runtime/Engine/Public/MaterialShared.h b/Engine/Source/Runtime/Engine/Public/MaterialShared.h
--- a/Engine/Source/Runtime/Engine/Public/MaterialShared.h
+++ b/Engine/Source/Runtime/Engine/Public/MaterialShared.h
@@ -111,9 +111,6 @@ template <class ElementType> class TLinkedList;
 
 #define HANDLE_CUSTOM_OUTPUTS_AS_MATERIAL_ATTRIBUTES 0
 
-// GPUCULL_TODO: This seems it does not even compile if turned off? (EDIT: compilation issues fixed in 18733736)
-#define ALLOW_DITHERED_LOD_FOR_INSTANCED_STATIC_MESHES (1)
-
 // disallow debug data in shipping or on non-desktop Test
 #define ALLOW_SHADERMAP_DEBUG_DATA (!(UE_BUILD_SHIPPING || (UE_BUILD_TEST && !PLATFORM_DESKTOP)))
 
