Strata - fix for material function output to not crash.
Preview is still not visible but that will be investigated further later.

diff --git a/Engine/Source/Runtime/Engine/Private/Materials/MaterialExpressions.cpp b/Engine/Source/Runtime/Engine/Private/Materials/MaterialExpressions.cpp
--- a/Engine/Source/Runtime/Engine/Private/Materials/MaterialExpressions.cpp
+++ b/Engine/Source/Runtime/Engine/Private/Materials/MaterialExpressions.cpp
@@ -2040,6 +2040,24 @@ void UMaterialExpression::ConnectToPreviewMaterial(UMaterial* InMaterial, int32
 				check(MaterialInput);
 				ConnectExpression(MaterialInput, OutputIndex);
 			}
+			else if (IsResultMaterialAttributes(0))
+			{
+				// Set preview material as unlit and not using material attribute
+				InMaterial->SetShadingModel(MSM_Unlit);
+				InMaterial->bUseMaterialAttributes = false;
+
+				// Connect the material attribute output to the material attribute break node
+				UMaterialExpressionBreakMaterialAttributes* BreakMatAtt = NewObject<UMaterialExpressionBreakMaterialAttributes>(this);
+				BreakMatAtt->MaterialAttributes.Connect(OutputIndex, this);
+
+				// Get BaseColor as unlit color preview from material attribute
+				UMaterialExpressionStrataUnlitBSDF* UnlitBSDF = NewObject<UMaterialExpressionStrataUnlitBSDF>(this);
+				UnlitBSDF->EmissiveColor.Connect(0, BreakMatAtt);
+
+				FExpressionInput* MaterialInput = InMaterial->GetExpressionInputForProperty(MP_FrontMaterial);
+				check(MaterialInput);
+				MaterialInput->Connect(0, UnlitBSDF);
+			}
 			else
 			{
 				InMaterial->SetShadingModel(MSM_Unlit);
