Fixed ContentBundleEditor to compare added actor's world against its injected world instead of relying on ActorDescContainer.

diff --git a/Engine/Source/Runtime/Engine/Private/WorldPartition/ContentBundle/ContentBundleEditor.cpp b/Engine/Source/Runtime/Engine/Private/WorldPartition/ContentBundle/ContentBundleEditor.cpp
--- a/Engine/Source/Runtime/Engine/Private/WorldPartition/ContentBundle/ContentBundleEditor.cpp
+++ b/Engine/Source/Runtime/Engine/Private/WorldPartition/ContentBundle/ContentBundleEditor.cpp
@@ -138,7 +138,7 @@ bool FContentBundleEditor::AddActor(AActor* InActor)
 {
 	check(GetStatus() == EContentBundleStatus::ContentInjected || GetStatus() == EContentBundleStatus::ReadyToInject);
 
-	if (InActor->GetWorld() != ActorDescContainer->GetWorld() || InActor->HasAllFlags(RF_Transient) || !InActor->IsMainPackageActor())
+	if (InActor->GetWorld() != GetInjectedWorld() || InActor->HasAllFlags(RF_Transient) || !InActor->IsMainPackageActor())
 	{
 		return false;
 	}
