Remove superfluous include.

diff --git a/Engine/Source/Runtime/Engine/Private/WorldPartition/WorldPartitionLevelHelper.cpp b/Engine/Source/Runtime/Engine/Private/WorldPartition/WorldPartitionLevelHelper.cpp
--- a/Engine/Source/Runtime/Engine/Private/WorldPartition/WorldPartitionLevelHelper.cpp
+++ b/Engine/Source/Runtime/Engine/Private/WorldPartition/WorldPartitionLevelHelper.cpp
@@ -17,7 +17,6 @@
 #include "UnrealEngine.h"
 #include "WorldPartition/WorldPartition.h"
 #include "WorldPartition/WorldPartitionPackageHelper.h"
-#include "AssetCompilingManager.h"
 #include "LevelUtils.h"
 #include "ActorFolder.h"
 
