World Partition Builders: always unload maps by default when invoking builders from the editor.

diff --git a/Engine/Source/Runtime/Engine/Public/WorldPartition/IWorldPartitionEditorModule.h b/Engine/Source/Runtime/Engine/Public/WorldPartition/IWorldPartitionEditorModule.h
--- a/Engine/Source/Runtime/Engine/Public/WorldPartition/IWorldPartitionEditorModule.h
+++ b/Engine/Source/Runtime/Engine/Public/WorldPartition/IWorldPartitionEditorModule.h
@@ -26,7 +26,7 @@ public:
 		UWorld* World = nullptr;
 		FString ExtraArgs;
 		FText OperationDescription;
-		bool bUnloadMap = false;
+		bool bUnloadMap = true;
 	};
 
 	virtual bool RunBuilder(TSubclassOf<UWorldPartitionBuilder> BuilderClass, UWorld* InWorld);
