Added a CVar to disallow unloading/deactivation of specific data layers on the server
Removed "wp.Runtime.EnableServerDataLayerStreamingOut" CVar

diff --git a/Engine/Source/Runtime/Engine/Private/WorldPartition/WorldPartition.cpp b/Engine/Source/Runtime/Engine/Private/WorldPartition/WorldPartition.cpp
--- a/Engine/Source/Runtime/Engine/Private/WorldPartition/WorldPartition.cpp
+++ b/Engine/Source/Runtime/Engine/Private/WorldPartition/WorldPartition.cpp
@@ -103,15 +103,6 @@ FAutoConsoleVariableRef UWorldPartition::CVarEnableServerStreamingOut(
 	TEXT("Changing the value while the game is running won't be considered."),
 	WorldPartition::ECVF_Runtime_ReadOnly);
 
-
-bool UWorldPartition::bEnableServerDataLayerStreamingOut = true;
-FAutoConsoleVariableRef UWorldPartition::CVarEnableServerDataLayerStreamingOut(
-	TEXT("wp.Runtime.EnableServerDataLayerStreamingOut"),
-	UWorldPartition::bEnableServerDataLayerStreamingOut,
-	TEXT("Turn on/off to allow or not the server to stream out levels of runtime cells linked to runtime data layers.\n")
-	TEXT("Changing the value while the game is running won't be considered."),
-	WorldPartition::ECVF_Runtime_ReadOnly);
-
 bool UWorldPartition::bUseMakingVisibleTransactionRequests = false;
 FAutoConsoleVariableRef UWorldPartition::CVarUseMakingVisibleTransactionRequests(
 	TEXT("wp.Runtime.UseMakingVisibleTransactionRequests"),
@@ -1075,18 +1066,6 @@ bool UWorldPartition::IsServerStreamingOutEnabled() const
 	return bCachedIsServerStreamingOutEnabled.Get(false);
 }
 
-bool UWorldPartition::IsServerDataLayerStreamingOutEnabled() const
-{
-	// Resolve once (we don't allow changing the state at runtime)
-	if (!bCachedIsServerDataLayerStreamingOutEnabled.IsSet())
-	{
-		UWorld* OwningWorld = GetWorld();
-		bCachedIsServerDataLayerStreamingOutEnabled = OwningWorld && OwningWorld->IsGameWorld() && IsServer() && UWorldPartition::bEnableServerDataLayerStreamingOut;
-	}
-
-	return bCachedIsServerDataLayerStreamingOutEnabled.Get(false);
-}
-
 bool UWorldPartition::UseMakingVisibleTransactionRequests() const
 {
 	// Resolve once (we don't allow changing the state at runtime)
diff --git a/Engine/Source/Runtime/Engine/Private/WorldPartition/WorldPartitionStreamingPolicy.cpp b/Engine/Source/Runtime/Engine/Private/WorldPartition/WorldPartitionStreamingPolicy.cpp
--- a/Engine/Source/Runtime/Engine/Private/WorldPartition/WorldPartitionStreamingPolicy.cpp
+++ b/Engine/Source/Runtime/Engine/Private/WorldPartition/WorldPartitionStreamingPolicy.cpp
@@ -46,6 +46,13 @@ static FAutoConsoleVariableRef CVarBlockOnSlowStreaming(
 	GBlockOnSlowStreaming,
 	TEXT("Set if streaming needs to block when to slow to catchup."));
 
+static FString GServerDisallowStreamingOutDataLayersString;
+static FAutoConsoleVariableRef CVarServerDisallowStreamingOutDataLayers(
+	TEXT("wp.Runtime.ServerDisallowStreamingOutDataLayers"),
+	GServerDisallowStreamingOutDataLayersString,
+	TEXT("Comma separated list of data layer names that aren't allowed to be unloaded or deactivated on the server"),
+	ECVF_ReadOnly);
+
 bool UWorldPartitionStreamingPolicy::IsUpdateOptimEnabled = false;
 FAutoConsoleVariableRef UWorldPartitionStreamingPolicy::CVarUpdateOptimEnabled(
 	TEXT("wp.Runtime.UpdateStreaming.EnableOptimization"),
@@ -353,6 +360,35 @@ bool UWorldPartitionStreamingPolicy::GetIntersectingCells(const TArray<FWorldPar
 	return true;
 }
 
+const TSet<FName>& UWorldPartitionStreamingPolicy::GetServerDisallowedStreamingOutDataLayers()
+{
+	if (!CachedServerDisallowStreamingOutDataLayers.IsSet())
+	{
+		TSet<FName> ServerDisallowStreamingOutDataLayers;
+
+		if (!GServerDisallowStreamingOutDataLayersString.IsEmpty())
+		{
+			TArray<FString> AllDLAssetsStrings;
+			GServerDisallowStreamingOutDataLayersString.ParseIntoArray(AllDLAssetsStrings, TEXT(","));
+
+			if (const UDataLayerSubsystem* DataLayerSubsystem = UWorld::GetSubsystem<UDataLayerSubsystem>(GetWorld()))
+			{
+				for (const FString& DataLayerAssetName : AllDLAssetsStrings)
+				{
+					if (const UDataLayerInstance* DataLayerInstance = DataLayerSubsystem->GetDataLayerInstanceFromAssetName(FName(DataLayerAssetName)))
+					{
+						ServerDisallowStreamingOutDataLayers.Add(DataLayerInstance->GetDataLayerFName());
+					}
+				}
+			}
+		}
+
+		CachedServerDisallowStreamingOutDataLayers = ServerDisallowStreamingOutDataLayers;
+	}
+	
+	return CachedServerDisallowStreamingOutDataLayers.GetValue();
+}
+
 void UWorldPartitionStreamingPolicy::UpdateStreamingState()
 {
 	TRACE_CPUPROFILER_EVENT_SCOPE(UWorldPartitionStreamingPolicy::UpdateStreamingState);
@@ -437,7 +473,7 @@ void UWorldPartitionStreamingPolicy::UpdateStreamingState()
 		if (bIsServer)
 		{
 			const bool bCanServerDeactivateOrUnloadCells = WorldPartition->IsServerStreamingOutEnabled();
-			const bool bCanServerDeactivateOrUnloadDataLayerCells = WorldPartition->IsServerDataLayerStreamingOutEnabled();
+			const TSet<FName>& ServerDisallowStreamingOutDataLayers = GetServerDisallowedStreamingOutDataLayers();
 
 			TRACE_CPUPROFILER_EVENT_SCOPE(UWorldPartitionStreamingPolicy::UpdateStreamingState_ServerUpdate);
 
@@ -468,10 +504,23 @@ void UWorldPartitionStreamingPolicy::UpdateStreamingState()
 			TSet<FName> EffectiveActiveDataLayerNames = DataLayerSubsystem->GetEffectiveActiveDataLayerNames();
 			TSet<FName> EffectiveLoadedDataLayerNames = DataLayerSubsystem->GetEffectiveLoadedDataLayerNames();
 
-			auto AddServerFrameCell = [this, bCanServerDeactivateOrUnloadDataLayerCells, &EffectiveLoadedDataLayerNames, &EffectiveActiveDataLayerNames](const UWorldPartitionRuntimeCell* Cell)
+			auto CanServerDeactivateOrUnloadDataLayerCell = [&ServerDisallowStreamingOutDataLayers](const UWorldPartitionRuntimeCell* Cell)
+			{
+				if (Cell->HasDataLayers())
+				{
+					if (Cell->HasAnyDataLayer(ServerDisallowStreamingOutDataLayers))
+					{
+						return false;
+					}
+				}
+
+				return true;
+			};
+
+			auto AddServerFrameCell = [this, CanServerDeactivateOrUnloadDataLayerCell, &EffectiveLoadedDataLayerNames, &EffectiveActiveDataLayerNames](const UWorldPartitionRuntimeCell* Cell)
 			{
-				// Keep Data Layer cells is their current state if server cannot deactivate/unload data layer cells
-				if (!bCanServerDeactivateOrUnloadDataLayerCells && Cell->HasDataLayers())
+				// Keep Data Layer cells in their current state if server cannot deactivate/unload data layer cells
+				if (!CanServerDeactivateOrUnloadDataLayerCell(Cell))
 				{
 					if (ActivatedCells.Contains(Cell))
 					{
diff --git a/Engine/Source/Runtime/Engine/Public/WorldPartition/WorldPartition.h b/Engine/Source/Runtime/Engine/Public/WorldPartition/WorldPartition.h
--- a/Engine/Source/Runtime/Engine/Public/WorldPartition/WorldPartition.h
+++ b/Engine/Source/Runtime/Engine/Public/WorldPartition/WorldPartition.h
@@ -255,7 +255,6 @@ public:
 	bool IsServer() const;
 	bool IsServerStreamingEnabled() const;
 	bool IsServerStreamingOutEnabled() const;
-	bool IsServerDataLayerStreamingOutEnabled() const;
 	bool UseMakingVisibleTransactionRequests() const;
 	bool UseMakingInvisibleTransactionRequests() const;
 
@@ -360,7 +359,6 @@ private:
 	mutable TOptional<bool> bCachedUseMakingVisibleTransactionRequests;
 	mutable TOptional<bool> bCachedIsServerStreamingEnabled;
 	mutable TOptional<bool> bCachedIsServerStreamingOutEnabled;
-	mutable TOptional<bool> bCachedIsServerDataLayerStreamingOutEnabled;
 
 	UPROPERTY(Transient)
 	mutable TObjectPtr<UWorldPartitionStreamingPolicy> StreamingPolicy;
@@ -384,12 +382,10 @@ private:
 
 	static int32 EnableServerStreaming;
 	static bool bEnableServerStreamingOut;
-	static bool bEnableServerDataLayerStreamingOut;
 	static bool bUseMakingVisibleTransactionRequests;
 	static bool bUseMakingInvisibleTransactionRequests;
 	static FAutoConsoleVariableRef CVarEnableServerStreaming;
 	static FAutoConsoleVariableRef CVarEnableServerStreamingOut;
-	static FAutoConsoleVariableRef CVarEnableServerDataLayerStreamingOut;
 	static FAutoConsoleVariableRef CVarUseMakingVisibleTransactionRequests;
 	static FAutoConsoleVariableRef CVarUseMakingInvisibleTransactionRequests;
 
diff --git a/Engine/Source/Runtime/Engine/Public/WorldPartition/WorldPartitionStreamingPolicy.h b/Engine/Source/Runtime/Engine/Public/WorldPartition/WorldPartitionStreamingPolicy.h
--- a/Engine/Source/Runtime/Engine/Public/WorldPartition/WorldPartitionStreamingPolicy.h
+++ b/Engine/Source/Runtime/Engine/Public/WorldPartition/WorldPartitionStreamingPolicy.h
@@ -127,6 +127,8 @@ private:
 	uint32 ComputeStreamingSourceHash(const FWorldPartitionStreamingSource& Source) const;
 	int32 ComputeServerStreamingEnabledEpoch() const;
 
+	const TSet<FName>& GetServerDisallowedStreamingOutDataLayers();
+
 	static bool IsUpdateStreamingOptimEnabled();
 
 	// CVars to control update optimization
@@ -148,6 +150,8 @@ private:
 	uint32 UpdateStreamingHash;
 	uint32 UpdateStreamingStateCalls;
 
+	TOptional<TSet<FName>> CachedServerDisallowStreamingOutDataLayers;
+
 	EWorldPartitionStreamingPerformance StreamingPerformance;
 #if !UE_BUILD_SHIPPING
 	void UpdateDebugCellsStreamingPriority(const TSet<const UWorldPartitionRuntimeCell*>& ActivateStreamingCells, const TSet<const UWorldPartitionRuntimeCell*>& LoadStreamingCells);
