Network insights
- Fixes to enabling capture of demorecordings using networkinsights

diff --git a/Engine/Source/Developer/TraceServices/Private/Analyzers/NetTraceAnalyzer.cpp b/Engine/Source/Developer/TraceServices/Private/Analyzers/NetTraceAnalyzer.cpp
--- a/Engine/Source/Developer/TraceServices/Private/Analyzers/NetTraceAnalyzer.cpp
+++ b/Engine/Source/Developer/TraceServices/Private/Analyzers/NetTraceAnalyzer.cpp
@@ -751,7 +751,7 @@ void FNetTraceAnalyzer::HandleConnectionClosedEvent(const FOnEventContext& Conte
 	else
 	{
 		// Incomplete trace?  Ignore?
-		check(false);
+		UE_LOG(LogNetTrace, Warning, TEXT("Connection %d is missing"), ConnectionId);
 	}
 }
 
diff --git a/Engine/Source/Runtime/Engine/Private/NetConnection.cpp b/Engine/Source/Runtime/Engine/Private/NetConnection.cpp
--- a/Engine/Source/Runtime/Engine/Private/NetConnection.cpp
+++ b/Engine/Source/Runtime/Engine/Private/NetConnection.cpp
@@ -567,6 +567,13 @@ void UNetConnection::InitConnection(UNetDriver* InDriver, EConnectionState InSta
 
 	InitChannelData();
 
+	// Cache instance id
+#if UE_NET_TRACE_ENABLED
+	NetTraceId = Driver->GetNetTraceId();
+#endif
+
+	SetConnectionId(InDriver->AllocateConnectionId());
+
 	// We won't be sending any packets, so use a default size
 	MaxPacket = (InMaxPacket == 0 || InMaxPacket > MAX_PACKET_SIZE) ? MAX_PACKET_SIZE : InMaxPacket;
 	PacketOverhead = 0;
@@ -600,6 +607,9 @@ void UNetConnection::InitConnection(UNetDriver* InDriver, EConnectionState InSta
 	auto PackageMapClient = NewObject<UPackageMapClient>(this);
 	PackageMapClient->Initialize(this, Driver->GuidCache);
 	PackageMap = PackageMapClient;
+
+	UE_NET_TRACE_CONNECTION_CREATED(NetTraceId, GetConnectionId());
+	UE_NET_TRACE_CONNECTION_STATE_UPDATED(NetTraceId, GetConnectionId(), static_cast<uint8>(GetConnectionState()));
 }
 
 void UNetConnection::InitHandler()
@@ -3777,14 +3787,18 @@ void UNetConnection::PrepareWriteBitsToSendBuffer(const int32 SizeInBits, const
 		FlushNet();
 	}
 
+#if UE_NET_TRACE_ENABLED
+	// If tracing is enabled setup the NetTraceCollector for outgoing data
+	if (SendBuffer.GetNumBits() == 0)
+	{
+		OutTraceCollector = UE_NET_TRACE_CREATE_COLLECTOR(ENetTraceVerbosity::Trace);
+	}
+#endif
+
 	// If this is the start of the queue, make sure to add the packet id
 	if ( SendBuffer.GetNumBits() == 0 && !IsInternalAck() )
 	{
-#if UE_NET_TRACE_ENABLED
-		// If tracing is enabled setup the NetTraceCollector for outgoing data
-		OutTraceCollector = UE_NET_TRACE_CREATE_COLLECTOR(ENetTraceVerbosity::Trace);
 		UE_NET_TRACE_SCOPE(PacketHeaderAndInfo, SendBuffer, OutTraceCollector, ENetTraceVerbosity::Trace);
-#endif
 
 		// Write Packet Header, before sending the packet we will go back and rewrite the data
 		WritePacketHeader(SendBuffer);
diff --git a/Engine/Source/Runtime/Engine/Private/NetDriver.cpp b/Engine/Source/Runtime/Engine/Private/NetDriver.cpp
--- a/Engine/Source/Runtime/Engine/Private/NetDriver.cpp
+++ b/Engine/Source/Runtime/Engine/Private/NetDriver.cpp
@@ -6289,8 +6289,8 @@ void UNetDriver::InitNetTraceId()
 	if (NetTraceId == NetTraceInvalidGameInstanceId)
 	{
 		// We just need to make sure that all active NetDriver`s running in the same game instance uses different NetTraceId`s.
-		static uint8 CurrentNetTraceId = 255;
-		NetTraceId = ++CurrentNetTraceId;
+		static uint8 CurrentNetTraceId = 0;
+		NetTraceId = 128 + ((++CurrentNetTraceId) & 127);
 	}
 }
 
diff --git a/Engine/Source/Runtime/Engine/Private/RepLayout.cpp b/Engine/Source/Runtime/Engine/Private/RepLayout.cpp
--- a/Engine/Source/Runtime/Engine/Private/RepLayout.cpp
+++ b/Engine/Source/Runtime/Engine/Private/RepLayout.cpp
@@ -2957,6 +2957,8 @@ static FORCEINLINE void WritePropertyHandle_BackwardsCompatible(
 	uint32			NetFieldExportHandle,
 	bool			bDoChecksum)
 {
+	UE_NET_TRACE_SCOPE(PropertyHandle, Writer, GetTraceCollector(Writer), ENetTraceVerbosity::Trace);
+
 	const int NumStartingBits = Writer.GetNumBits();
 
 	Writer.SerializeIntPacked(NetFieldExportHandle);
@@ -3053,6 +3055,8 @@ void FRepLayout::SendProperties_BackwardsCompatible_r(
 
 	FNetBitWriter TempWriter(Writer.PackageMap, 0);
 
+	UE_NET_TRACE_SCOPE(Properties, Writer, GetTraceCollector(Writer), ENetTraceVerbosity::Trace);
+
 	while (HandleIterator.NextHandle())
 	{
 		const FRepLayoutCmd& Cmd = Cmds[HandleIterator.CmdIndex];
@@ -3079,6 +3083,8 @@ void FRepLayout::SendProperties_BackwardsCompatible_r(
 
 		WritePropertyHandle_BackwardsCompatible(Writer, HandleIterator.CmdIndex + 1, bDoChecksum);
 
+		UE_NET_TRACE_DYNAMIC_NAME_SCOPE(Cmd.Property->GetFName(), Writer, GetTraceCollector(Writer), ENetTraceVerbosity::Trace);
+
 		if (Cmd.Type == ERepLayoutCmdType::DynamicArray)
 		{
 			const FScriptArray* Array = (FScriptArray *)Data.Data;
