Refactored DeriveNormalZ_Function to use DeriveNormalZ builtin function instead of duplicating its logic.

#udn https://udn.unrealengine.com/s/question/0D54z00008anTs9CAE/normalderivez-returning-nans?fromCase=1

diff --git a/Engine/Source/Runtime/Engine/Private/Materials/MaterialExpressions.cpp b/Engine/Source/Runtime/Engine/Private/Materials/MaterialExpressions.cpp
--- a/Engine/Source/Runtime/Engine/Private/Materials/MaterialExpressions.cpp
+++ b/Engine/Source/Runtime/Engine/Private/Materials/MaterialExpressions.cpp
@@ -13105,7 +13105,7 @@ int32 UMaterialExpressionDeriveNormalZ::Compile(class FMaterialCompiler* Compile
 		return Compiler->Errorf(TEXT("Missing input normal xy vector whose z should be derived."));
 	}
 
-	//z = sqrt(1 - ( x * x + y * y));
+	// z = sqrt(saturate(1 - ( x * x + y * y)));
 	int32 InputVector = Compiler->ForceCast(InXY.Compile(Compiler), MCT_Float2);
 	int32 DotResult = Compiler->Dot(InputVector, InputVector);
 	int32 InnerResult = Compiler->Sub(Compiler->Constant(1), DotResult);
