Fix regression in audio compression being always synchrone because we were pulling on GetNumChunks() even when unnecessary

#rnx

diff --git a/Engine/Source/Runtime/Engine/Private/SoundWave.cpp b/Engine/Source/Runtime/Engine/Private/SoundWave.cpp
--- a/Engine/Source/Runtime/Engine/Private/SoundWave.cpp
+++ b/Engine/Source/Runtime/Engine/Private/SoundWave.cpp
@@ -1678,13 +1678,16 @@ void USoundWave::PostLoad()
 				// if a sound class defined our loading behavior as something other than Retain and our cvar default is to retain, we need to release our handle.
 				ReleaseCompressedAudio();
 
-				const bool bHasMultipleChunks = GetNumChunks() > 1;
 				bool bShouldPrime = (ActualLoadingBehavior == ESoundWaveLoadingBehavior::PrimeOnLoad);
 				bShouldPrime |= (GIsEditor && (ActualLoadingBehavior == ESoundWaveLoadingBehavior::RetainOnLoad)); // treat this scenario like PrimeOnLoad
 
-				if (bShouldPrime && bHasMultipleChunks)
+				if (bShouldPrime)
 				{
-					IStreamingManager::Get().GetAudioStreamingManager().RequestChunk(Proxy, 1, [](EAudioChunkLoadResult) {});
+					const bool bHasMultipleChunks = GetNumChunks() > 1;
+					if (bHasMultipleChunks)
+					{
+						IStreamingManager::Get().GetAudioStreamingManager().RequestChunk(Proxy, 1, [](EAudioChunkLoadResult) {});
+					}
 				}
 			}
 
