Remove submix options from audio device properties, all we have to do is play the sound

#rnx
#pf skip

diff --git a/Engine/Source/Runtime/Engine/Classes/GameFramework/InputDeviceProperties.h b/Engine/Source/Runtime/Engine/Classes/GameFramework/InputDeviceProperties.h
--- a/Engine/Source/Runtime/Engine/Classes/GameFramework/InputDeviceProperties.h
+++ b/Engine/Source/Runtime/Engine/Classes/GameFramework/InputDeviceProperties.h
@@ -485,25 +485,9 @@ struct FAudioBasedVibrationData
 
 	FAudioBasedVibrationData();
 	
-	/** The sound to play on the gamepad */
+	/** The sound to play on the gamepad. Make sure the set the sound's submix sends to the gamepad audio and vibration endpoints! */
 	UPROPERTY(EditAnywhere, BlueprintReadWrite, Category = "DeviceProperty")
 	TObjectPtr<USoundBase> Sound;
-
-	/** The audio endpoint that the sound will be played on for gamepad vibration. */
-	UPROPERTY(EditAnywhere, BlueprintReadWrite, Category = "Vibration Audio")
-	TObjectPtr<UEndpointSubmix> VibrationEndpoint;
-
-	/** The audio endpoint that the sound will be played on for gamepad audio. */
-	UPROPERTY(EditAnywhere, BlueprintReadWrite, Category = "Gamepad Audio")
-	TObjectPtr<UEndpointSubmix> GamepadAudioEndpoint;
-
-	/** The amount of sound to send to the Vibration endpoint submix */
-	UPROPERTY(EditAnywhere, BlueprintReadWrite, Category = "Vibration Audio", meta = (UIMin = "0.0"))
-	float VibrationSoundLevel;
-
-	/** The amount of sound to send to the Gamepad Audio endpoint submix */
-	UPROPERTY(EditAnywhere, BlueprintReadWrite, Category = "Gamepad Audio", meta = (UIMin = "0.0"))
-	float GamepadAudioSoundLevel;
 };
 
 /**
diff --git a/Engine/Source/Runtime/Engine/Private/GameFramework/InputDeviceProperties.cpp b/Engine/Source/Runtime/Engine/Private/GameFramework/InputDeviceProperties.cpp
--- a/Engine/Source/Runtime/Engine/Private/GameFramework/InputDeviceProperties.cpp
+++ b/Engine/Source/Runtime/Engine/Private/GameFramework/InputDeviceProperties.cpp
@@ -428,43 +428,17 @@ int32 UInputDeviceTriggerVibrationProperty::GetVibrationAmplitudeValue(const FDe
 
 FAudioBasedVibrationData::FAudioBasedVibrationData()
 	: Sound(nullptr)
-	, VibrationEndpoint(nullptr)
-	, GamepadAudioEndpoint(nullptr)
-	, VibrationSoundLevel(1.0f)
-	, GamepadAudioSoundLevel(1.0f)
 {
 }
 
 void UInputDeviceAudioBasedVibrationProperty::EvaluateDeviceProperty_Implementation(const FPlatformUserId PlatformUser, const FInputDeviceId DeviceId, const float DeltaTime, const float Duration)
 {
-	if (const FAudioBasedVibrationData* DataToUse = GetRelevantData(PlatformUser, DeviceId))
-	{
-		// Set the endpoints to use on this sound
-		if (DataToUse->Sound)
-		{
-			if (DataToUse->VibrationEndpoint)
-			{
-				FSoundSubmixSendInfo SendInfo;
-				SendInfo.SoundSubmix = DataToUse->VibrationEndpoint;
-				SendInfo.SendLevel = DataToUse->VibrationSoundLevel;
-				DataToUse->Sound->SoundSubmixSends.Emplace(MoveTemp(SendInfo));
-			}
-
-			if (DataToUse->GamepadAudioEndpoint)
-			{
-				FSoundSubmixSendInfo SendInfo;
-				SendInfo.SoundSubmix = DataToUse->GamepadAudioEndpoint;
-				SendInfo.SendLevel = DataToUse->GamepadAudioSoundLevel;
-				DataToUse->Sound->SoundSubmixSends.Emplace(MoveTemp(SendInfo));
-			}
-		}
-	}
+	// There is nothing to do here, the endpoints are set on the sound asset.
 }
 
 void UInputDeviceAudioBasedVibrationProperty::ApplyDeviceProperty(const FPlatformUserId UserId, const FInputDeviceId DeviceId)
 {
-	// Audio Endpoints can already play vibrations and audio on an input device, if the device implements
-	// the interface for it. We can just do that! Woot
+	// This sound should have the submix sends set that it wants, all we have to do is play it.
 	if (const FAudioBasedVibrationData* DataToUse = GetRelevantData(UserId, DeviceId))
 	{
 		if (DataToUse->Sound)
@@ -474,7 +448,7 @@ void UInputDeviceAudioBasedVibrationProperty::ApplyDeviceProperty(const FPlatfor
 			{
 				// The sound endpoints will have been populated above in the Evaluate function
 				PC->ClientPlaySound(DataToUse->Sound);
-			}	
+			}
 		}		
 	}
 }
