Fix Crash when generating streaming and there is a non-editor world present.

diff --git a/Engine/Source/Runtime/Engine/Private/WorldPartition/ContentBundle/ContentBundleContainer.cpp b/Engine/Source/Runtime/Engine/Private/WorldPartition/ContentBundle/ContentBundleContainer.cpp
--- a/Engine/Source/Runtime/Engine/Private/WorldPartition/ContentBundle/ContentBundleContainer.cpp
+++ b/Engine/Source/Runtime/Engine/Private/WorldPartition/ContentBundle/ContentBundleContainer.cpp
@@ -52,9 +52,12 @@ void FContentBundleContainer::Initialize()
 	RegisterContentBundleClientEvents();
 
 #if WITH_EDITOR
-	UWorldPartition* WorldPartition = GetInjectedWorld()->GetWorldPartition();
-	WorldPartition->OnPreGenerateStreaming.AddRaw(this, &FContentBundleContainer::OnPreGenerateStreaming);
-	WorldPartition->OnBeginCook.AddRaw(this, &FContentBundleContainer::OnBeginCook);
+	if (UseEditorContentBundle())
+	{
+		UWorldPartition* WorldPartition = GetInjectedWorld()->GetWorldPartition();
+		WorldPartition->OnPreGenerateStreaming.AddRaw(this, &FContentBundleContainer::OnPreGenerateStreaming);
+		WorldPartition->OnBeginCook.AddRaw(this, &FContentBundleContainer::OnBeginCook);
+	}
 #endif
 }
 
@@ -67,9 +70,12 @@ void FContentBundleContainer::Deinitialize()
 		UE_LOG(LogContentBundle, Log, TEXT("%s Deleting container."), *ContentBundle::Log::MakeDebugInfoString(*this));
 
 #if WITH_EDITOR
-		UWorldPartition* WorldPartition = GetInjectedWorld()->GetWorldPartition();
-		WorldPartition->OnPreGenerateStreaming.RemoveAll(this);
-		WorldPartition->OnBeginCook.RemoveAll(this);
+		if (UseEditorContentBundle())
+		{
+			UWorldPartition* WorldPartition = GetInjectedWorld()->GetWorldPartition();
+			WorldPartition->OnPreGenerateStreaming.RemoveAll(this);
+			WorldPartition->OnBeginCook.RemoveAll(this);
+		}
 #endif
 	}
 
