Fix nullptr when running a fast cook build

diff --git a/Engine/Source/Runtime/Engine/Private/PSOPrecache.cpp b/Engine/Source/Runtime/Engine/Private/PSOPrecache.cpp
--- a/Engine/Source/Runtime/Engine/Private/PSOPrecache.cpp
+++ b/Engine/Source/Runtime/Engine/Private/PSOPrecache.cpp
@@ -395,7 +395,11 @@ void FMaterialPSOPrecacheCollectionTask::DoTask(ENamedThreads::Type CurrentThrea
 	FTaskTagScope ParallelGTScope(ETaskTag::EParallelGameThread);
 
 	// Collect pso and start the async compiles
-	FPSOPrecacheRequestResultArray PrecacheResults = PrecacheParams.Material->GetGameThreadShaderMap()->CollectPSOs(PrecacheParams);
+	FPSOPrecacheRequestResultArray PrecacheResults;
+	if (PrecacheParams.Material->GetGameThreadShaderMap())
+	{
+		PrecacheResults = PrecacheParams.Material->GetGameThreadShaderMap()->CollectPSOs(PrecacheParams);
+	}
 
 	// Won't touch the material interface anymore - PSO compile jobs take refs to all RHI resources while creating the task
 	TGraphTask<FMaterialInterfaceReleaseTask>::CreateTask().ConstructAndDispatchWhenReady(MaterialInterface);
