Add test code to help debug rare shipping editor crash on shutdown in FSubsystemCollectionBase
Copying class name to stack so it will show up in minidump uploaded to crash reporter

diff --git a/Engine/Source/Runtime/Engine/Private/Subsystems/SubsystemCollection.cpp b/Engine/Source/Runtime/Engine/Private/Subsystems/SubsystemCollection.cpp
--- a/Engine/Source/Runtime/Engine/Private/Subsystems/SubsystemCollection.cpp
+++ b/Engine/Source/Runtime/Engine/Private/Subsystems/SubsystemCollection.cpp
@@ -48,6 +48,14 @@ FSubsystemCollectionBase::FSubsystemCollectionBase(UClass* InBaseType)
 
 USubsystem* FSubsystemCollectionBase::GetSubsystemInternal(UClass* SubsystemClass) const
 {
+#if WITH_EDITOR && UE_BUILD_SHIPPING
+	TStringBuilder<200> DebugSubsystemClassName;
+	if (SubsystemClass && IsEngineExitRequested())
+	{
+		SubsystemClass->GetFName().AppendString(DebugSubsystemClassName);
+	}
+#endif
+
 	USubsystem* SystemPtr = SubsystemMap.FindRef(SubsystemClass);
 
 	if (SystemPtr)
