Fix for animation notifies on sequences are not triggered from montages in linked anim instances. Specifically, when slot node is located in the linked anim instance and not the main instance.

diff --git a/Engine/Source/Runtime/Engine/Private/Animation/AnimInstanceProxy.cpp b/Engine/Source/Runtime/Engine/Private/Animation/AnimInstanceProxy.cpp
--- a/Engine/Source/Runtime/Engine/Private/Animation/AnimInstanceProxy.cpp
+++ b/Engine/Source/Runtime/Engine/Private/Animation/AnimInstanceProxy.cpp
@@ -598,7 +598,38 @@ void FAnimInstanceProxy::PostUpdate(UAnimInstance* InAnimInstance) const
 		}
 	}
 #endif
+	
+	// Copy slot information to main instance if we are using the main instance's montage evaluation data.
+	// Note that linked anim instance's proxies PostUpdate() will be called before the main instance's proxy PostUpdate().
+	if (bUseMainInstanceMontageEvaluationData && GetMainInstanceProxy() && GetMainInstanceProxy() != this)
+	{
+		FAnimInstanceProxy& MainProxy = *GetMainInstanceProxy();
+		
+		for (const TTuple<FName, int> & LinkedSlotTrackerPair : SlotNameToTrackerIndex)
+		{
+			const int* MainTrackerIndexPtr = MainProxy.SlotNameToTrackerIndex.Find(LinkedSlotTrackerPair.Key);
+
+			// Ensure slot tracker exists for main instance.
+			if (!MainTrackerIndexPtr)
+			{
+				MainProxy.RegisterSlotNodeWithAnimInstance(LinkedSlotTrackerPair.Key);
+				MainTrackerIndexPtr = MainProxy.SlotNameToTrackerIndex.Find(LinkedSlotTrackerPair.Key);
+			}
 
+			// Update slot information for main instance.
+			{
+				const FMontageActiveSlotTracker & LinkedTracker = SlotWeightTracker[GetBufferReadIndex()][LinkedSlotTrackerPair.Value];
+				FMontageActiveSlotTracker& MainTracker = MainProxy.SlotWeightTracker[MainProxy.GetBufferWriteIndex()][*MainTrackerIndexPtr];
+				
+				MainTracker.MontageLocalWeight = FMath::Max(MainTracker.MontageLocalWeight, LinkedTracker.MontageLocalWeight);
+				MainTracker.NodeGlobalWeight = FMath::Max(MainTracker.NodeGlobalWeight, LinkedTracker.NodeGlobalWeight);
+				
+				MainTracker.bIsRelevantThisTick = MainTracker.bIsRelevantThisTick || LinkedTracker.bIsRelevantThisTick;
+				MainTracker.bWasRelevantOnPreviousTick = MainTracker.bWasRelevantOnPreviousTick || LinkedTracker.bWasRelevantOnPreviousTick;
+			}
+		}
+	}
+	
 	InAnimInstance->NotifyQueue.Append(NotifyQueue);
 	InAnimInstance->NotifyQueue.ApplyMontageNotifies(*this);
 
