Hololens SkyAtmosphereComponentCrash.
-Got a memory access violation, writing to protected memory, on this line on hololens.  Possibly this is a compiler bug, but the code is slightly wierd.  We don't need to call the FColor constructor here.

diff --git a/Engine/Source/Runtime/Engine/Private/Components/SkyAtmosphereComponent.cpp b/Engine/Source/Runtime/Engine/Private/Components/SkyAtmosphereComponent.cpp
--- a/Engine/Source/Runtime/Engine/Private/Components/SkyAtmosphereComponent.cpp
+++ b/Engine/Source/Runtime/Engine/Private/Components/SkyAtmosphereComponent.cpp
@@ -55,9 +55,9 @@ USkyAtmosphereComponent::USkyAtmosphereComponent(const FObjectInitializer& Objec
 	RayleighScatteringScale = RayleightScatteringRaw.B;
 	RayleighExponentialDistribution = EarthRayleighScaleHeight;
 
-	MieScattering = FColor(FColor::White);
+	MieScattering = FColor::White;
 	MieScatteringScale = 0.003996f;
-	MieAbsorption = FColor(FColor::White);
+	MieAbsorption = FColor::White;
 	MieAbsorptionScale = 0.000444f;
 	MieAnisotropy = 0.8f;
 	MieExponentialDistribution = EarthMieScaleHeight;
