Add a check to prevent misuse of the texture compiling manager

#rnx

diff --git a/Engine/Source/Runtime/Engine/Private/TextureCompiler.cpp b/Engine/Source/Runtime/Engine/Private/TextureCompiler.cpp
--- a/Engine/Source/Runtime/Engine/Private/TextureCompiler.cpp
+++ b/Engine/Source/Runtime/Engine/Private/TextureCompiler.cpp
@@ -195,6 +195,8 @@ void FTextureCompilingManager::UpdateCompilationNotification()
 
 void FTextureCompilingManager::PostCompilation(UTexture* Texture)
 {
+	TGuardValue PostCompilationGuard(bIsRoutingPostCompilation, true);
+
 	check(IsInGameThread());
 	TRACE_CPUPROFILER_EVENT_SCOPE(FTextureCompilingManager::PostCompilation);
 
@@ -238,8 +240,11 @@ int32 FTextureCompilingManager::GetNumRemainingAssets() const
 
 void FTextureCompilingManager::AddTextures(TArrayView<UTexture* const> InTextures)
 {
-	TRACE_CPUPROFILER_EVENT_SCOPE(FTextureCompilingManager::AddTextures)
 	check(IsInGameThread());
+	checkf(bIsRoutingPostCompilation == false,
+		TEXT("Registering a texture to the compile manager from inside a texture postcompilation is not supported and usually indicate that the previous async operation wasn't completed (i.e. missing call to PreEditChange) before modifying a texture property."));
+
+	TRACE_CPUPROFILER_EVENT_SCOPE(FTextureCompilingManager::AddTextures)
 
 	// Register new textures after ProcessTextures to avoid
 	// potential reentrant calls to CreateResource on the
diff --git a/Engine/Source/Runtime/Engine/Public/TextureCompiler.h b/Engine/Source/Runtime/Engine/Public/TextureCompiler.h
--- a/Engine/Source/Runtime/Engine/Public/TextureCompiler.h
+++ b/Engine/Source/Runtime/Engine/Public/TextureCompiler.h
@@ -102,6 +102,7 @@ private:
 
 	double LastReschedule = 0.0f;
 	bool bHasShutdown = false;
+	bool bIsRoutingPostCompilation = false;
 	TArray<TSet<TWeakObjectPtr<UTexture>>> RegisteredTextureBuckets;
 	FAsyncCompilationNotification Notification;
 
