Fix crash in shared linked anim layer module when we don't have a skeletal mesh

#tests PIE
#rnx

diff --git a/Engine/Source/Runtime/Engine/Private/Animation/AnimSubsystem_SharedLinkedAnimLayers.cpp b/Engine/Source/Runtime/Engine/Private/Animation/AnimSubsystem_SharedLinkedAnimLayers.cpp
--- a/Engine/Source/Runtime/Engine/Private/Animation/AnimSubsystem_SharedLinkedAnimLayers.cpp
+++ b/Engine/Source/Runtime/Engine/Private/Animation/AnimSubsystem_SharedLinkedAnimLayers.cpp
@@ -118,7 +118,11 @@ FAnimSubsystem_SharedLinkedAnimLayers* FAnimSubsystem_SharedLinkedAnimLayers::Ge
 		return nullptr;
 	}
 #endif
-	check(SkelMesh && SkelMesh->GetAnimInstance());
+	if (!SkelMesh)
+	{
+		return nullptr;
+	}
+	check(SkelMesh->GetAnimInstance());
 	return SkelMesh->GetAnimInstance()->FindSubsystem<FAnimSubsystem_SharedLinkedAnimLayers>();
 }
 
