Implemented support for importing animated OpenVDB sequences using the new import dialog. SVTViewerComponent can now properly play back animated SVT. Fixed warped looking artifact when raymarching non-cubic SVT volumes. Fixed SVT cooking not working. Slightly restructured and cleaned up SparseVolumeTexture.cpp and SparseVolumeTextureOpenVDBUtility.cpp. Also refactored SparseVolumeTextureFactory to share as much code as possible between static and animated SVT, now that they both use the same import UI.

#rnx

diff --git a/Engine/Shaders/Private/SparseVolumeTexture/VisualizeSparseVolumeTexture.usf b/Engine/Shaders/Private/SparseVolumeTexture/VisualizeSparseVolumeTexture.usf
--- a/Engine/Shaders/Private/SparseVolumeTexture/VisualizeSparseVolumeTexture.usf
+++ b/Engine/Shaders/Private/SparseVolumeTexture/VisualizeSparseVolumeTexture.usf
@@ -40,11 +40,12 @@ bool Slabs(float3 p0, float3 p1, float3 rayOrigin, float3 invRaydir, out float o
 	return maxtmin <= mintmax;
 }
 
-float3 VolumeBoundMinWorld;
-float3 VolumeBoundMaxWorld;
 float4 WorldToLocal0;
 float4 WorldToLocal1;
 float4 WorldToLocal2;
+float3 WorldToLocalNoScale0;
+float3 WorldToLocalNoScale1;
+float3 WorldToLocalNoScale2;
 uint ComponentToVisualize;
 
 
@@ -89,7 +90,7 @@ void VisualizeSparseVolumeTexturePS(
 	float TMin = 0.0f;
 	float TMax = 0.0f;
 	float Transmittance = 1.0f;
-	if (Slabs(-1.0f, 1.0f, RayOrig, 1.0 / RayDir, TMin, TMax) && TMax > 0.0f)
+	if (Slabs(-1.0f, 1.0f, RayOrig, 1.0f / RayDir, TMin, TMax) && TMax > 0.0f)
 	{
 		Transmittance = 0.5; // show the bounding box
 
@@ -111,9 +112,15 @@ void VisualizeSparseVolumeTexturePS(
 		if (P.y > (VolumeResolution.y)) P.y = (VolumeResolution.y);
 		if (P.z > (VolumeResolution.z)) P.z = (VolumeResolution.z);
 
+		// DDA requires the ray direction not to be non-(uniformly) scaled.
+		float3 RayDirNoScale;
+		RayDirNoScale.x = dot(WorldToLocalNoScale0, RayDirWorld);
+		RayDirNoScale.y = dot(WorldToLocalNoScale1, RayDirWorld);
+		RayDirNoScale.z = dot(WorldToLocalNoScale2, RayDirWorld);
+
 		// Amanatides 3D-DDA data preparation
-		float3 stepSign = sign(RayDir);
-		float3 tDelta = abs(1.0 / RayDir);
+		float3 stepSign = sign(RayDirNoScale);
+		float3 tDelta = abs(1.0f / RayDirNoScale);
 		float3 tMax = float3(0.0, 0.0, 0.0);
 		float3 refPoint = floor(P);
 		tMax.x = stepSign.x > 0.0 ? refPoint.x + 1.0 - P.x : P.x - refPoint.x;
diff --git a/Engine/Source/Editor/SparseVolumeTexture/Private/OpenVDBImportWindow.cpp b/Engine/Source/Editor/SparseVolumeTexture/Private/OpenVDBImportWindow.cpp
--- a/Engine/Source/Editor/SparseVolumeTexture/Private/OpenVDBImportWindow.cpp
+++ b/Engine/Source/Editor/SparseVolumeTexture/Private/OpenVDBImportWindow.cpp
@@ -98,25 +98,6 @@ void SOpenVDBImportWindow::Construct(const FArguments& InArgs)
 			+ SVerticalBox::Slot()
 			.AutoHeight()
 			.Padding(2)
-			[
-				SNew(STextBlock)
-				.Text(LOCTEXT("OpenVDBImportWindow_FileInfo", "Source File Grid Info"))
-			]
-			+ SVerticalBox::Slot()
-			.AutoHeight()
-			.Padding(2)
-			[
-				SNew(SBorder)
-				.Padding(FMargin(3))
-				.BorderImage(FAppStyle::GetBrush("ToolPanel.GroupBorder"))
-				[
-					SNew(STextBlock)
-					.Text(FText::FromString(InArgs._FileInfoString))
-				]
-			]
-			+ SVerticalBox::Slot()
-			.AutoHeight()
-			.Padding(2)
 			[
 				SNew(SUniformGridPanel)
 				.SlotPadding(2)
@@ -144,6 +125,31 @@ void SOpenVDBImportWindow::Construct(const FArguments& InArgs)
 		+ SVerticalBox::Slot()
 		.AutoHeight()
 		.Padding(2)
+		[
+			SNew(SHorizontalBox)
+			+ SHorizontalBox::Slot()
+			.VAlign(VAlign_Center)
+			.HAlign(HAlign_Right)
+			.AutoWidth()
+			.Padding(2.0f)
+			[
+				SNew(STextBlock)
+				.Text(LOCTEXT("ImportAsSequenceCheckBoxLabel", "Import Sequence"))
+			]
+			+ SHorizontalBox::Slot()
+			.VAlign(VAlign_Center)
+			.HAlign(HAlign_Right)
+			.AutoWidth()
+			.Padding(2.0f)
+			[
+				SAssignNew(ImportAsSequenceCheckBox, SCheckBox)
+				.ToolTipText(LOCTEXT("ImportAsSequenceCheckBoxTooltip", "Import multiple sequentially labeled .vdb files as a single animated SparseVirtualTexture sequence."))
+				.IsChecked(false)
+			]
+		]
+		+ SVerticalBox::Slot()
+		.AutoHeight()
+		.Padding(2)
 		[
 			SAssignNew(PackedDataAConfigurator, SOpenVDBPackedDataConfigurator)
 			.PackedData(PackedDataA)
@@ -161,6 +167,25 @@ void SOpenVDBImportWindow::Construct(const FArguments& InArgs)
 			.OpenVDBSupportedTargetFormats(OpenVDBSupportedTargetFormats)
 			.PackedDataName(LOCTEXT("OpenVDBImportWindow_PackedDataB", "Packed Data B"))
 		]
+		+ SVerticalBox::Slot()
+		.AutoHeight()
+		.Padding(2)
+		[
+			SNew(STextBlock)
+			.Text(LOCTEXT("OpenVDBImportWindow_FileInfo", "Source File Grid Info"))
+		]
+		+ SVerticalBox::Slot()
+		.AutoHeight()
+		.Padding(2)
+		[
+			SNew(SBorder)
+			.Padding(FMargin(3))
+			.BorderImage(FAppStyle::GetBrush("ToolPanel.GroupBorder"))
+			[
+				SNew(STextBlock)
+				.Text(FText::FromString(InArgs._FileInfoString))
+			]
+		]
 	);
 
 	SetDefaultGridAssignment();
@@ -228,6 +253,11 @@ bool SOpenVDBImportWindow::ShouldImport() const
 	return bShouldImport;
 }
 
+bool SOpenVDBImportWindow::ShouldImportAsSequence() const
+{
+	return ImportAsSequenceCheckBox->IsChecked();
+}
+
 EActiveTimerReturnType SOpenVDBImportWindow::SetFocusPostConstruct(double InCurrentTime, float InDeltaTime)
 {
 	if (ImportButton.IsValid())
@@ -265,7 +295,7 @@ FReply SOpenVDBImportWindow::OnResetToDefaultClick()
 
 FText SOpenVDBImportWindow::GetImportTypeDisplayText() const
 {
-	return LOCTEXT("OpenVDBImportWindow_ImportType", "Import Static OpenVDB");
+	return ShouldImportAsSequence() ? LOCTEXT("OpenVDBImportWindow_ImportTypeAnimated", "Import OpenVDB animation") : LOCTEXT("OpenVDBImportWindow_ImportTypeStatic", "Import static OpenVDB");
 }
 
 void SOpenVDBImportWindow::SetDefaultGridAssignment()
diff --git a/Engine/Source/Editor/SparseVolumeTexture/Private/OpenVDBImportWindow.h b/Engine/Source/Editor/SparseVolumeTexture/Private/OpenVDBImportWindow.h
--- a/Engine/Source/Editor/SparseVolumeTexture/Private/OpenVDBImportWindow.h
+++ b/Engine/Source/Editor/SparseVolumeTexture/Private/OpenVDBImportWindow.h
@@ -103,6 +103,7 @@ public:
 	FReply OnImport();
 	FReply OnCancel();
 	bool ShouldImport() const;
+	bool ShouldImportAsSequence() const;
 
 private:
 	FSparseVolumeRawSourcePackedData*					PackedDataA;
@@ -111,12 +112,13 @@ private:
 	TArray<TSharedPtr<ESparseVolumePackedDataFormat>>*	OpenVDBSupportedTargetFormats;
 	TSharedPtr<SOpenVDBPackedDataConfigurator>			PackedDataAConfigurator;
 	TSharedPtr<SOpenVDBPackedDataConfigurator>			PackedDataBConfigurator;
+	TSharedPtr<SCheckBox>								ImportAsSequenceCheckBox;
 	TSharedPtr<SButton>									ImportButton;
 	TWeakPtr<SWindow>									WidgetWindow;
 	bool												bShouldImport;
 
 	EActiveTimerReturnType SetFocusPostConstruct(double InCurrentTime, float InDeltaTime);
-	bool CanImport()  const;
+	bool CanImport() const;
 	FReply OnResetToDefaultClick();
 	FText GetImportTypeDisplayText() const;
 	void SetDefaultGridAssignment();
diff --git a/Engine/Source/Editor/SparseVolumeTexture/Private/SparseVolumeTextureFactory.cpp b/Engine/Source/Editor/SparseVolumeTexture/Private/SparseVolumeTextureFactory.cpp
--- a/Engine/Source/Editor/SparseVolumeTexture/Private/SparseVolumeTextureFactory.cpp
+++ b/Engine/Source/Editor/SparseVolumeTexture/Private/SparseVolumeTextureFactory.cpp
@@ -24,6 +24,161 @@
 
 DEFINE_LOG_CATEGORY_STATIC(LogSparseVolumeTextureFactory, Log, All);
 
+struct FOpenVDBPreviewData
+{
+	TArray<uint8> LoadedFile;
+	TArray<FOpenVDBGridInfo> GridInfo;
+	TArray<TSharedPtr<FOpenVDBGridComponentInfo>> GridComponentInfoPtrs;
+	FString FileInfoString;
+};
+
+static bool LoadOpenVDBPreviewData(const FString& Filename, FOpenVDBPreviewData* OutPreviewData)
+{
+	FOpenVDBPreviewData& Result = *OutPreviewData;
+	Result = {};
+
+	if (!FFileHelper::LoadFileToArray(Result.LoadedFile, *Filename))
+	{
+		UE_LOG(LogSparseVolumeTextureFactory, Error, TEXT("OpenVDB file could not be loaded: %s"), *Filename);
+		return false;
+	}
+	if (!GetOpenVDBGridInfo(Result.LoadedFile, true /*bCreateStrings*/, &Result.GridInfo))
+	{
+		UE_LOG(LogSparseVolumeTextureFactory, Error, TEXT("Failed to read OpenVDB file: %s"), *Filename);
+		return false;
+	}
+	if (Result.GridInfo.IsEmpty())
+	{
+		UE_LOG(LogSparseVolumeTextureFactory, Error, TEXT("OpenVDB file contains no grids: %s"), *Filename);
+		return false;
+	}
+
+	// We need a <None> option to leave channels empty
+	FOpenVDBGridComponentInfo NoneGridComponentInfo;
+	NoneGridComponentInfo.Index = INDEX_NONE;
+	NoneGridComponentInfo.ComponentIndex = INDEX_NONE;
+	NoneGridComponentInfo.Name = TEXT("<None>");
+	NoneGridComponentInfo.DisplayString = TEXT("<None>");
+
+	Result.GridComponentInfoPtrs.Add(MakeShared<FOpenVDBGridComponentInfo>(NoneGridComponentInfo));
+
+	// Create individual entries for each component of all valid source grids.
+	// This is an array of TSharedPtr because SComboBox requires its input to be wrapped in TSharedPtr.
+	bool bFoundSupportedGridType = false;
+	for (const FOpenVDBGridInfo& Grid : Result.GridInfo)
+	{
+		// Append all grids to the string, even if we don't actually support them
+		Result.FileInfoString.Append(Grid.DisplayString);
+		Result.FileInfoString.AppendChar(TEXT('\n'));
+
+		if (Grid.Type == EOpenVDBGridType::Unknown || !IsOpenVDBGridValid(Grid, Filename))
+		{
+			continue;
+		}
+
+		bFoundSupportedGridType = true;
+
+		// Create one entry per component
+		for (uint32 ComponentIdx = 0; ComponentIdx < Grid.NumComponents; ++ComponentIdx)
+		{
+			FOpenVDBGridComponentInfo ComponentInfo;
+			ComponentInfo.Index = Grid.Index;
+			ComponentInfo.ComponentIndex = ComponentIdx;
+			ComponentInfo.Name = Grid.Name;
+
+			const TCHAR* ComponentNames[] = { TEXT(".X"), TEXT(".Y"),TEXT(".Z"),TEXT(".W") };
+			FStringFormatOrderedArguments FormatArgs;
+			FormatArgs.Add(ComponentInfo.Index);
+			FormatArgs.Add(ComponentInfo.Name);
+			FormatArgs.Add(Grid.NumComponents == 1 ? TEXT("") : ComponentNames[ComponentIdx]);
+
+			ComponentInfo.DisplayString = FString::Format(TEXT("{0}. {1}{2}"), FormatArgs);
+
+			Result.GridComponentInfoPtrs.Add(MakeShared<FOpenVDBGridComponentInfo>(MoveTemp(ComponentInfo)));
+		}
+	}
+
+	if (!bFoundSupportedGridType)
+	{
+		UE_LOG(LogSparseVolumeTextureFactory, Error, TEXT("OpenVDB file contains no grids of supported type: %s"), *Filename);
+		return false;
+	}
+
+	return true;
+}
+
+struct FOpenVDBImportOptions
+{
+	FSparseVolumeRawSourcePackedData PackedDataA;
+	FSparseVolumeRawSourcePackedData PackedDataB;
+	bool bIsSequence;
+	bool bShouldImport;
+};
+
+static FOpenVDBImportOptions ShowOpenVDBImportWindow(const FString& Filename, const FString& FileInfoString, TArray<TSharedPtr<FOpenVDBGridComponentInfo>>& GridComponentInfo)
+{
+	FOpenVDBImportOptions ImportOptions{};
+
+	TSharedPtr<SWindow> ParentWindow;
+
+	if (FModuleManager::Get().IsModuleLoaded("MainFrame"))
+	{
+		IMainFrameModule& MainFrame = FModuleManager::LoadModuleChecked<IMainFrameModule>("MainFrame");
+		ParentWindow = MainFrame.GetParentWindow();
+	}
+
+	// Compute centered window position based on max window size, which include when all categories are expanded
+	const float ImportWindowWidth = 450.0f;
+	const float ImportWindowHeight = 750.0f;
+	FVector2D ImportWindowSize = FVector2D(ImportWindowWidth, ImportWindowHeight); // Max window size it can get based on current slate
+
+
+	FSlateRect WorkAreaRect = FSlateApplicationBase::Get().GetPreferredWorkArea();
+	FVector2D DisplayTopLeft(WorkAreaRect.Left, WorkAreaRect.Top);
+	FVector2D DisplaySize(WorkAreaRect.Right - WorkAreaRect.Left, WorkAreaRect.Bottom - WorkAreaRect.Top);
+
+	float ScaleFactor = FPlatformApplicationMisc::GetDPIScaleFactorAtPoint(DisplayTopLeft.X, DisplayTopLeft.Y);
+	ImportWindowSize *= ScaleFactor;
+
+	FVector2D WindowPosition = (DisplayTopLeft + (DisplaySize - ImportWindowSize) / 2.0f) / ScaleFactor;
+
+	TSharedRef<SWindow> Window = SNew(SWindow)
+		.Title(NSLOCTEXT("UnrealEd", "OpenVDBImportOptionsTitle", "OpenVDB Import Options"))
+		.SizingRule(ESizingRule::Autosized)
+		.AutoCenter(EAutoCenter::None)
+		.ClientSize(ImportWindowSize)
+		.ScreenPosition(WindowPosition);
+
+	TArray<TSharedPtr<ESparseVolumePackedDataFormat>> SupportedFormats =
+	{
+		MakeShared<ESparseVolumePackedDataFormat>(ESparseVolumePackedDataFormat::Float32),
+		MakeShared<ESparseVolumePackedDataFormat>(ESparseVolumePackedDataFormat::Float16),
+		MakeShared<ESparseVolumePackedDataFormat>(ESparseVolumePackedDataFormat::Unorm8)
+	};
+
+	TSharedPtr<SOpenVDBImportWindow> OpenVDBOptionWindow;
+	Window->SetContent
+	(
+		SAssignNew(OpenVDBOptionWindow, SOpenVDBImportWindow)
+		.PackedDataA(&ImportOptions.PackedDataA)
+		.PackedDataB(&ImportOptions.PackedDataB)
+		.OpenVDBGridComponentInfo(&GridComponentInfo)
+		.FileInfoString(FileInfoString)
+		.OpenVDBSupportedTargetFormats(&SupportedFormats)
+		.WidgetWindow(Window)
+		.FullPath(FText::FromString(Filename))
+		.MaxWindowHeight(ImportWindowHeight)
+		.MaxWindowWidth(ImportWindowWidth)
+	);
+
+	FSlateApplication::Get().AddModalWindow(Window, ParentWindow, false);
+
+	ImportOptions.bShouldImport = OpenVDBOptionWindow->ShouldImport();
+	ImportOptions.bIsSequence = OpenVDBOptionWindow->ShouldImportAsSequence();
+
+	return ImportOptions;
+}
+
 USparseVolumeTextureFactory::USparseVolumeTextureFactory(const FObjectInitializer& ObjectInitializer)
 	: Super(ObjectInitializer)
 {
@@ -108,288 +263,229 @@ UObject* USparseVolumeTextureFactory::FactoryCreateFile(UClass* InClass, UObject
 	GEditor->GetEditorSubsystem<UImportSubsystem>()->BroadcastAssetPreImport(this, InClass, InParent, InName, Parms);
 	TArray<UObject*> ResultAssets;
 
-	auto IsFilePotentiallyPartOfASequence = [](const FString& Filename)
+	// Load file and get info about each contained grid
+	FOpenVDBPreviewData PreviewData;
+	if (!LoadOpenVDBPreviewData(Filename, &PreviewData))
 	{
-		// The file is potnetially a sequence of the character before the `.vdb` is a number.
-		return FChar::IsDigit(Filename[Filename.Len() - 5]);
-	};
+		return nullptr;
+	}
 
-	if (IsFilePotentiallyPartOfASequence(Filename))
+	// Show dialog for import options
+	FOpenVDBImportOptions ImportOptions = ShowOpenVDBImportWindow(Filename, PreviewData.FileInfoString, PreviewData.GridComponentInfoPtrs);
+	if (!ImportOptions.bShouldImport)
 	{
-		// Import as an animated sparse volume texture asset.
-
-		const FString FilenameWithoutExt = Filename.LeftChop(4);
-		const int32 LastNonDigitIndex = FilenameWithoutExt.FindLastCharByPredicate([](TCHAR Letter) { return !FChar::IsDigit(Letter); }) + 1;
-		const int32 DigitCount = FilenameWithoutExt.Len() - LastNonDigitIndex;
-		FString FilenameWithoutSuffix = FilenameWithoutExt.LeftChop(FilenameWithoutExt.Len() - LastNonDigitIndex);
-		TCHAR LastDigit = FilenameWithoutExt[FilenameWithoutExt.Len() - 5];
+		bOutOperationCanceled = true;
+		return nullptr;
+	}
 
-		bool IndexStartsAtOne = false;
-		auto GetOpenVDBFileNameForFrame = [&](int32 FrameIndex)
+	// Utility function for computing the bounding box encompassing the bounds of all frames in the SVT.
+	// Accounts for padding caused by rounding up to a multiple of the page resolution.
+	auto ExpandVolumeBounds = [](const FOpenVDBImportOptions& ImportOptions, const TArray<FOpenVDBGridInfo>& GridInfoArray, FBox& VolumeBounds)
+	{
+		const FSparseVolumeRawSourcePackedData* PackedData[] = { &ImportOptions.PackedDataA, &ImportOptions.PackedDataB };
+		for (uint32 PackedDataIdx = 0; PackedDataIdx < 2; ++PackedDataIdx)
 		{
-			FString IndexString = FString::FromInt(FrameIndex + (IndexStartsAtOne ? 1 : 0));
-			// User must select a frame with index in [0-9] so that we can count leading 0s
-			check(DigitCount==1 || (DigitCount>1 && IndexString.Len() <= DigitCount));
-			const int32 MissingLeadingZeroCount = DigitCount - IndexString.Len();
-			const FString StringZero = FString::FromInt(0);
-			for (int32 i = 0; i < MissingLeadingZeroCount; ++i)
+			for (uint32 CompIdx = 0; CompIdx < 4; ++CompIdx)
 			{
-				IndexString = StringZero + IndexString;
+				const uint32 GridIdx = PackedData[PackedDataIdx]->SourceGridIndex[CompIdx];
+				if (GridIdx != INDEX_NONE)
+				{
+					const FOpenVDBGridInfo& GridInfo = GridInfoArray[GridIdx];
+					VolumeBounds.Min.X = FMath::Min(VolumeBounds.Min.X, GridInfo.VolumeActiveAABBMin.X);
+					VolumeBounds.Min.Y = FMath::Min(VolumeBounds.Min.Y, GridInfo.VolumeActiveAABBMin.Y);
+					VolumeBounds.Min.Z = FMath::Min(VolumeBounds.Min.Z, GridInfo.VolumeActiveAABBMin.Z);
+
+					const FVector PageTableVolumeResolution = FMath::DivideAndRoundUp(GridInfo.VolumeActiveDim, FVector(SPARSE_VOLUME_TILE_RES));
+					const FVector Max = GridInfo.VolumeActiveAABBMin + (PageTableVolumeResolution * SPARSE_VOLUME_TILE_RES);
+					VolumeBounds.Max.X = FMath::Max(VolumeBounds.Max.X, Max.X);
+					VolumeBounds.Max.Y = FMath::Max(VolumeBounds.Max.Y, Max.Y);
+					VolumeBounds.Max.Z = FMath::Max(VolumeBounds.Max.Z, Max.Z);
+				}
 			}
-			return FString(FilenameWithoutSuffix + IndexString) + TEXT(".vdb");
-		};
-
-		const FString VDBFileAt0 = GetOpenVDBFileNameForFrame(0);
-		const FString VDBFileAt1 = GetOpenVDBFileNameForFrame(1);
-		const bool VDBFileAt0Exists = FPaths::FileExists(VDBFileAt0);
-		const bool VDBFileAt1Exists = FPaths::FileExists(VDBFileAt1);
-		if (!VDBFileAt0Exists && !VDBFileAt1Exists)
-		{
-			UE_LOG(LogSparseVolumeTextureFactory, Error, TEXT("An OpenVDB animated sequence must start at index 0 or 1: %s or %s not found."), *VDBFileAt0, *VDBFileAt1);
-			return nullptr;
 		}
-		IndexStartsAtOne = !VDBFileAt0Exists;
-
-		FName NewName(InName .ToString()+ TEXT("VDBAnim"));
-		UAnimatedSparseVolumeTexture* AnimatedSVTexture = NewObject<UAnimatedSparseVolumeTexture>(InParent, UAnimatedSparseVolumeTexture::StaticClass(), NewName, Flags);
+	};
 
-		// Go over all the frame index and stop at the first missing one.
-		int32 FrameCount = 0;
-		while(FPaths::FileExists(GetOpenVDBFileNameForFrame(FrameCount)))
-		{
-			FrameCount++;
-		}
+	// Import as either single static SVT or a sequence of frames, making up an animated SVT
+	if (!ImportOptions.bIsSequence)
+	{
+		// Import as a static sparse volume texture asset.
 
-		FScopedSlowTask ImportTask(FrameCount, LOCTEXT("ImportingVDBAnim", "Importing OpenVDB animation"));
+		FScopedSlowTask ImportTask(1.0f, LOCTEXT("ImportingVDBStatic", "Importing static OpenVDB"));
 		ImportTask.MakeDialog(true);
 
-		// Allocate space for each frame
-		AnimatedSVTexture->FrameCount = FrameCount;
-		AnimatedSVTexture->AnimationFrames.SetNum(FrameCount);
-
-		UE_LOG(LogSparseVolumeTextureFactory, Display, TEXT("Serializing: %i frame"), FrameCount);
-
-		// Load vdb files for all frames.
-		for(int32 FrameIndex = 0; FrameIndex < FrameCount; ++FrameIndex)
-		{
-			const FString FrameFilename = GetOpenVDBFileNameForFrame(FrameIndex);
-			TArray<uint8> LoadedFile;
-
-			if (!FFileHelper::LoadFileToArray(LoadedFile, *FrameFilename))
-			{
-				UE_LOG(LogSparseVolumeTextureFactory, Error, TEXT("OpenVDB file could not be opened: %s"), *FrameFilename);
-				return nullptr;
-			}
-
-			uint32 DensityGridIndex = 0;
-			FOpenVDBData OpenVDBData;
-			if (!FindDensityGridIndex(LoadedFile, FrameFilename, &DensityGridIndex, &OpenVDBData))
-			{
-				UE_LOG(LogSparseVolumeTextureFactory, Error, TEXT("OpenVDB file contains no suitable density grid: %s"), *FrameFilename);
-				return nullptr;
-			}
+		FName NewName(InName.ToString() + TEXT("VDB"));
+		UStaticSparseVolumeTexture* StaticSVTexture = NewObject<UStaticSparseVolumeTexture>(InParent, UStaticSparseVolumeTexture::StaticClass(), NewName, Flags);
 
-			UE_LOG(LogSparseVolumeTextureFactory, Display, TEXT("  - frame %i (active dimension %i x %i x %i)"), FrameIndex,
-				OpenVDBData.VolumeActiveDim.X, OpenVDBData.VolumeActiveDim.Y, OpenVDBData.VolumeActiveDim.Z);
+		FBox VolumeBounds(FVector(FLT_MAX), FVector(-FLT_MAX));
+		ExpandVolumeBounds(ImportOptions, PreviewData.GridInfo, VolumeBounds);
+		StaticSVTexture->VolumeBounds = VolumeBounds;
 
-			FSparseVolumeRawSource SparseVolumeRawSource{};
-			SparseVolumeRawSource.PackedDataA.Format = ESparseVolumePackedDataFormat::Unorm8;
-			SparseVolumeRawSource.PackedDataA.SourceGridIndex = FUintVector4(0, INDEX_NONE, INDEX_NONE, INDEX_NONE);
-			SparseVolumeRawSource.PackedDataA.SourceComponentIndex = FUintVector4(0, INDEX_NONE, INDEX_NONE, INDEX_NONE);
-			SparseVolumeRawSource.PackedDataB.Format = ESparseVolumePackedDataFormat::Unorm8;
-			SparseVolumeRawSource.PackedDataB.SourceGridIndex = FUintVector4(INDEX_NONE, INDEX_NONE, INDEX_NONE, INDEX_NONE);
-			SparseVolumeRawSource.PackedDataB.SourceComponentIndex = FUintVector4(INDEX_NONE, INDEX_NONE, INDEX_NONE, INDEX_NONE);
-			SparseVolumeRawSource.SourceAssetFile = MoveTemp(LoadedFile);
-			LoadedFile.Reset();
+		FSparseVolumeRawSource SparseVolumeRawSource{};
+		SparseVolumeRawSource.PackedDataA = ImportOptions.PackedDataA;
+		SparseVolumeRawSource.PackedDataB = ImportOptions.PackedDataB;
+		SparseVolumeRawSource.SourceAssetFile = MoveTemp(PreviewData.LoadedFile);
+		PreviewData.LoadedFile.Reset();
 
-			// Serialize the raw source data from this frame into the asset object.
-			{
-				UE::Serialization::FEditorBulkDataWriter RawDataArchiveWriter(AnimatedSVTexture->AnimationFrames[FrameIndex].RawData);
-				SparseVolumeRawSource.Serialize(RawDataArchiveWriter);
-			}
+		// Serialize the raw source data into the asset object.
+		{
+			UE::Serialization::FEditorBulkDataWriter RawDataArchiveWriter(StaticSVTexture->StaticFrame.RawData);
+			SparseVolumeRawSource.Serialize(RawDataArchiveWriter);
+		}
 
-			if (ImportTask.ShouldCancel())
-			{
-				return nullptr;
-			}
-			ImportTask.EnterProgressFrame(1.0f, LOCTEXT("ConvertingVDBAnim", "Converting OpenVDB animation"));
+		if (ImportTask.ShouldCancel())
+		{
+			return nullptr;
 		}
+		ImportTask.EnterProgressFrame(1.0f, LOCTEXT("ConvertingVDBStatic", "Converting static OpenVDB"));
 
-		ResultAssets.Add(AnimatedSVTexture);
+		ResultAssets.Add(StaticSVTexture);
 	}
 	else
 	{
-		// Import as a static sparse volume texture asset.
+		// Import as an animated sparse volume texture asset.
+		
+		// Data from original file is no longer needed; we iterate over all frames later
+		PreviewData.LoadedFile.Empty();
 
-		FName NewName(InName.ToString() + TEXT("VDB"));
-		UStaticSparseVolumeTexture* StaticSVTexture = NewObject<UStaticSparseVolumeTexture>(InParent, UStaticSparseVolumeTexture::StaticClass(), NewName, Flags);
+		FName NewName(InName.ToString() + TEXT("VDBAnim"));
+		UAnimatedSparseVolumeTexture* AnimatedSVTexture = NewObject<UAnimatedSparseVolumeTexture>(InParent, UAnimatedSparseVolumeTexture::StaticClass(), NewName, Flags);
+		
+		// The file is potentially a sequence if the character before the `.vdb` is a number.
+		const bool bIsFilePotentiallyPartOfASequence = FChar::IsDigit(Filename[Filename.Len() - 5]);
 
-		// Load file and get info about each contained grid
-		TArray<uint8> LoadedFile;
-		TArray<TSharedPtr<FOpenVDBGridComponentInfo>> GridComponentInfoPtrs;
-		TArray<TSharedPtr<ESparseVolumePackedDataFormat>> SupportedFormats =
-		{ 
-			MakeShared<ESparseVolumePackedDataFormat>(ESparseVolumePackedDataFormat::Float32),
-			MakeShared<ESparseVolumePackedDataFormat>(ESparseVolumePackedDataFormat::Float16),
-			MakeShared<ESparseVolumePackedDataFormat>(ESparseVolumePackedDataFormat::Unorm8)
-		};
-		FString FileInfoString;
+		// Collect filenames of all files to be imported
+		TArray<FString> SequenceFileNames;
+		if (!bIsFilePotentiallyPartOfASequence)
 		{
-			if (!FFileHelper::LoadFileToArray(LoadedFile, *Filename))
-			{
-				UE_LOG(LogSparseVolumeTextureFactory, Error, TEXT("OpenVDB file could not be loaded: %s"), *Filename);
-				return nullptr;
-			}
-
-			TArray<FOpenVDBGridInfo> GridInfo;
-			if (!GetOpenVDBGridInfo(LoadedFile, &GridInfo))
+			SequenceFileNames.Add(Filename);
+		}
+		else
+		{
+			const FString FilenameWithoutExt = Filename.LeftChop(4);
+			const int32 LastNonDigitIndex = FilenameWithoutExt.FindLastCharByPredicate([](TCHAR Letter) { return !FChar::IsDigit(Letter); }) + 1;
+			const int32 DigitCount = FilenameWithoutExt.Len() - LastNonDigitIndex;
+			FString FilenameWithoutSuffix = FilenameWithoutExt.LeftChop(FilenameWithoutExt.Len() - LastNonDigitIndex);
+			TCHAR LastDigit = FilenameWithoutExt[FilenameWithoutExt.Len() - 5];
+
+			bool bIndexStartsAtOne = false;
+			auto GetOpenVDBFileNameForFrame = [&](int32 FrameIndex)
 			{
-				UE_LOG(LogSparseVolumeTextureFactory, Error, TEXT("Failed to read OpenVDB file: %s"), *Filename);
-				return nullptr;
-			}
-			if (GridInfo.IsEmpty())
+				FString IndexString = FString::FromInt(FrameIndex + (bIndexStartsAtOne ? 1 : 0));
+				// User must select a frame with index in [0-9] so that we can count leading 0s
+				check(DigitCount == 1 || (DigitCount > 1 && IndexString.Len() <= DigitCount));
+				const int32 MissingLeadingZeroCount = DigitCount - IndexString.Len();
+				const FString StringZero = FString::FromInt(0);
+				for (int32 i = 0; i < MissingLeadingZeroCount; ++i)
+				{
+					IndexString = StringZero + IndexString;
+				}
+				return FString(FilenameWithoutSuffix + IndexString) + TEXT(".vdb");
+			};
+
+			const FString VDBFileAt0 = GetOpenVDBFileNameForFrame(0);
+			const FString VDBFileAt1 = GetOpenVDBFileNameForFrame(1);
+			const bool VDBFileAt0Exists = FPaths::FileExists(VDBFileAt0);
+			const bool VDBFileAt1Exists = FPaths::FileExists(VDBFileAt1);
+			if (!VDBFileAt0Exists && !VDBFileAt1Exists)
 			{
-				UE_LOG(LogSparseVolumeTextureFactory, Error, TEXT("OpenVDB file contains no grids: %s"), *Filename);
+				UE_LOG(LogSparseVolumeTextureFactory, Error, TEXT("An OpenVDB animated sequence must start at index 0 or 1: %s or %s not found."), *VDBFileAt0, *VDBFileAt1);
 				return nullptr;
 			}
+			bIndexStartsAtOne = !VDBFileAt0Exists;
 
-			// Create individual entries for each component of all valid source grids
-			TArray<FOpenVDBGridComponentInfo> GridComponentInfo;
-			for (const FOpenVDBGridInfo& Grid : GridInfo)
+			// Go over all the frame index and stop at the first missing one.
+			for (int32 FrameIdx = 0; ; ++FrameIdx)
 			{
-				// Append all grids to the string, even if we don't actually support them
-				FileInfoString.Append(Grid.DisplayString);
-				FileInfoString.AppendChar(TEXT('\n'));
-
-				if (Grid.Type == EOpenVDBGridType::Unknown || !IsOpenVDBDataValid(Grid.OpenVDBData, Filename))
+				const FString FrameFileName = GetOpenVDBFileNameForFrame(FrameIdx);
+				if (!FPaths::FileExists(FrameFileName))
 				{
-					continue;
+					break;
 				}
+				SequenceFileNames.Add(FrameFileName);
+			}
+		}
 
-				// Create one entry per component
-				for (uint32 ComponentIdx = 0; ComponentIdx < Grid.NumComponents; ++ComponentIdx)
-				{
-					FOpenVDBGridComponentInfo ComponentInfo;
-					ComponentInfo.Index = Grid.Index;
-					ComponentInfo.ComponentIndex = ComponentIdx;
-					ComponentInfo.Name = Grid.Name;
+		const int32 NumFrames = SequenceFileNames.Num();
+
+		FScopedSlowTask ImportTask(NumFrames, LOCTEXT("ImportingVDBAnim", "Importing OpenVDB animation"));
+		ImportTask.MakeDialog(true);
 
-					const TCHAR* ComponentNames[] = { TEXT(".X"), TEXT(".Y"),TEXT(".Z"),TEXT(".W") };
-					FStringFormatOrderedArguments FormatArgs;
-					FormatArgs.Add(ComponentInfo.Index);
-					FormatArgs.Add(ComponentInfo.Name);
-					FormatArgs.Add(Grid.NumComponents == 1 ? TEXT("") : ComponentNames[ComponentIdx]);
+		// Allocate space for each frame
+		AnimatedSVTexture->FrameCount = NumFrames;
+		AnimatedSVTexture->AnimationFrames.SetNum(NumFrames);
 
-					ComponentInfo.DisplayString = FString::Format(TEXT("{0}. {1}{2}"), FormatArgs);
+		FBox VolumeBounds(FVector(FLT_MAX), FVector(-FLT_MAX));
 
-					GridComponentInfo.Add(MoveTemp(ComponentInfo));
-				}
+		// Load individual frames, check them for compatibility with the first loaded file and append them to the resulting asset
+		for (int32 FrameIdx = 0; FrameIdx < NumFrames; ++FrameIdx)
+		{
+			const FString& FrameFilename = SequenceFileNames[FrameIdx];
+
+			UE_LOG(LogSparseVolumeTextureFactory, Display, TEXT("Loading OpenVDB sequence frame #%i %s."), FrameIdx, *FrameFilename);
+
+			// Load file and get info about each contained grid
+			TArray<uint8> LoadedFrameFile;
+			if (!FFileHelper::LoadFileToArray(LoadedFrameFile, *FrameFilename))
+			{
+				UE_LOG(LogSparseVolumeTextureFactory, Error, TEXT("OpenVDB file could not be loaded: %s"), *FrameFilename);
+				return nullptr;
 			}
 
-			if (GridComponentInfo.IsEmpty())
+			TArray<FOpenVDBGridInfo> FrameGridInfo;
+			if (!GetOpenVDBGridInfo(LoadedFrameFile, true, &FrameGridInfo))
 			{
-				UE_LOG(LogSparseVolumeTextureFactory, Error, TEXT("OpenVDB file contains no grids of supported type: %s"), *Filename);
+				UE_LOG(LogSparseVolumeTextureFactory, Error, TEXT("Failed to read OpenVDB file: %s"), *FrameFilename);
 				return nullptr;
 			}
 
-			// Convert to array of TSharedPtr. SComboBox requires this.
-			GridComponentInfoPtrs.Empty(GridComponentInfo.Num() + 1);
-			
-			// We need a <None> option to leave channels empty
-			FOpenVDBGridComponentInfo NoneGridComponentInfo;
-			NoneGridComponentInfo.Index = INDEX_NONE;
-			NoneGridComponentInfo.ComponentIndex = INDEX_NONE;
-			NoneGridComponentInfo.Name = TEXT("<None>");
-			NoneGridComponentInfo.DisplayString = TEXT("<None>");
-			
-			GridComponentInfoPtrs.Add(MakeShared<FOpenVDBGridComponentInfo>(NoneGridComponentInfo));
-			
-			for (const FOpenVDBGridComponentInfo& GridCompInfo : GridComponentInfo)
+			// Sanity check for compatibility
+			const FSparseVolumeRawSourcePackedData* PackedData[] = { &ImportOptions.PackedDataA, &ImportOptions.PackedDataB };
+			for (uint32 PackedDataIdx = 0; PackedDataIdx < 2; ++PackedDataIdx)
 			{
-				GridComponentInfoPtrs.Add(MakeShared<FOpenVDBGridComponentInfo>(GridCompInfo));
+				for (uint32 DstCompIdx = 0; DstCompIdx < 4; ++DstCompIdx)
+				{
+					const uint32 SourceGridIndex = PackedData[PackedDataIdx]->SourceGridIndex[DstCompIdx];
+					if (SourceGridIndex != INDEX_NONE)
+					{
+						if ((int32)SourceGridIndex >= FrameGridInfo.Num())
+						{
+							UE_LOG(LogSparseVolumeTextureFactory, Error, TEXT("OpenVDB file is incompatible with other frames in the sequence: %s"), *FrameFilename);
+							return nullptr;
+						}
+						const auto& OrigSourceGrid = PreviewData.GridInfo[SourceGridIndex];
+						const auto& FrameSourceGrid = FrameGridInfo[SourceGridIndex];
+						if (OrigSourceGrid.Type != FrameSourceGrid.Type || OrigSourceGrid.Name != FrameSourceGrid.Name)
+						{
+							UE_LOG(LogSparseVolumeTextureFactory, Error, TEXT("OpenVDB file is incompatible with other frames in the sequence: %s"), *FrameFilename);
+							return nullptr;
+						}
+					}
+				}
 			}
-		}
 
-		// Show import dialog
-		FSparseVolumeRawSourcePackedData PackedDataA{};
-		FSparseVolumeRawSourcePackedData PackedDataB{};
-		{
-			TSharedPtr<SWindow> ParentWindow;
+			ExpandVolumeBounds(ImportOptions, FrameGridInfo, VolumeBounds);
+
+			FSparseVolumeRawSource SparseVolumeRawSource{};
+			SparseVolumeRawSource.PackedDataA = ImportOptions.PackedDataA;
+			SparseVolumeRawSource.PackedDataB = ImportOptions.PackedDataB;
+			SparseVolumeRawSource.SourceAssetFile = MoveTemp(LoadedFrameFile);
 
-			if (FModuleManager::Get().IsModuleLoaded("MainFrame"))
+			// Serialize the raw source data from this frame into the asset object.
 			{
-				IMainFrameModule& MainFrame = FModuleManager::LoadModuleChecked<IMainFrameModule>("MainFrame");
-				ParentWindow = MainFrame.GetParentWindow();
+				UE::Serialization::FEditorBulkDataWriter RawDataArchiveWriter(AnimatedSVTexture->AnimationFrames[FrameIdx].RawData);
+				SparseVolumeRawSource.Serialize(RawDataArchiveWriter);
 			}
 
-			// Compute centered window position based on max window size, which include when all categories are expanded
-			const float ImportWindowWidth = 450.0f;
-			const float ImportWindowHeight = 750.0f;
-			FVector2D ImportWindowSize = FVector2D(ImportWindowWidth, ImportWindowHeight); // Max window size it can get based on current slate
-
-
-			FSlateRect WorkAreaRect = FSlateApplicationBase::Get().GetPreferredWorkArea();
-			FVector2D DisplayTopLeft(WorkAreaRect.Left, WorkAreaRect.Top);
-			FVector2D DisplaySize(WorkAreaRect.Right - WorkAreaRect.Left, WorkAreaRect.Bottom - WorkAreaRect.Top);
-
-			float ScaleFactor = FPlatformApplicationMisc::GetDPIScaleFactorAtPoint(DisplayTopLeft.X, DisplayTopLeft.Y);
-			ImportWindowSize *= ScaleFactor;
-
-			FVector2D WindowPosition = (DisplayTopLeft + (DisplaySize - ImportWindowSize) / 2.0f) / ScaleFactor;
-
-			TSharedRef<SWindow> Window = SNew(SWindow)
-				.Title(NSLOCTEXT("UnrealEd", "OpenVDBImportOptionsTitle", "OpenVDB Import Options"))
-				.SizingRule(ESizingRule::Autosized)
-				.AutoCenter(EAutoCenter::None)
-				.ClientSize(ImportWindowSize)
-				.ScreenPosition(WindowPosition);
-
-			TSharedPtr<SOpenVDBImportWindow> OpenVDBOptionWindow;
-			Window->SetContent
-			(
-				SAssignNew(OpenVDBOptionWindow, SOpenVDBImportWindow)
-				.PackedDataA(&PackedDataA)
-				.PackedDataB(&PackedDataB)
-				.OpenVDBGridComponentInfo(&GridComponentInfoPtrs)
-				.FileInfoString(FileInfoString)
-				.OpenVDBSupportedTargetFormats(&SupportedFormats)
-				.WidgetWindow(Window)
-				.FullPath(FText::FromString(Filename))
-				.MaxWindowHeight(ImportWindowHeight)
-				.MaxWindowWidth(ImportWindowWidth)
-			);
-
-			FSlateApplication::Get().AddModalWindow(Window, ParentWindow, false);
-
-			if (!OpenVDBOptionWindow->ShouldImport())
+			if (ImportTask.ShouldCancel())
 			{
-				bOutOperationCanceled = true;
 				return nullptr;
 			}
+			ImportTask.EnterProgressFrame(1.0f, LOCTEXT("ConvertingVDBAnim", "Converting OpenVDB animation"));
 		}
 
-		FScopedSlowTask ImportTask(1.0f, LOCTEXT("ImportingVDBStatic", "Importing static OpenVDB"));
-		ImportTask.MakeDialog(true);
-
-		FSparseVolumeRawSource SparseVolumeRawSource{};
-		SparseVolumeRawSource.PackedDataA = PackedDataA;
-		SparseVolumeRawSource.PackedDataB = PackedDataB;
-		SparseVolumeRawSource.SourceAssetFile = MoveTemp(LoadedFile);
-		LoadedFile.Reset();
-
-		// Serialize the raw source data into the asset object.
-		{
-			UE::Serialization::FEditorBulkDataWriter RawDataArchiveWriter(StaticSVTexture->StaticFrame.RawData);
-			SparseVolumeRawSource.Serialize(RawDataArchiveWriter);
-		}
-
-		if (ImportTask.ShouldCancel())
-		{
-			return nullptr;
-		}
-		ImportTask.EnterProgressFrame(1.0f, LOCTEXT("ConvertingVDBStatic", "Converting static OpenVDB"));
+		AnimatedSVTexture->VolumeBounds = VolumeBounds;
 
-		ResultAssets.Add(StaticSVTexture);
+		ResultAssets.Add(AnimatedSVTexture);
 	}
 
 	// Now notify the system about the imported/updated/created assets
diff --git a/Engine/Source/Editor/SparseVolumeTexture/Private/SparseVolumeTextureOpenVDBUtility.cpp b/Engine/Source/Editor/SparseVolumeTexture/Private/SparseVolumeTextureOpenVDBUtility.cpp
--- a/Engine/Source/Editor/SparseVolumeTexture/Private/SparseVolumeTextureOpenVDBUtility.cpp
+++ b/Engine/Source/Editor/SparseVolumeTexture/Private/SparseVolumeTextureOpenVDBUtility.cpp
@@ -29,116 +29,111 @@ namespace
 	};
 }
 
-static FOpenVDBData GetOpenVDBData(openvdb::GridBase::Ptr GridBase)
+static FOpenVDBGridInfo GetOpenVDBGridInfo(openvdb::GridBase::Ptr Grid, uint32 GridIndex, bool bCreateStrings)
 {
-	FOpenVDBData OpenVDBData;
-	openvdb::CoordBBox VolumeActiveAABB = GridBase->evalActiveVoxelBoundingBox();
-	openvdb::Coord VolumeActiveDim = GridBase->evalActiveVoxelDim();
-	openvdb::Vec3d VolumeVoxelSize = GridBase->voxelSize();
-
-	OpenVDBData.VolumeActiveAABBMin = FVector(VolumeActiveAABB.min().x(), VolumeActiveAABB.min().y(), VolumeActiveAABB.min().z());
-	OpenVDBData.VolumeActiveAABBMax = FVector(VolumeActiveAABB.max().x(), VolumeActiveAABB.max().y(), VolumeActiveAABB.max().z());
-	OpenVDBData.VolumeActiveDim = FVector(VolumeActiveDim.x(), VolumeActiveDim.y(), VolumeActiveDim.z());
-	OpenVDBData.VolumeVoxelSize = FVector(VolumeVoxelSize.x(), VolumeVoxelSize.y(), VolumeVoxelSize.z());
-	OpenVDBData.bIsInWorldSpace = GridBase->isInWorldSpace();
-	OpenVDBData.bHasUniformVoxels = GridBase->hasUniformVoxels();
-
-	return OpenVDBData;
-}
-
-#endif
-
-bool IsOpenVDBDataValid(const FOpenVDBData& OpenVDBData, const FString& Filename)
-{
-	if (OpenVDBData.VolumeActiveDim.X * OpenVDBData.VolumeActiveDim.Y * OpenVDBData.VolumeActiveDim.Z == 0)
+	openvdb::CoordBBox VolumeActiveAABB = Grid->evalActiveVoxelBoundingBox();
+	openvdb::Coord VolumeActiveDim = Grid->evalActiveVoxelDim();
+	openvdb::Vec3d VolumeVoxelSize = Grid->voxelSize();
+	openvdb::math::MapBase::ConstPtr MapBase = Grid->constTransform().baseMap();
+	openvdb::Vec3d VoxelSize = MapBase->voxelSize();
+	openvdb::Mat4d GridTransformVDB = MapBase->getAffineMap()->getConstMat4();
+
+	FOpenVDBGridInfo GridInfo;
+	GridInfo.Index = GridIndex;
+	GridInfo.NumComponents = 0;
+	GridInfo.Type = EOpenVDBGridType::Unknown;
+	GridInfo.VolumeActiveAABBMin = FVector(VolumeActiveAABB.min().x(), VolumeActiveAABB.min().y(), VolumeActiveAABB.min().z());
+	GridInfo.VolumeActiveAABBMax = FVector(VolumeActiveAABB.max().x(), VolumeActiveAABB.max().y(), VolumeActiveAABB.max().z());
+	GridInfo.VolumeActiveDim = FVector(VolumeActiveDim.x(), VolumeActiveDim.y(), VolumeActiveDim.z());
+	GridInfo.VolumeVoxelSize = FVector(VolumeVoxelSize.x(), VolumeVoxelSize.y(), VolumeVoxelSize.z());
+	GridInfo.bIsInWorldSpace = Grid->isInWorldSpace();
+	GridInfo.bHasUniformVoxels = Grid->hasUniformVoxels();
+
+	for (int i = 0; i < 4; ++i)
 	{
-		// SVT_TODO we should gently handle that case
-		UE_LOG(LogSparseVolumeTextureOpenVDBUtility, Warning, TEXT("OpenVDB grid is empty due to volume size being 0: %s"), *Filename);
-		return false;
+		for (int j = 0; j < 4; ++j)
+		{
+			GridInfo.Transform.M[i][j] = static_cast<float>(GridTransformVDB[j][i]);
+		}
 	}
 
-	if (!OpenVDBData.bHasUniformVoxels)
+	// Figure out the type/format of the grid
+	if (Grid->isType<FOpenVDBFloat1Grid>())
 	{
-		UE_LOG(LogSparseVolumeTextureOpenVDBUtility, Warning, TEXT("OpenVDB importer cannot handle non uniform voxels: %s"), *Filename);
-		return false;
+		GridInfo.NumComponents = 1;
+		GridInfo.Type = EOpenVDBGridType::Float;
 	}
-	return true;
-}
-
-bool FindDensityGridIndex(TArray<uint8>& SourceFile, const FString& Filename, uint32* OutGridIndex, FOpenVDBData* OutOVDBData)
-{
-#if OPENVDB_AVAILABLE
-	FArrayUint8StreamBuf StreamBuf(SourceFile);
-	std::istream IStream(&StreamBuf);
-	openvdb::io::Stream Stream(IStream, false /*delayLoad*/);
-
-	openvdb::GridPtrVecPtr Grids = Stream.getGrids();
-
-	if (!Grids)
+	else if (Grid->isType<FOpenVDBFloat2Grid>())
 	{
-		UE_LOG(LogSparseVolumeTextureOpenVDBUtility, Warning, TEXT("OpenVDB file contains no grids: %s"), *Filename);
-		return false;
+		GridInfo.NumComponents = 2;
+		GridInfo.Type = EOpenVDBGridType::Float2;
 	}
-
-	uint32 GridIndex = 0;
-	openvdb::GridBase::Ptr BaseGrid = nullptr;
-	for (openvdb::GridBase::Ptr& Grid : *Grids)
+	else if (Grid->isType<FOpenVDBFloat3Grid>())
 	{
-		if (Grid->getName() == "density")
-		{
-			if (Grid->isType<openvdb::FloatGrid>())
-			{
-				BaseGrid = Grid;
-				break;
-			}
-			break;
-		}
-		++GridIndex;
+		GridInfo.NumComponents = 3;
+		GridInfo.Type = EOpenVDBGridType::Float3;
 	}
-
-	// If we have not found any density map, let's pick up the first float map
-	if (!BaseGrid)
+	else if (Grid->isType<FOpenVDBFloat4Grid>())
 	{
-		GridIndex = 0;
-		for (openvdb::GridBase::Ptr& Grid : *Grids)
-		{
-			if (Grid->isType<openvdb::FloatGrid>())
-			{
-				BaseGrid = Grid;
-				break;
-			}
-			++GridIndex;
-		}
+		GridInfo.NumComponents = 4;
+		GridInfo.Type = EOpenVDBGridType::Float4;
 	}
-
-	if (BaseGrid)
+	else if (Grid->isType<FOpenVDBDouble1Grid>())
 	{
-		openvdb::FloatGrid::Ptr DensityGrid = openvdb::gridPtrCast<openvdb::FloatGrid>(BaseGrid);
+		GridInfo.NumComponents = 1;
+		GridInfo.Type = EOpenVDBGridType::Double;
+	}
+	else if (Grid->isType<FOpenVDBDouble2Grid>())
+	{
+		GridInfo.NumComponents = 2;
+		GridInfo.Type = EOpenVDBGridType::Double2;
+	}
+	else if (Grid->isType<FOpenVDBDouble3Grid>())
+	{
+		GridInfo.NumComponents = 3;
+		GridInfo.Type = EOpenVDBGridType::Double3;
+	}
+	else if (Grid->isType<FOpenVDBDouble4Grid>())
+	{
+		GridInfo.NumComponents = 4;
+		GridInfo.Type = EOpenVDBGridType::Double4;
+	}
 
-		// Only open float grid for now
-		if (DensityGrid == nullptr)
-		{
-			UE_LOG(LogSparseVolumeTextureOpenVDBUtility, Error, TEXT("Could not convert the grid to float: %s"), *Filename);
-			return false;
-		}
+	if (bCreateStrings)
+	{
+		GridInfo.Name = Grid->getName().c_str();
 
-		FOpenVDBData OVDBData = GetOpenVDBData(DensityGrid);
-		if (!IsOpenVDBDataValid(OVDBData, Filename))
-		{
-			return false;
-		}
+		FStringFormatOrderedArguments FormatArgs;
+		FormatArgs.Add(GridInfo.Index);
+		FormatArgs.Add(OpenVDBGridTypeToString(GridInfo.Type));
+		FormatArgs.Add(GridInfo.Name);
 
-		*OutGridIndex = GridIndex;
-		*OutOVDBData = OVDBData;
-		return true;
+		GridInfo.DisplayString = FString::Format(TEXT("{0}. Type: {1}, Name: \"{2}\""), FormatArgs);
 	}
 
-#endif // OPENVDB_AVAILABLE
+	return GridInfo;
+}
 
-	return false;
+#endif
+
+bool IsOpenVDBGridValid(const FOpenVDBGridInfo& GridInfo, const FString& Filename)
+{
+	if (GridInfo.VolumeActiveDim.X * GridInfo.VolumeActiveDim.Y * GridInfo.VolumeActiveDim.Z == 0)
+	{
+		// SVT_TODO we should gently handle that case
+		UE_LOG(LogSparseVolumeTextureOpenVDBUtility, Warning, TEXT("OpenVDB grid is empty due to volume size being 0: %s"), *Filename);
+		return false;
+	}
+
+	if (!GridInfo.bHasUniformVoxels)
+	{
+		UE_LOG(LogSparseVolumeTextureOpenVDBUtility, Warning, TEXT("OpenVDB importer cannot handle non uniform voxels: %s"), *Filename);
+		return false;
+	}
+	return true;
 }
 
-bool GetOpenVDBGridInfo(TArray<uint8>& SourceFile, TArray<FOpenVDBGridInfo>* OutGridInfo)
+bool GetOpenVDBGridInfo(TArray<uint8>& SourceFile, bool bCreateStrings, TArray<FOpenVDBGridInfo>* OutGridInfo)
 {
 #if OPENVDB_AVAILABLE
 	FArrayUint8StreamBuf StreamBuf(SourceFile);
@@ -152,64 +147,8 @@ bool GetOpenVDBGridInfo(TArray<uint8>& SourceFile, TArray<FOpenVDBGridInfo>* Out
 	uint32 GridIndex = 0;
 	for (openvdb::GridBase::Ptr& Grid : *Grids)
 	{
-		// Create entry
-		FOpenVDBGridInfo GridInfo;
-		GridInfo.Index = GridIndex++;
-		GridInfo.NumComponents = 0;
-		GridInfo.Type = EOpenVDBGridType::Unknown;
-		GridInfo.Name = Grid->getName().c_str();
-		GridInfo.OpenVDBData = GetOpenVDBData(Grid);
-
-		// Figure out the type/format of the grid
-		if (Grid->isType<FOpenVDBFloat1Grid>())
-		{
-			GridInfo.NumComponents = 1;
-			GridInfo.Type = EOpenVDBGridType::Float;
-		}
-		else if (Grid->isType<FOpenVDBFloat2Grid>())
-		{
-			GridInfo.NumComponents = 2;
-			GridInfo.Type = EOpenVDBGridType::Float2;
-		}
-		else if (Grid->isType<FOpenVDBFloat3Grid>())
-		{
-			GridInfo.NumComponents = 3;
-			GridInfo.Type = EOpenVDBGridType::Float3;
-		}
-		else if (Grid->isType<FOpenVDBFloat4Grid>())
-		{
-			GridInfo.NumComponents = 4;
-			GridInfo.Type = EOpenVDBGridType::Float4;
-		}
-		else if (Grid->isType<FOpenVDBDouble1Grid>())
-		{
-			GridInfo.NumComponents = 1;
-			GridInfo.Type = EOpenVDBGridType::Double;
-		}
-		else if (Grid->isType<FOpenVDBDouble2Grid>())
-		{
-			GridInfo.NumComponents = 2;
-			GridInfo.Type = EOpenVDBGridType::Double2;
-		}
-		else if (Grid->isType<FOpenVDBDouble3Grid>())
-		{
-			GridInfo.NumComponents = 3;
-			GridInfo.Type = EOpenVDBGridType::Double3;
-		}
-		else if (Grid->isType<FOpenVDBDouble4Grid>())
-		{
-			GridInfo.NumComponents = 4;
-			GridInfo.Type = EOpenVDBGridType::Double4;
-		}
-
-		FStringFormatOrderedArguments FormatArgs;
-		FormatArgs.Add(GridInfo.Index);
-		FormatArgs.Add(OpenVDBGridTypeToString(GridInfo.Type));
-		FormatArgs.Add(GridInfo.Name);
-
-		GridInfo.DisplayString = FString::Format(TEXT("{0}. Type: {1}, Name: \"{2}\""), FormatArgs);
-
-		OutGridInfo->Add(MoveTemp(GridInfo));
+		OutGridInfo->Add(GetOpenVDBGridInfo(Grid, GridIndex, bCreateStrings));
+		++GridIndex;
 	}
 
 	return true;
@@ -400,25 +339,25 @@ bool ConvertOpenVDBToSparseVolumeTexture(
 
 			GridAdapters[PackedDataIdx][CompIdx] = UniqueGridAdapters[SourceGridIndex];
 
-			FOpenVDBData OVDBData = GetOpenVDBData(GridBase);
-			if (!IsOpenVDBDataValid(OVDBData, TEXT("")))
+			FOpenVDBGridInfo GridInfo = GetOpenVDBGridInfo(GridBase, 0, false);
+			if (!IsOpenVDBGridValid(GridInfo, TEXT("")))
 			{
 				return false;
 			}
 
 			if (bOverrideActiveMinMax)
 			{
-				OVDBData.VolumeActiveAABBMin = ActiveMin;
-				OVDBData.VolumeActiveAABBMax = ActiveMax;
-				OVDBData.VolumeActiveDim = ActiveMax + 1 - ActiveMin;
+				GridInfo.VolumeActiveAABBMin = ActiveMin;
+				GridInfo.VolumeActiveAABBMax = ActiveMax;
+				GridInfo.VolumeActiveDim = ActiveMax + 1 - ActiveMin;
 			}
 
-			Header.SourceVolumeResolution.X = FMath::Max(Header.SourceVolumeResolution.X, OVDBData.VolumeActiveDim.X);
-			Header.SourceVolumeResolution.Y = FMath::Max(Header.SourceVolumeResolution.Y, OVDBData.VolumeActiveDim.Y);
-			Header.SourceVolumeResolution.Z = FMath::Max(Header.SourceVolumeResolution.Z, OVDBData.VolumeActiveDim.Z);
-			SmallestAABBMin.X = FMath::Min(SmallestAABBMin.X, OVDBData.VolumeActiveAABBMin.X);
-			SmallestAABBMin.Y = FMath::Min(SmallestAABBMin.Y, OVDBData.VolumeActiveAABBMin.Y);
-			SmallestAABBMin.Z = FMath::Min(SmallestAABBMin.Z, OVDBData.VolumeActiveAABBMin.Z);
+			Header.SourceVolumeResolution.X = FMath::Max(Header.SourceVolumeResolution.X, GridInfo.VolumeActiveDim.X);
+			Header.SourceVolumeResolution.Y = FMath::Max(Header.SourceVolumeResolution.Y, GridInfo.VolumeActiveDim.Y);
+			Header.SourceVolumeResolution.Z = FMath::Max(Header.SourceVolumeResolution.Z, GridInfo.VolumeActiveDim.Z);
+			SmallestAABBMin.X = FMath::Min(SmallestAABBMin.X, GridInfo.VolumeActiveAABBMin.X);
+			SmallestAABBMin.Y = FMath::Min(SmallestAABBMin.Y, GridInfo.VolumeActiveAABBMin.Y);
+			SmallestAABBMin.Z = FMath::Min(SmallestAABBMin.Z, GridInfo.VolumeActiveAABBMin.Z);
 
 			GridBackgroundValues[PackedDataIdx][CompIdx] = GridAdapters[PackedDataIdx][CompIdx]->GetBackgroundValue(SourceComponentIndex);
 			if (bNormalizedFormat[PackedDataIdx] && PackedData[PackedDataIdx]->bRemapInputForUnorm)
@@ -432,6 +371,7 @@ bool ConvertOpenVDBToSparseVolumeTexture(
 			}
 		}
 	}
+	OutResult->Header->SourceVolumeAABBMin = SmallestAABBMin;
 	
 	FIntVector3 PageTableVolumeResolution = FIntVector3(
 		FMath::DivideAndRoundUp(Header.SourceVolumeResolution.X, SPARSE_VOLUME_TILE_RES),
diff --git a/Engine/Source/Editor/SparseVolumeTexture/Private/SparseVolumeTextureOpenVDBUtility.h b/Engine/Source/Editor/SparseVolumeTexture/Private/SparseVolumeTextureOpenVDBUtility.h
--- a/Engine/Source/Editor/SparseVolumeTexture/Private/SparseVolumeTextureOpenVDBUtility.h
+++ b/Engine/Source/Editor/SparseVolumeTexture/Private/SparseVolumeTextureOpenVDBUtility.h
@@ -6,16 +6,6 @@
 
 #include "CoreMinimal.h"
 
-struct FOpenVDBData
-{
-	FVector VolumeActiveAABBMin;
-	FVector VolumeActiveAABBMax;
-	FVector VolumeActiveDim;
-	FVector VolumeVoxelSize;
-	bool bIsInWorldSpace;
-	bool bHasUniformVoxels;
-};
-
 enum class EOpenVDBGridType : uint8
 {
 	Unknown = 0,
@@ -31,19 +21,23 @@ enum class EOpenVDBGridType : uint8
 
 struct FOpenVDBGridInfo
 {
+	FMatrix44f Transform;
+	FVector VolumeActiveAABBMin;
+	FVector VolumeActiveAABBMax;
+	FVector VolumeActiveDim;
+	FVector VolumeVoxelSize;
+	FString Name;
+	FString DisplayString; // Contains Index (into source file grids), Type and Name
 	uint32 Index;
 	uint32 NumComponents;
 	EOpenVDBGridType Type;
-	FString Name;
-	FString DisplayString; // Contains Index (into source file grids), Type and Name
-	FOpenVDBData OpenVDBData;
+	bool bIsInWorldSpace;
+	bool bHasUniformVoxels;
 };
 
-bool IsOpenVDBDataValid(const FOpenVDBData& OpenVDBData, const FString& Filename);
-
-bool FindDensityGridIndex(TArray<uint8>& SourceFile, const FString& Filename, uint32* OutGridIndex, FOpenVDBData* OutOVDBData);
+bool IsOpenVDBGridValid(const FOpenVDBGridInfo& GridInfo, const FString& Filename);
 
-bool GetOpenVDBGridInfo(TArray<uint8>& SourceFile, TArray<FOpenVDBGridInfo>* OutGridInfo);
+bool GetOpenVDBGridInfo(TArray<uint8>& SourceFile, bool bCreateStrings, TArray<FOpenVDBGridInfo>* OutGridInfo);
 
 bool ConvertOpenVDBToSparseVolumeTexture(
 	TArray<uint8>& SourceFile,
diff --git a/Engine/Source/Runtime/Engine/Classes/SparseVolumeTexture/SparseVolumeTexture.h b/Engine/Source/Runtime/Engine/Classes/SparseVolumeTexture/SparseVolumeTexture.h
--- a/Engine/Source/Runtime/Engine/Classes/SparseVolumeTexture/SparseVolumeTexture.h
+++ b/Engine/Source/Runtime/Engine/Classes/SparseVolumeTexture/SparseVolumeTexture.h
@@ -43,32 +43,6 @@
 	Float32 = 2,
 };
 
-struct ENGINE_API FSparseVolumeAssetHeader
-{
-	FIntVector3 PageTableVolumeResolution;
-	FIntVector3 TileDataVolumeResolution;
-	FIntVector3 SourceVolumeResolution;
-	EPixelFormat PackedDataAFormat;
-	EPixelFormat PackedDataBFormat;
-
-	// The current data format version for the header.
-	static const uint32 kVersion = 0;
-
-	// This version can be used to convert existing header to new version later.
-	uint32 Version;
-
-	void Serialize(FArchive& Ar);
-
-	FSparseVolumeAssetHeader()
-		: PageTableVolumeResolution(FIntVector3(0, 0, 0))
-		, TileDataVolumeResolution(FIntVector3(0, 0, 0))
-		, SourceVolumeResolution(FIntVector3(0, 0, 0))
-		, PackedDataAFormat(PF_Unknown)
-		, Version(kVersion)
-	{
-	}
-};
-
 // Describes how the grids of a source asset map to the components of a float4 of packed data in a SVT
 struct ENGINE_API FSparseVolumeRawSourcePackedData
 {
@@ -101,6 +75,33 @@
 	}
 };
 
+struct ENGINE_API FSparseVolumeAssetHeader
+{
+	FIntVector3 PageTableVolumeResolution;
+	FIntVector3 TileDataVolumeResolution;
+	FIntVector3 SourceVolumeResolution;
+	FIntVector3 SourceVolumeAABBMin;
+	EPixelFormat PackedDataAFormat;
+	EPixelFormat PackedDataBFormat;
+
+	// The current data format version for the header.
+	static const uint32 kVersion = 0;
+
+	// This version can be used to convert existing header to new version later.
+	uint32 Version;
+
+	void Serialize(FArchive& Ar);
+
+	FSparseVolumeAssetHeader()
+		: PageTableVolumeResolution(FIntVector3(0, 0, 0))
+		, TileDataVolumeResolution(FIntVector3(0, 0, 0))
+		, SourceVolumeResolution(FIntVector3(0, 0, 0))
+		, PackedDataAFormat(PF_Unknown)
+		, Version(kVersion)
+	{
+	}
+};
+
 // The structure represent the runtime data cooked runtime data.
 struct ENGINE_API FSparseVolumeTextureRuntime
 {
@@ -172,6 +173,7 @@
 {
 	FSparseVolumeTextureFrame();
 	virtual ~FSparseVolumeTextureFrame();
+	bool BuildRuntimeData();
 
 	// The frame data that can be streamed in when in game.
 	FByteBulkData						RuntimeStreamedInData;
@@ -209,12 +211,19 @@
 	virtual ~USparseVolumeTexture() = default;
 
 	virtual int32 GetFrameCount() const { return 0; }
-	virtual const FSparseVolumeAssetHeader* GetSparseVolumeTextureHeader() const { return nullptr; }
-	virtual FSparseVolumeTextureSceneProxy* GetSparseVolumeTextureSceneProxy() { return nullptr; }
-	virtual const FSparseVolumeTextureSceneProxy* GetSparseVolumeTextureSceneProxy() const { return nullptr; }
+	virtual const FSparseVolumeAssetHeader* GetSparseVolumeTextureHeader(int32 FrameIndex) const { return nullptr; }
+	virtual FSparseVolumeTextureSceneProxy* GetSparseVolumeTextureSceneProxy(int32 FrameIndex) { return nullptr; }
+	virtual const FSparseVolumeTextureSceneProxy* GetSparseVolumeTextureSceneProxy(int32 FrameIndex) const { return nullptr; }
 
 	/** Getter for the shader uniform parameters with index as ESparseVolumeTextureShaderUniform. */
-	virtual FVector4 GetUniformParameter(int32 Index) const { return FVector4(ForceInitToZero); }
+	virtual FVector4 GetUniformParameter(int32 Index, int32 FrameIndex) const { return FVector4(ForceInitToZero); }
 
+	virtual FBox GetVolumeBounds() const { return FBox(); }
+
+	/** In order to keep the contents of an animated SVT sequence stable in world space, we need to account for the fact that
+		different frames of the sequence have different AABBs. We solve this by scaling and biasing UVs that are relative to
+		the volume bounds into the UV space represented by the AABB of each animation frame.*/
+	void GetFrameUVScaleBias(int32 FrameIndex, FVector* OutScale, FVector* OutBias) const;
+
 	/** Getter for the shader uniform parameter type with index as ESparseVolumeTextureShaderUniform. */
 	static UE::Shader::EValueType GetUniformParameterType(int32 Index);
@@ -230,6 +239,9 @@
 
 public:
 
+	UPROPERTY(VisibleAnywhere, Category = "Rendering", meta = (DisplayName = "Volume Bounds"))
+	FBox VolumeBounds;
+
 	UStaticSparseVolumeTexture();
 	virtual ~UStaticSparseVolumeTexture() = default;
 
@@ -247,15 +259,14 @@
 
 	//~ Begin USparseVolumeTexture Interface.
 	int32 GetFrameCount() const override { return 1; }
-	const FSparseVolumeAssetHeader* GetSparseVolumeTextureHeader() const override;
-	FSparseVolumeTextureSceneProxy* GetSparseVolumeTextureSceneProxy() override;
-	const FSparseVolumeTextureSceneProxy* GetSparseVolumeTextureSceneProxy() const override;
+	const FSparseVolumeAssetHeader* GetSparseVolumeTextureHeader(int32 FrameIndex) const override;
+	FSparseVolumeTextureSceneProxy* GetSparseVolumeTextureSceneProxy(int32 FrameIndex) override;
+	const FSparseVolumeTextureSceneProxy* GetSparseVolumeTextureSceneProxy(int32 FrameIndex) const override;
 
 	/** Getter for the shader uniform parameters. */
-	FVector4 GetUniformParameter(int32 Index) const;
-	/** Getter for the shader uniform parameter type. */
-	static UE::Shader::EValueType GetUniformParameterType(int32 Index);
-	//~ Begin USparseVolumeTexture Interface.
+	FVector4 GetUniformParameter(int32 Index, int32 FrameIndex) const override;
+	FBox GetVolumeBounds() const override;
+	//~ End USparseVolumeTexture Interface.
 
 private:
 
@@ -263,7 +274,7 @@
 	friend class USparseVolumeTextureFactory; // Importer
 #endif
 
-	FSparseVolumeTextureFrame	StaticFrame;
+	FSparseVolumeTextureFrame StaticFrame;
 
 	void ConvertRawSourceDataToSparseVolumeTextureRuntime();
 	void GenerateOrLoadDDCRuntimeData();
@@ -278,6 +289,13 @@
 
 public:
 
+	// The asset frame count.
+	UPROPERTY(VisibleAnywhere, Category = "Animation", meta = (DisplayName = "Frame Count"))
+	int32 FrameCount = 0;
+
+	UPROPERTY(VisibleAnywhere, Category = "Rendering", meta = (DisplayName = "Volume Bounds"))
+	FBox VolumeBounds;
+
 	UAnimatedSparseVolumeTexture();
 	virtual ~UAnimatedSparseVolumeTexture() = default;
 
@@ -295,35 +313,28 @@
 
 	//~ Begin USparseVolumeTexture Interface.
 	int32 GetFrameCount() const override { return FrameCount; }
-	const FSparseVolumeAssetHeader* GetSparseVolumeTextureHeader() const override;
-	FSparseVolumeTextureSceneProxy* GetSparseVolumeTextureSceneProxy() override;
-	const FSparseVolumeTextureSceneProxy* GetSparseVolumeTextureSceneProxy() const override;
+	const FSparseVolumeAssetHeader* GetSparseVolumeTextureHeader(int32 FrameIndex) const override;
+	FSparseVolumeTextureSceneProxy* GetSparseVolumeTextureSceneProxy(int32 FrameIndex) override;
+	const FSparseVolumeTextureSceneProxy* GetSparseVolumeTextureSceneProxy(int32 FrameIndex) const override;
 
 	/** Getter for the shader uniform parameters. */
-	FVector4 GetUniformParameter(int32 Index) const;
-	/** Getter for the shader uniform parameter type. */
-	static UE::Shader::EValueType GetUniformParameterType(int32 Index);
-	//~ Begin USparseVolumeTexture Interface.
+	FVector4 GetUniformParameter(int32 Index, int32 FrameIndex) const override;
+	FBox GetVolumeBounds() const override;
+	//~ End USparseVolumeTexture Interface.
 
 	// Used for debugging a specific frame of an animated sequence.
 	const FSparseVolumeTextureSceneProxy* GetSparseVolumeTextureFrameSceneProxy(int32 FrameIndex) const;
 
-	// The asset frame count.
-	UPROPERTY(VisibleAnywhere, Category = "Animation", meta = (DisplayName = "Frame Count"))
-	int32 FrameCount = 0;
-
 private:
 
 #if WITH_EDITOR
 	friend class USparseVolumeTextureFactory; // Importer
 #endif
 	
-	const uint32 PreviewFrameIndex = 0;
 	const bool bLoadAllFramesToProxies = true;	// SVT_TODO remove that once streaming is working
-	int32 GetFrameCountToLoad() const ;
+	TArray<FSparseVolumeTextureFrame> AnimationFrames;
 
-	TArray<FSparseVolumeTextureFrame>	AnimationFrames;
-
+	int32 GetFrameCountToLoad() const;
 	void ConvertRawSourceDataToSparseVolumeTextureRuntime(int32 FrameIndex);
 	void GenerateOrLoadDDCRuntimeData(int32 FrameIndex);
 	void GenerateOrLoadDDCRuntimeDataAndCreateSceneProxy(int32 FrameIndex);
diff --git a/Engine/Source/Runtime/Engine/Private/Materials/MaterialUniformExpressions.cpp b/Engine/Source/Runtime/Engine/Private/Materials/MaterialUniformExpressions.cpp
--- a/Engine/Source/Runtime/Engine/Private/Materials/MaterialUniformExpressions.cpp
+++ b/Engine/Source/Runtime/Engine/Private/Materials/MaterialUniformExpressions.cpp
@@ -1203,7 +1203,7 @@ void FUniformExpressionSet::FillUniformBuffer(const FMaterialRenderContext& Mate
 					// 1 is tile data A
 					// 2 is tile data B
 
-					const FSparseVolumeTextureSceneProxy* SVTextureProxy = SVTexture->GetSparseVolumeTextureSceneProxy();
+					const FSparseVolumeTextureSceneProxy* SVTextureProxy = SVTexture->GetSparseVolumeTextureSceneProxy(0 /*SVT_TODO*/);
 					if (SVTextureProxy)
 					{
 						if (LayerIndex == 0)
diff --git a/Engine/Source/Runtime/Engine/Private/Shader/Preshader.cpp b/Engine/Source/Runtime/Engine/Private/Shader/Preshader.cpp
--- a/Engine/Source/Runtime/Engine/Private/Shader/Preshader.cpp
+++ b/Engine/Source/Runtime/Engine/Private/Shader/Preshader.cpp
@@ -823,7 +823,7 @@ static void EvaluateSparseTextureUniform(const FMaterialRenderContext& Context,
 	}
 	if (Texture != nullptr && VectorIndex != INDEX_NONE)
 	{
-		Stack.PushValue(FValue(Texture->GetUniformParameter(VectorIndex)));
+		Stack.PushValue(FValue(Texture->GetUniformParameter(VectorIndex, 0 /*SVT_TODO*/)));
 	}
 	else
 	{
diff --git a/Engine/Source/Runtime/Engine/Private/SparseVolumeTexture/SparseVolumeTexture.cpp b/Engine/Source/Runtime/Engine/Private/SparseVolumeTexture/SparseVolumeTexture.cpp
--- a/Engine/Source/Runtime/Engine/Private/SparseVolumeTexture/SparseVolumeTexture.cpp
+++ b/Engine/Source/Runtime/Engine/Private/SparseVolumeTexture/SparseVolumeTexture.cpp
@@ -33,48 +33,30 @@
 	return ConvertOpenVDBToSparseVolumeTextureDelegate;
 }
 
-void FSparseVolumeAssetHeader::Serialize(FArchive& Ar)
+namespace UE
 {
-
-	Ar << Version;
-
-	if (Version == 0)
+namespace SVT
+{
+namespace Private
+{
+	// SVT_TODO: This really should be a shared function.
+	template <typename Y, typename T>
+	void SerializeEnumAs(FArchive& Ar, T& Target)
 	{
-		Ar << PageTableVolumeResolution;
-		Ar << TileDataVolumeResolution;
-		Ar << SourceVolumeResolution;
-
-		uint32 FormatAsUint;
-		
-		FormatAsUint = (uint32)PackedDataAFormat;
-		Ar << FormatAsUint;
+		Y Buffer = static_cast<Y>(Target);
+		Ar << Buffer;
 		if (Ar.IsLoading())
 		{
-			PackedDataAFormat = static_cast<EPixelFormat>(FormatAsUint);
-		}
-		
-		FormatAsUint = (uint32)PackedDataBFormat;
-		Ar << FormatAsUint;
-		if (Ar.IsLoading())
-		{
-			PackedDataBFormat = static_cast<EPixelFormat>(FormatAsUint);
+			Target = static_cast<T>(Buffer);
 		}
 	}
-	else
-	{
-		// FSparseVolumeAssetHeader needs to account for new version
-		check(false);
-	}
-}
+} // Private
+} // SVT
+} // UE
 
 FArchive& operator<<(FArchive& Ar, FSparseVolumeRawSourcePackedData& PackedData)
 {
-	uint32 FormatAsUint = (uint32)PackedData.Format;
-	Ar << FormatAsUint;
-	if (Ar.IsLoading())
-	{
-		PackedData.Format = static_cast<ESparseVolumePackedDataFormat>(FormatAsUint);
-	}
+	UE::SVT::Private::SerializeEnumAs<uint8>(Ar, PackedData.Format);
 	Ar << PackedData.SourceGridIndex;
 	Ar << PackedData.SourceComponentIndex;
 	Ar << PackedData.bRemapInputForUnorm;
@@ -82,6 +64,8 @@
 	return Ar;
 }
 
+////////////////////////////////////////////////////////////////////////////////////////////////
+
 void FSparseVolumeRawSource::Serialize(FArchive& Ar)
 {
 	Ar << Version;
@@ -99,6 +83,30 @@
 	}
 }
 
+////////////////////////////////////////////////////////////////////////////////////////////////
+
+void FSparseVolumeAssetHeader::Serialize(FArchive& Ar)
+{
+	Ar << Version;
+
+	if (Version == 0)
+	{
+		Ar << PageTableVolumeResolution;
+		Ar << TileDataVolumeResolution;
+		Ar << SourceVolumeResolution;
+		Ar << SourceVolumeAABBMin;
+		UE::SVT::Private::SerializeEnumAs<uint8>(Ar, PackedDataAFormat);
+		UE::SVT::Private::SerializeEnumAs<uint8>(Ar, PackedDataBFormat);
+	}
+	else
+	{
+		// FSparseVolumeAssetHeader needs to account for new version
+		check(false);
+	}
+}
+
+////////////////////////////////////////////////////////////////////////////////////////////////
+
 void FSparseVolumeTextureRuntime::Serialize(FArchive& Ar)
 {
 	Header.Serialize(Ar);
@@ -128,6 +136,85 @@
 
 ////////////////////////////////////////////////////////////////////////////////////////////////
 
+FSparseVolumeTextureSceneProxy::FSparseVolumeTextureSceneProxy()
+	: FRenderResource()
+	, SparseVolumeTextureRuntime(nullptr)
+	, PageTableTextureRHI(nullptr)
+	, PhysicalTileDataATextureRHI(nullptr)
+	, PhysicalTileDataBTextureRHI(nullptr)
+{
+}
+
+FSparseVolumeTextureSceneProxy::~FSparseVolumeTextureSceneProxy()
+{
+}
+
+void FSparseVolumeTextureSceneProxy::InitialiseRuntimeData(FSparseVolumeTextureRuntime& SparseVolumeTextureRuntimeIn)
+{
+	SparseVolumeTextureRuntime = &SparseVolumeTextureRuntimeIn;
+}
+
+void FSparseVolumeTextureSceneProxy::InitRHI()
+{
+	// Page table
+	{
+		EPixelFormat PageEntryFormat = PF_R32_UINT;
+		FIntVector3 PageTableVolumeResolution = SparseVolumeTextureRuntime->Header.PageTableVolumeResolution;
+		const FRHITextureCreateDesc Desc =
+			FRHITextureCreateDesc::Create3D(TEXT("SparseVolumeTexture.PageTable.RHITexture"),
+				PageTableVolumeResolution.X, PageTableVolumeResolution.Y, PageTableVolumeResolution.Z, PageEntryFormat)
+			.SetFlags(ETextureCreateFlags::ShaderResource);
+
+		PageTableTextureRHI = RHICreateTexture(Desc);
+
+		const int32 FormatSize = GPixelFormats[PageEntryFormat].BlockBytes;
+		const FUpdateTextureRegion3D UpdateRegion(0, 0, 0, 0, 0, 0, PageTableVolumeResolution.X, PageTableVolumeResolution.Y, PageTableVolumeResolution.Z);
+		RHIUpdateTexture3D(PageTableTextureRHI, 0, UpdateRegion, UpdateRegion.Width * FormatSize, UpdateRegion.Width * UpdateRegion.Height * FormatSize, (const uint8*)SparseVolumeTextureRuntime->PageTable.GetData());
+	}
+
+	// Tile data
+	{
+		FIntVector3 TileDataVolumeResolution = SparseVolumeTextureRuntime->Header.TileDataVolumeResolution;
+		EPixelFormat VoxelFormatA = SparseVolumeTextureRuntime->Header.PackedDataAFormat;
+		EPixelFormat VoxelFormatB = SparseVolumeTextureRuntime->Header.PackedDataBFormat;
+		const FUpdateTextureRegion3D UpdateRegion(0, 0, 0, 0, 0, 0, TileDataVolumeResolution.X, TileDataVolumeResolution.Y, TileDataVolumeResolution.Z);
+
+		// A
+		if (VoxelFormatA != PF_Unknown)
+		{
+			const FRHITextureCreateDesc Desc =
+				FRHITextureCreateDesc::Create3D(TEXT("SparseVolumeTexture.PhysicalTileDataA.RHITexture"),
+					TileDataVolumeResolution.X, TileDataVolumeResolution.Y, TileDataVolumeResolution.Z, VoxelFormatA)
+				.SetFlags(ETextureCreateFlags::ShaderResource);
+
+			PhysicalTileDataATextureRHI = RHICreateTexture(Desc);
+
+			const int32 FormatSize = GPixelFormats[VoxelFormatA].BlockBytes;
+			RHIUpdateTexture3D(PhysicalTileDataATextureRHI, 0, UpdateRegion, UpdateRegion.Width * FormatSize, UpdateRegion.Width * UpdateRegion.Height * FormatSize, (const uint8*)SparseVolumeTextureRuntime->PhysicalTileDataA.GetData());
+		}
+
+		// B
+		if (VoxelFormatB != PF_Unknown)
+		{
+			const FRHITextureCreateDesc Desc =
+				FRHITextureCreateDesc::Create3D(TEXT("SparseVolumeTexture.PhysicalTileDataB.RHITexture"),
+					TileDataVolumeResolution.X, TileDataVolumeResolution.Y, TileDataVolumeResolution.Z, VoxelFormatB)
+				.SetFlags(ETextureCreateFlags::ShaderResource);
+
+			PhysicalTileDataBTextureRHI = RHICreateTexture(Desc);
+
+			const int32 FormatSize = GPixelFormats[VoxelFormatB].BlockBytes;
+			RHIUpdateTexture3D(PhysicalTileDataBTextureRHI, 0, UpdateRegion, UpdateRegion.Width * FormatSize, UpdateRegion.Width * UpdateRegion.Height * FormatSize, (const uint8*)SparseVolumeTextureRuntime->PhysicalTileDataB.GetData());
+		}
+	}
+}
+
+void FSparseVolumeTextureSceneProxy::ReleaseRHI()
+{
+}
+
+////////////////////////////////////////////////////////////////////////////////////////////////
+
 FSparseVolumeTextureFrame::FSparseVolumeTextureFrame()
 	: RuntimeStreamedInData()
 	, SparseVolumeTextureRuntime()
@@ -142,6 +229,47 @@
 {
 }
 
+bool FSparseVolumeTextureFrame::BuildRuntimeData()
+{
+#if WITH_EDITORONLY_DATA
+	// Check if the virtualized bulk data payload is available now
+	if (RawData.HasPayloadData())
+	{
+		// First, read the source data in from the raw data stored as bulk data
+		UE::Serialization::FEditorBulkDataReader RawDataArchiveReader(RawData);
+		FSparseVolumeRawSource SparseVolumeRawSource;
+		SparseVolumeRawSource.Serialize(RawDataArchiveReader);
+
+		// Then, convert the raw source data (OpenVDB) to SVT
+		FOpenVDBToSVTConversionResult SVTResult;
+		SVTResult.Header = &SparseVolumeTextureRuntime.Header;
+		SVTResult.PageTable = &SparseVolumeTextureRuntime.PageTable;
+		SVTResult.PhysicalTileDataA = &SparseVolumeTextureRuntime.PhysicalTileDataA;
+		SVTResult.PhysicalTileDataB = &SparseVolumeTextureRuntime.PhysicalTileDataB;
+
+		const bool bSuccess = ConvertOpenVDBToSparseVolumeTextureDelegate.IsBound() && ConvertOpenVDBToSparseVolumeTextureDelegate.Execute(
+			SparseVolumeRawSource.SourceAssetFile,
+			SparseVolumeRawSource.PackedDataA,
+			SparseVolumeRawSource.PackedDataB,
+			&SVTResult,
+			false, FVector::Zero(), FVector::Zero());
+		
+		if (!bSuccess)
+		{
+			// Clear any writes that may have happened to this data during the cook attempt
+			SparseVolumeTextureRuntime = {};
+			return false;
+		}
+
+		// Now unload the raw data
+		RawData.UnloadData();
+		
+		return true;
+	}
+#endif
+	return false;
+}
+
 ////////////////////////////////////////////////////////////////////////////////////////////////
 
 USparseVolumeTexture::USparseVolumeTexture(const FObjectInitializer& ObjectInitializer)
@@ -149,6 +277,23 @@
 {
 }
 
+void USparseVolumeTexture::GetFrameUVScaleBias(int32 FrameIndex, FVector* OutScale, FVector* OutBias) const
+{
+	*OutScale = FVector::One();
+	*OutBias = FVector::Zero();
+	const FSparseVolumeTextureSceneProxy* Proxy = GetSparseVolumeTextureSceneProxy(FrameIndex);
+	if (Proxy)
+	{
+		const FSparseVolumeAssetHeader& Header = Proxy->GetHeader();
+		const FBox VolumeBounds = GetVolumeBounds();
+		check(VolumeBounds.GetVolume() > 0.0);
+		const FBox FrameBounds = FBox(FVector(Header.SourceVolumeAABBMin), FVector(Header.SourceVolumeAABBMin + Header.PageTableVolumeResolution * SPARSE_VOLUME_TILE_RES)); // AABB of current frame
+
+		*OutScale = VolumeBounds.GetExtent() / FrameBounds.GetExtent();
+		*OutBias = -((FrameBounds.Min - VolumeBounds.Min) / (VolumeBounds.Max - VolumeBounds.Min) * *OutScale);
+	}
+}
+
 UE::Shader::EValueType USparseVolumeTexture::GetUniformParameterType(int32 Index)
 {
 	switch (Index)
@@ -197,15 +342,17 @@
 
 	FStripDataFlags StripFlags(Ar);
 
-	// The runtime bulk data for static sparse volume texture is always loaded, not streamed in.
-	StaticFrame.RuntimeStreamedInData.SetBulkDataFlags(BULKDATA_ForceInlinePayload);
-
 	if (StripFlags.IsEditorDataStripped() && Ar.IsLoadingFromCookedPackage())
 	{
 		// In this case we are loading in game with a cooked build so we only need to load the runtime data.
+
+		// Read cooked bulk data from archive
+		StaticFrame.RuntimeStreamedInData.Serialize(Ar, this);
+
+		// Create runtime data from cooked bulk data
 		{
-			FBulkDataReader RuntimeStreamedInData(StaticFrame.RuntimeStreamedInData);
-			StaticFrame.SparseVolumeTextureRuntime.Serialize(RuntimeStreamedInData);
+			FBulkDataReader BulkDataReader(StaticFrame.RuntimeStreamedInData);
+			StaticFrame.SparseVolumeTextureRuntime.Serialize(BulkDataReader);
 		}
 
 		// The bulk data is no longer needed
@@ -218,12 +365,21 @@
 	else if (Ar.IsCooking())
 	{
 		// We are cooking the game, serialize the asset out.
-		ConvertRawSourceDataToSparseVolumeTextureRuntime();
+		
+		// The runtime bulk data for static sparse volume texture is always loaded, not streamed in.
+		StaticFrame.RuntimeStreamedInData.SetBulkDataFlags(BULKDATA_ForceInlinePayload);
 
+		const bool bBuiltRuntimeData = StaticFrame.BuildRuntimeData();
+		check(bBuiltRuntimeData); // SVT_TODO: actual error handling
+		
+		// Write runtime data into RuntimeStreamedInData
 		{
-			FBulkDataWriter RuntimeStreamedInData(StaticFrame.RuntimeStreamedInData);
-			StaticFrame.SparseVolumeTextureRuntime.Serialize(RuntimeStreamedInData);
+			FBulkDataWriter BulkDataWriter(StaticFrame.RuntimeStreamedInData);
+			StaticFrame.SparseVolumeTextureRuntime.Serialize(BulkDataWriter);
 		}
+
+		// And now write the cooked bulk data to the archive
+		StaticFrame.RuntimeStreamedInData.Serialize(Ar, this);
 	}
 	else if (!Ar.IsObjectReferenceCollector())
 	{
@@ -251,36 +407,62 @@
 }
 #endif // WITH_EDITOR
 
+const FSparseVolumeAssetHeader* UStaticSparseVolumeTexture::GetSparseVolumeTextureHeader(int32 FrameIndex) const
+{
+	return &StaticFrame.SparseVolumeTextureRuntime.Header;
+}
+FSparseVolumeTextureSceneProxy* UStaticSparseVolumeTexture::GetSparseVolumeTextureSceneProxy(int32 FrameIndex)
+{ 
+	return &StaticFrame.SparseVolumeTextureSceneProxy; 
+}
+const FSparseVolumeTextureSceneProxy* UStaticSparseVolumeTexture::GetSparseVolumeTextureSceneProxy(int32 FrameIndex) const
+{ 
+	return &StaticFrame.SparseVolumeTextureSceneProxy; 
+}
+
+FVector4 UStaticSparseVolumeTexture::GetUniformParameter(int32 Index, int32 FrameIndex) const
+{
+	const FSparseVolumeTextureSceneProxy* Proxy = GetSparseVolumeTextureSceneProxy(0);
+	if (Proxy)
+	{
+		const FSparseVolumeAssetHeader& Header = Proxy->GetHeader();
+		switch (Index)
+		{
+		case ESparseVolumeTexture_PhysicalUVToPageUV:
+		{
+			// 3d uv coordinates are specified in [0, 1]. Before addressing the page tables which might have padding,
+			// since source volume resolution might not align to tile resolution, we have to rescale the uv so that [0,1] maps to the source texture boundaries.
+			FVector3f PhysicalUVToPageUV = FVector3f(Header.SourceVolumeResolution)/ FVector3f(Header.PageTableVolumeResolution * SPARSE_VOLUME_TILE_RES);
+			return FVector4(PhysicalUVToPageUV.X, PhysicalUVToPageUV.Y, PhysicalUVToPageUV.Z, 0.0f);
+		}
+		case ESparseVolumeTexture_PageTableSize:
+		{
+			return FVector4(Header.PageTableVolumeResolution.X, Header.PageTableVolumeResolution.Y, Header.PageTableVolumeResolution.Z, 0.0f);
+		}
+		case ESparseVolumeTexture_TileSize:				return FVector4(float(SPARSE_VOLUME_TILE_RES), 0.0f, 0.0f, 0.0f);
+		default:
+			break;
+		}
+		check(0);
+		return FVector4(ForceInitToZero);
+	}
+	// 0 while waiting for the proxy
+	return FVector4(ForceInitToZero);
+}
+
+FBox UStaticSparseVolumeTexture::GetVolumeBounds() const
+{
+	return VolumeBounds;
+}
+
 void UStaticSparseVolumeTexture::ConvertRawSourceDataToSparseVolumeTextureRuntime()
 {
 #if WITH_EDITOR
 	// Check if the virtualized bulk data payload is available now
 	if (StaticFrame.RawData.HasPayloadData())
 	{
-		// When cooking, the runtime data must be serialized out.
-
-		// First, read the source data in from the raw data stored as bulk data
-		UE::Serialization::FEditorBulkDataReader RawDataArchiveReader(StaticFrame.RawData);
-		FSparseVolumeRawSource SparseVolumeRawSource;
-		SparseVolumeRawSource.Serialize(RawDataArchiveReader);
-
-		FOpenVDBToSVTConversionResult SVTResult;
-		SVTResult.Header = &StaticFrame.SparseVolumeTextureRuntime.Header;
-		SVTResult.PageTable = &StaticFrame.SparseVolumeTextureRuntime.PageTable;
-		SVTResult.PhysicalTileDataA = &StaticFrame.SparseVolumeTextureRuntime.PhysicalTileDataA;
-		SVTResult.PhysicalTileDataB = &StaticFrame.SparseVolumeTextureRuntime.PhysicalTileDataB;
-
-		// Then, cook the runtime data and serialize it out.
-		const bool bSuccess = ConvertOpenVDBToSparseVolumeTextureDelegate.IsBound() && ConvertOpenVDBToSparseVolumeTextureDelegate.Execute(
-			SparseVolumeRawSource.SourceAssetFile,
-			SparseVolumeRawSource.PackedDataA,
-			SparseVolumeRawSource.PackedDataB,
-			&SVTResult,
-			false, FVector::Zero(), FVector::Zero());
+		const bool bSuccess = StaticFrame.BuildRuntimeData();
 		ensure(bSuccess);
-
-		// Now unload the raw data
-		StaticFrame.RawData.UnloadData();
 	}
 	else
 	{
@@ -290,21 +472,6 @@
 #endif // WITH_EDITOR
 }
 
-void UStaticSparseVolumeTexture::GenerateOrLoadDDCRuntimeDataAndCreateSceneProxy()
-{
-	// Release any previously allocated render thread proxy
-	BeginReleaseResource(&StaticFrame.SparseVolumeTextureSceneProxy);
-
-#if WITH_EDITORONLY_DATA
-	// We only fetch/put DDC when in editor. Otherwise, StaticFrame.SparseVolumeTextureRuntime is serialize in.
-	GenerateOrLoadDDCRuntimeData();
-#endif
-
-	// Runtime data is now valid, create the render thread proxy
-	StaticFrame.SparseVolumeTextureSceneProxy.InitialiseRuntimeData(StaticFrame.SparseVolumeTextureRuntime);
-	BeginInitResource(&StaticFrame.SparseVolumeTextureSceneProxy);
-}
-
 void UStaticSparseVolumeTexture::GenerateOrLoadDDCRuntimeData()
 {
 #if WITH_EDITORONLY_DATA
@@ -359,126 +526,19 @@
 #endif // WITH_EDITORONLY_DATA
 }
 
-const FSparseVolumeAssetHeader* UStaticSparseVolumeTexture::GetSparseVolumeTextureHeader() const
+void UStaticSparseVolumeTexture::GenerateOrLoadDDCRuntimeDataAndCreateSceneProxy()
 {
-	return &StaticFrame.SparseVolumeTextureRuntime.Header;
-}
-FSparseVolumeTextureSceneProxy* UStaticSparseVolumeTexture::GetSparseVolumeTextureSceneProxy()
-{ 
-	return &StaticFrame.SparseVolumeTextureSceneProxy; 
-}
-const FSparseVolumeTextureSceneProxy* UStaticSparseVolumeTexture::GetSparseVolumeTextureSceneProxy() const
-{ 
-	return &StaticFrame.SparseVolumeTextureSceneProxy; 
-}
+	// Release any previously allocated render thread proxy
+	BeginReleaseResource(&StaticFrame.SparseVolumeTextureSceneProxy);
 
-FVector4 UStaticSparseVolumeTexture::GetUniformParameter(int32 Index) const
-{
-	const FSparseVolumeTextureSceneProxy* Proxy = GetSparseVolumeTextureSceneProxy();
-	if (Proxy)
-	{
-		const FSparseVolumeAssetHeader& Header = Proxy->GetHeader();
-		switch (Index)
-		{
-		case ESparseVolumeTexture_PhysicalUVToPageUV:
-		{
-			// 3d uv coordinates are specified in [0, 1]. Before addressing the page tables which might have padding,
-			// since source volume resolution might not align to tile resolution, we have to rescale the uv so that [0,1] maps to the source texture boundaries.
-			FVector3f PhysicalUVToPageUV = FVector3f(Header.SourceVolumeResolution)/ FVector3f(Header.PageTableVolumeResolution * SPARSE_VOLUME_TILE_RES);
-			return FVector4(PhysicalUVToPageUV.X, PhysicalUVToPageUV.Y, PhysicalUVToPageUV.Z, 0.0f);
-		}
-		case ESparseVolumeTexture_PageTableSize:
-		{
-			return FVector4(Header.PageTableVolumeResolution.X, Header.PageTableVolumeResolution.Y, Header.PageTableVolumeResolution.Z, 0.0f);
-		}
-		case ESparseVolumeTexture_TileSize:				return FVector4(float(SPARSE_VOLUME_TILE_RES), 0.0f, 0.0f, 0.0f);
-		default:
-			break;
-		}
-		check(0);
-		return FVector4(ForceInitToZero);
-	}
-	// 0 while waiting for the proxy
-	return FVector4(ForceInitToZero);
+#if WITH_EDITORONLY_DATA
+	// We only fetch/put DDC when in editor. Otherwise, StaticFrame.SparseVolumeTextureRuntime is serialize in.
+	GenerateOrLoadDDCRuntimeData();
+#endif
+
+	// Runtime data is now valid, create the render thread proxy
+	StaticFrame.SparseVolumeTextureSceneProxy.InitialiseRuntimeData(StaticFrame.SparseVolumeTextureRuntime);
+	BeginInitResource(&StaticFrame.SparseVolumeTextureSceneProxy);
 }
 
 ////////////////////////////////////////////////////////////////////////////////////////////////
-
-FSparseVolumeTextureSceneProxy::FSparseVolumeTextureSceneProxy()
-	: FRenderResource()
-	, SparseVolumeTextureRuntime(nullptr)
-	, PageTableTextureRHI(nullptr)
-	, PhysicalTileDataATextureRHI(nullptr)
-	, PhysicalTileDataBTextureRHI(nullptr)
-{
-}
-
-FSparseVolumeTextureSceneProxy::~FSparseVolumeTextureSceneProxy()
-{
-}
-
-void FSparseVolumeTextureSceneProxy::InitialiseRuntimeData(FSparseVolumeTextureRuntime& SparseVolumeTextureRuntimeIn)
-{
-	SparseVolumeTextureRuntime = &SparseVolumeTextureRuntimeIn;
-}
-
-void FSparseVolumeTextureSceneProxy::InitRHI()
-{
-	// Page table
-	{
-		EPixelFormat PageEntryFormat = PF_R32_UINT;
-		FIntVector3 PageTableVolumeResolution = SparseVolumeTextureRuntime->Header.PageTableVolumeResolution;
-		const FRHITextureCreateDesc Desc =
-			FRHITextureCreateDesc::Create3D(TEXT("SparseVolumeTexture.PageTable.RHITexture"),
-				PageTableVolumeResolution.X, PageTableVolumeResolution.Y, PageTableVolumeResolution.Z, PageEntryFormat)
-			.SetFlags(ETextureCreateFlags::ShaderResource);
-
-		PageTableTextureRHI = RHICreateTexture(Desc);
-
-		const int32 FormatSize = GPixelFormats[PageEntryFormat].BlockBytes;
-		const FUpdateTextureRegion3D UpdateRegion(0, 0, 0, 0, 0, 0, PageTableVolumeResolution.X, PageTableVolumeResolution.Y, PageTableVolumeResolution.Z);
-		RHIUpdateTexture3D(PageTableTextureRHI, 0, UpdateRegion, UpdateRegion.Width * FormatSize, UpdateRegion.Width * UpdateRegion.Height * FormatSize, (const uint8*)SparseVolumeTextureRuntime->PageTable.GetData());
-	}
-
-	// Tile data
-	{
-		FIntVector3 TileDataVolumeResolution = SparseVolumeTextureRuntime->Header.TileDataVolumeResolution;
-		EPixelFormat VoxelFormatA = SparseVolumeTextureRuntime->Header.PackedDataAFormat;
-		EPixelFormat VoxelFormatB = SparseVolumeTextureRuntime->Header.PackedDataBFormat;
-		const FUpdateTextureRegion3D UpdateRegion(0, 0, 0, 0, 0, 0, TileDataVolumeResolution.X, TileDataVolumeResolution.Y, TileDataVolumeResolution.Z);
-
-		// A
-		if (VoxelFormatA != PF_Unknown)
-		{
-			const FRHITextureCreateDesc Desc =
-				FRHITextureCreateDesc::Create3D(TEXT("SparseVolumeTexture.PhysicalTileDataA.RHITexture"),
-					TileDataVolumeResolution.X, TileDataVolumeResolution.Y, TileDataVolumeResolution.Z, VoxelFormatA)
-				.SetFlags(ETextureCreateFlags::ShaderResource);
-
-			PhysicalTileDataATextureRHI = RHICreateTexture(Desc);
-
-			const int32 FormatSize = GPixelFormats[VoxelFormatA].BlockBytes;
-			RHIUpdateTexture3D(PhysicalTileDataATextureRHI, 0, UpdateRegion, UpdateRegion.Width * FormatSize, UpdateRegion.Width * UpdateRegion.Height * FormatSize, (const uint8*)SparseVolumeTextureRuntime->PhysicalTileDataA.GetData());
-		}
-		
-		// B
-		if (VoxelFormatB != PF_Unknown)
-		{
-			const FRHITextureCreateDesc Desc =
-				FRHITextureCreateDesc::Create3D(TEXT("SparseVolumeTexture.PhysicalTileDataB.RHITexture"),
-					TileDataVolumeResolution.X, TileDataVolumeResolution.Y, TileDataVolumeResolution.Z, VoxelFormatB)
-				.SetFlags(ETextureCreateFlags::ShaderResource);
-
-			PhysicalTileDataBTextureRHI = RHICreateTexture(Desc);
-
-			const int32 FormatSize = GPixelFormats[VoxelFormatB].BlockBytes;
-			RHIUpdateTexture3D(PhysicalTileDataBTextureRHI, 0, UpdateRegion, UpdateRegion.Width * FormatSize, UpdateRegion.Width * UpdateRegion.Height * FormatSize, (const uint8*)SparseVolumeTextureRuntime->PhysicalTileDataB.GetData());
-		}
-	}
-}
-
-void FSparseVolumeTextureSceneProxy::ReleaseRHI()
-{
-}
-
-////////////////////////////////////////////////////////////////////////////////////////////////
@@ -507,15 +567,6 @@
 	Super::PostLoad();
 }
 
-int32 UAnimatedSparseVolumeTexture::GetFrameCountToLoad() const 
-{
-	if (FrameCount > 0)
-	{
-		return bLoadAllFramesToProxies ? FrameCount : 1;
-	}
-	return 0;
-}
-
 void UAnimatedSparseVolumeTexture::BeginDestroy()
 {
 	Super::BeginDestroy();
@@ -536,11 +587,54 @@
 
 	if (StripFlags.IsEditorDataStripped() && Ar.IsLoadingFromCookedPackage())
 	{
-		check(false); // SVT_TODO Implement
+		// In this case we are loading in game with a cooked build so we only need to load the runtime data.
+
+		AnimationFrames.SetNum(FrameCount);
+		for (int32 i = 0; i < FrameCount; ++i)
+		{
+			FSparseVolumeTextureFrame& Frame = AnimationFrames[i];
+			
+			// Read cooked bulk data from archive
+			Frame.RuntimeStreamedInData.Serialize(Ar, this);
+
+			// Create runtime data from cooked bulk data
+			{
+				FBulkDataReader RuntimeStreamedInData(Frame.RuntimeStreamedInData);
+				Frame.SparseVolumeTextureRuntime.Serialize(RuntimeStreamedInData);
+			}
+
+			// The bulk data is no longer needed
+			Frame.RuntimeStreamedInData.RemoveBulkData();
+
+			// Runtime data is now valid, create the render thread proxy
+			Frame.SparseVolumeTextureSceneProxy.InitialiseRuntimeData(Frame.SparseVolumeTextureRuntime);
+			BeginInitResource(&Frame.SparseVolumeTextureSceneProxy);
+		}
 	}
 	else if (Ar.IsCooking())
 	{
-		check(false); // SVT_TODO Implement
+		// We are cooking the game, serialize the asset out.
+
+		check(AnimationFrames.Num() == FrameCount);
+		for (int32 i = 0; i < FrameCount; ++i)
+		{
+			FSparseVolumeTextureFrame& Frame = AnimationFrames[i];
+
+			// SVT_TODO: Until we have streaming, the runtime bulk data for all sparse volume texture frames is always loaded.
+			AnimationFrames[i].RuntimeStreamedInData.SetBulkDataFlags(BULKDATA_ForceInlinePayload);
+
+			const bool bBuiltRuntimeData = Frame.BuildRuntimeData();
+			check(bBuiltRuntimeData); // SVT_TODO
+
+			// Write runtime data into RuntimeStreamedInData
+			{
+				FBulkDataWriter RuntimeStreamedInData(Frame.RuntimeStreamedInData);
+				Frame.SparseVolumeTextureRuntime.Serialize(RuntimeStreamedInData);
+			}
+
+			// And now write the cooked bulk data to the archive
+			Frame.RuntimeStreamedInData.Serialize(Ar, this);
+		}
 	}
 	else if (!Ar.IsObjectReferenceCollector())
 	{
@@ -575,9 +669,96 @@
 void UAnimatedSparseVolumeTexture::PostEditChangeProperty(FPropertyChangedEvent& PropertyChangedEvent)
 {
 	Super::PostEditChangeProperty(PropertyChangedEvent);
+
+	const int32 FrameCountToLoad = GetFrameCountToLoad();
+	for (int32 FrameIndex = 0; FrameIndex < FrameCountToLoad; FrameIndex++)
+	{
+		GenerateOrLoadDDCRuntimeDataAndCreateSceneProxy(FrameIndex);
+	}
 }
 #endif // WITH_EDITOR
 
+const FSparseVolumeAssetHeader* UAnimatedSparseVolumeTexture::GetSparseVolumeTextureHeader(int32 FrameIndex) const
+{
+	// When an AnimatedSparseVolumeTexture is used as SparseVolumeTexture, it can only be previewed using its first frame.
+	check(AnimationFrames.Num() >= 1);
+	FrameIndex = FrameIndex % GetFrameCountToLoad();
+	const FSparseVolumeTextureFrame& Frame = AnimationFrames[FrameIndex];
+	return &Frame.SparseVolumeTextureRuntime.Header;
+}
+
+FSparseVolumeTextureSceneProxy* UAnimatedSparseVolumeTexture::GetSparseVolumeTextureSceneProxy(int32 FrameIndex)
+{
+	// When an AnimatedSparseVolumeTexture is used as SparseVolumeTexture, it can only be previewed using its first frame.
+	check(AnimationFrames.Num() >= 1);
+	FrameIndex = FrameIndex % GetFrameCountToLoad();
+	FSparseVolumeTextureFrame& Frame = AnimationFrames[FrameIndex];
+	return &Frame.SparseVolumeTextureSceneProxy;
+}
+
+const FSparseVolumeTextureSceneProxy* UAnimatedSparseVolumeTexture::GetSparseVolumeTextureSceneProxy(int32 FrameIndex) const
+{
+	// When an AnimatedSparseVolumeTexture is used as SparseVolumeTexture, it can only be previewed using its first frame.
+	check(AnimationFrames.Num() >= 1);
+	FrameIndex = FrameIndex % GetFrameCountToLoad();
+	const FSparseVolumeTextureFrame& Frame = AnimationFrames[FrameIndex];
+	return &Frame.SparseVolumeTextureSceneProxy;
+}
+
+FVector4 UAnimatedSparseVolumeTexture::GetUniformParameter(int32 Index, int32 FrameIndex) const
+{
+	// When an AnimatedSparseVolumeTexture is used as SparseVolumeTexture, it can only be previewed using its first frame.
+	const FSparseVolumeTextureSceneProxy* Proxy = GetSparseVolumeTextureSceneProxy(FrameIndex);
+	if (Proxy)
+	{
+		const FSparseVolumeAssetHeader& Header = Proxy->GetHeader();
+		switch (Index)
+		{
+		case ESparseVolumeTexture_PhysicalUVToPageUV:
+		{
+			// 3d uv coordinates are specified in [0, 1]. Before addressing the page tables which might have padding,
+			// since source volume resolution might not align to tile resolution, we have to rescale the uv so that [0,1] maps to the source texture boundaries.
+			FVector3f PhysicalUVToPageUV = FVector3f(Header.SourceVolumeResolution) / FVector3f(Header.PageTableVolumeResolution * SPARSE_VOLUME_TILE_RES);
+			return FVector4(PhysicalUVToPageUV.X, PhysicalUVToPageUV.Y, PhysicalUVToPageUV.Z, 0.0f);
+		}
+		case ESparseVolumeTexture_PageTableSize:
+		{
+			return FVector4(Header.PageTableVolumeResolution.X, Header.PageTableVolumeResolution.Y, Header.PageTableVolumeResolution.Z, 0.0f);
+		}
+		case ESparseVolumeTexture_TileSize:				return FVector4(float(SPARSE_VOLUME_TILE_RES), 0.0f, 0.0f, 0.0f);
+		default:
+			break;
+		}
+		check(0);
+		return FVector4(ForceInitToZero);
+	}
+
+	// 0 while waiting for the proxy
+	return FVector4(ForceInitToZero);
+}
+
+FBox UAnimatedSparseVolumeTexture::GetVolumeBounds() const
+{
+	return VolumeBounds;
+}
+
+const FSparseVolumeTextureSceneProxy* UAnimatedSparseVolumeTexture::GetSparseVolumeTextureFrameSceneProxy(int32 FrameIndex) const
+{
+	check(AnimationFrames.Num() >= 1);
+	FrameIndex = FrameIndex % GetFrameCountToLoad();
+	// SVT_TODO when streaming is enabled, this will likely change.
+	return &AnimationFrames[FrameIndex].SparseVolumeTextureSceneProxy;
+}
+
+int32 UAnimatedSparseVolumeTexture::GetFrameCountToLoad() const
+{
+	if (FrameCount > 0)
+	{
+		return bLoadAllFramesToProxies ? FrameCount : 1;
+	}
+	return 0;
+}
+
 void UAnimatedSparseVolumeTexture::ConvertRawSourceDataToSparseVolumeTextureRuntime(int32 FrameIndex)
 {
 #if WITH_EDITOR
@@ -586,30 +767,8 @@
 	// Check if the virtualized bulk data payload is available now
 	if (Frame.RawData.HasPayloadData())
 	{
-		// When cooking, the runtime data must be serialized out.
-
-		// First, read the source data in from the raw data stored as bulk data
-		UE::Serialization::FEditorBulkDataReader RawDataArchiveReader(Frame.RawData);
-		FSparseVolumeRawSource SparseVolumeRawSource;
-		SparseVolumeRawSource.Serialize(RawDataArchiveReader);
-
-		FOpenVDBToSVTConversionResult SVTResult;
-		SVTResult.Header = &Frame.SparseVolumeTextureRuntime.Header;
-		SVTResult.PageTable = &Frame.SparseVolumeTextureRuntime.PageTable;
-		SVTResult.PhysicalTileDataA = &Frame.SparseVolumeTextureRuntime.PhysicalTileDataA;
-		SVTResult.PhysicalTileDataB = &Frame.SparseVolumeTextureRuntime.PhysicalTileDataB;
-
-		// Then, cook the runtime data and serialize it out.
-		const bool bSuccess = ConvertOpenVDBToSparseVolumeTextureDelegate.IsBound() && ConvertOpenVDBToSparseVolumeTextureDelegate.Execute(
-			SparseVolumeRawSource.SourceAssetFile,
-			SparseVolumeRawSource.PackedDataA,
-			SparseVolumeRawSource.PackedDataB,
-			&SVTResult,
-			false, FVector::Zero(), FVector::Zero());
+		const bool bSuccess = Frame.BuildRuntimeData();
 		ensure(bSuccess);
-
-		// Now unload the raw data
-		Frame.RawData.UnloadData();
 	}
 	else
 	{
@@ -619,23 +778,6 @@
 #endif // WITH_EDITOR
 }
 
-void UAnimatedSparseVolumeTexture::GenerateOrLoadDDCRuntimeDataAndCreateSceneProxy(int32 FrameIndex)
-{
-	FSparseVolumeTextureFrame& Frame = AnimationFrames[FrameIndex];
-
-	// Release any previously allocated render thread proxy
-	BeginReleaseResource(&Frame.SparseVolumeTextureSceneProxy);
-
-#if WITH_EDITORONLY_DATA
-	// We only fetch/put DDC when in editor. Otherwise, StaticFrame.SparseVolumeTextureRuntime is serialize in.
-	GenerateOrLoadDDCRuntimeData(FrameIndex);
-#endif
-
-	// Runtime data is now valid, create the render thread proxy
-	Frame.SparseVolumeTextureSceneProxy.InitialiseRuntimeData(Frame.SparseVolumeTextureRuntime);
-	BeginInitResource(&Frame.SparseVolumeTextureSceneProxy);
-}
-
 void UAnimatedSparseVolumeTexture::GenerateOrLoadDDCRuntimeData(int32 FrameIndex)
 {
 #if WITH_EDITORONLY_DATA
@@ -693,69 +835,21 @@
 #endif // WITH_EDITORONLY_DATA
 }
 
-const FSparseVolumeAssetHeader* UAnimatedSparseVolumeTexture::GetSparseVolumeTextureHeader() const
+void UAnimatedSparseVolumeTexture::GenerateOrLoadDDCRuntimeDataAndCreateSceneProxy(int32 FrameIndex)
 {
-	// When an AnimatedSparseVolumeTexture is used as SparseVolumeTexture, it can only be previewed using its first frame.
-	check(AnimationFrames.Num() >= 1);
-	const FSparseVolumeTextureFrame& Frame = AnimationFrames[PreviewFrameIndex];
-	return &Frame.SparseVolumeTextureRuntime.Header;
+	FSparseVolumeTextureFrame& Frame = AnimationFrames[FrameIndex];
+
+	// Release any previously allocated render thread proxy
+	BeginReleaseResource(&Frame.SparseVolumeTextureSceneProxy);
+
+#if WITH_EDITORONLY_DATA
+	// We only fetch/put DDC when in editor. Otherwise, StaticFrame.SparseVolumeTextureRuntime is serialize in.
+	GenerateOrLoadDDCRuntimeData(FrameIndex);
+#endif
+
+	// Runtime data is now valid, create the render thread proxy
+	Frame.SparseVolumeTextureSceneProxy.InitialiseRuntimeData(Frame.SparseVolumeTextureRuntime);
+	BeginInitResource(&Frame.SparseVolumeTextureSceneProxy);
 }
 
-FSparseVolumeTextureSceneProxy* UAnimatedSparseVolumeTexture::GetSparseVolumeTextureSceneProxy()
-{
-	// When an AnimatedSparseVolumeTexture is used as SparseVolumeTexture, it can only be previewed using its first frame.
-	check(AnimationFrames.Num() >= 1);
-	FSparseVolumeTextureFrame& Frame = AnimationFrames[PreviewFrameIndex];
-	return &Frame.SparseVolumeTextureSceneProxy;
-}
-
-const FSparseVolumeTextureSceneProxy* UAnimatedSparseVolumeTexture::GetSparseVolumeTextureSceneProxy() const
-{
-	// When an AnimatedSparseVolumeTexture is used as SparseVolumeTexture, it can only be previewed using its first frame.
-	check(AnimationFrames.Num() >= 1);
-	const FSparseVolumeTextureFrame& Frame = AnimationFrames[PreviewFrameIndex];
-	return &Frame.SparseVolumeTextureSceneProxy;
-}
-
-FVector4 UAnimatedSparseVolumeTexture::GetUniformParameter(int32 Index) const
-{
-	// When an AnimatedSparseVolumeTexture is used as SparseVolumeTexture, it can only be previewed using its first frame.
-	const FSparseVolumeTextureSceneProxy* Proxy = GetSparseVolumeTextureSceneProxy();
-	if (Proxy)
-	{
-		const FSparseVolumeAssetHeader& Header = Proxy->GetHeader();
-		switch (Index)
-		{
-		case ESparseVolumeTexture_PhysicalUVToPageUV:
-		{
-			// 3d uv coordinates are specified in [0, 1]. Before addressing the page tables which might have padding,
-			// since source volume resolution might not align to tile resolution, we have to rescale the uv so that [0,1] maps to the source texture boundaries.
-			FVector3f PhysicalUVToPageUV = FVector3f(Header.SourceVolumeResolution) / FVector3f(Header.PageTableVolumeResolution * SPARSE_VOLUME_TILE_RES);
-			return FVector4(PhysicalUVToPageUV.X, PhysicalUVToPageUV.Y, PhysicalUVToPageUV.Z, 0.0f);
-		}
-		case ESparseVolumeTexture_PageTableSize:
-		{
-			return FVector4(Header.PageTableVolumeResolution.X, Header.PageTableVolumeResolution.Y, Header.PageTableVolumeResolution.Z, 0.0f);
-		}
-		case ESparseVolumeTexture_TileSize:				return FVector4(float(SPARSE_VOLUME_TILE_RES), 0.0f, 0.0f, 0.0f);
-		default:
-			break;
-		}
-		check(0);
-		return FVector4(ForceInitToZero);
-	}
-
-	// 0 while waiting for the proxy
-	return FVector4(ForceInitToZero);
-}
-
-
-const FSparseVolumeTextureSceneProxy* UAnimatedSparseVolumeTexture::GetSparseVolumeTextureFrameSceneProxy(int32 FrameIndex) const
-{
-	check(AnimationFrames.Num() >= 1);
-	FrameIndex = FrameIndex % GetFrameCountToLoad();
-	// SVT_TODO when streaming is enabled, this will likely change.
-	return &AnimationFrames[FrameIndex].SparseVolumeTextureSceneProxy;
-}
-
 ////////////////////////////////////////////////////////////////////////////////////////////////
diff --git a/Engine/Source/Runtime/Renderer/Private/SparseVolumeTexture/SparseVolumeTextureViewerComponent.cpp b/Engine/Source/Runtime/Renderer/Private/SparseVolumeTexture/SparseVolumeTextureViewerComponent.cpp
--- a/Engine/Source/Runtime/Renderer/Private/SparseVolumeTexture/SparseVolumeTextureViewerComponent.cpp
+++ b/Engine/Source/Runtime/Renderer/Private/SparseVolumeTexture/SparseVolumeTextureViewerComponent.cpp
@@ -20,6 +20,8 @@
 
 #define LOCTEXT_NAMESPACE "SparseVolumeTextureViewerComponent"
 
+constexpr double SVTViewerDefaultVolumeExtent = 100.0;
+
 /*=============================================================================
 	USparseVolumeTextureViewerComponent implementation.
 =============================================================================*/
@@ -75,27 +77,23 @@ void USparseVolumeTextureViewerComponent::Serialize(FArchive& Ar)
 
 FBoxSphereBounds USparseVolumeTextureViewerComponent::CalcBounds(const FTransform& LocalToWorld) const
 {
-	const float OneMeterInCentimeter = 100.0f;
 	FBoxSphereBounds NormalizedBound;
 	NormalizedBound.Origin = FVector(0.0f, 0.0f, 0.0f);
 
 	if (SparseVolumeTexturePreview)
 	{
-		// We assume that the maximum bound will have a length of 1 meter. 
+		// We assume that the maximum bound will have a length of SVTViewerDefaultVolumeExtent (1 meter). 
 		// Then the other dimensions are scaled relatively.
 		// All this account for the size of volume with page table (they add padding).
 		// TODO can we recover world size from OpenVDB meta data in meter?
-		const FSparseVolumeAssetHeader* SVTHeader = SparseVolumeTexturePreview->GetSparseVolumeTextureHeader();
-		check(SVTHeader);
-		const float MaxDim = (float)FMath::Max(FMath::Max(SVTHeader->PageTableVolumeResolution.X, SVTHeader->PageTableVolumeResolution.Y), SVTHeader->PageTableVolumeResolution.Z);
-
-		NormalizedBound.BoxExtent.X = 0.5f * OneMeterInCentimeter * float(SVTHeader->PageTableVolumeResolution.X) / MaxDim;
-		NormalizedBound.BoxExtent.Y = 0.5f * OneMeterInCentimeter * float(SVTHeader->PageTableVolumeResolution.Y) / MaxDim;
-		NormalizedBound.BoxExtent.Z = 0.5f * OneMeterInCentimeter * float(SVTHeader->PageTableVolumeResolution.Z) / MaxDim;
+		const FBox VolumeBounds = SparseVolumeTexturePreview->GetVolumeBounds();
+		const FVector Extent = VolumeBounds.GetExtent();
+		const double MaxDim = FMath::Max(FMath::Max(Extent.X, Extent.Y), Extent.Z);
+		NormalizedBound.BoxExtent = Extent / MaxDim * SVTViewerDefaultVolumeExtent;
 	}
 	else
 	{
-		NormalizedBound.BoxExtent = FVector(OneMeterInCentimeter, OneMeterInCentimeter, OneMeterInCentimeter);
+		NormalizedBound.BoxExtent = FVector(SVTViewerDefaultVolumeExtent);
 	}
 
 	NormalizedBound.SphereRadius = NormalizedBound.BoxExtent.Size();
@@ -156,23 +154,44 @@ void USparseVolumeTextureViewerComponent::SendRenderTransformCommand()
 {
 	if (SparseVolumeTextureViewerSceneProxy)
 	{
-		FTransform ToWorldMat = GetComponentToWorld();
-		ToWorldMat.SetRotation(FQuat::Identity);								// Do not allow rotation
-		FVector Scale3D = ToWorldMat.GetScale3D();
-		float MaxScaling = FMath::Max(Scale3D.X, FMath::Max(Scale3D.Y, Scale3D.Z));
-		ToWorldMat.SetScale3D(FVector(MaxScaling, MaxScaling, MaxScaling));		// Keep max scaling
-		FBoxSphereBounds ComponentBounds = CalcBounds(ToWorldMat);
-
-		FMatrix44f ToLocalMat = FMatrix44f(FTranslationMatrix(-ToWorldMat.GetTranslation()) * FRotationMatrix(FRotator(GetComponentToWorld().GetRotation().Inverse())) * FScaleMatrix(ComponentBounds.BoxExtent.Reciprocal()));
+		FVector VolumeExtent = FVector(SVTViewerDefaultVolumeExtent);
+		// In order to keep the contents of an animated SVT sequence stable in world space, we need to account for the fact that
+		// different frames of the sequence have different AABBs.
+		FVector UVScale = FVector::One();
+		FVector UVBias = FVector::Zero();
 		
-		uint32 CompToVisualize = (uint32)ComponentToVisualize;
+		if (SparseVolumeTexturePreview)
+		{
+			const FVector VolumeBoundsExtent = SparseVolumeTexturePreview->GetVolumeBounds().GetExtent();
+			const double MaxBoundsDim = FMath::Max(FMath::Max(VolumeBoundsExtent.X, VolumeBoundsExtent.Y), VolumeBoundsExtent.Z);
+			VolumeExtent = VolumeBoundsExtent / MaxBoundsDim * SVTViewerDefaultVolumeExtent;
+
+			SparseVolumeTexturePreview->GetFrameUVScaleBias(FrameIndex, &UVScale, &UVBias);
+		}
+
+		const FTransform ToWorldTransform = GetComponentTransform();
+		const FVector Scale3D = ToWorldTransform.GetScale3D();
+		// Keep max scaling since the DDA algorithm has trouble with non-uniform scale.
+		// Note that using the other components in FScaleMatrix below is fine 
+		// because they cancel out with the actual volume resolution, producing uniformly scaled voxels.
+		const float MaxScaling = FMath::Max(Scale3D.X, FMath::Max(Scale3D.Y, Scale3D.Z));
+
+		const FRotationMatrix WorldToLocalRotation = FRotationMatrix(FRotator(ToWorldTransform.GetRotation().Inverse()));
+		const FMatrix44f ToLocalMatNoScale = FMatrix44f(WorldToLocalRotation);
+		const FMatrix44f ToLocalMat = FMatrix44f(
+			FTranslationMatrix(-ToWorldTransform.GetTranslation()) * WorldToLocalRotation * FScaleMatrix((VolumeExtent * MaxScaling).Reciprocal())
+			* FScaleMatrix(0.5) * FTranslationMatrix(FVector(0.5))
+			* FScaleMatrix(UVScale) * FTranslationMatrix(UVBias)
+			* FScaleMatrix(2.0) * FTranslationMatrix(FVector(-1.0)));
+
+		const uint32 CompToVisualize = (uint32)ComponentToVisualize;
 
 		FSparseVolumeTextureViewerSceneProxy* SVTViewerSceneProxy = SparseVolumeTextureViewerSceneProxy;
 		ENQUEUE_RENDER_COMMAND(FUpdateSparseVolumeTextureViewerProxyTransformCommand)(
-			[SVTViewerSceneProxy, ComponentBounds, ToLocalMat, CompToVisualize](FRHICommandList& RHICmdList)
+			[SVTViewerSceneProxy, ToLocalMat, ToLocalMatNoScale, CompToVisualize](FRHICommandList& RHICmdList)
 			{
-				SVTViewerSceneProxy->VolumeWorldBounds = ComponentBounds;
 				SVTViewerSceneProxy->WorldToLocal = ToLocalMat;
+				SVTViewerSceneProxy->WorldToLocalNoScale = ToLocalMatNoScale;
 				SVTViewerSceneProxy->ComponentToVisualize = CompToVisualize;
 			});
 	}
@@ -182,7 +201,10 @@ void USparseVolumeTextureViewerComponent::TickComponent(float DeltaTime, enum EL
 {
 	if (bAnimate)
 	{
-		FrameIndex++;
+		const int32 NumFrames = SparseVolumeTexturePreview->GetFrameCount();
+		const float AnimationDuration = NumFrames / FrameRate;
+		FrameIndex = int32((AnimationTime / (AnimationDuration + SMALL_NUMBER)) * float(SparseVolumeTexturePreview->GetFrameCount()));
+		AnimationTime = FMath::Fmod(AnimationTime + DeltaTime, (AnimationDuration + SMALL_NUMBER));
 	}
 	else if(SparseVolumeTexturePreview)
 	{
diff --git a/Engine/Source/Runtime/Renderer/Private/SparseVolumeTexture/SparseVolumeTextureViewerComponent.h b/Engine/Source/Runtime/Renderer/Private/SparseVolumeTexture/SparseVolumeTextureViewerComponent.h
--- a/Engine/Source/Runtime/Renderer/Private/SparseVolumeTexture/SparseVolumeTextureViewerComponent.h
+++ b/Engine/Source/Runtime/Renderer/Private/SparseVolumeTexture/SparseVolumeTextureViewerComponent.h
@@ -35,6 +35,12 @@ class USparseVolumeTextureViewerComponent : public UPrimitiveComponent
 	UPROPERTY(EditAnywhere, BlueprintReadWrite, Category = "Asset Preview", meta = (UIMin = 0.0, UIMax = 1.0, ClampMin = 0.0, ClampMax = 1.0, EditCondition = "!bAnimate"))
 	float AnimationFrame;
 
+	UPROPERTY(EditAnywhere, BlueprintReadWrite, Category = "Asset Preview", meta = (UIMin = 0.0, UIMax = 120.0, ClampMin = 0.0, ClampMax = 120.0, EditCondition = "bAnimate"))
+	float FrameRate = 24.0f;
+
+	UPROPERTY(VisibleAnywhere, Category = "Asset Preview", meta = (UIMin = 0.0, UIMax = 60.0))
+	float AnimationTime;
+
 	UPROPERTY(EditAnywhere, BlueprintReadWrite, Category = "Asset Preview", meta = (UIMin = 0, UIMax = 7, ClampMin = 0, ClampMax = 7))
 	int32 ComponentToVisualize;
 
diff --git a/Engine/Source/Runtime/Renderer/Private/SparseVolumeTexture/SparseVolumeTextureViewerRendering.cpp b/Engine/Source/Runtime/Renderer/Private/SparseVolumeTexture/SparseVolumeTextureViewerRendering.cpp
--- a/Engine/Source/Runtime/Renderer/Private/SparseVolumeTexture/SparseVolumeTextureViewerRendering.cpp
+++ b/Engine/Source/Runtime/Renderer/Private/SparseVolumeTexture/SparseVolumeTextureViewerRendering.cpp
@@ -59,8 +59,9 @@ class FVisualizeSparseVolumeTexturePS : public FGlobalShader
 		SHADER_PARAMETER(FVector4f, WorldToLocal0)
 		SHADER_PARAMETER(FVector4f, WorldToLocal1)
 		SHADER_PARAMETER(FVector4f, WorldToLocal2)
-		SHADER_PARAMETER(FVector3f, VolumeBoundMinWorld)
-		SHADER_PARAMETER(FVector3f, VolumeBoundMaxWorld)
+		SHADER_PARAMETER(FVector3f, WorldToLocalNoScale0)
+		SHADER_PARAMETER(FVector3f, WorldToLocalNoScale1)
+		SHADER_PARAMETER(FVector3f, WorldToLocalNoScale2)
 		SHADER_PARAMETER(FVector3f, SparseVolumeTextureResolution)
 		SHADER_PARAMETER(FVector3f, SparseVolumeTexturePageTableResolution)
 		SHADER_PARAMETER(uint32, ComponentToVisualize)
@@ -113,8 +114,8 @@ void AddSparseVolumeTextureViewerRenderPass(FRDGBuilder& GraphBuilder, FSceneRen
 
 		for (auto& SVTProxy : Scene->SparseVolumeTextureViewers)
 		{
-			const FBoxSphereBounds& SVTProxyBound = SVTProxy->VolumeWorldBounds;
 			const FMatrix44f& WorldToLocal = SVTProxy->WorldToLocal;
+			const FMatrix44f& WorldToLocalNoScale = SVTProxy->WorldToLocalNoScale;
 			
 			FGlobalShaderMap* GlobalShaderMap = GetGlobalShaderMap(GMaxRHIFeatureLevel);
 			FVisualizeSparseVolumeTextureVS::FPermutationDomain VsPermutationVector;
@@ -128,9 +129,10 @@ void AddSparseVolumeTextureViewerRenderPass(FRDGBuilder& GraphBuilder, FSceneRen
 			PsPassParameters->WorldToLocal0 = FVector4f(WorldToLocal.M[0][0], WorldToLocal.M[1][0], WorldToLocal.M[2][0], WorldToLocal.M[3][0]);
 			PsPassParameters->WorldToLocal1 = FVector4f(WorldToLocal.M[0][1], WorldToLocal.M[1][1], WorldToLocal.M[2][1], WorldToLocal.M[3][1]);
 			PsPassParameters->WorldToLocal2 = FVector4f(WorldToLocal.M[0][2], WorldToLocal.M[1][2], WorldToLocal.M[2][2], WorldToLocal.M[3][2]);
+			PsPassParameters->WorldToLocalNoScale0 = FVector3f(WorldToLocalNoScale.M[0][0], WorldToLocalNoScale.M[1][0], WorldToLocalNoScale.M[2][0]);
+			PsPassParameters->WorldToLocalNoScale1 = FVector3f(WorldToLocalNoScale.M[0][1], WorldToLocalNoScale.M[1][1], WorldToLocalNoScale.M[2][1]);
+			PsPassParameters->WorldToLocalNoScale2 = FVector3f(WorldToLocalNoScale.M[0][2], WorldToLocalNoScale.M[1][2], WorldToLocalNoScale.M[2][2]);
 			PsPassParameters->ComponentToVisualize = SVTProxy->ComponentToVisualize;
-			PsPassParameters->VolumeBoundMinWorld = FVector3f(SVTProxyBound.Origin - SVTProxyBound.BoxExtent);
-			PsPassParameters->VolumeBoundMaxWorld = FVector3f(SVTProxyBound.Origin + SVTProxyBound.BoxExtent);
 			PsPassParameters->SparseVolumeTextureResolution = FVector3f::OneVector;
 			PsPassParameters->SparseVolumeTexturePageTableResolution = FVector3f::OneVector;
 			PsPassParameters->SparseVolumeTextureA = GBlackVolumeTexture->TextureRHI;
diff --git a/Engine/Source/Runtime/Renderer/Private/SparseVolumeTexture/SparseVolumeTextureViewerSceneProxy.cpp b/Engine/Source/Runtime/Renderer/Private/SparseVolumeTexture/SparseVolumeTextureViewerSceneProxy.cpp
--- a/Engine/Source/Runtime/Renderer/Private/SparseVolumeTexture/SparseVolumeTextureViewerSceneProxy.cpp
+++ b/Engine/Source/Runtime/Renderer/Private/SparseVolumeTexture/SparseVolumeTextureViewerSceneProxy.cpp
@@ -51,7 +51,7 @@ FSparseVolumeTextureViewerSceneProxy::FSparseVolumeTextureViewerSceneProxy(const
 		else
 		{
 			// Regular texture from proxy
-			SparseVolumeTextureSceneProxy = InComponent->SparseVolumeTexturePreview->GetSparseVolumeTextureSceneProxy();
+			SparseVolumeTextureSceneProxy = InComponent->SparseVolumeTexturePreview->GetSparseVolumeTextureSceneProxy(FrameIndex);
 		}
 	}
 }
diff --git a/Engine/Source/Runtime/Renderer/Private/SparseVolumeTexture/SparseVolumeTextureViewerSceneProxy.h b/Engine/Source/Runtime/Renderer/Private/SparseVolumeTexture/SparseVolumeTextureViewerSceneProxy.h
--- a/Engine/Source/Runtime/Renderer/Private/SparseVolumeTexture/SparseVolumeTextureViewerSceneProxy.h
+++ b/Engine/Source/Runtime/Renderer/Private/SparseVolumeTexture/SparseVolumeTextureViewerSceneProxy.h
@@ -18,8 +18,8 @@ public:
 	virtual ~FSparseVolumeTextureViewerSceneProxy() = default;
 
 	const FSparseVolumeTextureSceneProxy* SparseVolumeTextureSceneProxy = nullptr;
-	FBoxSphereBounds VolumeWorldBounds;
 	FMatrix44f WorldToLocal;
+	FMatrix44f WorldToLocalNoScale;
 	uint32 ComponentToVisualize;
 
 protected:
