Implement dynamic bool parameters to enable dynamic branching in static switches. This can cut shader permutation counts at the cost of larger more gernalized shaders. Which might or might not  have a performance penalty based on the use case.

diff --git a/Engine/Plugins/Enterprise/MDLImporter/Source/MDLImporter/Private/generator/GeneratorMaterialExpressions.cpp b/Engine/Plugins/Enterprise/MDLImporter/Source/MDLImporter/Private/generator/GeneratorMaterialExpressions.cpp
--- a/Engine/Plugins/Enterprise/MDLImporter/Source/MDLImporter/Private/generator/GeneratorMaterialExpressions.cpp
+++ b/Engine/Plugins/Enterprise/MDLImporter/Source/MDLImporter/Private/generator/GeneratorMaterialExpressions.cpp
@@ -127,7 +127,7 @@ namespace Generator
 		}
 		else if (Expression->IsA<UMaterialExpressionFunctionInput>())
 		{
-			return (Cast<UMaterialExpressionFunctionInput>(Expression)->InputType == FunctionInput_StaticBool);
+			return (Cast<UMaterialExpressionFunctionInput>(Expression)->InputType == FunctionInput_StaticBool) || (Cast<UMaterialExpressionFunctionInput>(Expression)->InputType == FunctionInput_Bool);
 		}
 		else if (Expression->IsA<UMaterialExpressionMaterialFunctionCall>())
 		{
@@ -626,6 +626,7 @@ namespace Generator
 			case FunctionInput_TextureCube:
 			case FunctionInput_VolumeTexture:
 			case FunctionInput_StaticBool:
+			case FunctionInput_Bool:
 			case FunctionInput_MaterialAttributes:
 			default:
 				ensure(false);
diff --git a/Engine/Source/Editor/UnrealEd/Private/MaterialGraphSchema.cpp b/Engine/Source/Editor/UnrealEd/Private/MaterialGraphSchema.cpp
--- a/Engine/Source/Editor/UnrealEd/Private/MaterialGraphSchema.cpp
+++ b/Engine/Source/Editor/UnrealEd/Private/MaterialGraphSchema.cpp
@@ -111,6 +111,9 @@ void FMaterialGraphSchemaAction_NewNode::SetFunctionInputType(UMaterialExpressio
 	case MCT_StaticBool:
 		FunctionInput->InputType = FunctionInput_StaticBool;
 		break;
+	case MCT_Bool:
+		FunctionInput->InputType = FunctionInput_Bool;
+		break;
 	case MCT_MaterialAttributes:
 		FunctionInput->InputType = FunctionInput_MaterialAttributes;
 		break;
diff --git a/Engine/Source/Runtime/Engine/Classes/Materials/MaterialExpressionFunctionInput.h b/Engine/Source/Runtime/Engine/Classes/Materials/MaterialExpressionFunctionInput.h
--- a/Engine/Source/Runtime/Engine/Classes/Materials/MaterialExpressionFunctionInput.h
+++ b/Engine/Source/Runtime/Engine/Classes/Materials/MaterialExpressionFunctionInput.h
@@ -28,6 +28,7 @@ enum EFunctionInputType : int
 	FunctionInput_StaticBool,
 	FunctionInput_MaterialAttributes,
 	FunctionInput_TextureExternal,
+	FunctionInput_Bool,
 	FunctionInput_MAX,
 };
 
diff --git a/Engine/Source/Runtime/Engine/Classes/Materials/MaterialExpressionStaticBoolParameter.h b/Engine/Source/Runtime/Engine/Classes/Materials/MaterialExpressionStaticBoolParameter.h
--- a/Engine/Source/Runtime/Engine/Classes/Materials/MaterialExpressionStaticBoolParameter.h
+++ b/Engine/Source/Runtime/Engine/Classes/Materials/MaterialExpressionStaticBoolParameter.h
@@ -17,16 +17,21 @@ class UMaterialExpressionStaticBoolParameter : public UMaterialExpressionParamet
 	UPROPERTY(EditAnywhere, Category=MaterialExpressionStaticBoolParameter, meta = (ShowAsInputPin = "Primary"))
 	uint32 DefaultValue:1;
 
+	/**Change Parameter from "static bool" to (dynamic) bool type which enables it to be used with dynamic branching*/
+	UPROPERTY(EditAnywhere, Category = MaterialExpressionStaticBoolParameter)
+	uint32 DynamicBranch:1;
+
 public:
 	//~ Begin UMaterialExpression Interface
 #if WITH_EDITOR
 	virtual int32 Compile(class FMaterialCompiler* Compiler, int32 OutputIndex) override;
 	virtual int32 CompilePreview(class FMaterialCompiler* Compiler, int32 OutputIndex) override;
 	virtual void GetCaption(TArray<FString>& OutCaptions) const override;
-	virtual uint32 GetOutputType(int32 OutputIndex) override {return MCT_StaticBool;}
+	virtual uint32 GetOutputType(int32 OutputIndex) override {return DynamicBranch ? MCT_Bool : MCT_StaticBool;}
 	virtual bool GetParameterValue(FMaterialParameterMetadata& OutMeta) const override
 	{
 		OutMeta.Value = (bool)DefaultValue;
+		OutMeta.bDynamicSwitchParameter = DynamicBranch;
 		return Super::GetParameterValue(OutMeta);
 	}
 	virtual bool SetParameterValue(const FName& Name, const FMaterialParameterMetadata& Meta, EMaterialExpressionSetParameterValueFlags Flags) override
diff --git a/Engine/Source/Runtime/Engine/Private/Materials/HLSLMaterialTranslator.cpp b/Engine/Source/Runtime/Engine/Private/Materials/HLSLMaterialTranslator.cpp
--- a/Engine/Source/Runtime/Engine/Private/Materials/HLSLMaterialTranslator.cpp
+++ b/Engine/Source/Runtime/Engine/Private/Materials/HLSLMaterialTranslator.cpp
@@ -202,6 +202,7 @@ static inline EMaterialValueType GetMaterialValueType(EMaterialParameterType Typ
 	case EMaterialParameterType::Scalar: return MCT_Float;
 	case EMaterialParameterType::Vector: return MCT_Float4;
 	case EMaterialParameterType::DoubleVector: return MCT_LWCVector4;
+	case EMaterialParameterType::StaticSwitch: return MCT_Bool;
 	default: checkNoEntry(); return MCT_Unknown;
 	}
 }
@@ -3374,7 +3375,55 @@ int32 FHLSLMaterialTranslator::AccessUniformExpression(int32 Index)
 	check(!(CodeChunk.Type & MCT_TextureVirtual) || TextureUniformExpression);
 
 	TStringBuilder<1024> FormattedCode;
-	if (IsFloatNumericType(CodeChunk.Type))
+
+	if (CodeChunk.Type == MCT_Bool)
+	{
+		check(CodeChunk.UniformExpression->GetType() == &FMaterialUniformExpressionStaticBoolParameter::StaticType);
+		FMaterialUniformExpressionStaticBoolParameter* StaticBoolParameter = static_cast<FMaterialUniformExpressionStaticBoolParameter*>(CodeChunk.UniformExpression.GetReference());
+		check(!bIsLWC);
+
+		const uint32 NumComponents = GetNumComponents(CodeChunk.Type);
+
+		if (CodeChunk.UniformExpression->UniformOffset == INDEX_NONE)
+		{
+			// 'Bool' uniforms are packed into bits
+			if (CurrentNumBoolComponents + NumComponents > 32u)
+			{
+				CurrentBoolUniformOffset = UniformPreshaderOffset++;
+				CurrentNumBoolComponents = 0u;
+			}
+
+			int32 UniformOffset = CurrentBoolUniformOffset * 32u + CurrentNumBoolComponents;
+
+			FUniformExpressionSet& UniformExpressionSet = MaterialCompilationOutput.UniformExpressionSet;
+			FMaterialUniformPreshaderHeader& PreshaderHeader = UniformExpressionSet.UniformPreshaders.AddDefaulted_GetRef();
+			PreshaderHeader.FieldIndex = UniformExpressionSet.UniformPreshaderFields.Num();
+			PreshaderHeader.NumFields = 1;
+			PreshaderHeader.OpcodeOffset = UniformExpressionSet.UniformPreshaderData.Num();
+			CodeChunk.UniformExpression->WriteNumberOpcodes(UniformExpressionSet.UniformPreshaderData);
+			PreshaderHeader.OpcodeSize = UniformExpressionSet.UniformPreshaderData.Num() - PreshaderHeader.OpcodeOffset;
+
+			FMaterialUniformPreshaderField& PreshaderField = MaterialCompilationOutput.UniformExpressionSet.UniformPreshaderFields.AddDefaulted_GetRef();
+			PreshaderField.BufferOffset = UniformOffset;
+			PreshaderField.Type = UE::Shader::MakeValueType(UE::Shader::EValueComponentType::Bool, NumComponents);
+			PreshaderField.ComponentIndex = 0;
+
+			CodeChunk.UniformExpression->UniformOffset = UniformOffset;
+			CurrentNumBoolComponents += NumComponents;
+		}
+
+		const uint32 UniformOffset = CodeChunk.UniformExpression->UniformOffset / 32u;
+		const uint32 NumBoolComponents = CodeChunk.UniformExpression->UniformOffset % 32u;
+
+		const uint32 RegisterIndex = UniformOffset / 4;
+		const uint32 RegisterOffset = UniformOffset % 4;
+		FormattedCode.Appendf(TEXT("UnpackUniform_%s(asuint(Material.PreshaderBuffer[%u][%u]), %u)"),
+			TEXT("bool"),
+			RegisterIndex,
+			RegisterOffset,
+			NumBoolComponents);
+	}
+	else if (IsFloatNumericType(CodeChunk.Type))
 	{
 		check(CodeChunk.DerivativeStatus != EDerivativeStatus::Valid);
 		const uint32 NumComponents = GetNumComponents(CodeChunk.Type);
@@ -5504,6 +5553,64 @@ int32 FHLSLMaterialTranslator::ActorWorldPosition()
 	return CastToNonLWCIfDisabled(Result);
 }
 
+int32 FHLSLMaterialTranslator::DynamicBranch(int32 Condition, int32 A, int32 B)
+{
+	if (Condition == INDEX_NONE)
+	{
+		return INDEX_NONE;
+	}
+
+	if (B == INDEX_NONE)
+	{
+		return A;
+	}
+
+	if (A == INDEX_NONE)
+	{
+		return B;
+	}
+
+	EMaterialValueType TypeA = GetParameterType(A);
+	EMaterialValueType TypeB = GetParameterType(B);
+	if (IsLWCType(TypeA) || IsLWCType(TypeB))
+	{
+		TypeA = MakeLWCType(TypeA);
+		TypeB = MakeLWCType(TypeB);
+	}
+
+	EMaterialValueType ResultType = MCT_Unknown;
+	if (TypeA == TypeB)
+	{
+		ResultType = TypeA;
+	}
+	else if (!IsFloatNumericType(TypeA) || !IsFloatNumericType(TypeB))
+	{
+		Errorf(TEXT("Cannot branch on non float numeric Types if they are not equal: %s %s"), DescribeType(TypeA), DescribeType(TypeB));
+		return INDEX_NONE;
+	}
+	else
+	{
+		ResultType = GetNumComponents(TypeA) > GetNumComponents(TypeB) ? TypeA : TypeB;
+	}
+
+	A = ForceCast(A, ResultType, MFCF_ReplicateValue);
+	B = ForceCast(B, ResultType, MFCF_ReplicateValue);
+	FString SymbolName = CreateSymbolName(TEXT("Static"));
+
+	checkf(Condition >= 0 && Condition < CurrentScopeChunks->Num(), TEXT("Index %d/%d, Platform=%d"), Condition, CurrentScopeChunks->Num(), (int)Platform);
+	const FShaderCodeChunk& CodeChunk = (*CurrentScopeChunks)[Condition];
+
+	if (CodeChunk.UniformExpression && CodeChunk.UniformExpression->GetType() == &FMaterialUniformExpressionStaticBoolParameter::StaticType)
+	{
+		FMaterialUniformExpressionStaticBoolParameter* StaticBoolParameter = static_cast<FMaterialUniformExpressionStaticBoolParameter*>(CodeChunk.UniformExpression.GetReference());
+		AddCodeChunk(MCT_VoidStatement, TEXT("//%s"), *StaticBoolParameter->GetParameterName().ToString());
+	}
+
+	AddCodeChunk(MCT_VoidStatement, TEXT("%s %s;"), HLSLTypeString(ResultType), *SymbolName);
+	AddCodeChunk(MCT_VoidStatement, TEXT("[branch] switch (int(%s)){ default: %s = %s; break; case 0: %s = %s; break;}"), *GetParameterCode(Condition), *SymbolName, *GetParameterCode(A), *SymbolName, *GetParameterCode(B));
+	return AddCodeChunk(ResultType, *SymbolName);
+}
+
 // Compare two inputs and return true if they are either the same, or if they evaluate to equal constant expressions.
 static bool AreEqualExpressions(int32 A, int32 B, TArray<FShaderCodeChunk> const& InCurrentScopeChunks, FMaterial const& InMaterial)
 {
@@ -7301,6 +7408,33 @@ int32 FHLSLMaterialTranslator::StaticBool(bool bValue)
 	return AddInlinedCodeChunk(MCT_StaticBool,(bValue ? TEXT("true") : TEXT("false")));
 }
 
+int32 FHLSLMaterialTranslator::DynamicBoolParameter(FName ParameterName, bool bDefaultValue)
+{
+	// Look up the value we are compiling with for this static parameter.
+	bool bValue = bDefaultValue;
+
+	FMaterialParameterInfo ParameterInfo = GetParameterAssociationInfo();
+	ParameterInfo.Name = ParameterName;
+
+	const UE::Shader::EValueType ValueType = GetShaderValueType(EMaterialParameterType::StaticSwitch);
+	UE::Shader::FValue DefaultValue(bValue);
+
+	const uint32* PrevDefaultOffset = DefaultUniformValues.Find(DefaultValue);
+	uint32 DefaultOffset;
+	if (PrevDefaultOffset)
+	{
+		DefaultOffset = *PrevDefaultOffset;
+	}
+	else
+	{
+		DefaultOffset = MaterialCompilationOutput.UniformExpressionSet.AddDefaultParameterValue(DefaultValue);
+		DefaultUniformValues.Add(DefaultValue, DefaultOffset);
+	}
+
+	const int32 ParameterIndex = MaterialCompilationOutput.UniformExpressionSet.FindOrAddNumericParameter(EMaterialParameterType::StaticSwitch, ParameterInfo, DefaultOffset);
+	return AddUniformExpression(new FMaterialUniformExpressionStaticBoolParameter(ParameterInfo, ParameterIndex), MCT_Bool, TEXT(""));
+}
+
 int32 FHLSLMaterialTranslator::StaticBoolParameter(FName ParameterName,bool bDefaultValue)
 {
 	// Look up the value we are compiling with for this static parameter.
diff --git a/Engine/Source/Runtime/Engine/Private/Materials/HLSLMaterialTranslator.h b/Engine/Source/Runtime/Engine/Private/Materials/HLSLMaterialTranslator.h
--- a/Engine/Source/Runtime/Engine/Private/Materials/HLSLMaterialTranslator.h
+++ b/Engine/Source/Runtime/Engine/Private/Materials/HLSLMaterialTranslator.h
@@ -295,6 +295,8 @@ protected:
 	TArray<TRefCountPtr<FMaterialUniformExpressionExternalTexture>> UniformExternalTextureExpressions;
 	TMap<UE::Shader::FValue, uint32> DefaultUniformValues;
 	uint32 UniformPreshaderOffset = 0u;
+	uint32 CurrentBoolUniformOffset = ~0u;
+	uint32 CurrentNumBoolComponents = 32u;
 
 	/** Parameter collections referenced by this material.  The position in this array is used as an index on the shader parameter. */
 	TArray<UMaterialParameterCollection*> ParameterCollections;
@@ -837,6 +839,7 @@ protected:
 
 	virtual int32 ActorWorldPosition() override;
 
+	virtual int32 DynamicBranch(int32 Condition, int32 A, int32 B) override;
 	virtual int32 If(int32 A, int32 B, int32 AGreaterThanB, int32 AEqualsB, int32 ALessThanB, int32 ThresholdArg) override;
 
 	void AllocateSlot(TBitArray<>& InBitArray, int32 InSlotIndex, int32 InSlotCount = 1) const;
@@ -930,6 +933,7 @@ protected:
 
 	virtual int32 StaticBool(bool bValue) override;
 	virtual int32 StaticBoolParameter(FName ParameterName, bool bDefaultValue) override;
+	virtual int32 DynamicBoolParameter(FName ParameterName, bool bDefaultValue) override;
 	virtual int32 StaticComponentMask(int32 Vector, FName ParameterName, bool bDefaultR, bool bDefaultG, bool bDefaultB, bool bDefaultA) override;
 	virtual const FMaterialLayersFunctions* GetMaterialLayers() override;
 
diff --git a/Engine/Source/Runtime/Engine/Private/Materials/MaterialCachedData.cpp b/Engine/Source/Runtime/Engine/Private/Materials/MaterialCachedData.cpp
--- a/Engine/Source/Runtime/Engine/Private/Materials/MaterialCachedData.cpp
+++ b/Engine/Source/Runtime/Engine/Private/Materials/MaterialCachedData.cpp
@@ -200,7 +200,8 @@ bool FMaterialCachedExpressionData::AddParameter(const FMaterialParameterInfo& P
 			break;
 
 		case EMaterialParameterType::StaticSwitch:
-			EditorOnlyData->StaticSwitchValues.Insert(ParameterMeta.Value.AsStaticSwitch(), Index);
+			StaticSwitchValues.Insert(ParameterMeta.Value.AsStaticSwitch(), Index);
+			DynamicSwitchValues.Insert(ParameterMeta.bDynamicSwitchParameter, Index);
 			break;
 
 		case EMaterialParameterType::StaticComponentMask:
@@ -666,13 +667,11 @@ void FMaterialCachedExpressionData::GetParameterValueByIndex(EMaterialParameterT
 	case EMaterialParameterType::Font:
 		OutResult.Value = FMaterialParameterValue(FontValues[ParameterIndex].LoadSynchronous(), FontPageValues[ParameterIndex]);
 		break;
-#if WITH_EDITORONLY_DATA
 	case EMaterialParameterType::StaticSwitch:
-		if (EditorOnlyData && !bIsEditorOnlyDataStripped)
-		{
-			OutResult.Value = EditorOnlyData->StaticSwitchValues[ParameterIndex];
-		}
+		OutResult.Value = StaticSwitchValues[ParameterIndex];
+		OutResult.bDynamicSwitchParameter = DynamicSwitchValues[ParameterIndex];
 		break;
+#if WITH_EDITORONLY_DATA
 	case EMaterialParameterType::StaticComponentMask:
 		if (EditorOnlyData && !bIsEditorOnlyDataStripped)
 		{
@@ -686,6 +685,28 @@ void FMaterialCachedExpressionData::GetParameterValueByIndex(EMaterialParameterT
 	}
 }
 
+void FMaterialCachedExpressionData::PostSerialize(const FArchive& Ar)
+{
+#if WITH_EDITORONLY_DATA
+	if(Ar.IsLoading())
+	{
+		bool bIsEditorOnlyDataStripped = true;
+		if (EditorOnlyData)
+		{
+			const FMaterialCachedParameterEditorEntry& EditorEntry = EditorOnlyData->EditorEntries[(int32)EMaterialParameterType::StaticSwitch];
+			bIsEditorOnlyDataStripped = EditorEntry.EditorInfo.Num() == 0;
+		}
+
+		if (EditorOnlyData && !bIsEditorOnlyDataStripped)
+		{
+			StaticSwitchValues = EditorOnlyData->StaticSwitchValues_DEPRECATED;
+			check(DynamicSwitchValues.Num() == 0);
+			DynamicSwitchValues.AddDefaulted(StaticSwitchValues.Num());
+		}
+	}
+#endif
+}
+
 bool FMaterialCachedExpressionData::GetParameterValue(EMaterialParameterType Type, const FMemoryImageMaterialParameterInfo& ParameterInfo, FMaterialParameterMetadata& OutResult) const
 {
 	const int32 Index = FindParameterIndex(Type, ParameterInfo);
diff --git a/Engine/Source/Runtime/Engine/Private/Materials/MaterialExpressions.cpp b/Engine/Source/Runtime/Engine/Private/Materials/MaterialExpressions.cpp
--- a/Engine/Source/Runtime/Engine/Private/Materials/MaterialExpressions.cpp
+++ b/Engine/Source/Runtime/Engine/Private/Materials/MaterialExpressions.cpp
@@ -428,6 +428,8 @@ void GetMaterialValueTypeDescriptions(const uint32 MaterialValueType, TArray<FTe
 
 	if (MaterialValueType & MCT_StaticBool)
 		OutDescriptions.Add(LOCTEXT("StaticBool", "Bool"));
+	if (MaterialValueType & MCT_Bool)
+		OutDescriptions.Add(LOCTEXT("Bool", "Bool"));
 	if (MaterialValueType & MCT_MaterialAttributes)
 		OutDescriptions.Add(LOCTEXT("MaterialAttributes", "Material Attributes"));
 	if (MaterialValueType & MCT_ShadingModel)
@@ -467,7 +469,7 @@ bool CanConnectMaterialValueTypes(const uint32 InputType, const uint32 OutputTyp
 	{
 		return true;
 	}
-	if (InputType == MCT_StaticBool && OutputType == MCT_Bool)
+	if (InputType == MCT_Bool && OutputType == MCT_StaticBool)
 	{
 		// StaticBool is allowed to connect to Bool (but not the other way around)
 		return true;
@@ -9027,6 +9029,26 @@ FExpressionInput* UMaterialExpressionStaticSwitchParameter::GetEffectiveInput(cl
 
 int32 UMaterialExpressionStaticSwitchParameter::Compile(class FMaterialCompiler* Compiler, int32 OutputIndex)
 {
+	if(DynamicBranch)
+	{
+		int32 V = Compiler->DynamicBoolParameter(ParameterName, DefaultValue);
+		if (V != INDEX_NONE)
+		{
+			return Compiler->DynamicBranch(V, A.Compile(Compiler), B.Compile(Compiler));
+		}
+		else
+		{
+			if(DefaultValue)
+			{
+				return A.Compile(Compiler);
+			}
+			else
+			{
+				return B.Compile(Compiler);
+			}
+		}
+	}
+
 	FExpressionInput* EffectiveInput = GetEffectiveInput(Compiler);
 	if (EffectiveInput)
 	{
@@ -9101,7 +9123,7 @@ UMaterialExpressionStaticBoolParameter::UMaterialExpressionStaticBoolParameter(c
 #if WITH_EDITOR
 int32 UMaterialExpressionStaticBoolParameter::Compile(class FMaterialCompiler* Compiler, int32 OutputIndex)
 {
-	return Compiler->StaticBoolParameter(ParameterName,DefaultValue);
+	return DynamicBranch ? Compiler->DynamicBoolParameter(ParameterName, DefaultValue) : Compiler->StaticBoolParameter(ParameterName,DefaultValue);
 }
 
 int32 UMaterialExpressionStaticBoolParameter::CompilePreview(class FMaterialCompiler* Compiler, int32 OutputIndex)
@@ -9236,12 +9258,29 @@ FExpressionInput* UMaterialExpressionStaticSwitch::GetEffectiveInput(class FMate
 
 int32 UMaterialExpressionStaticSwitch::Compile(class FMaterialCompiler* Compiler, int32 OutputIndex)
 {
-	FExpressionInput* EffectiveInput = GetEffectiveInput(Compiler);
-	if (EffectiveInput)
+	int32 V = INDEX_NONE;
+	if (Value.GetTracedInput().Expression)
 	{
-		return EffectiveInput->Compile(Compiler);
+		V = Value.Compile(Compiler);
 	}
-	return INDEX_NONE;
+
+	if(V != INDEX_NONE && Compiler->GetParameterType(V) == MCT_Bool)
+	{
+		return Compiler->DynamicBranch(V, A.Compile(Compiler), B.Compile(Compiler));
+	}
+
+	bool bValue = DefaultValue;
+	if (V != INDEX_NONE)
+	{
+		bool bSucceeded;
+		bValue = Compiler->GetStaticBoolValue(V, bSucceeded);
+		if (!bSucceeded)
+		{
+			return INDEX_NONE;
+		}
+	}
+
+	return bValue ? A.Compile(Compiler) : B.Compile(Compiler);
 }
 
 void UMaterialExpressionStaticSwitch::GetCaption(TArray<FString>& OutCaptions) const
@@ -9273,15 +9312,7 @@ uint32 UMaterialExpressionStaticSwitch::GetInputType(int32 InputIndex)
 	}
 	else
 	{
-		if (IsUsingNewHLSLGenerator())
-		{
-			// Allow non-static bool when using new HLSL generator
-			return MCT_Bool;
-		}
-		else
-		{
-			return MCT_StaticBool;
-		}
+		return MCT_Bool;
 	}
 }
 
@@ -16686,7 +16717,8 @@ void UMaterialExpressionFunctionInput::GetCaption(TArray<FString>& OutCaptions)
 		TEXT("VolumeTexture"),
 		TEXT("StaticBool"),
 		TEXT("MaterialAttributes"),
-		TEXT("External")
+		TEXT("External"),
+		TEXT("Bool")
 	};
 	check(InputType < FunctionInput_MAX);
 	OutCaptions.Add(FString(TEXT("Input ")) + InputName.ToString() + TEXT(" (") + TypeNames[InputType] + TEXT(")"));
@@ -16736,6 +16768,7 @@ int32 UMaterialExpressionFunctionInput::CompilePreviewValue(FMaterialCompiler* C
 		case FunctionInput_Texture2DArray:
 		case FunctionInput_TextureExternal:
 		case FunctionInput_StaticBool:
+		case FunctionInput_Bool:
 			return Compiler->Errorf(TEXT("Missing Preview connection for function input '%s'"), *InputName.ToString());
 		default:
 			return Compiler->Errorf(TEXT("Unknown input type"));
@@ -16912,6 +16945,8 @@ uint32 UMaterialExpressionFunctionInput::GetInputType(int32 InputIndex)
 		return MCT_VolumeTexture;
 	case FunctionInput_StaticBool:
 		return MCT_StaticBool;
+	case FunctionInput_Bool:
+		return MCT_Bool;
 	case FunctionInput_MaterialAttributes:
 		return MCT_MaterialAttributes;
 	default:
diff --git a/Engine/Source/Runtime/Engine/Private/Materials/MaterialHLSLGenerator.cpp b/Engine/Source/Runtime/Engine/Private/Materials/MaterialHLSLGenerator.cpp
--- a/Engine/Source/Runtime/Engine/Private/Materials/MaterialHLSLGenerator.cpp
+++ b/Engine/Source/Runtime/Engine/Private/Materials/MaterialHLSLGenerator.cpp
@@ -494,6 +494,7 @@ const UE::HLSLTree::FExpression* FMaterialHLSLGenerator::AcquireFunctionInputExp
 			case FunctionInput_Texture2DArray:
 			case FunctionInput_VolumeTexture:
 			case FunctionInput_StaticBool:
+			case FunctionInput_Bool:
 			case FunctionInput_TextureExternal:
 				Errorf(TEXT("Missing Preview connection for function input '%s'"), *MaterialExpression->InputName.ToString());
 				return nullptr;
diff --git a/Engine/Source/Runtime/Engine/Private/Materials/MaterialInstance.cpp b/Engine/Source/Runtime/Engine/Private/Materials/MaterialInstance.cpp
--- a/Engine/Source/Runtime/Engine/Private/Materials/MaterialInstance.cpp
+++ b/Engine/Source/Runtime/Engine/Private/Materials/MaterialInstance.cpp
@@ -272,6 +272,7 @@ bool FMaterialInstanceResource::GetParameterValue(EMaterialParameterType Type, c
 		// Check for instances overrides
 		switch (Type)
 		{
+		case EMaterialParameterType::StaticSwitch: bResult = RenderThread_GetParameterValue<bool>(ParameterInfo, OutValue); break;
 		case EMaterialParameterType::Scalar: bResult = RenderThread_GetParameterValue<float>(ParameterInfo, OutValue); break;
 		case EMaterialParameterType::Vector: bResult = RenderThread_GetParameterValue<FLinearColor>(ParameterInfo, OutValue); break;
 		case EMaterialParameterType::DoubleVector: bResult = RenderThread_GetParameterValue<FVector4d>(ParameterInfo, OutValue); break;
@@ -371,14 +372,17 @@ void FMaterialInstanceResource::InitMIParameters(FMaterialInstanceParameterSet&
 	ParameterSet.RuntimeVirtualTextureParameters.Sort(SortMaterialInstanceParametersPredicate<const URuntimeVirtualTexture*>);
 	ParameterSet.SparseVolumeTextureParameters.Sort(SortMaterialInstanceParametersPredicate<const USparseVolumeTexture*>);
 
-	Swap(ScalarParameterArray.Array, ParameterSet.ScalarParameters);
-	Swap(VectorParameterArray.Array, ParameterSet.VectorParameters);
-	Swap(DoubleVectorParameterArray.Array, ParameterSet.DoubleVectorParameters);
-	Swap(TextureParameterArray.Array, ParameterSet.TextureParameters);
-	Swap(RuntimeVirtualTextureParameterArray.Array, ParameterSet.RuntimeVirtualTextureParameters);
-	Swap(SparseVolumeTextureParameterArray.Array, ParameterSet.SparseVolumeTextureParameters);
+	StaticSwitchParameterArray.Array = MoveTemp(ParameterSet.StaticSwitchParameters);
+	ScalarParameterArray.Array = MoveTemp(ParameterSet.ScalarParameters);
+	VectorParameterArray.Array = MoveTemp(ParameterSet.VectorParameters);
+	DoubleVectorParameterArray.Array = MoveTemp(ParameterSet.DoubleVectorParameters);
+	TextureParameterArray.Array = MoveTemp(ParameterSet.TextureParameters);
+	RuntimeVirtualTextureParameterArray.Array = MoveTemp(ParameterSet.RuntimeVirtualTextureParameters);
+	SparseVolumeTextureParameterArray.Array = MoveTemp(ParameterSet.SparseVolumeTextureParameters);
+
 
 	// Build hash tables.
+	StaticSwitchParameterArray.HashAddAllItems();
 	ScalarParameterArray.HashAddAllItems();
 	VectorParameterArray.HashAddAllItems();
 	DoubleVectorParameterArray.HashAddAllItems();
@@ -759,6 +763,20 @@ void GameThread_InitMIParameters(const UMaterialInstance& Instance)
 		ParamRef.Value = FSparseVolumeTextureParameterValue::GetValue(Parameter);
 	}
 	
+	FStaticParameterSet StaticParamSet = Instance.GetStaticParameters();
+	ParameterSet.StaticSwitchParameters.Reserve(StaticParamSet.StaticSwitchParameters.Num());
+	for(const FStaticSwitchParameter& Param : StaticParamSet.StaticSwitchParameters)
+	{
+		if(Param.IsOverride())
+		{
+			FMaterialParameterMetadata Result;
+			Param.GetValue(Result);
+			auto& ParamRef = ParameterSet.StaticSwitchParameters.AddDefaulted_GetRef();
+			ParamRef.Info = FHashedMaterialParameterInfo(Param.ParameterInfo);
+			ParamRef.Value = Result.Value.AsStaticSwitch();
+		}
+	}
+
 	ENQUEUE_RENDER_COMMAND(InitMIParameters)(
 		[Resource, Parameters = MoveTemp(ParameterSet)](FRHICommandListImmediate& RHICmdList) mutable
 		{
diff --git a/Engine/Source/Runtime/Engine/Private/Materials/MaterialInstanceSupport.h b/Engine/Source/Runtime/Engine/Private/Materials/MaterialInstanceSupport.h
--- a/Engine/Source/Runtime/Engine/Private/Materials/MaterialInstanceSupport.h
+++ b/Engine/Source/Runtime/Engine/Private/Materials/MaterialInstanceSupport.h
@@ -394,6 +394,8 @@ private:
 	/** The game thread accessible parent of the material instance. */
 	UMaterialInterface* GameThreadParent;
 	
+	/** StaticSwitch parameters to select material permutation **/
+	THashedMaterialParameterMap<bool> StaticSwitchParameterArray;
 	/** Vector parameters for this material instance. */
 	THashedMaterialParameterMap<FLinearColor> VectorParameterArray;
 	/** DoubleVector parameters for this material instance. */
@@ -410,12 +412,14 @@ private:
 	TArray<int32> ParentLayerIndexRemap;
 };
 
+template <> FORCEINLINE THashedMaterialParameterMap<bool>& FMaterialInstanceResource::GetValueArray() { return StaticSwitchParameterArray; }
 template <> FORCEINLINE THashedMaterialParameterMap<float>& FMaterialInstanceResource::GetValueArray() { return ScalarParameterArray; }
 template <> FORCEINLINE THashedMaterialParameterMap<FLinearColor>& FMaterialInstanceResource::GetValueArray() { return VectorParameterArray; }
 template <> FORCEINLINE THashedMaterialParameterMap<FVector4d>& FMaterialInstanceResource::GetValueArray() { return DoubleVectorParameterArray; }
 template <> FORCEINLINE THashedMaterialParameterMap<const UTexture*>& FMaterialInstanceResource::GetValueArray() { return TextureParameterArray; }
 template <> FORCEINLINE THashedMaterialParameterMap<const URuntimeVirtualTexture*>& FMaterialInstanceResource::GetValueArray() { return RuntimeVirtualTextureParameterArray; }
 template <> FORCEINLINE THashedMaterialParameterMap<const USparseVolumeTexture*>& FMaterialInstanceResource::GetValueArray() { return SparseVolumeTextureParameterArray; }
+template <> FORCEINLINE const THashedMaterialParameterMap<bool>& FMaterialInstanceResource::GetValueArray() const { return StaticSwitchParameterArray; }
 template <> FORCEINLINE const THashedMaterialParameterMap<float>& FMaterialInstanceResource::GetValueArray() const { return ScalarParameterArray; }
 template <> FORCEINLINE const THashedMaterialParameterMap<FLinearColor>& FMaterialInstanceResource::GetValueArray() const { return VectorParameterArray; }
 template <> FORCEINLINE const THashedMaterialParameterMap<FVector4d>& FMaterialInstanceResource::GetValueArray() const { return DoubleVectorParameterArray; }
@@ -425,6 +429,7 @@ template <> FORCEINLINE const THashedMaterialParameterMap<const USparseVolumeTex
 
 struct FMaterialInstanceParameterSet
 {
+	TArray<THashedMaterialParameterMap<bool>::TNamedParameter>							StaticSwitchParameters;
 	TArray<THashedMaterialParameterMap<float>::TNamedParameter>							ScalarParameters;
 	TArray<THashedMaterialParameterMap<FLinearColor>::TNamedParameter>					VectorParameters;
 	TArray<THashedMaterialParameterMap<FVector4d>::TNamedParameter>						DoubleVectorParameters;
diff --git a/Engine/Source/Runtime/Engine/Private/Materials/MaterialShader.cpp b/Engine/Source/Runtime/Engine/Private/Materials/MaterialShader.cpp
--- a/Engine/Source/Runtime/Engine/Private/Materials/MaterialShader.cpp
+++ b/Engine/Source/Runtime/Engine/Private/Materials/MaterialShader.cpp
@@ -498,7 +498,10 @@ void FStaticParameterSet::AddParametersOfType(EMaterialParameterType Type, const
 		{
 			const FMaterialParameterMetadata& Meta = It.Value;
 			check(Meta.Value.Type == Type);
-			StaticSwitchParameters.Emplace(It.Key, Meta.Value.AsStaticSwitch(), Meta.bOverride, Meta.ExpressionGuid);
+			if(!Meta.bDynamicSwitchParameter)
+			{
+				StaticSwitchParameters.Emplace(It.Key, Meta.Value.AsStaticSwitch(), Meta.bOverride, Meta.ExpressionGuid);
+			}
 		}
 		break;
 	case EMaterialParameterType::StaticComponentMask:
diff --git a/Engine/Source/Runtime/Engine/Private/Materials/MaterialUniformExpressions.cpp b/Engine/Source/Runtime/Engine/Private/Materials/MaterialUniformExpressions.cpp
--- a/Engine/Source/Runtime/Engine/Private/Materials/MaterialUniformExpressions.cpp
+++ b/Engine/Source/Runtime/Engine/Private/Materials/MaterialUniformExpressions.cpp
@@ -1739,6 +1739,7 @@ IMPLEMENT_MATERIALUNIFORMEXPRESSION_TYPE(FMaterialUniformExpressionTexture);
 IMPLEMENT_MATERIALUNIFORMEXPRESSION_TYPE(FMaterialUniformExpressionConstant);
 IMPLEMENT_MATERIALUNIFORMEXPRESSION_TYPE(FMaterialUniformExpressionGenericConstant);
 IMPLEMENT_MATERIALUNIFORMEXPRESSION_TYPE(FMaterialUniformExpressionNumericParameter);
+IMPLEMENT_MATERIALUNIFORMEXPRESSION_TYPE(FMaterialUniformExpressionStaticBoolParameter);
 IMPLEMENT_MATERIALUNIFORMEXPRESSION_TYPE(FMaterialUniformExpressionTextureParameter);
 IMPLEMENT_MATERIALUNIFORMEXPRESSION_TYPE(FMaterialUniformExpressionExternalTextureBase);
 IMPLEMENT_MATERIALUNIFORMEXPRESSION_TYPE(FMaterialUniformExpressionExternalTexture);
diff --git a/Engine/Source/Runtime/Engine/Private/Materials/MaterialUniformExpressions.h b/Engine/Source/Runtime/Engine/Private/Materials/MaterialUniformExpressions.h
--- a/Engine/Source/Runtime/Engine/Private/Materials/MaterialUniformExpressions.h
+++ b/Engine/Source/Runtime/Engine/Private/Materials/MaterialUniformExpressions.h
@@ -312,6 +312,57 @@ private:
 	int32 ParameterIndex;
 };
 
+/**
+ */
+class FMaterialUniformExpressionStaticBoolParameter : public FMaterialUniformExpression
+{
+	DECLARE_MATERIALUNIFORMEXPRESSION_TYPE(FMaterialUniformExpressionStaticBoolParameter);
+public:
+
+	FMaterialUniformExpressionStaticBoolParameter() {}
+	FMaterialUniformExpressionStaticBoolParameter(const FMaterialParameterInfo& InParameterInfo, uint32 InParameterIndex)
+		: ParameterIndex(InParameterIndex)
+		, ParameterInfo(InParameterInfo)
+	{
+		check(InParameterIndex >= 0 && InParameterIndex <= 0xffff);
+	}
+
+	// FMaterialUniformExpression interface.
+	virtual void WriteNumberOpcodes(UE::Shader::FPreshaderData& OutData) const override
+	{
+		OutData.WriteOpcode(UE::Shader::EPreshaderOpcode::Parameter);
+		OutData.Write((uint16)ParameterIndex);
+	}
+
+	virtual bool IsConstant() const
+	{
+		return false;
+	}
+
+	const FHashedMaterialParameterInfo& GetParameterInfo() const
+	{
+		return ParameterInfo;
+	}
+
+	FName GetParameterName() const
+	{
+		return ParameterInfo.GetName();
+	}
+
+	virtual bool IsIdentical(const FMaterialUniformExpression* OtherExpression) const
+	{
+		if (GetType() != OtherExpression->GetType())
+		{
+			return false;
+		}
+		FMaterialUniformExpressionStaticBoolParameter* OtherParameter = (FMaterialUniformExpressionStaticBoolParameter*)OtherExpression;
+		return ParameterInfo == OtherParameter->ParameterInfo && ParameterIndex == OtherParameter->ParameterIndex;
+	}
+private:
+	uint32 ParameterIndex;
+	FHashedMaterialParameterInfo ParameterInfo;
+};
+
 /** @return The texture that was associated with the given index when the given material had its uniform expressions/HLSL code generated. */
 template<typename TextureType>
 static TextureType* GetIndexedTexture(const FMaterial& Material, int32 TextureIndex)
diff --git a/Engine/Source/Runtime/Engine/Public/MaterialCachedData.h b/Engine/Source/Runtime/Engine/Public/MaterialCachedData.h
--- a/Engine/Source/Runtime/Engine/Public/MaterialCachedData.h
+++ b/Engine/Source/Runtime/Engine/Public/MaterialCachedData.h
@@ -141,7 +141,7 @@ struct FMaterialCachedExpressionEditorOnlyData
 	FMaterialCachedParameterEditorEntry EditorEntries[NumMaterialParameterTypes];
 
 	UPROPERTY()
-	TArray<bool> StaticSwitchValues;
+	TArray<bool> StaticSwitchValues_DEPRECATED;
 
 	UPROPERTY()
 	TArray<FStaticComponentMaskValue> StaticComponentMaskValues;
@@ -253,6 +253,8 @@ struct FMaterialCachedExpressionData
 		PropertyConnectedBitmask |= (1 << (uint32)Property);
 	}
 
+	void PostSerialize(const FArchive& Ar);
+
 #if WITH_EDITORONLY_DATA
 	TSharedPtr<FMaterialCachedExpressionEditorOnlyData> EditorOnlyData;
 #endif // WITH_EDITORONLY_DATA
@@ -269,6 +271,12 @@ struct FMaterialCachedExpressionData
 	UPROPERTY()
 	TArray<float> ScalarValues;
 
+	UPROPERTY()
+	TArray<bool> StaticSwitchValues;
+
+	UPROPERTY()
+	TArray<bool> DynamicSwitchValues;
+	
 	UPROPERTY()
 	TArray<FLinearColor> VectorValues;
 
diff --git a/Engine/Source/Runtime/Engine/Public/MaterialCompiler.h b/Engine/Source/Runtime/Engine/Public/MaterialCompiler.h
--- a/Engine/Source/Runtime/Engine/Public/MaterialCompiler.h
+++ b/Engine/Source/Runtime/Engine/Public/MaterialCompiler.h
@@ -354,6 +354,7 @@ public:
 	virtual int32 ParticleSpeed() = 0;
 	virtual int32 ParticleSize() = 0;
 
+	virtual int32 DynamicBranch(int32 Condition, int32 A, int32 B) = 0;
 	virtual int32 If(int32 A,int32 B,int32 AGreaterThanB,int32 AEqualsB,int32 ALessThanB,int32 Threshold) = 0;
 	virtual int32 Switch(int32 SwitchValueInput, int32 DefaultInput, TArray<int32>& CompiledInputs) = 0;
 
@@ -493,6 +494,7 @@ public:
 
 	virtual int32 StaticBool(bool Value) = 0;
 	virtual int32 StaticBoolParameter(FName ParameterName,bool bDefaultValue) = 0;
+	virtual int32 DynamicBoolParameter(FName ParameterName,bool bDefaultValue) = 0;
 	virtual int32 StaticComponentMask(int32 Vector,FName ParameterName,bool bDefaultR,bool bDefaultG,bool bDefaultB,bool bDefaultA) = 0;
 	virtual const FMaterialLayersFunctions* GetMaterialLayers() = 0;
 	virtual bool GetStaticBoolValue(int32 BoolIndex, bool& bSucceeded) = 0;
@@ -872,6 +874,7 @@ public:
 	virtual int32 ParticleRadius() override { return Compiler->ParticleRadius(); }
 	virtual int32 SphericalParticleOpacity(int32 Density) override { return Compiler->SphericalParticleOpacity(Density); }
 
+	virtual int32 DynamicBranch(int32 Condition, int32 A, int32 B) override { return Compiler->DynamicBranch(Condition, A, B); };
 	virtual int32 If(int32 A, int32 B, int32 AGreaterThanB, int32 AEqualsB, int32 ALessThanB, int32 Threshold) override { return Compiler->If(A, B, AGreaterThanB, AEqualsB, ALessThanB, Threshold); }
 	virtual int32 Switch(int32 SwitchValueInput, int32 DefaultInput, TArray<int32>& CompiledInputs) override { return Compiler->Switch(SwitchValueInput, DefaultInput, CompiledInputs); };
 	
@@ -934,6 +937,7 @@ public:
 
 	virtual int32 StaticBool(bool Value) override { return Compiler->StaticBool(Value); }
 	virtual int32 StaticBoolParameter(FName ParameterName, bool bDefaultValue) override { return Compiler->StaticBoolParameter(ParameterName, bDefaultValue); }
+	virtual int32 DynamicBoolParameter(FName ParameterName, bool bDefaultValue) override { return Compiler->DynamicBoolParameter(ParameterName, bDefaultValue); }
 	virtual int32 StaticComponentMask(int32 Vector, FName ParameterName, bool bDefaultR, bool bDefaultG, bool bDefaultB, bool bDefaultA) override { return Compiler->StaticComponentMask(Vector, ParameterName, bDefaultR, bDefaultG, bDefaultB, bDefaultA); }
 	virtual const FMaterialLayersFunctions* GetMaterialLayers() override { return Compiler->GetMaterialLayers(); }
 	virtual bool GetStaticBoolValue(int32 BoolIndex, bool& bSucceeded) override { return Compiler->GetStaticBoolValue(BoolIndex, bSucceeded); }
diff --git a/Engine/Source/Runtime/Engine/Public/MaterialShared.h b/Engine/Source/Runtime/Engine/Public/MaterialShared.h
--- a/Engine/Source/Runtime/Engine/Public/MaterialShared.h
+++ b/Engine/Source/Runtime/Engine/Public/MaterialShared.h
@@ -222,6 +222,8 @@ inline uint32 GetNumComponents(EMaterialValueType Type)
 	case MCT_LWCVector2: return 2;
 	case MCT_LWCVector3: return 3;
 	case MCT_LWCVector4: return 4;
+	case MCT_StaticBool: return 1;
+	case MCT_Bool: return 1;
 	default: return 0;
 	}
 }
diff --git a/Engine/Source/Runtime/Engine/Public/MaterialTypes.h b/Engine/Source/Runtime/Engine/Public/MaterialTypes.h
--- a/Engine/Source/Runtime/Engine/Public/MaterialTypes.h
+++ b/Engine/Source/Runtime/Engine/Public/MaterialTypes.h
@@ -194,12 +194,12 @@ enum class EMaterialParameterType : uint8
 	Font,
 	RuntimeVirtualTexture,
 	SparseVolumeTexture,
+	StaticSwitch,
 
 	NumRuntime, // Runtime parameter types must go above here, and editor-only ones below
 
 	// TODO - Would be nice to make static parameter values editor-only, but will save that for a future-refactor
-	StaticSwitch = NumRuntime,
-	StaticComponentMask,
+	StaticComponentMask = NumRuntime,
 
 	Num,
 	None = 0xff,
@@ -458,6 +458,9 @@ struct FMaterialParameterMetadata
 	/** Is the parameter overriden on the material it was queried from? */
 	bool bOverride = false;
 #endif // WITH_EDITORONLY_DATA
+
+	/** Is a Static Switch parameter using a dynamic branch? */
+	bool bDynamicSwitchParameter = false;
 };
 
 class FSHA1;
