Strata - added support for Strata as pin type for material function input.

diff --git a/Engine/Source/Runtime/Engine/Classes/Materials/MaterialExpressionFunctionInput.h b/Engine/Source/Runtime/Engine/Classes/Materials/MaterialExpressionFunctionInput.h
--- a/Engine/Source/Runtime/Engine/Classes/Materials/MaterialExpressionFunctionInput.h
+++ b/Engine/Source/Runtime/Engine/Classes/Materials/MaterialExpressionFunctionInput.h
@@ -29,6 +29,7 @@ enum EFunctionInputType : int
 	FunctionInput_MaterialAttributes,
 	FunctionInput_TextureExternal,
 	FunctionInput_Bool,
+	FunctionInput_Strata,
 	FunctionInput_MAX,
 };
 
@@ -108,6 +109,10 @@ class UMaterialExpressionFunctionInput : public UMaterialExpression
 	virtual int32 Compile(class FMaterialCompiler* Compiler, int32 OutputIndex) override;
 	virtual void GetCaption(TArray<FString>& OutCaptions) const override;
 
+
+	virtual bool IsResultStrataMaterial(int32 OutputIndex) override;
+	virtual void GatherStrataMaterialInfo(FStrataMaterialInfo& StrataMaterialInfo, int32 OutputIndex) override;
+	virtual FStrataOperator* StrataGenerateMaterialTopologyTree(class FMaterialCompiler* Compiler, class UMaterialExpression* Parent, int32 OutputIndex) override;
 	virtual bool IsResultMaterialAttributes(int32 OutputIndex) override;
 	virtual uint32 GetInputType(int32 InputIndex) override;
 	virtual uint32 GetOutputType(int32 OutputIndex) override;
diff --git a/Engine/Source/Runtime/Engine/Private/Materials/HLSLMaterialTranslator.cpp b/Engine/Source/Runtime/Engine/Private/Materials/HLSLMaterialTranslator.cpp
--- a/Engine/Source/Runtime/Engine/Private/Materials/HLSLMaterialTranslator.cpp
+++ b/Engine/Source/Runtime/Engine/Private/Materials/HLSLMaterialTranslator.cpp
@@ -789,7 +789,10 @@ bool FHLSLMaterialTranslator::Translate()
 
 		bEnableExecutionFlow = Material->IsUsingControlFlow();
 		bCompileForComputeShader = Material->IsLightFunction();
-
+		
+		//
+		// Process the strata tree representing the material topology.
+		//
 		const bool bStrataEnabled = Engine_IsStrataEnabled();
 		FStrataMaterialInput* FrontMaterialInput = Material->GetMaterialInterface() ? &Material->GetMaterialInterface()->GetMaterial()->GetEditorOnlyData()->FrontMaterial : nullptr;
 		UMaterialExpression* FrontMaterialExpr = FrontMaterialInput ? FrontMaterialInput->GetTracedInput().Expression : nullptr;
diff --git a/Engine/Source/Runtime/Engine/Private/Materials/MaterialExpressions.cpp b/Engine/Source/Runtime/Engine/Private/Materials/MaterialExpressions.cpp
--- a/Engine/Source/Runtime/Engine/Private/Materials/MaterialExpressions.cpp
+++ b/Engine/Source/Runtime/Engine/Private/Materials/MaterialExpressions.cpp
@@ -16068,8 +16068,9 @@ void UMaterialExpressionMaterialFunctionCall::PostEditChangeProperty(FPropertyCh
 void UMaterialExpressionMaterialFunctionCall::LinkFunctionIntoCaller(FMaterialCompiler* Compiler)
 {
 	MaterialFunction->LinkIntoCaller(FunctionInputs);
-	// Update parameter owner when stepping into layer functions
-	if (MaterialFunction->GetMaterialFunctionUsage() != EMaterialFunctionUsage::Default)
+	// Update parameter owner when stepping into layer functions.
+	// This is an optional step when we only want to march the material graph (e.g. to gather Strata material topology)
+	if (Compiler && MaterialFunction->GetMaterialFunctionUsage() != EMaterialFunctionUsage::Default)
 	{
 		Compiler->PushParameterOwner(FunctionParameterInfo);
 	}
@@ -16077,7 +16078,7 @@ void UMaterialExpressionMaterialFunctionCall::LinkFunctionIntoCaller(FMaterialCo
 
 void UMaterialExpressionMaterialFunctionCall::UnlinkFunctionFromCaller(FMaterialCompiler* Compiler)
 {
-	if (MaterialFunction->GetMaterialFunctionUsage() != EMaterialFunctionUsage::Default)
+	if (Compiler && MaterialFunction->GetMaterialFunctionUsage() != EMaterialFunctionUsage::Default)
 	{
 		const FMaterialParameterInfo PoppedParameterInfo = Compiler->PopParameterOwner();
 		check(PoppedParameterInfo == FunctionParameterInfo);
@@ -16175,7 +16176,8 @@ static const TCHAR* GetInputTypeName(uint8 InputType)
 		TEXT("TVol"),
 		TEXT("B"),
 		TEXT("MA"),
-		TEXT("TExt")
+		TEXT("TExt"),
+		TEXT("Stra")
 	};
 
 	check(InputType < FunctionInput_MAX);
@@ -16593,7 +16595,10 @@ void UMaterialExpressionMaterialFunctionCall::GatherStrataMaterialInfo(FStrataMa
 {
 	if (OutputIndex >= 0 && OutputIndex < FunctionOutputs.Num() && FunctionOutputs[OutputIndex].ExpressionOutput)
 	{
-		return FunctionOutputs[OutputIndex].ExpressionOutput->GatherStrataMaterialInfo(StrataMaterialInfo, 0);
+		this->LinkFunctionIntoCaller(nullptr);
+		FunctionOutputs[OutputIndex].ExpressionOutput->GatherStrataMaterialInfo(StrataMaterialInfo, 0);
+		this->UnlinkFunctionFromCaller(nullptr);
+		return;
 	}
 }
 
@@ -16601,7 +16606,10 @@ FStrataOperator* UMaterialExpressionMaterialFunctionCall::StrataGenerateMaterial
 {
 	if (OutputIndex >= 0 && OutputIndex < FunctionOutputs.Num() && FunctionOutputs[OutputIndex].ExpressionOutput)
 	{
-		return FunctionOutputs[OutputIndex].ExpressionOutput->StrataGenerateMaterialTopologyTree(Compiler, Parent, 0);
+		this->LinkFunctionIntoCaller(nullptr);
+		FStrataOperator* ResultingOperator = FunctionOutputs[OutputIndex].ExpressionOutput->StrataGenerateMaterialTopologyTree(Compiler, Parent, 0);
+		this->UnlinkFunctionFromCaller(nullptr);
+		return ResultingOperator;
 	}
 	return nullptr;
 }
@@ -16718,7 +16726,8 @@ void UMaterialExpressionFunctionInput::GetCaption(TArray<FString>& OutCaptions)
 		TEXT("StaticBool"),
 		TEXT("MaterialAttributes"),
 		TEXT("External"),
-		TEXT("Bool")
+		TEXT("Bool"),
+		TEXT("Strata")
 	};
 	check(InputType < FunctionInput_MAX);
 	OutCaptions.Add(FString(TEXT("Input ")) + InputName.ToString() + TEXT(" (") + TypeNames[InputType] + TEXT(")"));
@@ -16763,6 +16772,8 @@ int32 UMaterialExpressionFunctionInput::CompilePreviewValue(FMaterialCompiler* C
 			return Compiler->Constant4(PreviewValue.X, PreviewValue.Y, PreviewValue.Z, PreviewValue.W);
 		case FunctionInput_MaterialAttributes:		
 			return FMaterialAttributeDefinitionMap::CompileDefaultExpression(Compiler, AttributeID);
+		case FunctionInput_Strata:
+			return Compiler->StrataCreateAndRegisterNullMaterial();
 		case FunctionInput_Texture2D:
 		case FunctionInput_TextureCube:
 		case FunctionInput_Texture2DArray:
@@ -16790,7 +16801,8 @@ int32 UMaterialExpressionFunctionInput::Compile(class FMaterialCompiler* Compile
 		MCT_VolumeTexture,
 		MCT_StaticBool,
 		MCT_MaterialAttributes,
-		MCT_TextureExternal
+		MCT_TextureExternal,
+		MCT_Strata
 	};
 	check(InputType < FunctionInput_MAX);
 
@@ -16921,6 +16933,51 @@ bool UMaterialExpressionFunctionInput::IsResultMaterialAttributes(int32 OutputIn
 	}
 }
 
+bool UMaterialExpressionFunctionInput::IsResultStrataMaterial(int32 OutputIndex)
+{
+	if (FunctionInput_Strata == InputType)
+	{
+		return true;
+	}
+	else
+	{
+		return false;
+	}
+}
+
+void UMaterialExpressionFunctionInput::GatherStrataMaterialInfo(FStrataMaterialInfo& StrataMaterialInfo, int32 OutputIndex)
+{
+	if (FunctionInput_Strata == InputType)
+	{
+		FExpressionInput EffectivePreviewDuringCompileTracedInput = EffectivePreviewDuringCompile.GetTracedInput();
+		int32 ExpressionResult = INDEX_NONE;
+		if (EffectivePreviewDuringCompileTracedInput.GetTracedInput().Expression && !EffectivePreviewDuringCompileTracedInput.Expression->ContainsInputLoop())
+		{
+			EffectivePreviewDuringCompileTracedInput.Expression->GatherStrataMaterialInfo(StrataMaterialInfo, EffectivePreviewDuringCompileTracedInput.OutputIndex);
+		}
+	}
+}
+
+FStrataOperator* UMaterialExpressionFunctionInput::StrataGenerateMaterialTopologyTree(class FMaterialCompiler* Compiler, class UMaterialExpression* Parent, int32 OutputIndex)
+{
+	if (FunctionInput_Strata == InputType)
+	{
+		FExpressionInput EffectivePreviewDuringCompileTracedInput = EffectivePreviewDuringCompile.GetTracedInput();
+		int32 ExpressionResult = INDEX_NONE;
+		if (EffectivePreviewDuringCompileTracedInput.GetTracedInput().Expression && !EffectivePreviewDuringCompileTracedInput.Expression->ContainsInputLoop())
+		{
+			return EffectivePreviewDuringCompileTracedInput.Expression->StrataGenerateMaterialTopologyTree(Compiler, Parent, EffectivePreviewDuringCompileTracedInput.OutputIndex);
+		}
+	}
+	// If we are parsing for a material function input we always needs to return a default valid BSDF operator at least 
+	// If the material requires an input thenthe UI will forcethe user to proviced one.
+	// If this is not the case however, we need to compile a default material to compile the shader, or to preview the material function.
+	FStrataOperator& DefaultSlabOperator = Compiler->StrataCompilationRegisterOperator(STRATA_OPERATOR_BSDF, Compiler->StrataTreeStackGetPathUniqueId(), Parent, Compiler->StrataTreeStackGetParentPathUniqueId());
+	DefaultSlabOperator.BSDFType = STRATA_BSDF_TYPE_SLAB;
+	DefaultSlabOperator.ThicknessIndex = Compiler->StrataThicknessStackGetThicknessIndex();
+	return &DefaultSlabOperator;
+}
+
 uint32 UMaterialExpressionFunctionInput::GetInputType(int32 InputIndex)
 {
 	switch (InputType)
@@ -16949,6 +17006,8 @@ uint32 UMaterialExpressionFunctionInput::GetInputType(int32 InputIndex)
 		return MCT_Bool;
 	case FunctionInput_MaterialAttributes:
 		return MCT_MaterialAttributes;
+	case FunctionInput_Strata:
+		return MCT_Strata;
 	default:
 		return MCT_Unknown;
 	}
