Fixed some missing code to handle Bool and Strata material function input.
Noticed after 23845504 was merged with 23843814.

diff --git a/Engine/Source/Editor/UnrealEd/Private/MaterialGraphSchema.cpp b/Engine/Source/Editor/UnrealEd/Private/MaterialGraphSchema.cpp
--- a/Engine/Source/Editor/UnrealEd/Private/MaterialGraphSchema.cpp
+++ b/Engine/Source/Editor/UnrealEd/Private/MaterialGraphSchema.cpp
@@ -117,6 +117,9 @@ void FMaterialGraphSchemaAction_NewNode::SetFunctionInputType(UMaterialExpressio
 	case MCT_MaterialAttributes:
 		FunctionInput->InputType = FunctionInput_MaterialAttributes;
 		break;
+	case MCT_Strata:
+		FunctionInput->InputType = FunctionInput_Strata;
+		break;
 	default:
 		break;
 	}
diff --git a/Engine/Source/Runtime/Engine/Private/Materials/MaterialExpressions.cpp b/Engine/Source/Runtime/Engine/Private/Materials/MaterialExpressions.cpp
--- a/Engine/Source/Runtime/Engine/Private/Materials/MaterialExpressions.cpp
+++ b/Engine/Source/Runtime/Engine/Private/Materials/MaterialExpressions.cpp
@@ -16174,9 +16174,10 @@ static const TCHAR* GetInputTypeName(uint8 InputType)
 		TEXT("TCube"),
 		TEXT("T2dArr"),
 		TEXT("TVol"),
-		TEXT("B"),
+		TEXT("SB"),
 		TEXT("MA"),
 		TEXT("TExt"),
+		TEXT("B"),
 		TEXT("Stra")
 	};
 
@@ -16802,6 +16803,7 @@ int32 UMaterialExpressionFunctionInput::Compile(class FMaterialCompiler* Compile
 		MCT_StaticBool,
 		MCT_MaterialAttributes,
 		MCT_TextureExternal,
+		MCT_Bool,
 		MCT_Strata
 	};
 	check(InputType < FunctionInput_MAX);
