Fix crash when generating thumbnails for cooked additive animations

diff --git a/Engine/Source/Runtime/Engine/Private/Animation/AnimSequence.cpp b/Engine/Source/Runtime/Engine/Private/Animation/AnimSequence.cpp
--- a/Engine/Source/Runtime/Engine/Private/Animation/AnimSequence.cpp
+++ b/Engine/Source/Runtime/Engine/Private/Animation/AnimSequence.cpp
@@ -1645,13 +1645,13 @@ void UAnimSequence::GetAdditiveBasePose(FCompactPose& OutPose, FBlendedCurve& Ou
 	GetAdditiveBasePose(OutAnimationPoseData, ExtractionContext);
 }
 
-static void GetSequencePose(FAnimationPoseData& OutAnimationPoseData, const FAnimExtractContext& ExtractionContext, const UAnimSequence &Seq, int32 FrameIndex)
+static void GetSequencePose(FAnimationPoseData& OutAnimationPoseData, const FAnimExtractContext& ExtractionContext, const UAnimSequence &Seq, int32 FrameIndex, bool bForceUseRawData)
 {
 	const double Fraction = (Seq.GetNumberOfSampledKeys() > 0) ? FMath::Clamp<double>((double)FrameIndex / (double)Seq.GetNumberOfSampledKeys(), 0.0, 1.0) : 0.0;
 	const double BasePoseTime = Seq.GetPlayLength() * Fraction;
 	FAnimExtractContext BasePoseExtractionContext(ExtractionContext);
 	BasePoseExtractionContext.CurrentTime = BasePoseTime;
-	Seq.GetBonePose(OutAnimationPoseData, BasePoseExtractionContext, true);
+	Seq.GetBonePose(OutAnimationPoseData, BasePoseExtractionContext, bForceUseRawData);
 }
 
 void UAnimSequence::GetAdditiveBasePose(FAnimationPoseData& OutAnimationPoseData, const FAnimExtractContext& ExtractionContext) const
@@ -1669,19 +1669,19 @@ void UAnimSequence::GetAdditiveBasePose(FAnimationPoseData& OutAnimationPoseData
 
 			FAnimExtractContext BasePoseExtractionContext(ExtractionContext);
 			BasePoseExtractionContext.CurrentTime = BasePoseTime;
-			RefPoseSeq->GetBonePose(OutAnimationPoseData, BasePoseExtractionContext, true);
+			RefPoseSeq->GetBonePose(OutAnimationPoseData, BasePoseExtractionContext, CanEvaluateRawAnimationData());
 			break;
 		}
 		// use animation as a base pose. Need BasePoseSeq and RefFrameIndex (will clamp if outside).
 		case ABPT_AnimFrame:
 		{
-			GetSequencePose(OutAnimationPoseData, ExtractionContext, *RefPoseSeq, RefFrameIndex);
+			GetSequencePose(OutAnimationPoseData, ExtractionContext, *RefPoseSeq, RefFrameIndex, CanEvaluateRawAnimationData());
 			break;
 		}
 		// use this animation as a base pose. Need RefFrameIndex (will clamp if outside).
 		case ABPT_LocalAnimFrame:
 		{
-			GetSequencePose(OutAnimationPoseData, ExtractionContext, *this, RefFrameIndex);
+			GetSequencePose(OutAnimationPoseData, ExtractionContext, *this, RefFrameIndex, CanEvaluateRawAnimationData());
 			break;
 		}
 		// use ref pose of Skeleton as base
@@ -1714,7 +1714,7 @@ void UAnimSequence::GetBonePose_AdditiveMeshRotationOnly(FAnimationPoseData& Out
 	}
 
 	// Get target pose
-	GetBonePose(OutAnimationPoseData, ExtractionContext, true);
+	GetBonePose(OutAnimationPoseData, ExtractionContext, CanEvaluateRawAnimationData());
 
 	// get base pose
 	FCompactPose BasePose;
