IWorldPartitionCell is now a UINTERFACE and exposed GetCellOwner to the interface.

diff --git a/Engine/Source/Runtime/Engine/Private/WorldPartition/WorldPartitionRuntimeCellInterface.cpp b/Engine/Source/Runtime/Engine/Private/WorldPartition/WorldPartitionRuntimeCellInterface.cpp
new file mode 100644
--- /dev/null
+++ b/Engine/Source/Runtime/Engine/Private/WorldPartition/WorldPartitionRuntimeCellInterface.cpp
@@ -0,0 +1,9 @@
+// Copyright Epic Games, Inc. All Rights Reserved.
+
+#include "WorldPartition/WorldPartitionRuntimeCellInterface.h"
+
+#include UE_INLINE_GENERATED_CPP_BY_NAME(WorldPartitionRuntimeCellInterface)
+
+UWorldPartitionCell::UWorldPartitionCell(const FObjectInitializer& ObjectInitializer)
+	: Super(ObjectInitializer)
+{}
diff --git a/Engine/Source/Runtime/Engine/Public/WorldPartition/WorldPartitionRuntimeCell.h b/Engine/Source/Runtime/Engine/Public/WorldPartition/WorldPartitionRuntimeCell.h
--- a/Engine/Source/Runtime/Engine/Public/WorldPartition/WorldPartitionRuntimeCell.h
+++ b/Engine/Source/Runtime/Engine/Public/WorldPartition/WorldPartitionRuntimeCell.h
@@ -180,8 +180,6 @@ class ENGINE_API UWorldPartitionRuntimeCell : public UObject, public IWorldParti
 	virtual FGuid const& GetContentBundleID() const { return ContentBundleID; }
 	virtual TArray<FName> GetActors() const PURE_VIRTUAL(UWorldPartitionRuntimeCell::GetActors, return TArray<FName>(););
 
-	IWorldPartitionRuntimeCellOwner* GetCellOwner() const { return CastChecked<IWorldPartitionRuntimeCellOwner>(GetOuter()); }
-
 	/** Caches information on streaming source that will be used later on to sort cell. */
 	bool ShouldResetStreamingSourceInfo() const;
 	void ResetStreamingSourceInfo() const;
@@ -202,6 +200,7 @@ class ENGINE_API UWorldPartitionRuntimeCell : public UObject, public IWorldParti
 	virtual const FBox& GetContentBounds() const override;
 	virtual FBox GetCellBounds() const override;
 	virtual FName GetLevelPackageName() const override;
+	virtual IWorldPartitionRuntimeCellOwner* GetCellOwner() const override { return CastChecked<IWorldPartitionRuntimeCellOwner>(GetOuter()); }
 	//~End IWorldPartitionCell Interface
 
 	bool GetBlockOnSlowLoading() const { return bBlockOnSlowLoading; }
diff --git a/Engine/Source/Runtime/Engine/Public/WorldPartition/WorldPartitionRuntimeCellInterface.h b/Engine/Source/Runtime/Engine/Public/WorldPartition/WorldPartitionRuntimeCellInterface.h
--- a/Engine/Source/Runtime/Engine/Public/WorldPartition/WorldPartitionRuntimeCellInterface.h
+++ b/Engine/Source/Runtime/Engine/Public/WorldPartition/WorldPartitionRuntimeCellInterface.h
@@ -2,12 +2,23 @@
 #pragma once
 
 #include "CoreMinimal.h"
+#include "UObject/Interface.h"
+#include "WorldPartitionRuntimeCellInterface.generated.h"
 
 class UDataLayerInstance;
 class UDataLayerAsset;
+class IWorldPartitionRuntimeCellOwner;
+
+UINTERFACE()
+class ENGINE_API UWorldPartitionCell : public UInterface
+{
+	GENERATED_UINTERFACE_BODY()
+};
 
 class IWorldPartitionCell
 {
+	GENERATED_IINTERFACE_BODY()
+
 public:
 	/** Returns the cell content associated data layer instances. */
 	virtual TArray<const UDataLayerInstance*> GetDataLayerInstances() const = 0;
@@ -27,6 +38,8 @@ public:
 	virtual FBox GetCellBounds() const = 0;
 	/** Returns the associated level package name. */
 	virtual FName GetLevelPackageName() const = 0;
+	/** Returns cell owner */
+	virtual IWorldPartitionRuntimeCellOwner* GetCellOwner() const = 0;
 
 #if WITH_EDITOR
 	/** Returns the referenced actor packages. */
