Strata - StrataData as type of input to material function now requires an explicite default value.

diff --git a/Engine/Source/Runtime/Engine/Private/Materials/MaterialExpressions.cpp b/Engine/Source/Runtime/Engine/Private/Materials/MaterialExpressions.cpp
--- a/Engine/Source/Runtime/Engine/Private/Materials/MaterialExpressions.cpp
+++ b/Engine/Source/Runtime/Engine/Private/Materials/MaterialExpressions.cpp
@@ -16773,14 +16773,13 @@ int32 UMaterialExpressionFunctionInput::CompilePreviewValue(FMaterialCompiler* C
 			return Compiler->Constant4(PreviewValue.X, PreviewValue.Y, PreviewValue.Z, PreviewValue.W);
 		case FunctionInput_MaterialAttributes:		
 			return FMaterialAttributeDefinitionMap::CompileDefaultExpression(Compiler, AttributeID);
-		case FunctionInput_Strata:
-			return Compiler->StrataCreateAndRegisterNullMaterial();
 		case FunctionInput_Texture2D:
 		case FunctionInput_TextureCube:
 		case FunctionInput_Texture2DArray:
 		case FunctionInput_TextureExternal:
 		case FunctionInput_StaticBool:
 		case FunctionInput_Bool:
+		case FunctionInput_Strata:
 			return Compiler->Errorf(TEXT("Missing Preview connection for function input '%s'"), *InputName.ToString());
 		default:
 			return Compiler->Errorf(TEXT("Unknown input type"));
