InstancedStaticMesh - handle hidden materials for ray tracing in the same way as for StaticMesh

diff --git a/Engine/Source/Runtime/Engine/Private/InstancedStaticMesh.cpp b/Engine/Source/Runtime/Engine/Private/InstancedStaticMesh.cpp
--- a/Engine/Source/Runtime/Engine/Private/InstancedStaticMesh.cpp
+++ b/Engine/Source/Runtime/Engine/Private/InstancedStaticMesh.cpp
@@ -1885,9 +1885,20 @@ void FInstancedStaticMeshSceneProxy::GetDynamicRayTracingInstances(struct FRayTr
 			FMeshBatch MeshBatch;
 			FMeshBatch DynamicMeshBatch;
 
-			GetMeshElement(LOD, 0, SectionIdx, 0, false, false, DynamicMeshBatch);
+			if (!GetMeshElement(LOD, 0, SectionIdx, 0, false, false, DynamicMeshBatch))
+			{
+				DynamicMeshBatch.MaterialRenderProxy = UMaterial::GetDefaultMaterial(MD_Surface)->GetRenderProxy();
+				DynamicMeshBatch.SegmentIndex = SectionIdx;
+				DynamicMeshBatch.MeshIdInPrimitive = SectionIdx;
+			}
 
-			FStaticMeshSceneProxy::GetMeshElement(LOD, 0, SectionIdx, 0, false, false, MeshBatch);
+			if (!FStaticMeshSceneProxy::GetMeshElement(LOD, 0, SectionIdx, 0, false, false, MeshBatch))
+			{				
+				MeshBatch.MaterialRenderProxy = UMaterial::GetDefaultMaterial(MD_Surface)->GetRenderProxy();
+				MeshBatch.VertexFactory = &RenderData->LODVertexFactories[LOD].VertexFactory;
+				MeshBatch.SegmentIndex = SectionIdx;
+				MeshBatch.MeshIdInPrimitive = SectionIdx;
+			};
 
 			DynamicMeshBatch.VertexFactory = &InstancedRenderData.VertexFactories[LOD];
 
@@ -2115,8 +2126,17 @@ void FInstancedStaticMeshSceneProxy::GetDynamicRayTracingInstances(struct FRayTr
 		{
 			//#dxr_todo: so far we use the parent static mesh path to get material data
 			FMeshBatch MeshBatch;
-			FStaticMeshSceneProxy::GetMeshElement(LOD, 0, SectionIdx, 0, false, false, MeshBatch);
 
+			bool bResult = FStaticMeshSceneProxy::GetMeshElement(LOD, 0, SectionIdx, 0, false, false, MeshBatch);
+			if (!bResult)
+			{
+				// Hidden material
+				MeshBatch.MaterialRenderProxy = UMaterial::GetDefaultMaterial(MD_Surface)->GetRenderProxy();
+				MeshBatch.VertexFactory = &RenderData->LODVertexFactories[LOD].VertexFactory;
+				MeshBatch.SegmentIndex = SectionIdx;
+				MeshBatch.MeshIdInPrimitive = SectionIdx;
+			}
+			
 			RayTracingInstanceTemplate.Materials.Add(MeshBatch);
 		}
 
