Expose Montage_GetEffectivePlayRate, which accounts for PlayRate * RateScale.
The current Montage_GetPlayRate returns only PlayRate, so it doesn't provide an accurate final play rate if the Montage asset has a RateScale that isn't 1.

diff --git a/Engine/Source/Runtime/Engine/Classes/Animation/AnimInstance.h b/Engine/Source/Runtime/Engine/Classes/Animation/AnimInstance.h
--- a/Engine/Source/Runtime/Engine/Classes/Animation/AnimInstance.h
+++ b/Engine/Source/Runtime/Engine/Classes/Animation/AnimInstance.h
@@ -642,12 +642,18 @@ public:
 	UFUNCTION(BlueprintPure, Category = "Animation|Montage")
 	float Montage_GetBlendTime(const UAnimMontage* Montage) const;
 
-	/** Get PlayRate for Montage.
+	/** Get PlayRate for Montage. This does not account for RateScale, so it may not reflect the actual play rate seen in game (see Montage_GetEffectivePlayRate).
 	If Montage reference is NULL, PlayRate for any Active Montage will be returned.
 	If Montage is not playing, 0 is returned. */
 	UFUNCTION(BlueprintPure, Category = "Animation|Montage")
 	float Montage_GetPlayRate(const UAnimMontage* Montage) const;
 
+	/** Get scaled PlayRate for Montage. This accounts for RateScale, so it will reflect the actual play rate seen in game.
+	If Montage reference is NULL, scaled PlayRate for any Active Montage will be returned.
+	If Montage is not playing, 0 is returned. */
+	UFUNCTION(BlueprintPure, Category = "Animation|Montage")
+	float Montage_GetEffectivePlayRate(const UAnimMontage* Montage) const;
+
 	/*********************************************************************************************
 	* AnimMontage sync. See notes in AnimMontage.h
 	********************************************************************************************* */
diff --git a/Engine/Source/Runtime/Engine/Private/Animation/AnimInstance.cpp b/Engine/Source/Runtime/Engine/Private/Animation/AnimInstance.cpp
--- a/Engine/Source/Runtime/Engine/Private/Animation/AnimInstance.cpp
+++ b/Engine/Source/Runtime/Engine/Private/Animation/AnimInstance.cpp
@@ -2704,6 +2704,32 @@ float UAnimInstance::Montage_GetPlayRate(const UAnimMontage* Montage) const
 	return 0.f;
 }
 
+float UAnimInstance::Montage_GetEffectivePlayRate(const UAnimMontage* Montage) const
+{
+	if (Montage)
+	{
+		const FAnimMontageInstance* MontageInstance = GetActiveInstanceForMontage(Montage);
+		if (MontageInstance)
+		{
+			return MontageInstance->GetPlayRate() * Montage->RateScale;
+		}
+	}
+	else
+	{
+		// If no Montage reference, use first active one found.
+		for (int32 InstanceIndex = 0; InstanceIndex < MontageInstances.Num(); InstanceIndex++)
+		{
+			const FAnimMontageInstance* MontageInstance = MontageInstances[InstanceIndex];
+			if (MontageInstance && MontageInstance->IsActive() && MontageInstance->Montage)
+			{
+				return MontageInstance->GetPlayRate() * MontageInstance->Montage->RateScale;
+			}
+		}
+	}
+
+	return 0.f;
+}
+
 void UAnimInstance::MontageSync_Follow(const UAnimMontage* MontageFollower, const UAnimInstance* OtherAnimInstance, const UAnimMontage* MontageLeader)
 {
 	if (!MontageFollower || !OtherAnimInstance || !MontageLeader)
