[Navmesh] Fix FCompositeNavModifier IsEmpty() to handle navmesh resolutions.
#rnx

diff --git a/Engine/Source/Runtime/Engine/Public/AI/NavigationModifier.h b/Engine/Source/Runtime/Engine/Public/AI/NavigationModifier.h
--- a/Engine/Source/Runtime/Engine/Public/AI/NavigationModifier.h
+++ b/Engine/Source/Runtime/Engine/Public/AI/NavigationModifier.h
@@ -275,7 +275,8 @@ struct ENGINE_API FCompositeNavModifier : public FNavigationModifier
 	{ 
 		return (Areas.Num() == 0) && (SimpleLinks.Num() == 0) && (CustomLinks.Num() == 0) && 
 			!bFillCollisionUnderneathForNavmesh && 
-			!bMaskFillCollisionUnderneathForNavmesh;
+			!bMaskFillCollisionUnderneathForNavmesh &&
+			NavMeshResolution == ENavigationDataResolution::Invalid;
 	}
 
 	void Add(const FAreaNavModifier& Area)
