World Partition: unified generic hash code for streaming generation.

diff --git a/Engine/Source/Runtime/Engine/Private/WorldPartition/WorldPartitionRuntimeCell.cpp b/Engine/Source/Runtime/Engine/Private/WorldPartition/WorldPartitionRuntimeCell.cpp
--- a/Engine/Source/Runtime/Engine/Private/WorldPartition/WorldPartitionRuntimeCell.cpp
+++ b/Engine/Source/Runtime/Engine/Private/WorldPartition/WorldPartitionRuntimeCell.cpp
@@ -39,7 +39,7 @@ void UWorldPartitionRuntimeCell::PostDuplicate(bool bDuplicateForPIE)
 		{
 			if (ensure(Actor))
 			{
-				// Don't use AActor::Rename here since the actor is not par of the world, it's only a duplication template.	
+				// Don't use AActor::Rename here since the actor is not part of the world, it's only a duplication template.	
 				Actor->UObject::Rename(nullptr, UnsavedActorsContainer);
 			}
 		}
diff --git a/Engine/Source/Runtime/Engine/Private/WorldPartition/WorldPartitionRuntimeHash.cpp b/Engine/Source/Runtime/Engine/Private/WorldPartition/WorldPartitionRuntimeHash.cpp
--- a/Engine/Source/Runtime/Engine/Private/WorldPartition/WorldPartitionRuntimeHash.cpp
+++ b/Engine/Source/Runtime/Engine/Private/WorldPartition/WorldPartitionRuntimeHash.cpp
@@ -6,6 +6,9 @@
 #include "WorldPartition/WorldPartition.h"
 #include "WorldPartition/WorldPartitionActorDescView.h"
 #include "WorldPartition/WorldPartitionStreamingSource.h"
+#if WITH_EDITOR
+#include "WorldPartition/Cook/WorldPartitionCookPackage.h"
+#endif
 
 #include UE_INLINE_GENERATED_CPP_BY_NAME(WorldPartitionRuntimeHash)
 
@@ -48,30 +51,46 @@ void UWorldPartitionRuntimeHash::OnEndPlay()
 	ModifiedActorDescListForPIE.Empty();
 }
 
+bool UWorldPartitionRuntimeHash::GenerateStreaming(class UWorldPartitionStreamingPolicy* StreamingPolicy, const IStreamingGenerationContext* StreamingGenerationContext, TArray<FString>* OutPackagesToGenerate)
+{
+	return PackagesToGenerateForCook.IsEmpty();
+}
+
+void UWorldPartitionRuntimeHash::FlushStreaming()
+{
+	PackagesToGenerateForCook.Empty();
+}
+
 // In PIE, Always loaded cell is not generated. Instead, always loaded actors will be added to AlwaysLoadedActorsForPIE.
 // This will trigger loading/registration of these actors in the PersistentLevel (if not already loaded).
 // Then, duplication of world for PIE will duplicate only these actors. 
 // When stopping PIE, WorldPartition will release these FWorldPartitionReferences which 
 // will unload actors that were not already loaded in the non PIE world.
-bool UWorldPartitionRuntimeHash::ConditionalRegisterAlwaysLoadedActorsForPIE(const FWorldPartitionActorDescView& ActorDescView, bool bIsMainWorldPartition, bool bIsMainContainer, bool bIsCellAlwaysLoaded)
+bool UWorldPartitionRuntimeHash::ConditionalRegisterAlwaysLoadedActorsForPIE(const IStreamingGenerationContext::FActorSetInstance* ActorSetInstance, bool bIsMainWorldPartition, bool bIsMainContainer, bool bIsCellAlwaysLoaded)
 {
 	if (bIsMainWorldPartition && bIsMainContainer && bIsCellAlwaysLoaded && !IsRunningCookCommandlet())
 	{
-		// This will load the actor if it isn't already loaded
-		FWorldPartitionReference Reference(GetOuterUWorldPartition(), ActorDescView.GetGuid());
-
-		if (AActor* AlwaysLoadedActor = FindObject<AActor>(nullptr, *ActorDescView.GetActorSoftPath().ToString()))
+		for (const FGuid& ActorGuid : ActorSetInstance->ActorSet->Actors)
 		{
-			AlwaysLoadedActorsForPIE.Emplace(Reference, AlwaysLoadedActor);
+			IStreamingGenerationContext::FActorInstance ActorInstance(ActorGuid, ActorSetInstance);
+			const FWorldPartitionActorDescView& ActorDescView = ActorInstance.GetActorDescView();
 
-			// Handle child actors
-			AlwaysLoadedActor->ForEachComponent<UChildActorComponent>(true, [this, &Reference](UChildActorComponent* ChildActorComponent)
+			// This will load the actor if it isn't already loaded
+			FWorldPartitionReference Reference(GetOuterUWorldPartition(), ActorDescView.GetGuid());
+
+			if (AActor* AlwaysLoadedActor = FindObject<AActor>(nullptr, *ActorDescView.GetActorSoftPath().ToString()))
 			{
-				if (AActor* ChildActor = ChildActorComponent->GetChildActor())
+				AlwaysLoadedActorsForPIE.Emplace(Reference, AlwaysLoadedActor);
+
+				// Handle child actors
+				AlwaysLoadedActor->ForEachComponent<UChildActorComponent>(true, [this, &Reference](UChildActorComponent* ChildActorComponent)
 				{
-					AlwaysLoadedActorsForPIE.Emplace(Reference, ChildActor);
-				}
-			});
+					if (AActor* ChildActor = ChildActorComponent->GetChildActor())
+					{
+						AlwaysLoadedActorsForPIE.Emplace(Reference, ChildActor);
+					}
+				});
+			}
 		}
 
 		return true;
@@ -80,6 +99,134 @@ bool UWorldPartitionRuntimeHash::ConditionalRegisterAlwaysLoadedActorsForPIE(con
 	return false;
 }
 
+void UWorldPartitionRuntimeHash::PopulateRuntimeCell(UWorldPartitionRuntimeCell* RuntimeCell, const TArray<IStreamingGenerationContext::FActorInstance>& ActorInstances, TArray<FString>* OutPackagesToGenerate)
+{
+	for (const IStreamingGenerationContext::FActorInstance& ActorInstance : ActorInstances)
+	{
+		if (ActorInstance.GetContainerID().IsMainContainer())
+		{
+			const FWorldPartitionActorDescView& ActorDescView = ActorInstance.GetActorDescView();
+			if (AActor* Actor = FindObject<AActor>(nullptr, *ActorDescView.GetActorSoftPath().ToString()))
+			{
+				if (ModifiedActorDescListForPIE.GetActorDesc(ActorDescView.GetGuid()))
+				{
+					// Create an actor container to make sure duplicated actors will share an outer to properly remap inter-actors references
+					RuntimeCell->UnsavedActorsContainer = NewObject<UActorContainer>(RuntimeCell);
+					break;
+				}
+			}
+		}
+	}
+
+	FBox CellContentBounds(ForceInit);
+	for (const IStreamingGenerationContext::FActorInstance& ActorInstance : ActorInstances)
+	{
+		const FWorldPartitionActorDescView& ActorDescView = ActorInstance.GetActorDescView();
+		RuntimeCell->AddActorToCell(ActorDescView, ActorInstance.GetContainerID(), ActorInstance.GetTransform(), ActorInstance.GetActorDescContainer());
+		CellContentBounds += ActorDescView.GetRuntimeBounds().TransformBy(ActorInstance.GetTransform());
+					
+		if (ActorInstance.GetContainerID().IsMainContainer() && RuntimeCell->UnsavedActorsContainer)
+		{
+			if (AActor* Actor = FindObject<AActor>(nullptr, *ActorDescView.GetActorSoftPath().ToString()))
+			{
+				RuntimeCell->UnsavedActorsContainer->Actors.Add(Actor->GetFName(), Actor);
+
+				// Handle child actors
+				Actor->ForEachComponent<UChildActorComponent>(true, [RuntimeCell](UChildActorComponent* ChildActorComponent)
+				{
+					if (AActor* ChildActor = ChildActorComponent->GetChildActor())
+					{
+						RuntimeCell->UnsavedActorsContainer->Actors.Add(ChildActor->GetFName(), ChildActor);
+					}
+				});
+			}
+		}
+	}
+
+	RuntimeCell->RuntimeCellData->ContentBounds = CellContentBounds;
+
+	// Always loaded cell actors are transfered to World's Persistent Level (see UWorldPartitionRuntimeSpatialHash::PopulateGeneratorPackageForCook)
+	if (OutPackagesToGenerate && RuntimeCell->GetActorCount() && !RuntimeCell->IsAlwaysLoaded())
+	{
+		const FString PackageRelativePath = RuntimeCell->GetPackageNameToCreate();
+		check(!PackageRelativePath.IsEmpty());
+
+		OutPackagesToGenerate->Add(PackageRelativePath);
+
+		// Map relative package to StreamingCell for PopulateGeneratedPackageForCook/PopulateGeneratorPackageForCook/GetCellForPackage
+		PackagesToGenerateForCook.Add(PackageRelativePath, RuntimeCell);
+
+		UE_CLOG(IsRunningCookCommandlet(), LogWorldPartition, Log, TEXT("Creating runtime streaming cells %s."), *RuntimeCell->GetName());
+	}
+}
+
+bool UWorldPartitionRuntimeHash::PopulateGeneratedPackageForCook(const FWorldPartitionCookPackage& InPackagesToCook, TArray<UPackage*>& OutModifiedPackages)
+{
+	OutModifiedPackages.Reset();
+	if (UWorldPartitionRuntimeCell** MatchingCell = PackagesToGenerateForCook.Find(InPackagesToCook.RelativePath))
+	{
+		UWorldPartitionRuntimeCell* Cell = *MatchingCell;
+		if (ensure(Cell))
+		{
+			return Cell->PopulateGeneratedPackageForCook(InPackagesToCook.GetPackage(), OutModifiedPackages);
+		}
+	}
+	return false;
+}
+
+UWorldPartitionRuntimeCell* UWorldPartitionRuntimeHash::GetCellForPackage(const FWorldPartitionCookPackage& PackageToCook) const
+{
+	UWorldPartitionRuntimeCell** MatchingCell = const_cast<UWorldPartitionRuntimeCell**>(PackagesToGenerateForCook.Find(PackageToCook.RelativePath));
+	return MatchingCell ? *MatchingCell : nullptr;
+}
+
+TArray<UWorldPartitionRuntimeCell*> UWorldPartitionRuntimeHash::GetAlwaysLoadedCells() const
+{
+	TArray<UWorldPartitionRuntimeCell*> Result;
+	ForEachStreamingCells([&Result](const UWorldPartitionRuntimeCell* Cell)
+	{
+		if (Cell->IsAlwaysLoaded())
+		{
+			Result.Add(const_cast<UWorldPartitionRuntimeCell*>(Cell));
+		}
+		return true;
+	});
+	return Result;
+}
+
+bool UWorldPartitionRuntimeHash::PrepareGeneratorPackageForCook(TArray<UPackage*>& OutModifiedPackages)
+{
+	check(IsRunningCookCommandlet());
+
+	for (UWorldPartitionRuntimeCell* Cell : GetAlwaysLoadedCells())
+	{
+		check(Cell->IsAlwaysLoaded());
+		if (!Cell->PopulateGeneratorPackageForCook(OutModifiedPackages))
+		{
+			return false;
+		}
+	}
+
+	//@todo_ow: here we can safely remove always loaded cells as they are not part of the OutPackagesToGenerate
+	return true;
+}
+
+bool UWorldPartitionRuntimeHash::PopulateGeneratorPackageForCook(const TArray<FWorldPartitionCookPackage*>& InPackagesToCook, TArray<UPackage*>& OutModifiedPackages)
+{
+	check(IsRunningCookCommandlet());
+
+	for (const FWorldPartitionCookPackage* CookPackage : InPackagesToCook)
+	{
+		UWorldPartitionRuntimeCell** MatchingCell = PackagesToGenerateForCook.Find(CookPackage->RelativePath);
+		UWorldPartitionRuntimeCell* Cell = MatchingCell ? *MatchingCell : nullptr;
+		if (!Cell || !Cell->PrepareCellForCook(CookPackage->GetPackage()))
+		{
+			return false;
+		}
+	}
+	return true;
+}
+
 void UWorldPartitionRuntimeHash::DumpStateLog(FHierarchicalLogArchive& Ar)
 {
 	Ar.Printf(TEXT("----------------------------------------------------------------------------------------------------------------"));
diff --git a/Engine/Source/Runtime/Engine/Private/WorldPartition/WorldPartitionRuntimeSpatialHash.cpp b/Engine/Source/Runtime/Engine/Private/WorldPartition/WorldPartitionRuntimeSpatialHash.cpp
--- a/Engine/Source/Runtime/Engine/Private/WorldPartition/WorldPartitionRuntimeSpatialHash.cpp
+++ b/Engine/Source/Runtime/Engine/Private/WorldPartition/WorldPartitionRuntimeSpatialHash.cpp
@@ -35,7 +35,6 @@
 #include UE_INLINE_GENERATED_CPP_BY_NAME(WorldPartitionRuntimeSpatialHash)
 
 #if WITH_EDITOR
-#include "WorldPartition/Cook/WorldPartitionCookPackage.h"
 extern UNREALED_API class UEditorEngine* GEditor;
 #else
 #include "Misc/Paths.h"
@@ -1108,6 +1107,8 @@ void UWorldPartitionRuntimeSpatialHash::SetDefaultValues()
 
 bool UWorldPartitionRuntimeSpatialHash::GenerateStreaming(UWorldPartitionStreamingPolicy* StreamingPolicy, const IStreamingGenerationContext* StreamingGenerationContext, TArray<FString>* OutPackagesToGenerate)
 {
+	verify(Super::GenerateStreaming(StreamingPolicy, StreamingGenerationContext, OutPackagesToGenerate));
+
 	TRACE_CPUPROFILER_EVENT_SCOPE(UWorldPartitionRuntimeSpatialHash::GenerateStreaming);
 	UWorldPartition* WorldPartition = GetOuterUWorldPartition();
 
@@ -1122,7 +1123,6 @@ bool UWorldPartitionRuntimeSpatialHash::GenerateStreaming(UWorldPartitionStreami
 	// Fix case where StreamingGrids might have been persisted.
 	bIsNameToGridMappingDirty = true;
 	StreamingGrids.Empty();
-	check(PackagesToGenerateForCook.IsEmpty());
 
 	// Append grids from ASpatialHashRuntimeGridInfo actors to runtime spatial hash grids
 	TArray<FSpatialHashRuntimeGrid> AllGrids;
@@ -1353,23 +1353,16 @@ bool UWorldPartitionRuntimeSpatialHash::CreateStreamingGrid(const FSpatialHashRu
 				// Cell cannot be treated as always loaded if it has data layers
 				const bool bIsCellAlwaysLoaded = (&TempCell == &PartionedActors.GetAlwaysLoadedCell()) && !GridCellDataChunk.HasDataLayers() && !GridCellDataChunk.GetContentBundleID().IsValid();
 				
-				TArray<IStreamingGenerationContext::FActorInstance> FilteredActors;
-				
-				// We must be in immediate mode here so child actors are registered before adding them to the always loaded actors array. This is
-				// needed because the actor can contain a construction script that will invalidate the child actor components and recreate them, 
-				// leaving a dangling pointer in the array since it happens on the loading context destruction.
-				FWorldPartitionLoadingContext::FImmediate LoadingContext;
+				const TArray<const IStreamingGenerationContext::FActorSetInstance*> ActorSetInstances = GridCellDataChunk.GetActorSetInstances();
 
-				for (const IStreamingGenerationContext::FActorSetInstance* ActorSetInstance : GridCellDataChunk.GetActorSetInstances())
+				TArray<IStreamingGenerationContext::FActorInstance> FilteredActors;
+				for (const IStreamingGenerationContext::FActorSetInstance* ActorSetInstance : ActorSetInstances)
 				{
-					for (const FGuid& ActorGuid : ActorSetInstance->ActorSet->Actors)
+					if (!ConditionalRegisterAlwaysLoadedActorsForPIE(ActorSetInstance, bIsMainWorldPartition, ActorSetInstance->ContainerID.IsMainContainer(), bIsCellAlwaysLoaded))
 					{
-						IStreamingGenerationContext::FActorInstance ActorInstance(ActorGuid, ActorSetInstance);
-						const FWorldPartitionActorDescView& ActorDescView = ActorInstance.GetActorDescView();
-
-						if (!ConditionalRegisterAlwaysLoadedActorsForPIE(ActorDescView, bIsMainWorldPartition, ActorSetInstance->ContainerID.IsMainContainer(), bIsCellAlwaysLoaded))
+						for (const FGuid& ActorGuid : ActorSetInstance->ActorSet->Actors)
 						{
-							FilteredActors.Emplace(ActorInstance);
+							FilteredActors.Emplace(ActorGuid, ActorSetInstance);
 						}
 					}
 				}
@@ -1384,21 +1377,9 @@ bool UWorldPartitionRuntimeSpatialHash::CreateStreamingGrid(const FSpatialHashRu
 				const FString CellName = GetCellNameString(OuterWorld->GetPackage()->GetFName(), CurrentStreamingGrid.GridName, CellGlobalCoords, GridCellDataChunk.GetDataLayersID(), GridCellDataChunk.GetContentBundleID());
 				const FGuid CellGuid = GetCellGuid(CurrentStreamingGrid.GridName, CellGlobalCoords, GridCellDataChunk.GetDataLayersID(), GridCellDataChunk.GetContentBundleID());
 
-				UWorldPartitionRuntimeCell* StreamingCell = NewObject<UWorldPartitionRuntimeCell>(this, StreamingPolicy->GetRuntimeCellClass(), FName(CellName));
+				UWorldPartitionRuntimeCell* StreamingCell = CreateRuntimeCell(StreamingPolicy->GetRuntimeCellClass(), UWorldPartitionRuntimeCellDataSpatialHash::StaticClass(), FName(CellName));
+				UWorldPartitionRuntimeCellDataSpatialHash* CellDataSpatialHash = CastChecked<UWorldPartitionRuntimeCellDataSpatialHash>(StreamingCell->RuntimeCellData);
 
-				int32 LayerCellIndex;
-				int32* LayerCellIndexPtr = GridLevel.LayerCellsMapping.Find(CellIndex);
-				if (LayerCellIndexPtr)
-				{
-					LayerCellIndex = *LayerCellIndexPtr;
-				}
-				else
-				{
-					LayerCellIndex = GridLevel.LayerCells.AddDefaulted();
-					GridLevel.LayerCellsMapping.Add(CellIndex, LayerCellIndex);
-				}
-
-				GridLevel.LayerCells[LayerCellIndex].GridCells.Add(StreamingCell);
 				StreamingCell->SetIsAlwaysLoaded(bIsCellAlwaysLoaded);
 				StreamingCell->SetDataLayers(GridCellDataChunk.GetDataLayers());
 				StreamingCell->SetContentBundleUID(GridCellDataChunk.GetContentBundleID());
@@ -1409,154 +1390,42 @@ bool UWorldPartitionRuntimeSpatialHash::CreateStreamingGrid(const FSpatialHashRu
 				StreamingCell->SetIsHLOD(RuntimeGrid.HLODLayer ? true : false);
 				StreamingCell->SetGuid(CellGuid);
 
-				UWorldPartitionRuntimeCellDataSpatialHash* StreamingSourceInfoSpatialHash = NewObject<UWorldPartitionRuntimeCellDataSpatialHash>(StreamingCell);
-				StreamingSourceInfoSpatialHash->Level = Level;
 				FBox2D Bounds;
-				verify(TempLevel.GetCellBounds(FGridCellCoord2(CellCoordX, CellCoordY), Bounds));
-				StreamingSourceInfoSpatialHash->Position = FVector(Bounds.GetCenter(), 0.f);
+				verify(TempLevel.GetCellBounds(FGridCellCoord2(CellCoordX, CellCoordY), Bounds));				
 				const double CellExtent = Bounds.GetExtent().X;
 				check(CellExtent < MAX_flt);
-				StreamingSourceInfoSpatialHash->Extent = CellExtent;
-				StreamingCell->RuntimeCellData = StreamingSourceInfoSpatialHash;
 
-				check(!StreamingCell->UnsavedActorsContainer);
-				for (const IStreamingGenerationContext::FActorInstance& ActorInstance : FilteredActors)
-				{
-					const FWorldPartitionActorDescView& ActorDescView = ActorInstance.GetActorDescView();
-					if (AActor* Actor = FindObject<AActor>(nullptr, *ActorDescView.GetActorSoftPath().ToString()))
-					{
-						if (ModifiedActorDescListForPIE.GetActorDesc(ActorDescView.GetGuid()) != nullptr)
-						{
-							// Create an actor container to make sure duplicated actors will share an outer to properly remap inter-actors references
-							StreamingCell->UnsavedActorsContainer = NewObject<UActorContainer>(StreamingCell);
-							break;
-						}
-					}
-				}
+				CellDataSpatialHash->Level = Level;
+				CellDataSpatialHash->Position = FVector(Bounds.GetCenter(), 0.f);
+				CellDataSpatialHash->Extent = (float)CellExtent;
 
-				FBox CellContentBounds(ForceInit);
-				for (const IStreamingGenerationContext::FActorInstance& ActorInstance : FilteredActors)
-				{
-					const FWorldPartitionActorDescView& ActorDescView = ActorInstance.GetActorDescView();
-					StreamingCell->AddActorToCell(ActorDescView, ActorInstance.GetContainerID(), ActorInstance.GetTransform(), ActorInstance.GetActorDescContainer());
-					
-					CellContentBounds += ActorDescView.GetRuntimeBounds().TransformBy(ActorInstance.GetTransform());
+				PopulateRuntimeCell(StreamingCell, FilteredActors, OutPackagesToGenerate);
 
-					if (ActorInstance.GetContainerID().IsMainContainer() && StreamingCell->UnsavedActorsContainer)
-					{
-						if (AActor* Actor = FindObject<AActor>(nullptr, *ActorDescView.GetActorSoftPath().ToString()))
-						{
-							StreamingCell->UnsavedActorsContainer->Actors.Add(Actor->GetFName(), Actor);
-
-							// Handle child actors
-							Actor->ForEachComponent<UChildActorComponent>(true, [StreamingCell](UChildActorComponent* ChildActorComponent)
-							{
-								if (AActor* ChildActor = ChildActorComponent->GetChildActor())
-								{
-									StreamingCell->UnsavedActorsContainer->Actors.Add(ChildActor->GetFName(), ChildActor);
-								}
-							});
-						}
-					}
+				int32 LayerCellIndex;
+				int32* LayerCellIndexPtr = GridLevel.LayerCellsMapping.Find(CellIndex);
+				if (LayerCellIndexPtr)
+				{
+					LayerCellIndex = *LayerCellIndexPtr;
 				}
-				StreamingSourceInfoSpatialHash->ContentBounds = CellContentBounds;
-
-				// Always loaded cell actors are transfered to World's Persistent Level (see UWorldPartitionRuntimeSpatialHash::PopulateGeneratorPackageForCook)
-				if (StreamingCell->GetActorCount() && !StreamingCell->IsAlwaysLoaded())
+				else
 				{
-					if (OutPackagesToGenerate)
-					{
-						const FString PackageRelativePath = StreamingCell->GetPackageNameToCreate();
-						check(!PackageRelativePath.IsEmpty());
-						OutPackagesToGenerate->Add(PackageRelativePath);
-
-						// Map relative package to StreamingCell for PopulateGeneratedPackageForCook/PopulateGeneratorPackageForCook/GetCellForPackage
-						PackagesToGenerateForCook.Add(PackageRelativePath, StreamingCell);
-
-						UE_CLOG(IsRunningCookCommandlet(), LogWorldPartition, Log, TEXT("Creating runtime streaming cells %s."), *StreamingCell->GetName());
-					}
+					LayerCellIndex = GridLevel.LayerCells.AddDefaulted();
+					GridLevel.LayerCellsMapping.Add(CellIndex, LayerCellIndex);
 				}
-			}
-		}
-	}
-
-	return true;
-}
-
-bool UWorldPartitionRuntimeSpatialHash::PopulateGeneratedPackageForCook(const FWorldPartitionCookPackage& InPackagesToCook, TArray<UPackage*>& OutModifiedPackages)
-{
-	OutModifiedPackages.Reset();
-	if (UWorldPartitionRuntimeCell** MatchingCell = PackagesToGenerateForCook.Find(InPackagesToCook.RelativePath))
-	{
-		UWorldPartitionRuntimeCell* Cell = *MatchingCell;
-		if (ensure(Cell))
-		{
-			return Cell->PopulateGeneratedPackageForCook(InPackagesToCook.GetPackage(), OutModifiedPackages);
-		}
-	}
-	return false;
-}
 
-UWorldPartitionRuntimeCell* UWorldPartitionRuntimeSpatialHash::GetCellForPackage(const FWorldPartitionCookPackage& PackageToCook) const
-{
-	UWorldPartitionRuntimeCell** MatchingCell = const_cast<UWorldPartitionRuntimeCell**>(PackagesToGenerateForCook.Find(PackageToCook.RelativePath));
-	return MatchingCell ? *MatchingCell : nullptr;
-}
-
-TArray<UWorldPartitionRuntimeCell*> UWorldPartitionRuntimeSpatialHash::GetAlwaysLoadedCells() const
-{
-	TArray<UWorldPartitionRuntimeCell*> Result;
-	ForEachStreamingCells([&Result](const UWorldPartitionRuntimeCell* Cell)
-	{
-		if (Cell->IsAlwaysLoaded())
-		{
-			Result.Add(const_cast<UWorldPartitionRuntimeCell*>(Cell));
-		}
-		return true;
-	});
-	return Result;
-}
-
-bool UWorldPartitionRuntimeSpatialHash::PrepareGeneratorPackageForCook(TArray<UPackage*>& OutModifiedPackages)
-{
-	check(IsRunningCookCommandlet());
-	check(ExternalStreamingObjects.IsEmpty()); // Make sure we don't have external streaming objects grids so we won't gather their cells
-
-	for (UWorldPartitionRuntimeCell* Cell : GetAlwaysLoadedCells())
-	{
-		check(Cell->IsAlwaysLoaded());
-		if (!Cell->PopulateGeneratorPackageForCook(OutModifiedPackages))
-		{
-			return false;
+				GridLevel.LayerCells[LayerCellIndex].GridCells.Add(StreamingCell);
+			}
 		}
 	}
 
-	//@todo_ow: here we can safely remove always loaded cells as they are not part of the OutPackagesToGenerate
-	return true;
-}
-
-bool UWorldPartitionRuntimeSpatialHash::PopulateGeneratorPackageForCook(const TArray<FWorldPartitionCookPackage*>& InPackagesToCook, TArray<UPackage*>& OutModifiedPackages)
-{
-	check(IsRunningCookCommandlet());
-	check(ExternalStreamingObjects.IsEmpty()); // Make sure we don't have external streaming objects grids so we won't gather their cells
-
-	for (const FWorldPartitionCookPackage* CookPackage : InPackagesToCook)
-	{
-		UWorldPartitionRuntimeCell** MatchingCell = PackagesToGenerateForCook.Find(CookPackage->RelativePath);
-		UWorldPartitionRuntimeCell* Cell = MatchingCell ? *MatchingCell : nullptr;
-		if (!Cell || !Cell->PrepareCellForCook(CookPackage->GetPackage()))
-		{
-			return false;
-		}
-	}
 	return true;
 }
 
 void UWorldPartitionRuntimeSpatialHash::FlushStreaming()
 {
+	Super::FlushStreaming();
 	bIsNameToGridMappingDirty = true;
 	StreamingGrids.Empty();
-	PackagesToGenerateForCook.Empty();
 }
 
 bool UWorldPartitionRuntimeSpatialHash::IsValidGrid(FName GridName) const
@@ -1695,65 +1564,44 @@ void UWorldPartitionRuntimeSpatialHash::ForEachStreamingCellsSources(const TArra
 
 bool UWorldPartitionRuntimeSpatialHash::InjectExternalStreamingObject(URuntimeHashExternalStreamingObjectBase* ExternalStreamingObject)
 {
-	if (!ExternalStreamingObjects.Contains(ExternalStreamingObject))
-	{
-		if (URuntimeSpatialHashExternalStreamingObject* SpatialHashExternalStreamingObject = Cast<URuntimeSpatialHashExternalStreamingObject>(ExternalStreamingObject))
-		{
-			// Validate that there's a corresponding streaming grid for each streaming grid of this external streaming object
-			for (const FSpatialHashStreamingGrid& ExternalStreamingGrid : SpatialHashExternalStreamingObject->StreamingGrids)
-			{
-				if (!GetStreamingGridByName(ExternalStreamingGrid.GridName))
-				{
-					UE_LOG(LogWorldPartition, Error, TEXT("Failed to inject external streaming object %s, can't find matching streaming grid %s."),
-						*ExternalStreamingObject->GetName(), *ExternalStreamingGrid.GridName.ToString());
-					return false;
-				}
-			}
-
-			ExternalStreamingObjects.Add(SpatialHashExternalStreamingObject);
+	URuntimeSpatialHashExternalStreamingObject* SpatialHashExternalStreamingObject = CastChecked<URuntimeSpatialHashExternalStreamingObject>(ExternalStreamingObject);
 
-			for (const FSpatialHashStreamingGrid& ExternalStreamingGrid : SpatialHashExternalStreamingObject->StreamingGrids)
-			{
-				const FSpatialHashStreamingGrid* SourceGrid = GetStreamingGridByName(ExternalStreamingGrid.GridName);
-				SourceGrid->InjectExternalStreamingObjectGrid(ExternalStreamingGrid);
-			}
-			return true;
-		}
-		else
+	// Validate that there's a corresponding streaming grid for each streaming grid of this external streaming object
+	for (const FSpatialHashStreamingGrid& ExternalStreamingGrid : SpatialHashExternalStreamingObject->StreamingGrids)
+	{
+		if (!GetStreamingGridByName(ExternalStreamingGrid.GridName))
 		{
-			UE_LOG(LogWorldPartition, Error, TEXT("Incompatible external streaming object to inject. Hash Type: %s. Streaming Object Type: %s."), 
-				*UWorldPartitionRuntimeSpatialHash::StaticClass()->GetName(), *ExternalStreamingObject->GetClass()->GetName());
+			UE_LOG(LogWorldPartition, Error, TEXT("Failed to inject external streaming object %s, can't find matching streaming grid %s."),
+				*ExternalStreamingObject->GetName(), *ExternalStreamingGrid.GridName.ToString());
+			return false;
 		}
 	}
-	return false;
+
+	for (const FSpatialHashStreamingGrid& ExternalStreamingGrid : SpatialHashExternalStreamingObject->StreamingGrids)
+	{
+		const FSpatialHashStreamingGrid* SourceGrid = GetStreamingGridByName(ExternalStreamingGrid.GridName);
+		SourceGrid->InjectExternalStreamingObjectGrid(ExternalStreamingGrid);
+	}
+
+	return true;
 }
 
 bool UWorldPartitionRuntimeSpatialHash::RemoveExternalStreamingObject(URuntimeHashExternalStreamingObjectBase* ExternalStreamingObject)
 {
-	if (URuntimeSpatialHashExternalStreamingObject* SpatialHashExternalStreamingObject = Cast<URuntimeSpatialHashExternalStreamingObject>(ExternalStreamingObject))
+	URuntimeSpatialHashExternalStreamingObject* SpatialHashExternalStreamingObject = CastChecked<URuntimeSpatialHashExternalStreamingObject>(ExternalStreamingObject);
+
+	for (const FSpatialHashStreamingGrid& ExternalStreamingGrid : SpatialHashExternalStreamingObject->StreamingGrids)
 	{
-		if (ExternalStreamingObjects.Remove(SpatialHashExternalStreamingObject) > 0)
+		// In PIE, UWorldPartition::OnEndPlay() will call UWorldPartitionRuntimeSpatialHash::FlushStreaming() 
+		// which will empty StreamingGrids before RemoveExternalStreamingObject is called.
+		// In this case, we don't need to remove injected cells as the source grid was removed.
+		if (const FSpatialHashStreamingGrid* SourceGrid = GetStreamingGridByName(ExternalStreamingGrid.GridName))
 		{
-			for (const FSpatialHashStreamingGrid& ExternalStreamingGrid : SpatialHashExternalStreamingObject->StreamingGrids)
-			{
-				// In PIE, UWorldPartition::OnEndPlay() will call UWorldPartitionRuntimeSpatialHash::FlushStreaming() 
-				// which will empty StreamingGrids before RemoveExternalStreamingObject is called.
-				// In this case, we don't need to remove injected cells as the source grid was removed.
-				if (const FSpatialHashStreamingGrid* SourceGrid = GetStreamingGridByName(ExternalStreamingGrid.GridName))
-				{
-					SourceGrid->RemoveExternalStreamingObjectGrid(ExternalStreamingGrid);
-				}
-			}
-			return true;
+			SourceGrid->RemoveExternalStreamingObjectGrid(ExternalStreamingGrid);
 		}
 	}
-	else
-	{
-		UE_LOG(LogWorldPartition, Error, TEXT("Incompatible external streaming object to remove. Hash Type: %s. Streaming Object Type: %s."),
-			*UWorldPartitionRuntimeSpatialHash::StaticClass()->GetName(), *ExternalStreamingObject->GetClass()->GetName());
-	}
-	
-	return false;
+
+	return true;
 }
 
 uint32 UWorldPartitionRuntimeSpatialHash::GetNumGrids() const
@@ -1828,15 +1676,16 @@ EWorldPartitionStreamingPerformance UWorldPartitionRuntimeSpatialHash::GetStream
 	check(StreamingCell);
 
 	const FSpatialHashStreamingGrid* StreamingGrid = GetStreamingGridByName(StreamingCell->GetGridName());
+
 	if (ensure(StreamingGrid))
 	{
 		const float LoadingRange = StreamingGrid->LoadingRange;
 
-		UWorldPartitionRuntimeCellDataSpatialHash* StreamingSourceInfoSpatialHash = CastChecked<UWorldPartitionRuntimeCellDataSpatialHash>(Cell->RuntimeCellData);
+		UWorldPartitionRuntimeCellDataSpatialHash* CellDataSpatialHash = CastChecked<UWorldPartitionRuntimeCellDataSpatialHash>(Cell->RuntimeCellData);
 
-		if (StreamingSourceInfoSpatialHash->IsBlockingSource())
+		if (CellDataSpatialHash->IsBlockingSource())
 		{
-			const double Distance = FMath::Sqrt(StreamingSourceInfoSpatialHash->GetMinSquareDistanceToBlockingSource()) - ((double)StreamingGrid->GetCellSize(StreamingSourceInfoSpatialHash->Level) / 2);
+			const double Distance = FMath::Sqrt(CellDataSpatialHash->GetMinSquareDistanceToBlockingSource()) - ((double)StreamingGrid->GetCellSize(CellDataSpatialHash->Level) / 2);
 
 			const double Ratio = Distance / LoadingRange;
 
diff --git a/Engine/Source/Runtime/Engine/Public/WorldPartition/WorldPartitionRuntimeCell.h b/Engine/Source/Runtime/Engine/Public/WorldPartition/WorldPartitionRuntimeCell.h
--- a/Engine/Source/Runtime/Engine/Public/WorldPartition/WorldPartitionRuntimeCell.h
+++ b/Engine/Source/Runtime/Engine/Public/WorldPartition/WorldPartitionRuntimeCell.h
@@ -15,6 +15,7 @@
 #include "Algo/AnyOf.h"
 #include "WorldPartitionRuntimeCell.generated.h"
 
+class UActorContainer;
 class UDataLayerAsset;
 class UDataLayerInstance;
 struct FHierarchicalLogArchive;
diff --git a/Engine/Source/Runtime/Engine/Public/WorldPartition/WorldPartitionRuntimeHash.h b/Engine/Source/Runtime/Engine/Public/WorldPartition/WorldPartitionRuntimeHash.h
--- a/Engine/Source/Runtime/Engine/Public/WorldPartition/WorldPartitionRuntimeHash.h
+++ b/Engine/Source/Runtime/Engine/Public/WorldPartition/WorldPartitionRuntimeHash.h
@@ -12,6 +12,7 @@
 #include "WorldPartition/WorldPartitionRuntimeCellOwner.h"
 #include "WorldPartition/WorldPartitionActorDescViewProxy.h"
 #include "WorldPartition/WorldPartitionStreamingGeneration.h"
+#include "WorldPartition/WorldPartitionStreamingGenerationContext.h"
 #if WITH_EDITOR
 #include "CookPackageSplitter.h"
 #include "Misc/HierarchicalLogArchive.h"
@@ -19,7 +20,6 @@
 #include "WorldPartitionRuntimeHash.generated.h"
 
 class FActorDescViewMap;
-class IStreamingGenerationContext;
 struct FHierarchicalLogArchive;
 
 UENUM()
@@ -83,12 +83,13 @@ class ENGINE_API UWorldPartitionRuntimeHash : public UObject, public IWorldParti
 #if WITH_EDITOR
 	virtual void SetDefaultValues() {}
 	virtual bool SupportsHLODs() const { return false; }
-	virtual bool PrepareGeneratorPackageForCook(TArray<UPackage*>& OutModifiedPackages) { return false; }
-	virtual bool PopulateGeneratorPackageForCook(const TArray<FWorldPartitionCookPackage*>& PackagesToCook, TArray<UPackage*>& OutModifiedPackages) { return false; }
-	virtual bool PopulateGeneratedPackageForCook(const FWorldPartitionCookPackage& PackagesToCook, TArray<UPackage*>& OutModifiedPackages) { return false; }
-	virtual UWorldPartitionRuntimeCell* GetCellForPackage(const FWorldPartitionCookPackage& PackageToCook) const { return nullptr; }
-	virtual bool GenerateStreaming(class UWorldPartitionStreamingPolicy* StreamingPolicy, const IStreamingGenerationContext* StreamingGenerationContext, TArray<FString>* OutPackagesToGenerate) { return false; }
-	virtual void FlushStreaming() {}
+	virtual bool PrepareGeneratorPackageForCook(TArray<UPackage*>& OutModifiedPackages);
+	virtual bool PopulateGeneratorPackageForCook(const TArray<FWorldPartitionCookPackage*>& PackagesToCook, TArray<UPackage*>& OutModifiedPackages);
+	virtual bool PopulateGeneratedPackageForCook(const FWorldPartitionCookPackage& PackagesToCook, TArray<UPackage*>& OutModifiedPackages);
+	virtual UWorldPartitionRuntimeCell* GetCellForPackage(const FWorldPartitionCookPackage& PackageToCook) const;
+	virtual TArray<UWorldPartitionRuntimeCell*> GetAlwaysLoadedCells() const;
+	virtual bool GenerateStreaming(class UWorldPartitionStreamingPolicy* StreamingPolicy, const IStreamingGenerationContext* StreamingGenerationContext, TArray<FString>* OutPackagesToGenerate);
+	virtual void FlushStreaming();
 	virtual bool GenerateHLOD(ISourceControlHelper* SourceControlHelper, const IStreamingGenerationContext* StreamingGenerationContext, bool bCreateActorsOnly) { return false; }
 	virtual bool IsValidGrid(FName GridName) const { return false; }
 	virtual void DrawPreview() const {}
@@ -101,8 +102,17 @@ class ENGINE_API UWorldPartitionRuntimeHash : public UObject, public IWorldParti
 	void OnBeginPlay();
 	void OnEndPlay();
 
+	inline UWorldPartitionRuntimeCell* CreateRuntimeCell(UClass* CellClass, UClass* CellDataClass, FName CellName)
+	{
+		check(!FindObject<UWorldPartitionRuntimeCell>(this, *CellName.ToString()));
+		UWorldPartitionRuntimeCell* RuntimeCell = NewObject<UWorldPartitionRuntimeCell>(this, CellClass, CellName);
+		RuntimeCell->RuntimeCellData = NewObject<UWorldPartitionRuntimeCellData>(RuntimeCell, CellDataClass);
+		return RuntimeCell;
+	}
+
 protected:
-	bool ConditionalRegisterAlwaysLoadedActorsForPIE(const FWorldPartitionActorDescView& ActorDescView, bool bIsMainWorldPartition, bool bIsMainContainer, bool bIsCellAlwaysLoaded);
+	bool ConditionalRegisterAlwaysLoadedActorsForPIE(const IStreamingGenerationContext::FActorSetInstance* ActorSetInstance, bool bIsMainWorldPartition, bool bIsMainContainer, bool bIsCellAlwaysLoaded);
+	void PopulateRuntimeCell(UWorldPartitionRuntimeCell* RuntimeCell, const TArray<IStreamingGenerationContext::FActorInstance>& ActorInstances, TArray<FString>* OutPackagesToGenerate);
 #endif
 
 public:
@@ -171,6 +181,8 @@ protected:
 
 	TArray<FAlwaysLoadedActorForPIE> AlwaysLoadedActorsForPIE;
 
+	TMap<FString, UWorldPartitionRuntimeCell*> PackagesToGenerateForCook;
+
 public:
 	mutable FActorDescList ModifiedActorDescListForPIE;
 #endif
diff --git a/Engine/Source/Runtime/Engine/Public/WorldPartition/WorldPartitionRuntimeSpatialHash.h b/Engine/Source/Runtime/Engine/Public/WorldPartition/WorldPartitionRuntimeSpatialHash.h
--- a/Engine/Source/Runtime/Engine/Public/WorldPartition/WorldPartitionRuntimeSpatialHash.h
+++ b/Engine/Source/Runtime/Engine/Public/WorldPartition/WorldPartitionRuntimeSpatialHash.h
@@ -267,10 +267,6 @@ public:
 
 	virtual void SetDefaultValues() override;
 	virtual bool SupportsHLODs() const override { return true; }
-	virtual bool PrepareGeneratorPackageForCook(TArray<UPackage*>& OutModifiedPackages) override;
-	virtual bool PopulateGeneratorPackageForCook(const TArray<FWorldPartitionCookPackage*>& PackagesToCook, TArray<UPackage*>& OutModifiedPackage) override;
-	virtual bool PopulateGeneratedPackageForCook(const FWorldPartitionCookPackage& PackageToCook, TArray<UPackage*>& OutModifiedPackages) override;
-	virtual UWorldPartitionRuntimeCell* GetCellForPackage(const FWorldPartitionCookPackage& PackageToCook) const override;
 	virtual void FlushStreaming() override;
 	virtual bool GenerateHLOD(ISourceControlHelper* SourceControlHelper, const IStreamingGenerationContext* StreamingGenerationContext, bool bCreateActorsOnly) override;
 	virtual bool IsValidGrid(FName GridName) const override;
@@ -325,8 +321,6 @@ private:
 
 	UPROPERTY(Transient)
 	mutable FWorldPartitionRuntimeSpatialHashGridPreviewer GridPreviewer;
-
-	TMap<FString, UWorldPartitionRuntimeCell*> PackagesToGenerateForCook;
 #endif
 
 	/** Whether this hash enables Z culling. */
@@ -343,19 +337,14 @@ protected:
 	mutable bool bIsNameToGridMappingDirty;
 
 private:
-	UPROPERTY(Transient)
-	TArray<TWeakObjectPtr<URuntimeSpatialHashExternalStreamingObject>> ExternalStreamingObjects;
-	
 	virtual bool Draw2D(class UCanvas* Canvas, const TArray<FWorldPartitionStreamingSource>& Sources, const FVector2D& PartitionCanvasSize, const FVector2D& Offset, FVector2D& OutUsedCanvasSize) const override;
 	virtual void Draw3D(const TArray<FWorldPartitionStreamingSource>& Sources) const override;
 	virtual bool ContainsRuntimeHash(const FString& Name) const override;
 	virtual bool IsStreaming3D() const override;
 
 	void GetAlwaysLoadedStreamingCells(const FSpatialHashStreamingGrid& StreamingGrid, TSet<const UWorldPartitionRuntimeCell*>& Cells) const;
-	void GetStreamingCells(const FVector& Position, const FSpatialHashStreamingGrid& StreamingGrid, TSet<const UWorldPartitionRuntimeCell*>& Cells) const;
 	const TMap<FName, const FSpatialHashStreamingGrid*>& GetNameToGridMapping() const;
 #if WITH_EDITOR
-	TArray<UWorldPartitionRuntimeCell*> GetAlwaysLoadedCells() const;
 	bool CreateStreamingGrid(const FSpatialHashRuntimeGrid& RuntimeGrid, const FSquare2DGridHelper& PartionedActors, UWorldPartitionStreamingPolicy* StreamingPolicy, TArray<FString>* OutPackagesToGenerate = nullptr);
 #endif
 	TArray<const FSpatialHashStreamingGrid*> GetFilteredStreamingGrids() const;
