Deprecate URig types in preparation for full removal in 5.4.

diff --git a/Engine/Plugins/Editor/EngineAssetDefinitions/Source/Private/AssetDefinition_Rig.h b/Engine/Plugins/Editor/EngineAssetDefinitions/Source/Private/AssetDefinition_Rig.h
--- a/Engine/Plugins/Editor/EngineAssetDefinitions/Source/Private/AssetDefinition_Rig.h
+++ b/Engine/Plugins/Editor/EngineAssetDefinitions/Source/Private/AssetDefinition_Rig.h
@@ -16,7 +16,9 @@ public:
 	// UAssetDefinition Begin
 	virtual FText GetAssetDisplayName() const override { return NSLOCTEXT("AssetTypeActions", "AssetTypeActions_Rig", "Rig"); }
 	virtual FLinearColor GetAssetColor() const override { return FLinearColor(FColor(255,201,14)); }
+	PRAGMA_DISABLE_DEPRECATION_WARNINGS
 	virtual TSoftClassPtr<UObject> GetAssetClass() const override { return URig::StaticClass(); }
+	PRAGMA_ENABLE_DEPRECATION_WARNINGS
 	virtual TConstArrayView<FAssetCategoryPath> GetAssetCategories() const override
 	{
 		static const auto Categories = { EAssetCategoryPaths::Animation };
diff --git a/Engine/Source/Editor/DetailCustomizations/Private/RigDetails.cpp b/Engine/Source/Editor/DetailCustomizations/Private/RigDetails.cpp
--- a/Engine/Source/Editor/DetailCustomizations/Private/RigDetails.cpp
+++ b/Engine/Source/Editor/DetailCustomizations/Private/RigDetails.cpp
@@ -41,7 +41,8 @@ class SWidget;
 #define NODE_TABLE_NODENAME		TEXT("NodeName")
 #define NODE_TABLE_PARENTNAME	TEXT("ParentName")
 
-
+// January 2023 - This whole compilation unit is deprecated to be removed in 5.4
+PRAGMA_DISABLE_DEPRECATION_WARNINGS
 TSharedRef<IDetailCustomization> FRigDetails::MakeInstance()
 {
 	return MakeShareable(new FRigDetails);
@@ -288,7 +289,7 @@ void FRigDetails::GenerateTransformBaseArrayElementWidget(TSharedRef<IPropertyHa
 	// orientation
 	TSharedPtr<IPropertyHandleArray> TransformConstraintsProp_R = ConstraintsProp->GetElement(1)->GetChildHandle(GET_MEMBER_NAME_CHECKED(FTransformBaseConstraint, TransformConstraints))->AsArray();
 	TSharedRef<IPropertyHandle> ParentNameProp_R = TransformConstraintsProp_R->GetElement(0)->GetChildHandle(GET_MEMBER_NAME_CHECKED(FRigTransformConstraint, ParentSpace)).ToSharedRef();
-
+	
 	// the interface will be node [display name] [parent node]
 	// delegate for display name
 	FText NodeName;
@@ -560,3 +561,4 @@ ECheckBoxState FRigDetails::AdvancedCheckBoxIsChecked(TSharedRef<IPropertyHandle
 	return ECheckBoxState::Undetermined;
 }
 #undef LOCTEXT_NAMESPACE
+PRAGMA_ENABLE_DEPRECATION_WARNINGS
\ No newline at end of file
diff --git a/Engine/Source/Editor/Persona/Private/SRigPicker.cpp b/Engine/Source/Editor/Persona/Private/SRigPicker.cpp
--- a/Engine/Source/Editor/Persona/Private/SRigPicker.cpp
+++ b/Engine/Source/Editor/Persona/Private/SRigPicker.cpp
@@ -12,6 +12,7 @@
 
 static const FVector2f ContentBrowserWindowSize(300.0f, 300.0f);
 
+PRAGMA_DISABLE_DEPRECATION_WARNINGS
 URig* SRigPicker::EngineHumanoidRig = NULL;
 
 void SRigPicker::Construct(const FArguments& InArgs)
@@ -128,3 +129,5 @@ void SRigPicker::OnSelectDefault()
 }
 
 #undef LOCTEXT_NAMESPACE
+
+PRAGMA_ENABLE_DEPRECATION_WARNINGS
\ No newline at end of file
diff --git a/Engine/Source/Editor/Persona/Private/SRigPicker.h b/Engine/Source/Editor/Persona/Private/SRigPicker.h
--- a/Engine/Source/Editor/Persona/Private/SRigPicker.h
+++ b/Engine/Source/Editor/Persona/Private/SRigPicker.h
@@ -63,6 +63,7 @@ private:
 
 	/** Delegate for closing the containing menu */
 	FSimpleDelegate OnClose;
-
+PRAGMA_DISABLE_DEPRECATION_WARNINGS
 	static URig* EngineHumanoidRig;
+PRAGMA_ENABLE_DEPRECATION_WARNINGS
 };
diff --git a/Engine/Source/Editor/SkeletonEditor/Public/IEditableSkeleton.h b/Engine/Source/Editor/SkeletonEditor/Public/IEditableSkeleton.h
--- a/Engine/Source/Editor/SkeletonEditor/Public/IEditableSkeleton.h
+++ b/Engine/Source/Editor/SkeletonEditor/Public/IEditableSkeleton.h
@@ -159,6 +159,7 @@ public:
 	/** Remove a compatible skeleton */
 	virtual void RemoveCompatibleSkeleton(const USkeleton* InCompatibleSkeleton) = 0;
 
+	PRAGMA_DISABLE_DEPRECATION_WARNINGS
 	/** Refreshes the rig config, validating the mappings */
 	virtual void RefreshRigConfig() = 0;
 
@@ -170,6 +171,7 @@ public:
 
 	/** Set multiple bone mappings */
 	virtual void SetRigBoneMappings(const TMap<FName, FName>& InMappings) = 0;
+	PRAGMA_ENABLE_DEPRECATION_WARNINGS
 
 	/** Remove any bones that are not used by any skeletal meshes */
 	virtual void RemoveUnusedBones() = 0;
diff --git a/Engine/Source/Runtime/Engine/Classes/Animation/Rig.h b/Engine/Source/Runtime/Engine/Classes/Animation/Rig.h
--- a/Engine/Source/Runtime/Engine/Classes/Animation/Rig.h
+++ b/Engine/Source/Runtime/Engine/Classes/Animation/Rig.h
@@ -21,9 +21,12 @@
 class USkeleton;
 struct FPropertyChangedEvent;
 
+// January 2023 - Whole compilation unit is being deprecated, so disabled warnings for whole thing. 
+PRAGMA_DISABLE_DEPRECATION_WARNINGS
+
 /** Rig Controller for bone transform **/
 USTRUCT()
-struct FNode
+struct UE_DEPRECATED(5.2, "URig retargeting is no longer supported. Use the IK Retargeter instead.") FNode
 {
 	GENERATED_USTRUCT_BODY()
 
@@ -67,7 +70,7 @@ struct FNode
 UENUM()
 namespace EControlConstraint
 {
-	enum Type : int
+	enum UE_DEPRECATED(5.2, "URig retargeting is no longer supported. Use the IK Retargeter instead.") Type : int
 	{
 		/** Rotation constraint. */
 		Orientation,
@@ -83,7 +86,7 @@ namespace EControlConstraint
 UENUM()
 namespace EConstraintTransform
 {
-	enum Type : int
+	enum UE_DEPRECATED(5.2, "URig retargeting is no longer supported. Use the IK Retargeter instead.") Type : int
 	{
 		/** Absolute value. */
 		Absolute,
@@ -93,10 +96,10 @@ namespace EConstraintTransform
 }
 
 USTRUCT()
-struct FRigTransformConstraint
+struct UE_DEPRECATED(5.2, "URig retargeting is no longer supported. Use the IK Retargeter instead.") FRigTransformConstraint
 {
 	GENERATED_USTRUCT_BODY()
-
+	
 	/** What transform type **/
 	UPROPERTY(/*EditAnywhere, Category="FTransformBaseConstraint"*/)
 	TEnumAsByte<EConstraintTransform::Type>	TranformType = EConstraintTransform::Absolute;
@@ -112,10 +115,10 @@ struct FRigTransformConstraint
 
 /** This defines what constraint it is defined */
 USTRUCT()
-struct FTransformBaseConstraint
+struct UE_DEPRECATED(5.2, "URig retargeting is no longer supported. Use the IK Retargeter instead.") FTransformBaseConstraint
 {
 	GENERATED_USTRUCT_BODY()
-
+	
 	/** What transform type **/
 	UPROPERTY(EditAnywhere, Category="FTransformBaseConstraint")
 	TArray<FRigTransformConstraint>			TransformConstraints;
@@ -124,7 +127,7 @@ struct FTransformBaseConstraint
 
 /** This is a mapping table between bone in a particular skeletal mesh and bone of this skeleton set. */
 USTRUCT()
-struct FTransformBase
+struct UE_DEPRECATED(5.2, "URig retargeting is no longer supported. Use the IK Retargeter instead.") FTransformBase
 {
 	GENERATED_USTRUCT_BODY()
 
@@ -143,7 +146,7 @@ DECLARE_DELEGATE_RetVal_OneParam(int32, FGetParentIndex, FName);
  *		- support to share different animations
  */
 UCLASS(hidecategories=Object, MinimalAPI)
-class URig : public UObject, public INodeMappingProviderInterface
+class UE_DEPRECATED(5.2, "URig retargeting is no longer supported. Use the IK Retargeter instead.") URig : public UObject, public INodeMappingProviderInterface
 {
 	GENERATED_UCLASS_BODY()
 
@@ -223,3 +226,6 @@ private:
 //	void CalculateLocalSpace(int32 NodeIndex, const FTransform& ComponentTransform, const TArray<FTransform> & TransformBuffer, const FGetParentIndex& DelegateToGetParentIndex, FTransform& OutLocalSpaceTransform) const;
 };
 
+// January 2023 - Whole compilation unit is being deprecated, so disabled warnings for whole thing. 
+PRAGMA_ENABLE_DEPRECATION_WARNINGS
+
diff --git a/Engine/Source/Runtime/Engine/Classes/Animation/Skeleton.h b/Engine/Source/Runtime/Engine/Classes/Animation/Skeleton.h
--- a/Engine/Source/Runtime/Engine/Classes/Animation/Skeleton.h
+++ b/Engine/Source/Runtime/Engine/Classes/Animation/Skeleton.h
@@ -218,8 +218,10 @@ struct FRigConfiguration
 {
 	GENERATED_USTRUCT_BODY()
 
+	PRAGMA_DISABLE_DEPRECATION_WARNINGS
 	UPROPERTY()
 	TObjectPtr<class URig>  Rig = nullptr;
+	PRAGMA_ENABLE_DEPRECATION_WARNINGS
 
 	// @todo in the future we can make this to be run-time data
 	UPROPERTY()
@@ -535,8 +537,10 @@ private:
 	UPROPERTY(duplicatetransient, AssetRegistrySearchable)
 	TSoftObjectPtr<class UDataAsset> AdditionalPreviewSkeletalMeshes;
 
+	PRAGMA_DISABLE_DEPRECATION_WARNINGS
 	UPROPERTY()
 	FRigConfiguration RigConfig;
+	PRAGMA_ENABLE_DEPRECATION_WARNINGS
 
 	/** rig property will be saved separately */
 	ENGINE_API virtual void GetAssetRegistryTags(TArray<FAssetRegistryTag>& OutTags) const override;
@@ -950,6 +954,7 @@ public:
 	ENGINE_API static const FName CompatibleSkeletonsNameTag;
 	ENGINE_API static const FString CompatibleSkeletonsTagDelimiter;
 
+	PRAGMA_DISABLE_DEPRECATION_WARNINGS
 	// rig Configs
 	ENGINE_API static const FName RigTag;
 	ENGINE_API void SetRigConfig(URig * Rig);
@@ -962,6 +967,7 @@ public:
 	ENGINE_API void RefreshRigConfig();
 	int32 FindRigBoneMapping(const FName& NodeName) const;
 	ENGINE_API URig * GetRig() const;
+	PRAGMA_ENABLE_DEPRECATION_WARNINGS
 
 #endif
 
diff --git a/Engine/Source/Runtime/Engine/Private/Animation/AnimSequence.cpp b/Engine/Source/Runtime/Engine/Private/Animation/AnimSequence.cpp
--- a/Engine/Source/Runtime/Engine/Private/Animation/AnimSequence.cpp
+++ b/Engine/Source/Runtime/Engine/Private/Animation/AnimSequence.cpp
@@ -2078,6 +2078,8 @@ int32 FindMeshBoneIndexFromBoneName(USkeleton * Skeleton, const FName &BoneName)
 
 	return BoneIndex;
 }
+
+PRAGMA_DISABLE_DEPRECATION_WARNINGS
 void FillUpTransformBasedOnRig(USkeleton* Skeleton, TArray<FTransform>& NodeSpaceBases, TArray<FTransform> &Rotations, TArray<FVector>& Translations, TArray<bool>& TranslationParentFlags)
 {
 	TArray<FTransform> SpaceBases;
@@ -2182,6 +2184,8 @@ int32 FindValidTransformParentTrack(const URig* Rig, int32 NodeIndex, bool bTran
 	return INDEX_NONE;
 }
 
+PRAGMA_ENABLE_DEPRECATION_WARNINGS
+
 
 void UAnimSequence::RemapTracksToNewSkeleton( USkeleton* NewSkeleton, bool bConvertSpaces )
 {
@@ -2198,6 +2202,7 @@ void UAnimSequence::RemapTracksToNewSkeleton( USkeleton* NewSkeleton, bool bConv
 	{
 		USkeleton* OldSkeleton = GetSkeleton();
 
+		PRAGMA_DISABLE_DEPRECATION_WARNINGS
 		// first check if both has same rig, if so, we'll retarget using it
 		if (OldSkeleton && OldSkeleton->GetRig() != nullptr && NewSkeleton->GetRig() == OldSkeleton->GetRig() && OldSkeleton->GetPreviewMesh() && NewSkeleton->GetPreviewMesh())
 		{
@@ -2406,6 +2411,7 @@ void UAnimSequence::RemapTracksToNewSkeleton( USkeleton* NewSkeleton, bool bConv
 			// convert back to animated data with new skeleton
 			ConvertRiggingDataToAnimationData(RiggingAnimationData);
 		}
+		PRAGMA_ENABLE_DEPRECATION_WARNINGS
 		// @todo end rig testing
 		// @IMPORTANT: now otherwise this will try to do bone to bone mapping
 		else if(OldSkeleton && OldSkeleton != NewSkeleton)
@@ -2691,6 +2697,7 @@ void UAnimSequence::ReplaceReferredAnimations(const TMap<UAnimationAsset*, UAnim
 	}
 }
 
+PRAGMA_DISABLE_DEPRECATION_WARNINGS
 int32 FindParentNodeIndex(URig* Rig, USkeleton* Skeleton, FName ParentNodeName)
 {
 	const int32& ParentNodeIndex = Rig->FindNode(ParentNodeName);
@@ -2698,6 +2705,7 @@ int32 FindParentNodeIndex(URig* Rig, USkeleton* Skeleton, FName ParentNodeName)
 	
 	return Skeleton->GetReferenceSkeleton().FindBoneIndex(ParentBoneName);
 }
+PRAGMA_ENABLE_DEPRECATION_WARNINGS
 
 int32 UAnimSequence::GetSpaceBasedAnimationData(TArray< TArray<FTransform> >& AnimationDataInComponentSpace, FAnimSequenceTrackContainer * RiggingAnimationData) const
 {
@@ -2722,6 +2730,7 @@ int32 UAnimSequence::GetSpaceBasedAnimationData(TArray< TArray<FTransform> >& An
 		AnimationDataInComponentSpace[BoneIndex].AddUninitialized(NumKeys);
 	}
 
+	PRAGMA_DISABLE_DEPRECATION_WARNINGS
 	if (RiggingAnimationData)
 	{
 		const URig* Rig = MySkeleton->GetRig();
@@ -2873,6 +2882,7 @@ int32 UAnimSequence::GetSpaceBasedAnimationData(TArray< TArray<FTransform> >& An
 			}
 		} while (bCompleted == false);
 	}
+	PRAGMA_ENABLE_DEPRECATION_WARNINGS
 	else
 	{
 		// now calculating old animated space bases
@@ -2925,6 +2935,7 @@ int32 UAnimSequence::GetSpaceBasedAnimationData(TArray< TArray<FTransform> >& An
 	return AnimationDataInComponentSpace.Num();
 }
 
+PRAGMA_DISABLE_DEPRECATION_WARNINGS
 bool UAnimSequence::ConvertAnimationDataToRiggingData(FAnimSequenceTrackContainer& RiggingAnimationData)
 {
 	USkeleton* MySkeleton = GetSkeleton();
@@ -3058,6 +3069,7 @@ bool UAnimSequence::ConvertAnimationDataToRiggingData(FAnimSequenceTrackContaine
 
 	return false;
 }
+PRAGMA_ENABLE_DEPRECATION_WARNINGS
 
 bool UAnimSequence::ConvertRiggingDataToAnimationData(FAnimSequenceTrackContainer& RiggingAnimationData)
 {
@@ -3096,7 +3108,6 @@ bool UAnimSequence::ConvertRiggingDataToAnimationData(FAnimSequenceTrackContaine
 
 		const int32 NumModelKeys = DataModelInterface->GetNumberOfKeys();
 		const FReferenceSkeleton& RefSkeleton = MySkeleton->GetReferenceSkeleton();
-		const URig* Rig = MySkeleton->GetRig();
 		for (int32 NodeIndex = 0; NodeIndex < ValidNumNodes; ++NodeIndex)
 		{
 			FName BoneName = MySkeleton->GetRigBoneMapping(ValidNodeNames[NodeIndex]);
diff --git a/Engine/Source/Runtime/Engine/Private/Animation/Rig.cpp b/Engine/Source/Runtime/Engine/Private/Animation/Rig.cpp
--- a/Engine/Source/Runtime/Engine/Private/Animation/Rig.cpp
+++ b/Engine/Source/Runtime/Engine/Private/Animation/Rig.cpp
@@ -13,6 +13,9 @@
 
 //@todo should move all this window stuff somewhere else. Persona?
 
+// January 2023 - Whole compilation unit is being deprecated, so disabled warnings for whole thing. 
+PRAGMA_DISABLE_DEPRECATION_WARNINGS
+
 #define LOCTEXT_NAMESPACE "Rig"
 
 FName URig::WorldNodeName(TEXT("World"));
@@ -514,3 +517,5 @@ void URig::GetMappableNodeData(TArray<FName>& OutNames, TArray<FNodeItem>& OutNo
 }
 #undef LOCTEXT_NAMESPACE 
 
+// January 2023 - Whole compilation unit is being deprecated, so disabled warnings for whole thing. 
+PRAGMA_ENABLE_DEPRECATION_WARNINGS
\ No newline at end of file
diff --git a/Engine/Source/Runtime/Engine/Private/Animation/Skeleton.cpp b/Engine/Source/Runtime/Engine/Private/Animation/Skeleton.cpp
--- a/Engine/Source/Runtime/Engine/Private/Animation/Skeleton.cpp
+++ b/Engine/Source/Runtime/Engine/Private/Animation/Skeleton.cpp
@@ -2075,6 +2075,9 @@ void USkeleton::HandleVirtualBoneChanges()
 }
 
 #if WITH_EDITOR
+
+PRAGMA_DISABLE_DEPRECATION_WARNINGS
+
 void USkeleton::SetRigConfig(URig * Rig)
 {
 	if (RigConfig.Rig != Rig)
@@ -2226,11 +2229,14 @@ URig * USkeleton::GetRig() const
 {
 	return RigConfig.Rig;
 }
+PRAGMA_ENABLE_DEPRECATION_WARNINGS
 
 void USkeleton::GetAssetRegistryTags(TArray<FAssetRegistryTag>& OutTags) const
 {
 	Super::GetAssetRegistryTags(OutTags);
+	PRAGMA_DISABLE_DEPRECATION_WARNINGS
 	FString RigFullName = (RigConfig.Rig)? RigConfig.Rig->GetFullName() : TEXT("");
+	PRAGMA_ENABLE_DEPRECATION_WARNINGS
 
 	OutTags.Add(FAssetRegistryTag(USkeleton::RigTag, RigFullName, FAssetRegistryTag::TT_Hidden));
 
