Expose getter/setter for world_asset and LoadLevelInstance/IsLoaded functions to BP in LevelInstance

diff --git a/Engine/Source/Runtime/Engine/Public/LevelInstance/LevelInstanceInterface.h b/Engine/Source/Runtime/Engine/Public/LevelInstance/LevelInstanceInterface.h
--- a/Engine/Source/Runtime/Engine/Public/LevelInstance/LevelInstanceInterface.h
+++ b/Engine/Source/Runtime/Engine/Public/LevelInstance/LevelInstanceInterface.h
@@ -11,10 +11,11 @@
 class ULevelInstanceComponent;
 class ULevelStreamingLevelInstance;
 
-UINTERFACE()
+UINTERFACE(BlueprintType, meta = (CannotImplementInterfaceInBlueprint))
 class ENGINE_API ULevelInstanceInterface : public UInterface
 {
 	GENERATED_UINTERFACE_BODY()
+
 };
 
 /**
@@ -24,6 +25,7 @@ class ENGINE_API ULevelInstanceInterface : public UInterface
  * 
  * @see FLevelInstanceActorImpl to use as a member or base class member of the implementing class. It provides most of the boiler plate code so that the implementing class behaves properly in editor (Undo/Redo/Delete/Selection/Imporrt/etc)
  */
+
 class ENGINE_API ILevelInstanceInterface
 {
 	GENERATED_IINTERFACE_BODY()
@@ -35,11 +37,15 @@ class ENGINE_API ILevelInstanceInterface
 
 	virtual const FGuid& GetLevelInstanceGuid() const = 0;
 
+	
+	UFUNCTION(BlueprintCallable, Category = Default)
 	virtual const TSoftObjectPtr<UWorld>& GetWorldAsset() const = 0;
 
 	virtual bool IsLoadingEnabled() const = 0;
 
 #if WITH_EDITOR
+
+	UFUNCTION(BlueprintCallable, Category = Default)
 	virtual bool SetWorldAsset(TSoftObjectPtr<UWorld> WorldAsset) = 0;
 
 	virtual ULevelInstanceComponent* GetLevelInstanceComponent() const = 0;
@@ -60,10 +66,13 @@ class ENGINE_API ILevelInstanceInterface
 
 	virtual ULevelInstanceSubsystem* GetLevelInstanceSubsystem() const;
 
+	UFUNCTION(BlueprintCallable, Category = Default)
 	virtual bool IsLoaded() const;
-			
+	
+	UFUNCTION(BlueprintCallable, Category = Default)
 	virtual void LoadLevelInstance();
 
+	UFUNCTION(BlueprintCallable, Category = Default)
 	virtual void UnloadLevelInstance();
 
 	virtual bool IsInitiallyVisible() const { return true; }
@@ -73,6 +82,7 @@ class ENGINE_API ILevelInstanceInterface
 	virtual TSubclassOf<ULevelStreamingLevelInstance> GetLevelStreamingClass() const;
 
 #if WITH_EDITOR
+	UFUNCTION(BlueprintCallable, Category = Default)
 	virtual ULevel* GetLoadedLevel() const;
 
 	virtual void UpdateLevelInstanceFromWorldAsset();
