Iris
- Updated Iris code based on changes from PackageMapClient
- Scaling
- Vetting of package loading in PIE

diff --git a/Engine/Source/Runtime/Engine/Private/Net/Iris/ReplicationSystem/ActorReplicationBridge.cpp b/Engine/Source/Runtime/Engine/Private/Net/Iris/ReplicationSystem/ActorReplicationBridge.cpp
--- a/Engine/Source/Runtime/Engine/Private/Net/Iris/ReplicationSystem/ActorReplicationBridge.cpp
+++ b/Engine/Source/Runtime/Engine/Private/Net/Iris/ReplicationSystem/ActorReplicationBridge.cpp
@@ -102,6 +102,7 @@ UActorReplicationBridge::UActorReplicationBridge()
 , NetDriver(nullptr)
 , MaxPollFrequency(0.0f)
 , ObjectReferencePackageMap(nullptr)
+, SpawnInfoFlags(0U)
 {
 	SetInstancePreUpdateFunction(UE::Net::Private::ActorReplicationBridgePreUpdateFunction);
 	SetInstanceGetWorldLocationFunction(UE::Net::Private::ActorReplicationBridgeGetActorWorldLocation);
@@ -153,6 +154,9 @@ void UActorReplicationBridge::Initialize(UReplicationSystem* InReplicationSystem
 
 	ObjectReferencePackageMap = NewObject<UIrisObjectReferencePackageMap>();
 	ObjectReferencePackageMap->AddToRoot();
+
+	// Get spawn info flags from cvars
+	SpawnInfoFlags = UE::Net::Private::GetActorReplicationBridgeSpawnInfoFlags();
 }
 
 UActorReplicationBridge::~UActorReplicationBridge()
@@ -487,7 +491,7 @@ bool UActorReplicationBridge::WriteCreationHeader(UE::Net::FNetSerializationCont
 		// Serialize the data
 		// Indicate that this is an actor
 		Writer->WriteBool(true);
-		WriteActorCreationHeader(Context, Header);
+		WriteActorCreationHeader(Context, Header, SpawnInfoFlags);
 
 		return !Writer->IsOverflown();
 	}
@@ -605,7 +609,7 @@ FObjectReplicationBridgeInstantiateResult UActorReplicationBridge::BeginInstanti
 					// Set Scale if it differs from Default
 					if (!Header->SpawnInfo.Scale.Equals(DefaultSpawnInfo.Scale, Epsilon))
 					{
-						Actor->SetActorScale3D(Header->SpawnInfo.Scale);
+						Actor->SetActorRelativeScale3D(Header->SpawnInfo.Scale);
 					}
 
 					UE_LOG_ACTORREPLICATIONBRIDGE(Verbose, TEXT("OnBeginInstantiateFromRemote Spawned Actor %s"), *Actor->GetPathName());
@@ -951,6 +955,19 @@ void UActorReplicationBridge::GetActorCreationHeader(const AActor* Actor, UE::Ne
 			Header.SpawnInfo.Location = FRepMovement::RebaseOntoZeroOrigin(Actor->GetActorLocation(), Actor);
 			Header.SpawnInfo.Rotation = Actor->GetActorRotation();
 			Header.SpawnInfo.Scale = Actor->GetActorScale();
+			FVector Scale = Actor->GetActorScale();
+
+			if (USceneComponent* AttachParent = RootComponent->GetAttachParent())
+			{
+				// If this actor is attached, when the scale is serialized on the client, the attach parent property won't be set yet.
+				// USceneComponent::SetWorldScale3D (which got called by AActor::SetActorScale3D, which we used to do but no longer).
+				// would perform this transformation so that what is sent is relative to the parent. If we don't do this, we will
+				// apply the world scale on the client, which will then get applied a second time when the attach parent property is received.
+				FTransform ParentToWorld = AttachParent->GetSocketTransform(RootComponent->GetAttachSocketName());
+				Scale = Scale * ParentToWorld.GetSafeScaleReciprocal(ParentToWorld.GetScale3D());
+			}
+
+			Header.SpawnInfo.Scale = Scale;
 			Header.SpawnInfo.Velocity = Actor->GetVelocity();
 		}
 		else
diff --git a/Engine/Source/Runtime/Engine/Private/Net/Iris/ReplicationSystem/ActorReplicationBridgeInternal.cpp b/Engine/Source/Runtime/Engine/Private/Net/Iris/ReplicationSystem/ActorReplicationBridgeInternal.cpp
--- a/Engine/Source/Runtime/Engine/Private/Net/Iris/ReplicationSystem/ActorReplicationBridgeInternal.cpp
+++ b/Engine/Source/Runtime/Engine/Private/Net/Iris/ReplicationSystem/ActorReplicationBridgeInternal.cpp
@@ -7,7 +7,10 @@
 #include "Iris/Serialization/NetBitStreamReader.h"
 #include "Iris/Serialization/NetBitStreamWriter.h"
 #include "Iris/Serialization/NetBitStreamUtil.h"
+#include "Iris/Serialization/NetSerializers.h"
+#include "Iris/Serialization/PackedVectorNetSerializers.h"
 #include "Iris/Serialization/ObjectNetSerializer.h"
+#include "HAL/IConsoleManager.h"
 
 namespace UE::Net::Private
 {
@@ -19,32 +22,103 @@ const FActorReplicationBridgeSpawnInfo& GetDefaultActorReplicationBridgeSpawnInf
 	return DefaultActorReplicationBridgeSpawnInfo;
 }
 
+// Write vector using default value compression using specified serializer
+static void InternalWriteConditionallyQuantizedVector(FNetBitStreamWriter* Writer, const FVector& Vector, const FVector& DefaultValue, bool bQuantize)
+{
+	// We use 0.01f for comparing when using quantization, because we will only send a single point of precision anyway.
+	// We could probably get away with 0.1f, but that may introduce edge cases for rounding.
+	static constexpr float Epsilon_Quantized = 0.01f;
+				
+	// We use KINDA_SMALL_NUMBER for comparing when not using quantization, because that's the default for FVector::Equals.
+	const float Epsilon = bQuantize ? Epsilon_Quantized : UE_KINDA_SMALL_NUMBER;
+	
+	if (Writer->WriteBool(!Vector.Equals(DefaultValue, Epsilon)))
+	{
+		Writer->WriteBool(bQuantize);
+
+		const FNetSerializer& Serializer = bQuantize ? UE_NET_GET_SERIALIZER(FVectorNetQuantize10NetSerializer) : UE_NET_GET_SERIALIZER(FVectorNetSerializer);
+
+		FNetSerializationContext Context(Writer);
+
+		alignas(16) uint8 QuantizedState[32] = {};
+		checkSlow(sizeof(QuantizedState) >= Serializer.QuantizedTypeSize);
+
+		FNetQuantizeArgs QuantizeArgs;
+		QuantizeArgs.Version = Serializer.Version;
+		QuantizeArgs.Source = NetSerializerValuePointer(&Vector);
+		QuantizeArgs.Target = NetSerializerValuePointer(&QuantizedState[0]);
+		QuantizeArgs.NetSerializerConfig = NetSerializerConfigParam(Serializer.DefaultConfig);
+		Serializer.Quantize(Context, QuantizeArgs);
+
+		FNetSerializeArgs SerializeArgs;
+		SerializeArgs.Version = Serializer.Version;
+		SerializeArgs.Source = QuantizeArgs.Target;
+		SerializeArgs.NetSerializerConfig = NetSerializerConfigParam(Serializer.DefaultConfig);
+		Serializer.Serialize(Context, SerializeArgs);
+	}
+}
+
+// Read default value compressed vector using specified serializer
+static void InternalReadConditionallyQuantizedVector(FNetBitStreamReader* Reader, FVector& OutVector, const FVector& DefaultValue)
+{
+	if (Reader->ReadBool())
+	{
+		const bool bIsQuantized = Reader->ReadBool();
+
+		const FNetSerializer& Serializer = bIsQuantized ? UE_NET_GET_SERIALIZER(FVectorNetQuantize10NetSerializer) : UE_NET_GET_SERIALIZER(FVectorNetSerializer);
+
+		FNetSerializationContext Context(Reader);
+
+		alignas(16) uint8 QuantizedState[32] = {};
+		checkSlow(sizeof(QuantizedState) >= Serializer.QuantizedTypeSize);
+
+		FNetDeserializeArgs DeserializeArgs;
+		DeserializeArgs.Version = Serializer.Version;
+		DeserializeArgs.Target = NetSerializerValuePointer(&QuantizedState[0]);
+		DeserializeArgs.NetSerializerConfig = NetSerializerConfigParam(Serializer.DefaultConfig);
+		Serializer.Deserialize(Context, DeserializeArgs);
+
+		FNetDequantizeArgs DequantizeArgs;
+		DequantizeArgs.Version = Serializer.Version;
+		DequantizeArgs.Source = DeserializeArgs.Target;
+		DequantizeArgs.Target = NetSerializerValuePointer(&OutVector);
+		DequantizeArgs.NetSerializerConfig = NetSerializerConfigParam(Serializer.DefaultConfig);
+		Serializer.Dequantize(Context, DequantizeArgs);
+	}
+	else
+	{
+		OutVector = DefaultValue;
+	}
+}
+
 /*
  * IRIS: TODO: As soon as we have added code generation for Serialization code this will no longer be required
  */
-void WriteActorReplicationBridgeSpawnInfo(FNetSerializationContext& Context, const FActorReplicationBridgeSpawnInfo& SpawnInfo)
+void WriteActorReplicationBridgeSpawnInfo(FNetSerializationContext& Context, const FActorReplicationBridgeSpawnInfo& SpawnInfo, uint32 Flags)
 {
 	FNetBitStreamWriter* Writer = Context.GetBitStreamWriter();
 
-	constexpr float epsilon = 0.001f;
+	InternalWriteConditionallyQuantizedVector(Writer, SpawnInfo.Location, DefaultActorReplicationBridgeSpawnInfo.Location, (Flags & EActorReplicationBridgeSpawnInfoFlags_QuantizeLocation) != 0U);
+	InternalWriteConditionallyQuantizedVector(Writer, SpawnInfo.Scale, DefaultActorReplicationBridgeSpawnInfo.Scale, (Flags & EActorReplicationBridgeSpawnInfoFlags_QuantizeScale) != 0U);
+	InternalWriteConditionallyQuantizedVector(Writer, SpawnInfo.Velocity, DefaultActorReplicationBridgeSpawnInfo.Velocity, (Flags & EActorReplicationBridgeSpawnInfoFlags_QuantizeVelocity) != 0U);
 
-	WriteVector(Writer, SpawnInfo.Location, DefaultActorReplicationBridgeSpawnInfo.Location, epsilon);
-	WriteRotator(Writer, SpawnInfo.Rotation, DefaultActorReplicationBridgeSpawnInfo.Rotation, epsilon);
-	WriteVector(Writer, SpawnInfo.Scale, DefaultActorReplicationBridgeSpawnInfo.Scale, epsilon);
-	WriteVector(Writer, SpawnInfo.Velocity, DefaultActorReplicationBridgeSpawnInfo.Velocity, epsilon);
+	// For rotation we use 0.001f for Rotation comparison to keep consistency with old behavior.
+	static constexpr float RotationEpsilon = 0.001f;
+	WriteRotator(Writer, SpawnInfo.Rotation, DefaultActorReplicationBridgeSpawnInfo.Rotation, RotationEpsilon);
 }
 
 void ReadActorReplicationBridgeSpawnInfo(FNetSerializationContext& Context, FActorReplicationBridgeSpawnInfo& SpawnInfo)
 {
 	FNetBitStreamReader* Reader = Context.GetBitStreamReader();
 
-	ReadVector(Reader, SpawnInfo.Location, DefaultActorReplicationBridgeSpawnInfo.Location);
+	InternalReadConditionallyQuantizedVector(Reader, SpawnInfo.Location, DefaultActorReplicationBridgeSpawnInfo.Location);
+	InternalReadConditionallyQuantizedVector(Reader, SpawnInfo.Scale, DefaultActorReplicationBridgeSpawnInfo.Scale);
+	InternalReadConditionallyQuantizedVector(Reader, SpawnInfo.Velocity, DefaultActorReplicationBridgeSpawnInfo.Velocity);
+
 	ReadRotator(Reader, SpawnInfo.Rotation, DefaultActorReplicationBridgeSpawnInfo.Rotation);
-	ReadVector(Reader, SpawnInfo.Scale, DefaultActorReplicationBridgeSpawnInfo.Scale);
-	ReadVector(Reader, SpawnInfo.Velocity, DefaultActorReplicationBridgeSpawnInfo.Velocity);
 }
 
-void WriteActorCreationHeader(FNetSerializationContext& Context, const FActorCreationHeader& Header)
+void WriteActorCreationHeader(FNetSerializationContext& Context, const FActorCreationHeader& Header, uint32 Flags)
 {
 	FNetBitStreamWriter* Writer = Context.GetBitStreamWriter();
 
@@ -60,7 +134,7 @@ void WriteActorCreationHeader(FNetSerializationContext& Context, const FActorCre
 		}
 		
 		// Write actor spawn info
-		WriteActorReplicationBridgeSpawnInfo(Context, Header.SpawnInfo);
+		WriteActorReplicationBridgeSpawnInfo(Context, Header.SpawnInfo, Flags);
 	}
 	else // should be possible to refer by path??
 	{
@@ -156,6 +230,43 @@ void ReadSubObjectCreationHeader(FNetSerializationContext& Context, FSubObjectCr
 	}
 }
 
+uint32 GetActorReplicationBridgeSpawnInfoFlags()
+{
+	// Init spawninfo flags from CVARs
+	uint32 Flags = 0U;
+	{
+		bool bQuantizeActorScaleOnSpawn = false;
+		IConsoleVariable* CVar = IConsoleManager::Get().FindConsoleVariable(TEXT("net.QuantizeActorScaleOnSpawn"));
+		if (ensure(CVar))
+		{
+			bQuantizeActorScaleOnSpawn = CVar->GetBool();
+		}
+		Flags |= bQuantizeActorScaleOnSpawn ? EActorReplicationBridgeSpawnInfoFlags_QuantizeScale : 0U;
+	}
+
+	{
+		bool bQuantizeActorLocationOnSpawn = true;
+		IConsoleVariable* CVar = IConsoleManager::Get().FindConsoleVariable(TEXT("net.QuantizeActorLocationOnSpawn"));
+		if (ensure(CVar))
+		{
+			bQuantizeActorLocationOnSpawn = CVar->GetBool();
+		}
+		Flags |= bQuantizeActorLocationOnSpawn ? EActorReplicationBridgeSpawnInfoFlags_QuantizeLocation : 0U;
+	}
+
+	{
+		bool bQuantizeActorVelocityOnSpawn = true;
+		IConsoleVariable* CVar = IConsoleManager::Get().FindConsoleVariable(TEXT("net.QuantizeActorVelocityOnSpawn"));
+		if (ensure(CVar))
+		{
+			bQuantizeActorVelocityOnSpawn = CVar->GetBool();
+		}
+		Flags |= bQuantizeActorVelocityOnSpawn ? EActorReplicationBridgeSpawnInfoFlags_QuantizeVelocity : 0U;
+	}
+
+	return Flags;
+}
+
 }
 
 #endif // UE_WITH_IRIS
diff --git a/Engine/Source/Runtime/Engine/Private/Net/Iris/ReplicationSystem/ActorReplicationBridgeInternal.h b/Engine/Source/Runtime/Engine/Private/Net/Iris/ReplicationSystem/ActorReplicationBridgeInternal.h
--- a/Engine/Source/Runtime/Engine/Private/Net/Iris/ReplicationSystem/ActorReplicationBridgeInternal.h
+++ b/Engine/Source/Runtime/Engine/Private/Net/Iris/ReplicationSystem/ActorReplicationBridgeInternal.h
@@ -14,6 +14,14 @@ namespace UE::Net
 namespace UE::Net::Private
 {
 
+enum EActorReplicationBridgeSpawnInfoFlags : uint32
+{
+	EActorReplicationBridgeSpawnInfoFlags_QuantizeScale = 1U,
+	EActorReplicationBridgeSpawnInfoFlags_QuantizeLocation = EActorReplicationBridgeSpawnInfoFlags_QuantizeScale << 1U,
+	EActorReplicationBridgeSpawnInfoFlags_QuantizeVelocity = EActorReplicationBridgeSpawnInfoFlags_QuantizeLocation << 1U,
+	EActorReplicationBridgeSpawnInfoFlags_BitCount = 3U,
+};
+
 struct FActorReplicationBridgeSpawnInfo
 {
 	FVector Location;
@@ -47,15 +55,17 @@ struct FSubObjectCreationHeader : public FActorReplicationBridgeCreationHeader
 
 const FActorReplicationBridgeSpawnInfo& GetDefaultActorReplicationBridgeSpawnInfo();
 
-void WriteActorReplicationBridgeSpawnInfo(FNetSerializationContext& Context, const FActorReplicationBridgeSpawnInfo& SpawnInfo);
+void WriteActorReplicationBridgeSpawnInfo(FNetSerializationContext& Context, const FActorReplicationBridgeSpawnInfo& SpawnInfo, uint32 Flags);
 void ReadActorReplicationBridgeSpawnInfo(FNetSerializationContext& Context, FActorReplicationBridgeSpawnInfo& SpawnInfo);
 
-void WriteActorCreationHeader(FNetSerializationContext& Context, const FActorCreationHeader& Header);
+void WriteActorCreationHeader(FNetSerializationContext& Context, const FActorCreationHeader& Header, uint32 Flags);
 void ReadActorCreationHeader(FNetSerializationContext& Context, FActorCreationHeader& Header);
 
 void WriteSubObjectCreationHeader(FNetSerializationContext& Context, const FSubObjectCreationHeader& Header);
 void ReadSubObjectCreationHeader(FNetSerializationContext& Context, FSubObjectCreationHeader& Header); 
 
+uint32 GetActorReplicationBridgeSpawnInfoFlags();
+
 }
 
 #endif // UE_WITH_IRIS
diff --git a/Engine/Source/Runtime/Engine/Public/Net/Iris/ReplicationSystem/ActorReplicationBridge.h b/Engine/Source/Runtime/Engine/Public/Net/Iris/ReplicationSystem/ActorReplicationBridge.h
--- a/Engine/Source/Runtime/Engine/Public/Net/Iris/ReplicationSystem/ActorReplicationBridge.h
+++ b/Engine/Source/Runtime/Engine/Public/Net/Iris/ReplicationSystem/ActorReplicationBridge.h
@@ -104,5 +104,8 @@ private:
 	float MaxPollFrequency;
 
 	UIrisObjectReferencePackageMap* ObjectReferencePackageMap;
+
+	uint32 SpawnInfoFlags;
+
 #endif // UE_WITH_IRIS
 };
diff --git a/Engine/Source/Runtime/Experimental/Iris/Core/Private/Iris/ReplicationSystem/ObjectReferenceCache.cpp b/Engine/Source/Runtime/Experimental/Iris/Core/Private/Iris/ReplicationSystem/ObjectReferenceCache.cpp
--- a/Engine/Source/Runtime/Experimental/Iris/Core/Private/Iris/ReplicationSystem/ObjectReferenceCache.cpp
+++ b/Engine/Source/Runtime/Experimental/Iris/Core/Private/Iris/ReplicationSystem/ObjectReferenceCache.cpp
@@ -114,13 +114,19 @@ bool FObjectReferenceCache::CanClientLoadObjectInternal(const UObject* Object, b
 {
 	// PackageMapClient can't load maps, we must wait for the client to load the map when ready
 	// These guids are special guids, where the guid and all child guids resolve once the map has been loaded
-	if (bIsDynamic || Object->GetOutermost()->ContainsMap())
+	if (bIsDynamic)
+	{
+		return false;
+	}
+
+	if (Object->GetPackage()->ContainsMap())
 	{
 		return false;
 	}
 
 #if WITH_EDITOR
 	// For objects using external package, we need to do the test on the package of their outer most object
+	// (this is currently only possible in Editor)
 	UObject* OutermostObject = Object->GetOutermostObject();
 	if (OutermostObject && OutermostObject->GetPackage()->ContainsMap())
 	{
@@ -166,7 +172,7 @@ bool FObjectReferenceCache::ShouldIgnoreWhenMissing(FNetRefHandle RefHandle) con
 			return true;
 		}
 		// Sometimes, other systems async load packages, which we don't track, but still must be aware of
-		if (OutermostCacheObject->Object != nullptr && !OutermostCacheObject->Object->GetOutermost()->IsFullyLoaded())
+		if (OutermostCacheObject->Object != nullptr && !OutermostCacheObject->Object->GetPackage()->IsFullyLoaded())
 		{
 			return true;
 		}
@@ -635,7 +641,7 @@ UObject* FObjectReferenceCache::ResolveObjectReferenceHandleInternal(FNetRefHand
 
 	check(!CacheObjectPtr->bIsPending);
 
-	if (!ensure(ObjOuter == nullptr || ObjOuter->GetOutermost()->IsFullyLoaded() || ObjOuter->GetOutermost()->HasAnyPackageFlags(TreatAsLoadedFlags)))
+	if (!ensure(ObjOuter == nullptr || ObjOuter->GetPackage()->IsFullyLoaded() || ObjOuter->GetPackage()->HasAnyPackageFlags(TreatAsLoadedFlags)))
 	{
 		UE_LOG(LogIris, Error, TEXT("GetObjectFromRefHandle: Outer is null or package is not fully loaded. FullPath: %s Outer: %s"), ToCStr(FullPath(RefHandle, ResolveContext)), *GetFullNameSafe(ObjOuter));
 	}
@@ -652,6 +658,13 @@ UObject* FObjectReferenceCache::ResolveObjectReferenceHandleInternal(FNetRefHand
 
 	// See if this object is in memory
 	Object = StaticFindObject(UObject::StaticClass(), ObjOuter, *ObjectPath, false);
+#if WITH_EDITOR
+	// Object must be null if the package is a dynamic PIE package with pending external objects still loading, as it would normally while object is async loading
+	if (Object && Object->GetPackage()->IsDynamicPIEPackagePending())
+	{
+		Object = nullptr;
+	}
+#endif
 
 	// Assume this is a package if the outer is invalid and this is a static guid
 	const bool bIsPackage = CacheObjectPtr->bIsPackage;
