Fix Strata SLW forcing opaque blend mode, while masked blend mode is also supported.

diff --git a/Engine/Source/Runtime/Engine/Private/Materials/Material.cpp b/Engine/Source/Runtime/Engine/Private/Materials/Material.cpp
--- a/Engine/Source/Runtime/Engine/Private/Materials/Material.cpp
+++ b/Engine/Source/Runtime/Engine/Private/Materials/Material.cpp
@@ -4873,7 +4873,10 @@ void UMaterial::RebuildShadingModelField()
 			{
 				MaterialDomain = EMaterialDomain::MD_Surface;
 				ShadingModel = MSM_SingleLayerWater;
-				BlendMode = EBlendMode::BLEND_Opaque; // STRATA_TODO water can also be masked: check Mask input from the main node to automatically enabled that?
+				if (BlendMode != EBlendMode::BLEND_Opaque && BlendMode != EBlendMode::BLEND_Masked)
+				{
+					BlendMode = EBlendMode::BLEND_Opaque;
+				}
 			}
 			else if (StrataMaterialInfo.HasOnlyShadingModel(SSM_LightFunction))
 			{
