Iris
* Ignore actor seamless travel callbacks when Iris is enabled

diff --git a/Engine/Source/Runtime/Engine/Private/NetDriver.cpp b/Engine/Source/Runtime/Engine/Private/NetDriver.cpp
--- a/Engine/Source/Runtime/Engine/Private/NetDriver.cpp
+++ b/Engine/Source/Runtime/Engine/Private/NetDriver.cpp
@@ -3519,6 +3519,13 @@ void UNetDriver::NotifyActorTearOff(AActor* Actor)
 
 void UNetDriver::NotifyActorIsTraveling(AActor* TravelingActor)
 {
+#if UE_WITH_IRIS
+	if (ReplicationSystem)
+	{
+		return;
+	}
+#endif
+
 	NetworkObjects->OnActorIsTraveling(TravelingActor);
 }
 
