FWorldPartitionHelpers: GetWorldPartition(UObject*)

#rnx

diff --git a/Engine/Source/Runtime/Engine/Private/WorldPartition/WorldPartitionHelpers.cpp b/Engine/Source/Runtime/Engine/Private/WorldPartition/WorldPartitionHelpers.cpp
--- a/Engine/Source/Runtime/Engine/Private/WorldPartition/WorldPartitionHelpers.cpp
+++ b/Engine/Source/Runtime/Engine/Private/WorldPartition/WorldPartitionHelpers.cpp
@@ -4,7 +4,10 @@
 
 #include "WorldPartition/WorldPartition.h"
 #include "WorldPartition/WorldPartitionEditorHash.h"
+#include "WorldPartition/WorldPartitionRuntimeCellInterface.h"
+#include "WorldPartition/WorldPartitionRuntimeCellOwner.h"
 #include "Engine/World.h"
+#include "Engine/Level.h"
 #include "Algo/AnyOf.h"
 
 #include "Commandlets/Commandlet.h"
@@ -12,6 +15,50 @@
 
 DEFINE_LOG_CATEGORY_STATIC(LogWorldPartitionHelpers, Log, All);
 
+UWorldPartition* FWorldPartitionHelpers::GetWorldPartition(const UObject* InObject)
+{
+	if (!IsValid(InObject))
+	{
+		return nullptr;
+	}
+
+	const UWorld* OuterWorld = nullptr;
+
+	if (const UWorld* InWorld = Cast<const UWorld>(InObject))
+	{
+		return GetWorldPartition(InWorld->PersistentLevel);
+	}
+	else if (const ULevel* Level = Cast<const ULevel>(InObject))
+	{
+		if (const IWorldPartitionCell* Cell = Level->GetWorldPartitionRuntimeCell())
+		{
+			return GetWorldPartition(Cast<UObject>(Cell));
+		}
+	}
+	else if (const IWorldPartitionCell* Cell = Cast<const IWorldPartitionCell>(InObject))
+	{
+		if (const IWorldPartitionRuntimeCellOwner* CellOwner = Cell->GetCellOwner())
+		{
+			return GetWorldPartition(Cast<UObject>(CellOwner));
+		}
+	}
+	else if (const IWorldPartitionRuntimeCellOwner* CellOwner = Cast<const IWorldPartitionRuntimeCellOwner>(InObject))
+	{
+		OuterWorld = CellOwner->GetOuterWorld();
+	}
+	else if (const ULevel* OuterLevel = InObject->GetTypedOuter<ULevel>())
+	{
+		return GetWorldPartition(OuterLevel);
+	}
+
+	if (!OuterWorld)
+	{
+		OuterWorld = InObject ? InObject->GetTypedOuter<UWorld>() : nullptr;
+	}
+
+	return OuterWorld ? OuterWorld->GetWorldPartition() : nullptr;
+}
+
 #if WITH_EDITOR
 
 bool FWorldPartitionHelpers::IsActorDescClassCompatibleWith(const FWorldPartitionActorDesc* ActorDesc, const UClass* Class)
diff --git a/Engine/Source/Runtime/Engine/Public/WorldPartition/WorldPartitionHelpers.h b/Engine/Source/Runtime/Engine/Public/WorldPartition/WorldPartitionHelpers.h
--- a/Engine/Source/Runtime/Engine/Public/WorldPartition/WorldPartitionHelpers.h
+++ b/Engine/Source/Runtime/Engine/Public/WorldPartition/WorldPartitionHelpers.h
@@ -3,16 +3,26 @@
 
 #include "CoreMinimal.h"
 #include "GameFramework/Actor.h"
-#include "WorldPartition/WorldPartitionHandle.h"
 
 #if WITH_EDITOR
+#include "WorldPartition/WorldPartitionHandle.h"
+#endif
 
 class UWorld;
 class UWorldPartition;
+
+#if WITH_EDITOR
 class FWorldPartitionActorDesc;
+#endif
 
 class ENGINE_API FWorldPartitionHelpers
 {
+public:
+	/** Returns the owning World Partition for this object. */
+	static UWorldPartition* GetWorldPartition(const UObject* InObject);
+
+#if WITH_EDITOR
+private:
 	static bool IsActorDescClassCompatibleWith(const FWorldPartitionActorDesc* ActorDesc, const UClass* Class);
 
 public:
@@ -80,6 +90,8 @@ public:
 	// Editor/Runtime conversions
 	static bool ConvertRuntimePathToEditorPath(const FSoftObjectPath& InPath, FSoftObjectPath& OutPath);
 	static bool ConvertEditorPathToRuntimePath(const FSoftObjectPath& InPath, FSoftObjectPath& OutPath);
-};
 
 #endif // WITH_EDITOR
+};
+
+
