If morph vertex buffer hasn't updated this frame or for some frames, use the last updated morph buffer for the previous morph data, so it cancels out motion vector.
This fixes a jiggery issue in MorphTargetPreview tab in editor previewing using morph sliders.

diff --git a/Engine/Source/Runtime/Engine/Private/GPUSkinVertexFactory.cpp b/Engine/Source/Runtime/Engine/Private/GPUSkinVertexFactory.cpp
--- a/Engine/Source/Runtime/Engine/Private/GPUSkinVertexFactory.cpp
+++ b/Engine/Source/Runtime/Engine/Private/GPUSkinVertexFactory.cpp
@@ -809,6 +809,13 @@ const FMorphVertexBuffer* TGPUSkinVertexFactory<BoneInfluenceType>::GetMorphVert
 	return Data->MorphVertexBufferPool ? &Data->MorphVertexBufferPool->GetMorphVertexBufferForReading(bPrevious) : nullptr;
 }
 
+template <GPUSkinBoneInfluenceType BoneInfluenceType>
+uint32 TGPUSkinVertexFactory<BoneInfluenceType>::GetMorphVertexBufferUpdatedFrameNumber() const
+{
+	check(Data.IsValid());
+	return Data->MorphVertexBufferPool ? Data->MorphVertexBufferPool->GetUpdatedFrameNumber() : 0;
+}
+
 /*-----------------------------------------------------------------------------
 TGPUSkinAPEXClothVertexFactory
 -----------------------------------------------------------------------------*/
@@ -915,9 +922,10 @@ public:
 		const bool bIsMobile = IsMobilePlatform(ShaderPlatform);
 		if (!bIsMobile)
 		{
-			const FMorphVertexBuffer* MorphVertexBuffer = nullptr;
 			const auto* GPUSkinVertexFactory = (const FGPUBaseSkinVertexFactory*)VertexFactory;
-			MorphVertexBuffer = GPUSkinVertexFactory->GetMorphVertexBuffer(!View->Family->bWorldIsPaused);
+			bool bMorphUpdatedThisFrame = (View->Family->FrameNumber == GPUSkinVertexFactory->GetMorphVertexBufferUpdatedFrameNumber());
+			bool bPrevious = !View->Family->bWorldIsPaused && bMorphUpdatedThisFrame;
+			const FMorphVertexBuffer* MorphVertexBuffer = GPUSkinVertexFactory->GetMorphVertexBuffer(bPrevious);
 			ShaderBindings.Add(PreviousMorphBufferParameter, MorphVertexBuffer ? MorphVertexBuffer->GetSRV() : GNullVertexBuffer.VertexBufferSRV.GetReference());
 		}
 	}
diff --git a/Engine/Source/Runtime/Engine/Private/SkeletalRenderGPUSkin.cpp b/Engine/Source/Runtime/Engine/Private/SkeletalRenderGPUSkin.cpp
--- a/Engine/Source/Runtime/Engine/Private/SkeletalRenderGPUSkin.cpp
+++ b/Engine/Source/Runtime/Engine/Private/SkeletalRenderGPUSkin.cpp
@@ -625,6 +625,7 @@ void FSkeletalMeshObjectGPUSkin::ProcessUpdatedDynamicData(EGPUSkinCacheEntryMod
 		if(bMorphNeedsUpdate)
 		{
 			UpdateMorphVertexBuffer(RHICmdList, Mode, LOD, LODData, bGPUSkinCacheEnabled, MorphVertexBuffer);
+			LOD.MorphVertexBufferPool.SetUpdatedFrameNumber(FrameNumberToPrepare);
 		}
 	}
 	else
diff --git a/Engine/Source/Runtime/Engine/Private/SkeletalRenderGPUSkin.h b/Engine/Source/Runtime/Engine/Private/SkeletalRenderGPUSkin.h
--- a/Engine/Source/Runtime/Engine/Private/SkeletalRenderGPUSkin.h
+++ b/Engine/Source/Runtime/Engine/Private/SkeletalRenderGPUSkin.h
@@ -316,7 +316,9 @@ public:
 	void ReleaseResources();
 	SIZE_T GetResourceSize() const;
 	void EnableDoubleBuffer();
-	bool IsDoubleBuffered() const	{ return bDoubleBuffer; }
+	bool IsDoubleBuffered() const					{ return bDoubleBuffer; }
+	void SetUpdatedFrameNumber(uint32 FrameNumber)	{ UpdatedFrameNumber = FrameNumber; }
+	uint32 GetUpdatedFrameNumber() const			{ return UpdatedFrameNumber; }
 	void SetCurrentRevisionNumber(uint32 RevisionNumber);
 	const FMorphVertexBuffer& GetMorphVertexBufferForReading(bool bPrevious) const;
 	FMorphVertexBuffer& GetMorphVertexBufferForWriting();
@@ -332,6 +334,8 @@ private:
 	// RevisionNumber Tracker
 	uint32 PreviousRevisionNumber = 0;
 	uint32 CurrentRevisionNumber = 0;
+	// Frame number of the morph vertex buffer that is last updated
+	uint32 UpdatedFrameNumber = 0;
 };
 
 /**
diff --git a/Engine/Source/Runtime/Engine/Public/GPUSkinVertexFactory.h b/Engine/Source/Runtime/Engine/Public/GPUSkinVertexFactory.h
--- a/Engine/Source/Runtime/Engine/Public/GPUSkinVertexFactory.h
+++ b/Engine/Source/Runtime/Engine/Public/GPUSkinVertexFactory.h
@@ -391,6 +391,7 @@ public:
 	/** Morph vertex factory functions */
 	virtual void UpdateMorphVertexStream(const class FMorphVertexBuffer* MorphVertexBuffer) {}
 	virtual const class FMorphVertexBuffer* GetMorphVertexBuffer(bool bPrevious) const { return nullptr; }
+	virtual uint32 GetMorphVertexBufferUpdatedFrameNumber() const { return 0; }
 	/** Cloth vertex factory access. */
 	virtual class FGPUBaseSkinAPEXClothVertexFactory* GetClothVertexFactory() { return nullptr; }
 	virtual class FGPUBaseSkinAPEXClothVertexFactory const* GetClothVertexFactory() const { return nullptr; }
@@ -468,7 +469,7 @@ public:
 	/** FGPUBaseSkinVertexFactory overrides */
 	virtual void UpdateMorphVertexStream(const class FMorphVertexBuffer* MorphVertexBuffer) override;
 	virtual const class FMorphVertexBuffer* GetMorphVertexBuffer(bool bPrevious) const override;
-
+	virtual uint32 GetMorphVertexBufferUpdatedFrameNumber() const override;
 
 	// FRenderResource interface.
 	virtual void InitRHI() override;
