[shaders] unify handling of single job bSucceeded flag; wasn't getting set properly when not using workers

diff --git a/Engine/Source/Programs/ShaderCompileWorker/Private/ShaderCompileWorker.cpp b/Engine/Source/Programs/ShaderCompileWorker/Private/ShaderCompileWorker.cpp
--- a/Engine/Source/Programs/ShaderCompileWorker/Private/ShaderCompileWorker.cpp
+++ b/Engine/Source/Programs/ShaderCompileWorker/Private/ShaderCompileWorker.cpp
@@ -739,6 +739,8 @@ private:
 			{
 				FShaderCompileJob& Job = SingleJobs[JobIndex];
 				OutputFile << Job.Output;
+				bool bSucceeded = (bool)Job.bSucceeded;
+				OutputFile << bSucceeded;
 				UpdateFileSize(OutputFile, FileSizePosition);
 			}
 		}
diff --git a/Engine/Source/Runtime/Engine/Private/ShaderCompiler/ShaderCompiler.cpp b/Engine/Source/Runtime/Engine/Private/ShaderCompiler/ShaderCompiler.cpp
--- a/Engine/Source/Runtime/Engine/Private/ShaderCompiler/ShaderCompiler.cpp
+++ b/Engine/Source/Runtime/Engine/Private/ShaderCompiler/ShaderCompiler.cpp
@@ -1750,7 +1750,7 @@ static void ProcessErrors(const FShaderCompileJob& CurrentJob, TArray<FString>&
 	}
 }
 
-static bool ReadSingleJob(FShaderCompileJob* CurrentJob, FArchive& OutputFile)
+static bool ReadSingleJob(FShaderCompileJob* CurrentJob, FArchive& WorkerOutputFileReader)
 {
 	TRACE_CPUPROFILER_EVENT_SCOPE(ReadSingleJob);
 
@@ -1758,9 +1758,10 @@ static bool ReadSingleJob(FShaderCompileJob* CurrentJob, FArchive& OutputFile)
 	CurrentJob->bFinalized = true;
 
 	// Deserialize the shader compilation output.
-	OutputFile << CurrentJob->Output;
-
-	CurrentJob->bSucceeded = CurrentJob->Output.bSucceeded;
+	WorkerOutputFileReader << CurrentJob->Output;
+	bool bSucceeded = false;
+	WorkerOutputFileReader << bSucceeded;
+	CurrentJob->bSucceeded = bSucceeded;
 	// The job should already have a non-zero output hash
 	checkf(CurrentJob->Output.OutputHash != FSHAHash() || !CurrentJob->bSucceeded, TEXT("OutputHash for a successful job was not set in the shader compile worker!"));
 
diff --git a/Engine/Source/Runtime/RenderCore/Private/ShaderCore.cpp b/Engine/Source/Runtime/RenderCore/Private/ShaderCore.cpp
--- a/Engine/Source/Runtime/RenderCore/Private/ShaderCore.cpp
+++ b/Engine/Source/Runtime/RenderCore/Private/ShaderCore.cpp
@@ -896,6 +896,8 @@ void CompileShader(const TArray<const IShaderFormat*>& ShaderFormats, FShaderCom
 #else
 	CompileShaderInternal(Compiler, Job, WorkingDirectory, CompileCount);
 #endif
+
+	Job.bSucceeded = Job.Output.bSucceeded;
 }
 
 void CompileShaderPipeline(const TArray<const IShaderFormat*>& ShaderFormats, FShaderPipelineCompileJob* PipelineJob, const FString& Dir, int32* CompileCount)
diff --git a/Engine/Source/Runtime/RenderCore/Public/ShaderCompilerCore.h b/Engine/Source/Runtime/RenderCore/Public/ShaderCompilerCore.h
--- a/Engine/Source/Runtime/RenderCore/Public/ShaderCompilerCore.h
+++ b/Engine/Source/Runtime/RenderCore/Public/ShaderCompilerCore.h
@@ -23,7 +23,7 @@ class FShaderPipelineCompileJob;
 // this is for the protocol, not the data, bump if FShaderCompilerInput or ProcessInputFromArchive changes.
 inline const int32 ShaderCompileWorkerInputVersion = 18;
 // this is for the protocol, not the data, bump if FShaderCompilerOutput or WriteToOutputArchive changes.
-inline const int32 ShaderCompileWorkerOutputVersion = 11;
+inline const int32 ShaderCompileWorkerOutputVersion = 12;
 // this is for the protocol, not the data.
 inline const int32 ShaderCompileWorkerSingleJobHeader = 'S';
 // this is for the protocol, not the data.
