Add a GetTypeHash overload for FHardwareDeviceIdentifier so that you can make TSet/TMap of it

#rnx

diff --git a/Engine/Source/Runtime/Engine/Classes/GameFramework/InputSettings.h b/Engine/Source/Runtime/Engine/Classes/GameFramework/InputSettings.h
--- a/Engine/Source/Runtime/Engine/Classes/GameFramework/InputSettings.h
+++ b/Engine/Source/Runtime/Engine/Classes/GameFramework/InputSettings.h
@@ -330,6 +330,8 @@ struct ENGINE_API FHardwareDeviceIdentifier
 	{
 		return Other.InputClassName == InputClassName && Other.HardwareDeviceIdentifier == HardwareDeviceIdentifier;
 	}
+
+	ENGINE_API friend uint32 GetTypeHash(const FHardwareDeviceIdentifier& InDevice);
 };
 
 /** Per-Platform input options */
diff --git a/Engine/Source/Runtime/Engine/Private/UserInterface/InputSettings.cpp b/Engine/Source/Runtime/Engine/Private/UserInterface/InputSettings.cpp
--- a/Engine/Source/Runtime/Engine/Private/UserInterface/InputSettings.cpp
+++ b/Engine/Source/Runtime/Engine/Private/UserInterface/InputSettings.cpp
@@ -516,6 +516,11 @@ FHardwareDeviceIdentifier::FHardwareDeviceIdentifier(const FName InClassName, co
 
 }
 
+uint32 GetTypeHash(const FHardwareDeviceIdentifier& InDevice)
+{
+	return HashCombine(GetTypeHash(InDevice.InputClassName), GetTypeHash(InDevice.HardwareDeviceIdentifier));
+}
+
 bool FHardwareDeviceIdentifier::IsValid() const
 {
 	return InputClassName.IsValid() && HardwareDeviceIdentifier.IsValid();
