- WorldPartitionSubsystem is now created for inactive words.
- Fixed UEditorLevelUtils::CreateNewStreamingLevelForWorld to properly uninitialized newly created worlds prior to adding its level as a streaming level.

diff --git a/Engine/Source/Editor/UnrealEd/Private/EditorLevelUtils.cpp b/Engine/Source/Editor/UnrealEd/Private/EditorLevelUtils.cpp
--- a/Engine/Source/Editor/UnrealEd/Private/EditorLevelUtils.cpp
+++ b/Engine/Source/Editor/UnrealEd/Private/EditorLevelUtils.cpp
@@ -777,6 +777,14 @@ ULevelStreaming* UEditorLevelUtils::CreateNewStreamingLevelForWorld(UWorld& InWo
 	ULevel* NewLevel = nullptr;
 	if (bNewWorldSaved)
 	{
+		// Make sure to uninitialize the world since the level will be used as a streaming level
+		// This will make sure that the initialization order will be respected.
+		// One example is world partition initialization done inside ULevel::OnLevelLoaded.
+		if (NewLevelWorld->bIsWorldInitialized)
+		{
+			NewLevelWorld->CleanupWorld();
+		}
+
 		NewStreamingLevel = AddLevelToWorld(WorldToAddLevelTo, *NewPackageName, LevelStreamingClass, InTransform);
 		if (NewStreamingLevel != nullptr)
 		{
diff --git a/Engine/Source/Runtime/Engine/Private/LevelInstance/LevelInstanceSubsystem.cpp b/Engine/Source/Runtime/Engine/Private/LevelInstance/LevelInstanceSubsystem.cpp
--- a/Engine/Source/Runtime/Engine/Private/LevelInstance/LevelInstanceSubsystem.cpp
+++ b/Engine/Source/Runtime/Engine/Private/LevelInstance/LevelInstanceSubsystem.cpp
@@ -998,6 +998,20 @@ ILevelInstanceInterface* ULevelInstanceSubsystem::CreateLevelInstanceFrom(const
 	ULevelStreamingLevelInstanceEditor* LevelStreaming = nullptr;
 	{
 		const bool bIsPartitioned = (CreationParams.Type != ELevelInstanceCreationType::PackedLevelActor) && GetWorld()->IsPartitionedWorld();
+
+		// We want to properly setup the world partition prior to its initialization
+		FDelegateHandle PreWorldInit = FWorldDelegates::OnPreWorldInitialization.AddLambda([&bIsPartitioned](UWorld* World, const UWorld::InitializationValues IVS)
+		{
+			if (bIsPartitioned)
+			{
+				UWorldPartition* WorldPartition = World->GetWorldPartition();
+				if (ensure(WorldPartition))
+				{
+					WorldPartition->SetCanBeUsedByLevelInstance(true);
+				}
+			}
+		});
+
 		LevelStreaming = StaticCast<ULevelStreamingLevelInstanceEditor*>(EditorLevelUtils::CreateNewStreamingLevelForWorld(
 		*GetWorld(), ULevelStreamingLevelInstanceEditor::StaticClass(), CreationParams.UseExternalActors(), LevelFilename, &ActorsToMove, CreationParams.TemplateWorld, /*bUseSaveAs*/true, bIsPartitioned, [this, bIsPartitioned, &ActorsToMove](ULevel* InLevel)
 		{
@@ -1015,9 +1029,9 @@ ILevelInstanceInterface* ULevelInstanceSubsystem::CreateLevelInstanceFrom(const
 			{
 				UWorldPartition* WorldPartition = InLevel->GetWorldPartition();
 				check(WorldPartition);
-				
-				// Flag the world partition that it can be used by a Level Instance
-				WorldPartition->SetCanBeUsedByLevelInstance(true);
+
+				// Make sure that the world partition can be used by a Level Instance
+				check(WorldPartition->CanBeUsedByLevelInstance());
 
 				// Make sure new level's AWorldDataLayers contains all the necessary Data Layer Instances before moving actors
 				if (UDataLayerSubsystem* DataLayerSubsystem = UWorld::GetSubsystem<UDataLayerSubsystem>(GetWorld()))
@@ -1045,6 +1059,8 @@ ILevelInstanceInterface* ULevelInstanceSubsystem::CreateLevelInstanceFrom(const
 				check(!WorldPartition->IsStreamingEnabled());
 			}
 		}));
+
+		FWorldDelegates::OnPreWorldInitialization.Remove(PreWorldInit);
 	}
 
 	if (!LevelStreaming)
diff --git a/Engine/Source/Runtime/Engine/Private/WorldPartition/HLOD/HLODSubsystem.cpp b/Engine/Source/Runtime/Engine/Private/WorldPartition/HLOD/HLODSubsystem.cpp
--- a/Engine/Source/Runtime/Engine/Private/WorldPartition/HLOD/HLODSubsystem.cpp
+++ b/Engine/Source/Runtime/Engine/Private/WorldPartition/HLOD/HLODSubsystem.cpp
@@ -291,11 +291,11 @@ UHLODSubsystem::FCellData* UHLODSubsystem::GetCellData(const UWorldPartitionRunt
 	}
 
 	const UWorldPartition* WorldPartition = InCell->GetCellOwner()->GetOuterWorld()->GetWorldPartition();
-	check(WorldPartition && WorldPartition->IsStreamingEnabled());
-
+	check(WorldPartition);
 	FWorldPartitionHLODRuntimeData* WorldPartitionHLODRuntimeData = WorldPartitionsHLODRuntimeData.Find(WorldPartition);
 	if (WorldPartitionHLODRuntimeData)
 	{
+		check(WorldPartition->IsStreamingEnabled());
 		return WorldPartitionHLODRuntimeData->CellsData.Find(InCell->GetGuid());
 	}
 
@@ -305,11 +305,11 @@ UHLODSubsystem::FCellData* UHLODSubsystem::GetCellData(const UWorldPartitionRunt
 UHLODSubsystem::FCellData* UHLODSubsystem::GetCellData(AWorldPartitionHLOD* InWorldPartitionHLOD)
 {
 	const UWorldPartition* WorldPartition = FHLODSubsystem::GetWorldPartition(InWorldPartitionHLOD);
-	check(WorldPartition && WorldPartition->IsStreamingEnabled());
-
+	check(WorldPartition);
 	FWorldPartitionHLODRuntimeData* WorldPartitionHLODRuntimeData = WorldPartitionsHLODRuntimeData.Find(WorldPartition);
 	if (WorldPartitionHLODRuntimeData)
 	{
+		check(WorldPartition->IsStreamingEnabled());
 		const FGuid CellGuid = InWorldPartitionHLOD->GetSourceCellGuid();
 		return WorldPartitionHLODRuntimeData->CellsData.Find(CellGuid);
 	}
diff --git a/Engine/Source/Runtime/Engine/Private/WorldPartition/WorldPartition.cpp b/Engine/Source/Runtime/Engine/Private/WorldPartition/WorldPartition.cpp
--- a/Engine/Source/Runtime/Engine/Private/WorldPartition/WorldPartition.cpp
+++ b/Engine/Source/Runtime/Engine/Private/WorldPartition/WorldPartition.cpp
@@ -733,6 +733,8 @@ void UWorldPartition::Uninitialize()
 PRAGMA_DISABLE_DEPRECATION_WARNINGS
 		OnWorldPartitionUninitialized.Broadcast(this);
 PRAGMA_ENABLE_DEPRECATION_WARNINGS
+
+		World = nullptr;
 	}
 }
 
@@ -1475,8 +1477,9 @@ bool UWorldPartition::IsStreamingCompleted(EWorldPartitionRuntimeCellState Query
 
 void UWorldPartition::OnCellShown(const UWorldPartitionRuntimeCell* InCell)
 {
+	check(IsInitialized());
 	// Discard Cell's LevelStreaming notification when once WorldPartition is unitialized (can happen for instanced WorldPartition)
-	if (GetWorld()->IsGameWorld() && IsInitialized())
+	if (GetWorld()->IsGameWorld())
 	{
 		if (IsStreamingEnabled())
 		{
@@ -1488,8 +1491,9 @@ void UWorldPartition::OnCellShown(const UWorldPartitionRuntimeCell* InCell)
 
 void UWorldPartition::OnCellHidden(const UWorldPartitionRuntimeCell* InCell)
 {
+	check(IsInitialized());
 	// Discard Cell's LevelStreaming notification when once WorldPartition is unitialized (can happen for instanced WorldPartition)
-	if (GetWorld()->IsGameWorld() && IsInitialized())
+	if (GetWorld()->IsGameWorld())
 	{
 		if (IsStreamingEnabled())
 		{
diff --git a/Engine/Source/Runtime/Engine/Private/WorldPartition/WorldPartitionRuntimeLevelStreamingCell.cpp b/Engine/Source/Runtime/Engine/Private/WorldPartition/WorldPartitionRuntimeLevelStreamingCell.cpp
--- a/Engine/Source/Runtime/Engine/Private/WorldPartition/WorldPartitionRuntimeLevelStreamingCell.cpp
+++ b/Engine/Source/Runtime/Engine/Private/WorldPartition/WorldPartitionRuntimeLevelStreamingCell.cpp
@@ -477,7 +477,8 @@ void UWorldPartitionRuntimeLevelStreamingCell::Deactivate() const
 
 void UWorldPartitionRuntimeLevelStreamingCell::OnLevelShown()
 {
-	if (UWorldPartition* WorldPartition = GetCellOwner()->GetWorldPartition())
+	UWorldPartition* WorldPartition = GetCellOwner()->GetWorldPartition();
+	if (WorldPartition && WorldPartition->IsInitialized())
 	{
 		WorldPartition->OnCellShown(this);
 	}
@@ -485,7 +486,8 @@ void UWorldPartitionRuntimeLevelStreamingCell::OnLevelShown()
 
 void UWorldPartitionRuntimeLevelStreamingCell::OnLevelHidden()
 {
-	if (UWorldPartition* WorldPartition = GetCellOwner()->GetWorldPartition())
+	UWorldPartition* WorldPartition = GetCellOwner()->GetWorldPartition();
+	if (WorldPartition && WorldPartition->IsInitialized())
 	{
 		WorldPartition->OnCellHidden(this);
 	}
diff --git a/Engine/Source/Runtime/Engine/Private/WorldPartition/WorldPartitionSubsystem.cpp b/Engine/Source/Runtime/Engine/Private/WorldPartition/WorldPartitionSubsystem.cpp
--- a/Engine/Source/Runtime/Engine/Private/WorldPartition/WorldPartitionSubsystem.cpp
+++ b/Engine/Source/Runtime/Engine/Private/WorldPartition/WorldPartitionSubsystem.cpp
@@ -163,6 +163,11 @@ void UWorldPartitionSubsystem::Deinitialize()
 	Super::Deinitialize();
 }
 
+bool UWorldPartitionSubsystem::DoesSupportWorldType(const EWorldType::Type WorldType) const
+{
+	return Super::DoesSupportWorldType(WorldType) || WorldType == EWorldType::Inactive;
+}
+
 #if WITH_EDITOR
 void UWorldPartitionSubsystem::ForEachWorldPartition(TFunctionRef<bool(UWorldPartition*)> Func)
 {
diff --git a/Engine/Source/Runtime/Engine/Public/WorldPartition/WorldPartitionSubsystem.h b/Engine/Source/Runtime/Engine/Public/WorldPartition/WorldPartitionSubsystem.h
--- a/Engine/Source/Runtime/Engine/Public/WorldPartition/WorldPartitionSubsystem.h
+++ b/Engine/Source/Runtime/Engine/Public/WorldPartition/WorldPartitionSubsystem.h
@@ -20,6 +20,9 @@ class ENGINE_API UWorldPartitionSubsystem : public UTickableWorldSubsystem
 {
 	GENERATED_BODY()
 
+protected:
+	virtual bool DoesSupportWorldType(const EWorldType::Type WorldType) const override;
+
 public:
 	UWorldPartitionSubsystem();
 
