World Partition: moved more spatial hash specifric code into corresponding cell data object to reduce amount of duplicated code between different hash implementations.

diff --git a/Engine/Source/Runtime/Engine/Private/WorldPartition/WorldPartitionRuntimeCell.cpp b/Engine/Source/Runtime/Engine/Private/WorldPartition/WorldPartitionRuntimeCell.cpp
--- a/Engine/Source/Runtime/Engine/Private/WorldPartition/WorldPartitionRuntimeCell.cpp
+++ b/Engine/Source/Runtime/Engine/Private/WorldPartition/WorldPartitionRuntimeCell.cpp
@@ -62,49 +62,6 @@ void UWorldPartitionRuntimeCell::SetDataLayers(const TArray<const UDataLayerInst
 		DataLayers.Add(DataLayerInstance->GetDataLayerFName());
 	}
 	DataLayers.Sort([](const FName& A, const FName& B) { return A.ToString() < B.ToString(); });
-	UpdateDebugName();
-}
-
-void UWorldPartitionRuntimeCell::SetDebugInfo(int64 InCoordX, int64 InCoordY, int64 InCoordZ, FName InGridName)
-{
-	DebugInfo.CoordX = InCoordX;
-	DebugInfo.CoordY = InCoordY;
-	DebugInfo.CoordZ = InCoordZ;
-	DebugInfo.GridName = InGridName;
-	UpdateDebugName();
-}
-
-void UWorldPartitionRuntimeCell::UpdateDebugName()
-{
-	TStringBuilder<512> Builder;
-	Builder += DebugInfo.GridName.ToString();
-	Builder += TEXT("_");
-	Builder += FString::Printf(TEXT("L%d_X%d_Y%d"), DebugInfo.CoordZ, DebugInfo.CoordX, DebugInfo.CoordY);
-	int32 DataLayerCount = DataLayers.Num();
-
-	if (DataLayerCount > 0)
-	{
-		if (const UDataLayerSubsystem* DataLayerSubsystem = UWorld::GetSubsystem<UDataLayerSubsystem>(GetCellOwner()->GetOuterWorld()))
-		{
-			TArray<const UDataLayerInstance*> DataLayerObjects;
-			Builder += TEXT(" DL[");
-			for (int i = 0; i < DataLayerCount; ++i)
-			{
-				const UDataLayerInstance* DataLayer = DataLayerSubsystem->GetDataLayerInstance(DataLayers[i]);
-				DataLayerObjects.Add(DataLayer);
-				Builder += DataLayer->GetDataLayerShortName();
-				Builder += TEXT(",");
-			}
-			Builder += FString::Printf(TEXT("ID:%X]"), FDataLayersID(DataLayerObjects).GetHash());
-		}
-	}
-
-	if (ContentBundleID.IsValid())
-	{
-		Builder += FString::Printf(TEXT(" CB[%s]"), *UContentBundleDescriptor::GetContentBundleCompactString(ContentBundleID));
-	}
-
-	DebugInfo.Name = Builder.ToString();
 }
 
 void UWorldPartitionRuntimeCell::DumpStateLog(FHierarchicalLogArchive& Ar)
@@ -113,26 +70,6 @@ void UWorldPartitionRuntimeCell::DumpStateLog(FHierarchicalLogArchive& Ar)
 }
 #endif
 
-bool UWorldPartitionRuntimeCell::ShouldResetStreamingSourceInfo() const
-{
-	return RuntimeCellData->ShouldResetStreamingSourceInfo();
-}
-
-void UWorldPartitionRuntimeCell::ResetStreamingSourceInfo() const
-{
-	RuntimeCellData->ResetStreamingSourceInfo();
-}
-
-void UWorldPartitionRuntimeCell::AppendStreamingSourceInfo(const FWorldPartitionStreamingSource& Source, const FSphericalSector& SourceShape) const
-{
-	RuntimeCellData->AppendStreamingSourceInfo(Source, SourceShape);
-}
-
-void UWorldPartitionRuntimeCell::MergeStreamingSourceInfo() const
-{
-	RuntimeCellData->MergeStreamingSourceInfo();
-}
-
 int32 UWorldPartitionRuntimeCell::SortCompare(const UWorldPartitionRuntimeCell* Other, bool bCanUseSortingCache) const
 {
 	const int32 Comparison = RuntimeCellData->SortCompare(Other->RuntimeCellData, bCanUseSortingCache);
@@ -143,11 +80,10 @@ int32 UWorldPartitionRuntimeCell::SortCompare(const UWorldPartitionRuntimeCell*
 
 bool UWorldPartitionRuntimeCell::IsDebugShown() const
 {
-	return FWorldPartitionDebugHelper::IsDebugRuntimeHashGridShown(GetGridName()) &&
-		   FWorldPartitionDebugHelper::IsDebugStreamingStatusShown(GetStreamingStatus()) &&
+	return FWorldPartitionDebugHelper::IsDebugStreamingStatusShown(GetStreamingStatus()) &&
 	       FWorldPartitionDebugHelper::AreDebugDataLayersShown(DataLayers) &&
-		   FWorldPartitionDebugHelper::IsDebugCellNameShown(DebugInfo.Name) &&
-		   (FWorldPartitionDebugHelper::CanDrawContentBundles() || !ContentBundleID.IsValid());
+		   (FWorldPartitionDebugHelper::CanDrawContentBundles() || !ContentBundleID.IsValid()) &&
+			RuntimeCellData->IsDebugShown();
 }
 
 FLinearColor UWorldPartitionRuntimeCell::GetDebugStreamingPriorityColor() const
@@ -189,16 +125,6 @@ bool UWorldPartitionRuntimeCell::ContainsDataLayer(const UDataLayerInstance* Dat
 	return GetDataLayers().Contains(DataLayerInstance->GetDataLayerFName());
 }
 
-const FBox& UWorldPartitionRuntimeCell::GetContentBounds() const
-{
-	return RuntimeCellData->GetContentBounds();
-}
-
-FBox UWorldPartitionRuntimeCell::GetCellBounds() const
-{
-	return RuntimeCellData->GetCellBounds();
-}
-
 FName UWorldPartitionRuntimeCell::GetLevelPackageName() const
 { 
 #if WITH_EDITOR
diff --git a/Engine/Source/Runtime/Engine/Private/WorldPartition/WorldPartitionRuntimeCellDataSpatialHash.cpp b/Engine/Source/Runtime/Engine/Private/WorldPartition/WorldPartitionRuntimeCellDataSpatialHash.cpp
--- a/Engine/Source/Runtime/Engine/Private/WorldPartition/WorldPartitionRuntimeCellDataSpatialHash.cpp
+++ b/Engine/Source/Runtime/Engine/Private/WorldPartition/WorldPartitionRuntimeCellDataSpatialHash.cpp
@@ -1,6 +1,7 @@
 // Copyright Epic Games, Inc. All Rights Reserved.
 
 #include "WorldPartition/WorldPartitionRuntimeCellDataSpatialHash.h"
+#include "WorldPartition/WorldPartitionDebugHelper.h"
 #include "Engine/Level.h"
 
 #include UE_INLINE_GENERATED_CPP_BY_NAME(WorldPartitionRuntimeCellDataSpatialHash)
@@ -163,3 +164,15 @@ FBox UWorldPartitionRuntimeCellDataSpatialHash::GetCellBounds() const
 	Box.Max.Z = GetContentBounds().Max.Z;
 	return Box;
 }
+
+bool UWorldPartitionRuntimeCellDataSpatialHash::IsDebugShown() const
+{
+	return Super::IsDebugShown() && 
+		   FWorldPartitionDebugHelper::IsDebugRuntimeHashGridShown(GridName) &&
+		   FWorldPartitionDebugHelper::IsDebugCellNameShown(GetName());
+}
+
+FString UWorldPartitionRuntimeCellDataSpatialHash::GetDebugName() const
+{
+	return DebugName;
+}
\ No newline at end of file
diff --git a/Engine/Source/Runtime/Engine/Private/WorldPartition/WorldPartitionRuntimeSpatialHash.cpp b/Engine/Source/Runtime/Engine/Private/WorldPartition/WorldPartitionRuntimeSpatialHash.cpp
--- a/Engine/Source/Runtime/Engine/Private/WorldPartition/WorldPartitionRuntimeSpatialHash.cpp
+++ b/Engine/Source/Runtime/Engine/Private/WorldPartition/WorldPartitionRuntimeSpatialHash.cpp
@@ -1384,7 +1384,6 @@ bool UWorldPartitionRuntimeSpatialHash::CreateStreamingGrid(const FSpatialHashRu
 				StreamingCell->SetDataLayers(GridCellDataChunk.GetDataLayers());
 				StreamingCell->SetContentBundleUID(GridCellDataChunk.GetContentBundleID());
 				StreamingCell->SetPriority(RuntimeGrid.Priority);
-				StreamingCell->SetDebugInfo(CellGlobalCoords.X, CellGlobalCoords.Y, CellGlobalCoords.Z, CurrentStreamingGrid.GridName);
 				StreamingCell->SetClientOnlyVisible(CurrentStreamingGrid.bClientOnlyVisible);
 				StreamingCell->SetBlockOnSlowLoading(CurrentStreamingGrid.bBlockOnSlowStreaming);
 				StreamingCell->SetIsHLOD(RuntimeGrid.HLODLayer ? true : false);
@@ -1398,6 +1397,44 @@ bool UWorldPartitionRuntimeSpatialHash::CreateStreamingGrid(const FSpatialHashRu
 				CellDataSpatialHash->Level = Level;
 				CellDataSpatialHash->Position = FVector(Bounds.GetCenter(), 0.f);
 				CellDataSpatialHash->Extent = (float)CellExtent;
+				CellDataSpatialHash->GridName = RuntimeGrid.GridName;
+				
+				auto ComputeCellDebugName = [StreamingCell, CellDataSpatialHash, World, &CellGlobalCoords]()
+				{
+					TStringBuilder<512> Builder;
+					Builder += CellDataSpatialHash->GridName.ToString();
+					Builder += TEXT("_");
+					Builder += FString::Printf(TEXT("L%d_X%d_Y%d"), CellGlobalCoords.Z, CellGlobalCoords.X, CellGlobalCoords.Y);
+
+					const TArray<FName>& DataLayers = StreamingCell->GetDataLayers();
+					int32 DataLayerCount = DataLayers.Num();
+
+					if (DataLayerCount > 0)
+					{
+						if (const UDataLayerSubsystem* DataLayerSubsystem = UWorld::GetSubsystem<UDataLayerSubsystem>(World))
+						{
+							TArray<const UDataLayerInstance*> DataLayerObjects;
+							Builder += TEXT(" DL[");
+							for (int i = 0; i < DataLayerCount; ++i)
+							{
+								const UDataLayerInstance* DataLayer = DataLayerSubsystem->GetDataLayerInstance(DataLayers[i]);
+								DataLayerObjects.Add(DataLayer);
+								Builder += DataLayer->GetDataLayerShortName();
+								Builder += TEXT(",");
+							}
+							Builder += FString::Printf(TEXT("ID:%X]"), FDataLayersID(DataLayerObjects).GetHash());
+						}
+					}
+
+					if (StreamingCell->GetContentBundleID().IsValid())
+					{
+						Builder += FString::Printf(TEXT(" CB[%s]"), *UContentBundleDescriptor::GetContentBundleCompactString(StreamingCell->GetContentBundleID()));
+					}
+
+					return Builder.ToString();
+				};
+
+				CellDataSpatialHash->DebugName = ComputeCellDebugName();
 
 				PopulateRuntimeCell(StreamingCell, FilteredActors, OutPackagesToGenerate);
 
@@ -1675,14 +1712,13 @@ EWorldPartitionStreamingPerformance UWorldPartitionRuntimeSpatialHash::GetStream
 	const UWorldPartitionRuntimeCell* StreamingCell = Cast<const UWorldPartitionRuntimeCell>(Cell);
 	check(StreamingCell);
 
-	const FSpatialHashStreamingGrid* StreamingGrid = GetStreamingGridByName(StreamingCell->GetGridName());
+	const UWorldPartitionRuntimeCellDataSpatialHash* CellDataSpatialHash = CastChecked<UWorldPartitionRuntimeCellDataSpatialHash>(Cell->RuntimeCellData);
+	const FSpatialHashStreamingGrid* StreamingGrid = GetStreamingGridByName(CellDataSpatialHash->GridName);
 
 	if (ensure(StreamingGrid))
 	{
 		const float LoadingRange = StreamingGrid->LoadingRange;
 
-		UWorldPartitionRuntimeCellDataSpatialHash* CellDataSpatialHash = CastChecked<UWorldPartitionRuntimeCellDataSpatialHash>(Cell->RuntimeCellData);
-
 		if (CellDataSpatialHash->IsBlockingSource())
 		{
 			const double Distance = FMath::Sqrt(CellDataSpatialHash->GetMinSquareDistanceToBlockingSource()) - ((double)StreamingGrid->GetCellSize(CellDataSpatialHash->Level) / 2);
diff --git a/Engine/Source/Runtime/Engine/Private/WorldPartition/WorldPartitionStreamingPolicy.cpp b/Engine/Source/Runtime/Engine/Private/WorldPartition/WorldPartitionStreamingPolicy.cpp
--- a/Engine/Source/Runtime/Engine/Private/WorldPartition/WorldPartitionStreamingPolicy.cpp
+++ b/Engine/Source/Runtime/Engine/Private/WorldPartition/WorldPartitionStreamingPolicy.cpp
@@ -739,7 +739,7 @@ void UWorldPartitionStreamingPolicy::UpdateDebugCellsStreamingPriority(const TSe
 
 		for (const UWorldPartitionRuntimeCell* Cell : Cells)
 		{
-			Cell->RuntimeCellData->MergeStreamingSourceInfo();
+			Cell->MergeStreamingSourceInfo();
 		}
 
 		SortStreamingCellsByImportance(Cells);
diff --git a/Engine/Source/Runtime/Engine/Public/WorldPartition/WorldPartitionRuntimeCell.h b/Engine/Source/Runtime/Engine/Public/WorldPartition/WorldPartitionRuntimeCell.h
--- a/Engine/Source/Runtime/Engine/Public/WorldPartition/WorldPartitionRuntimeCell.h
+++ b/Engine/Source/Runtime/Engine/Public/WorldPartition/WorldPartitionRuntimeCell.h
@@ -174,20 +174,10 @@ class ENGINE_API UWorldPartitionRuntimeCell : public UObject, public IWorldParti
 	virtual void SetStreamingPriority(int32 InStreamingPriority) const PURE_VIRTUAL(UWorldPartitionRuntimeCell::SetStreamingPriority,);
 	virtual EStreamingStatus GetStreamingStatus() const { return LEVEL_Unloaded; }
 	virtual bool IsLoading() const { return false; }
-	virtual const FString& GetDebugName() const { return DebugInfo.Name; }
-	virtual bool IsDebugShown() const;
-	virtual FName GetGridName() const { return DebugInfo.GridName; }
 	bool GetClientOnlyVisible() const { return bClientOnlyVisible; }
 	virtual FGuid const& GetContentBundleID() const { return ContentBundleID; }
 	virtual TArray<FName> GetActors() const PURE_VIRTUAL(UWorldPartitionRuntimeCell::GetActors, return TArray<FName>(););
 
-	/** Caches information on streaming source that will be used later on to sort cell. */
-	bool ShouldResetStreamingSourceInfo() const;
-	void ResetStreamingSourceInfo() const;
-	void AppendStreamingSourceInfo(const FWorldPartitionStreamingSource& Source, const FSphericalSector& SourceShape) const;
-	void MergeStreamingSourceInfo() const;
-	int32 SortCompare(const UWorldPartitionRuntimeCell* Other, bool bCanUseSortingCache = true) const;
-
 	//~Begin IWorldPartitionCell Interface
 	virtual TArray<const UDataLayerInstance*> GetDataLayerInstances() const override;
 	virtual bool ContainsDataLayer(const UDataLayerAsset* DataLayerAsset) const override;
@@ -198,8 +188,6 @@ class ENGINE_API UWorldPartitionRuntimeCell : public UObject, public IWorldParti
 	{
 		return Algo::AnyOf(DataLayers, [&InDataLayers](const FName& DataLayer) { return InDataLayers.Contains(DataLayer); });
 	}
-	virtual const FBox& GetContentBounds() const override;
-	virtual FBox GetCellBounds() const override;
 	virtual FName GetLevelPackageName() const override;
 	virtual IWorldPartitionRuntimeCellOwner* GetCellOwner() const override { return CastChecked<IWorldPartitionRuntimeCellOwner>(GetOuter()); }
 	//~End IWorldPartitionCell Interface
@@ -211,7 +199,6 @@ class ENGINE_API UWorldPartitionRuntimeCell : public UObject, public IWorldParti
 	void SetClientOnlyVisible(bool bInClientOnlyVisible) { bClientOnlyVisible = bInClientOnlyVisible; }
 	void SetDataLayers(const TArray<const UDataLayerInstance*>& InDataLayerInstances);
 	void SetContentBundleUID(const FGuid& InContentBundleID) { ContentBundleID = InContentBundleID; }
-	void SetDebugInfo(int64 InCoordX, int64 InCoordY, int64 InCoordZ, FName InGridName);
 	void SetLevelPackageName(const FName& InLevelPackageName) { LevelPackageName = InLevelPackageName; }
 	
 	//~Begin IWorldPartitionCell Interface
@@ -249,10 +236,6 @@ class ENGINE_API UWorldPartitionRuntimeCell : public UObject, public IWorldParti
 protected:
 	FLinearColor GetDebugStreamingPriorityColor() const;
 
-#if WITH_EDITOR
-	void UpdateDebugName();
-#endif
-
 	UPROPERTY()
 	bool bIsAlwaysLoaded;
 
@@ -260,10 +243,6 @@ private:
 	UPROPERTY()
 	TArray<FName> DataLayers;
 
-	// Debug Info
-	UPROPERTY()
-	FWorldPartitionRuntimeCellDebugInfo DebugInfo;
-
 	// Custom Priority
 	UPROPERTY()
 	int32 Priority;
@@ -294,6 +273,18 @@ protected:
 #endif
 
 public:
+	//~Begin UWorldPartitionRuntimeCellData Proxy
+	bool ShouldResetStreamingSourceInfo() const { return RuntimeCellData->ShouldResetStreamingSourceInfo(); }
+	void ResetStreamingSourceInfo() const { RuntimeCellData->ResetStreamingSourceInfo(); }
+	void AppendStreamingSourceInfo(const FWorldPartitionStreamingSource& Source, const FSphericalSector& SourceShape) const { RuntimeCellData->AppendStreamingSourceInfo(Source, SourceShape); }
+	void MergeStreamingSourceInfo() const { RuntimeCellData->MergeStreamingSourceInfo(); }
+	int32 SortCompare(const UWorldPartitionRuntimeCell* Other, bool bCanUseSortingCache = true) const;
+	const FBox& GetContentBounds() const { return RuntimeCellData->GetContentBounds(); }
+	FBox GetCellBounds() const { return RuntimeCellData->GetCellBounds(); }
+	virtual bool IsDebugShown() const;
+	virtual FString GetDebugName() const { return RuntimeCellData->GetDebugName(); }
+	//~End UWorldPartitionRuntimeCellData Proxy
+
 	UPROPERTY()
 	TObjectPtr<UWorldPartitionRuntimeCellData> RuntimeCellData;
 };
diff --git a/Engine/Source/Runtime/Engine/Public/WorldPartition/WorldPartitionRuntimeCellData.h b/Engine/Source/Runtime/Engine/Public/WorldPartition/WorldPartitionRuntimeCellData.h
--- a/Engine/Source/Runtime/Engine/Public/WorldPartition/WorldPartitionRuntimeCellData.h
+++ b/Engine/Source/Runtime/Engine/Public/WorldPartition/WorldPartitionRuntimeCellData.h
@@ -24,6 +24,9 @@ class ENGINE_API UWorldPartitionRuntimeCellData : public UObject
 	virtual const FBox& GetContentBounds() const;
 	virtual FBox GetCellBounds() const;
 
+	virtual bool IsDebugShown() const { return true; }
+	virtual FString GetDebugName() const { return GetName(); }
+
 	static int32 StreamingSourceCacheEpoch;
 	static void DirtyStreamingSourceCacheEpoch() { ++StreamingSourceCacheEpoch; }
 
diff --git a/Engine/Source/Runtime/Engine/Public/WorldPartition/WorldPartitionRuntimeCellDataSpatialHash.h b/Engine/Source/Runtime/Engine/Public/WorldPartition/WorldPartitionRuntimeCellDataSpatialHash.h
--- a/Engine/Source/Runtime/Engine/Public/WorldPartition/WorldPartitionRuntimeCellDataSpatialHash.h
+++ b/Engine/Source/Runtime/Engine/Public/WorldPartition/WorldPartitionRuntimeCellDataSpatialHash.h
@@ -19,6 +19,8 @@ class ENGINE_API UWorldPartitionRuntimeCellDataSpatialHash : public UWorldPartit
 	virtual void MergeStreamingSourceInfo() const override;
 	virtual int32 SortCompare(const UWorldPartitionRuntimeCellData* InOther, bool bCanUseSortingCache = true) const override;
 	virtual FBox GetCellBounds() const override;
+	virtual bool IsDebugShown() const override;
+	virtual FString GetDebugName() const override;
 	//~End UWorldPartitionRuntimeCellData
 
 	bool IsBlockingSource() const { return bCachedIsBlockingSource; }
@@ -33,6 +35,12 @@ class ENGINE_API UWorldPartitionRuntimeCellDataSpatialHash : public UWorldPartit
 	UPROPERTY()
 	int32 Level;
 
+	UPROPERTY()
+	FName GridName;
+
+	UPROPERTY()
+	FString DebugName;
+
 private:
 	float ComputeSourceToCellAngleFactor(const FSphericalSector& SourceShape) const;
 
