Make sure to remove the stand alone flag on objects in actor packages when destroying a level.

diff --git a/Engine/Source/Runtime/Engine/Private/LevelStreamingGCHelper.cpp b/Engine/Source/Runtime/Engine/Private/LevelStreamingGCHelper.cpp
--- a/Engine/Source/Runtime/Engine/Private/LevelStreamingGCHelper.cpp
+++ b/Engine/Source/Runtime/Engine/Private/LevelStreamingGCHelper.cpp
@@ -172,7 +172,13 @@ void FLevelStreamingGCHelper::PrepareStreamedOutLevelForGC(ULevel* InLevel)
 			Packages.Add(Package, &bIsAlreadyInSet);
 			if (!bIsAlreadyInSet)
 			{
-				ForEachObjectWithPackage(Package, [](UObject* PackageObject) { PackageObject->MarkAsGarbage(); return true; }, true, RF_NoFlags, EInternalObjectFlags::Garbage);
+				ForEachObjectWithPackage(Package, [](UObject* PackageObject)
+				{
+					PackageObject->ClearFlags(RF_Standalone);
+					PackageObject->MarkAsGarbage();
+					return true;
+				}, true, RF_NoFlags, EInternalObjectFlags::Garbage);
+
 				Package->MarkAsGarbage();
 			}
 		}, true, RF_NoFlags, EInternalObjectFlags::Garbage);
diff --git a/Engine/Source/Runtime/Engine/Private/WorldPartition/WorldPartitionLevelStreamingDynamic.cpp b/Engine/Source/Runtime/Engine/Private/WorldPartition/WorldPartitionLevelStreamingDynamic.cpp
--- a/Engine/Source/Runtime/Engine/Private/WorldPartition/WorldPartitionLevelStreamingDynamic.cpp
+++ b/Engine/Source/Runtime/Engine/Private/WorldPartition/WorldPartitionLevelStreamingDynamic.cpp
@@ -125,6 +125,7 @@ void UWorldPartitionLevelStreamingDynamic::BeginDestroy()
 	if (IsValid(RuntimeLevel))
 	{
 		RuntimeLevel->OnCleanupLevel.Remove(OnCleanupLevelDelegateHandle);
+		OnCleanupLevelDelegateHandle.Reset();
 	}
 	PackageReferencer.RemoveReferences();
 	Super::BeginDestroy();
@@ -500,7 +501,9 @@ void UWorldPartitionLevelStreamingDynamic::OnCleanupLevel()
 
 	if (RuntimeLevel)
 	{
+		check(OnCleanupLevelDelegateHandle.IsValid());
 		RuntimeLevel->OnCleanupLevel.Remove(OnCleanupLevelDelegateHandle);
+		OnCleanupLevelDelegateHandle.Reset();
 
 		// If reusing levels is enabled, trash world partition level/actor packages
 		// as it won't be done by ULevel::CleanupLevel
@@ -535,6 +538,10 @@ void UWorldPartitionLevelStreamingDynamic::OnCleanupLevel()
 
 		RuntimeLevel = nullptr;
 	}
+	else
+	{
+		check(!OnCleanupLevelDelegateHandle.IsValid());
+	}
 }
 
 // Overriding base class to make sure the world outliner doesn't show runtime cell levels as root object.
