StaticMesh - AccumulateBounds: Fixed bug introduced in CL 23644351

diff --git a/Engine/Source/Runtime/Engine/Private/StaticMesh.cpp b/Engine/Source/Runtime/Engine/Private/StaticMesh.cpp
--- a/Engine/Source/Runtime/Engine/Private/StaticMesh.cpp
+++ b/Engine/Source/Runtime/Engine/Private/StaticMesh.cpp
@@ -3796,7 +3796,7 @@ static void AccumulateBounds(FBox& Bounds, const FStaticMeshLODResources& LODMod
 		const FIndexArrayView IndexBuffer = LODModel.IndexBuffer.GetArrayView();
 		for (int32 Index = FirstIndex; Index < LastIndex; ++Index)
 		{
-			Bounds += Transform.TransformPosition(FVector(LODModel.VertexBuffers.PositionVertexBuffer.VertexPosition(Index)));
+			Bounds += Transform.TransformPosition(FVector(LODModel.VertexBuffers.PositionVertexBuffer.VertexPosition(IndexBuffer[Index])));
 		}
 	}
 }
