Made the world partition object transactional when we call Modify() on it so that the changes are properly added to the Undo/Redo buffer.

diff --git a/Engine/Source/Runtime/Engine/Private/WorldPartition/WorldPartition.cpp b/Engine/Source/Runtime/Engine/Private/WorldPartition/WorldPartition.cpp
--- a/Engine/Source/Runtime/Engine/Private/WorldPartition/WorldPartition.cpp
+++ b/Engine/Source/Runtime/Engine/Private/WorldPartition/WorldPartition.cpp
@@ -1175,6 +1175,7 @@ void UWorldPartition::SetEnableStreaming(bool bInEnableStreaming)
 	{
 		const FScopedTransaction Transaction(LOCTEXT("EditorWorldPartitionSetEnableStreaming", "Set WorldPartition EnableStreaming"));
 
+		SetFlags(RF_Transactional);
 		Modify();
 		bEnableStreaming = bInEnableStreaming;
 		OnEnableStreamingChanged();
@@ -1195,6 +1196,7 @@ void UWorldPartition::SetCanBeUsedByLevelInstance(bool bInCanBeUsedByLevelInstan
 	{
 		const FScopedTransaction Transaction(LOCTEXT("EditorWorldPartitionCanBeUsedByLevelInstance", "Set WorldPartition CanBeUsedByLevelInstance"));
 
+		SetFlags(RF_Transactional);
 		Modify();
 		bCanBeUsedByLevelInstance = bInCanBeUsedByLevelInstance;
 		if (bCanBeUsedByLevelInstance)
