Content Bundle: move generic code into RuntimeHashExternalStreamingObjectBase to ease working on new hashes.

diff --git a/Engine/Source/Runtime/Engine/Private/WorldPartition/WorldPartitionRuntimeHash.cpp b/Engine/Source/Runtime/Engine/Private/WorldPartition/WorldPartitionRuntimeHash.cpp
--- a/Engine/Source/Runtime/Engine/Private/WorldPartition/WorldPartitionRuntimeHash.cpp
+++ b/Engine/Source/Runtime/Engine/Private/WorldPartition/WorldPartitionRuntimeHash.cpp
@@ -6,6 +6,8 @@
 #include "WorldPartition/WorldPartition.h"
 #include "WorldPartition/WorldPartitionActorDescView.h"
 #include "WorldPartition/WorldPartitionStreamingSource.h"
+#include "WorldPartition/WorldPartitionRuntimeLevelStreamingCell.h"
+#include "WorldPartition/WorldPartitionLevelStreamingDynamic.h"
 #if WITH_EDITOR
 #include "WorldPartition/Cook/WorldPartitionCookPackage.h"
 #endif
@@ -28,6 +30,40 @@ void URuntimeHashExternalStreamingObjectBase::ForEachStreamingCells(TFunctionRef
 	}
 }
 
+void URuntimeHashExternalStreamingObjectBase::OnStreamingObjectLoaded()
+{
+	bool bIsACookedObject = !CellToLevelStreamingPackage.IsEmpty();
+	if (bIsACookedObject)
+	{
+		// Cooked streaming object's Cells do not have LevelStreaming and are outered to the streaming object.
+		// Create their level streaming now and re-outered to the runtime hash.
+		ForEachStreamingCells([this](UWorldPartitionRuntimeCell& Cell)
+		{
+			UWorldPartitionRuntimeLevelStreamingCell* RuntimeCell = CastChecked<UWorldPartitionRuntimeLevelStreamingCell>(&Cell);
+
+			FName LevelStreamingPackage = CellToLevelStreamingPackage.FindChecked(RuntimeCell->GetFName());
+			RuntimeCell->Rename(nullptr, GetOuterWorld()->GetWorldPartition()->RuntimeHash);
+			RuntimeCell->CreateAndSetLevelStreaming(*LevelStreamingPackage.ToString());
+		});
+	}
+}
+
+#if WITH_EDITOR
+void URuntimeHashExternalStreamingObjectBase::PopulateGeneratorPackageForCook()
+{
+	ForEachStreamingCells([this](UWorldPartitionRuntimeCell& Cell)
+	{
+		UWorldPartitionRuntimeLevelStreamingCell* RuntimeCell = CastChecked<UWorldPartitionRuntimeLevelStreamingCell>(&Cell);
+		UWorldPartitionLevelStreamingDynamic* LevelStreamingDynamic = RuntimeCell->GetLevelStreaming();
+		CellToLevelStreamingPackage.Add(RuntimeCell->GetFName(), LevelStreamingDynamic->PackageNameToLoad);
+
+		// Level streaming are outered to the world and would not be saved within the ExternalStreamingObject.
+		// Do not save them, instead they will be created once the external streaming object is loaded at runtime. 
+		LevelStreamingDynamic->SetFlags(RF_Transient);
+	});
+}
+#endif
+
 UWorldPartitionRuntimeHash::UWorldPartitionRuntimeHash(const FObjectInitializer& ObjectInitializer)
 	: Super(ObjectInitializer)
 {}
diff --git a/Engine/Source/Runtime/Engine/Private/WorldPartition/WorldPartitionRuntimeSpatialHash.cpp b/Engine/Source/Runtime/Engine/Private/WorldPartition/WorldPartitionRuntimeSpatialHash.cpp
--- a/Engine/Source/Runtime/Engine/Private/WorldPartition/WorldPartitionRuntimeSpatialHash.cpp
+++ b/Engine/Source/Runtime/Engine/Private/WorldPartition/WorldPartitionRuntimeSpatialHash.cpp
@@ -959,42 +959,6 @@ ASpatialHashRuntimeGridInfo::ASpatialHashRuntimeGridInfo(const FObjectInitialize
 #endif
 }
 
-void URuntimeSpatialHashExternalStreamingObject::OnStreamingObjectLoaded()
-{
-	bool bIsACookedObject = !CellToLevelStreamingPackage.IsEmpty();
-	if (bIsACookedObject)
-	{
-		// Cooked streaming object's Cells do not have LevelStreaming and are outered to the streaming object.
-		// Create their level streaming now and re-outered to the runtime hash.
-		ForEachStreamingCells([this](UWorldPartitionRuntimeCell& Cell)
-		{
-			UWorldPartitionRuntimeLevelStreamingCell* RuntimeCell = CastChecked<UWorldPartitionRuntimeLevelStreamingCell>(&Cell);
-
-			FName LevelStreamingPackage = CellToLevelStreamingPackage.FindChecked(RuntimeCell->GetFName());
-			RuntimeCell->Rename(nullptr, GetOuterWorld()->GetWorldPartition()->RuntimeHash);
-			RuntimeCell->CreateAndSetLevelStreaming(*LevelStreamingPackage.ToString());
-		});
-	}
-}
-
-#if WITH_EDITOR
-void URuntimeSpatialHashExternalStreamingObject::PopulateGeneratorPackageForCook()
-{
-	ForEachStreamingCells([this](UWorldPartitionRuntimeCell& Cell)
-	{
-		UWorldPartitionRuntimeLevelStreamingCell* RuntimeCell = CastChecked<UWorldPartitionRuntimeLevelStreamingCell>(&Cell);
-		RuntimeCell->Rename(nullptr, this);
-
-		UWorldPartitionLevelStreamingDynamic* LevelStreamingDynamic = RuntimeCell->GetLevelStreaming();
-		CellToLevelStreamingPackage.Add(RuntimeCell->GetFName(), LevelStreamingDynamic->PackageNameToLoad);
-
-		// Level streaming are outered to the world and would not be saved within the ExternalStreamingObject.
-		// Do not save them, instead they will be created once the external streaming object is loaded at runtime. 
-		LevelStreamingDynamic->SetFlags(RF_Transient);
-	});
-}
-#endif
-
 UWorldPartitionRuntimeSpatialHash::UWorldPartitionRuntimeSpatialHash(const FObjectInitializer& ObjectInitializer)
 	: Super(ObjectInitializer)
 #if WITH_EDITORONLY_DATA
diff --git a/Engine/Source/Runtime/Engine/Public/WorldPartition/WorldPartitionRuntimeHash.h b/Engine/Source/Runtime/Engine/Public/WorldPartition/WorldPartitionRuntimeHash.h
--- a/Engine/Source/Runtime/Engine/Public/WorldPartition/WorldPartitionRuntimeHash.h
+++ b/Engine/Source/Runtime/Engine/Public/WorldPartition/WorldPartitionRuntimeHash.h
@@ -43,26 +43,29 @@ public:
 	}
 
 	//~ Begin IWorldPartitionRuntimeCellOwner Interface
-	virtual UWorld* GetOwningWorld() const override { return OwningWorld.Get(); }
-	virtual UWorld* GetOuterWorld() const override { return OuterWorld.Get(); }
+	virtual UWorld* GetOwningWorld() const override final { return OwningWorld.Get(); }
+	virtual UWorld* GetOuterWorld() const override final { return OuterWorld.Get(); }
 	//~ End IWorldPartitionRuntimeCellOwner Interface
 
-	virtual class UWorld* GetWorld() const override { return GetOwningWorld(); }
+	virtual class UWorld* GetWorld() const override final { return GetOwningWorld(); }
 
-	virtual void ForEachStreamingCells(TFunctionRef<void(UWorldPartitionRuntimeCell&)> Func);
+	void ForEachStreamingCells(TFunctionRef<void(UWorldPartitionRuntimeCell&)> Func);
 	
-	virtual void OnStreamingObjectLoaded() {};
+	void OnStreamingObjectLoaded();
 
 #if WITH_EDITOR
-	virtual void PopulateGeneratorPackageForCook() {};
+	void PopulateGeneratorPackageForCook();
 #endif
 
-private:
+protected:
 	UPROPERTY();
 	TSoftObjectPtr<UWorld> OwningWorld;
 
 	UPROPERTY();
 	TSoftObjectPtr<UWorld> OuterWorld;
+
+	UPROPERTY();
+	TMap<FName, FName> CellToLevelStreamingPackage;
 };
 
 struct FWorldPartitionQueryCache
diff --git a/Engine/Source/Runtime/Engine/Public/WorldPartition/WorldPartitionRuntimeSpatialHash.h b/Engine/Source/Runtime/Engine/Public/WorldPartition/WorldPartitionRuntimeSpatialHash.h
--- a/Engine/Source/Runtime/Engine/Public/WorldPartition/WorldPartitionRuntimeSpatialHash.h
+++ b/Engine/Source/Runtime/Engine/Public/WorldPartition/WorldPartitionRuntimeSpatialHash.h
@@ -234,20 +234,8 @@ class URuntimeSpatialHashExternalStreamingObject : public URuntimeHashExternalSt
 	GENERATED_BODY()
 
 public:
-	//~ Begin URuntimeHashExternalStreamingObjectBase Interface
-	virtual void OnStreamingObjectLoaded() override;
-
-#if WITH_EDITOR
-	virtual void PopulateGeneratorPackageForCook() override;
-#endif
-	//~ End URuntimeHashExternalStreamingObjectBase Interface
-
 	UPROPERTY();
 	TArray<FSpatialHashStreamingGrid> StreamingGrids;
-
-private:
-	UPROPERTY();
-	TMap<FName, FName> CellToLevelStreamingPackage;
 };
 
 UCLASS()
