World Partition - Fixed crash in FWorldPartitionRuntimeCell::GetDataLayerInstances() when trying to access a null UWorld at PIE shutdown

diff --git a/Engine/Source/Runtime/Engine/Private/WorldPartition/WorldPartitionRuntimeCell.cpp b/Engine/Source/Runtime/Engine/Private/WorldPartition/WorldPartitionRuntimeCell.cpp
--- a/Engine/Source/Runtime/Engine/Private/WorldPartition/WorldPartitionRuntimeCell.cpp
+++ b/Engine/Source/Runtime/Engine/Private/WorldPartition/WorldPartitionRuntimeCell.cpp
@@ -99,7 +99,7 @@ FLinearColor UWorldPartitionRuntimeCell::GetDebugStreamingPriorityColor() const
 
 TArray<const UDataLayerInstance*> UWorldPartitionRuntimeCell::GetDataLayerInstances() const
 {
-	if (const UDataLayerSubsystem* DataLayerSubsystem = GetWorld()->GetSubsystem<UDataLayerSubsystem>())
+	if (const UDataLayerSubsystem* DataLayerSubsystem = UWorld::GetSubsystem<UDataLayerSubsystem>(GetWorld()))
 	{
 		return DataLayerSubsystem->GetDataLayerInstances(GetDataLayers());
 	}
@@ -109,7 +109,7 @@ TArray<const UDataLayerInstance*> UWorldPartitionRuntimeCell::GetDataLayerInstan
 
 bool UWorldPartitionRuntimeCell::ContainsDataLayer(const UDataLayerAsset* DataLayerAsset) const
 {
-	if (const UDataLayerSubsystem* DataLayerSubsystem = GetWorld()->GetSubsystem<UDataLayerSubsystem>())
+	if (const UDataLayerSubsystem* DataLayerSubsystem = UWorld::GetSubsystem<UDataLayerSubsystem>(GetWorld()))
 	{
 		if (const UDataLayerInstance* DataLayerInstance = DataLayerSubsystem->GetDataLayerInstance(DataLayerAsset))
 		{
