Fixed wireframe colors on shipping builds.

#p4v-preflight-copy 23872299

diff --git a/Engine/Source/Runtime/Engine/Private/PrimitiveDrawingUtils.cpp b/Engine/Source/Runtime/Engine/Private/PrimitiveDrawingUtils.cpp
--- a/Engine/Source/Runtime/Engine/Private/PrimitiveDrawingUtils.cpp
+++ b/Engine/Source/Runtime/Engine/Private/PrimitiveDrawingUtils.cpp
@@ -1565,7 +1565,7 @@ void ApplyViewModeOverrides(
 	FMeshElementCollector& Collector
 	)
 {
-#if !(UE_BUILD_SHIPPING || UE_BUILD_TEST)
+#if !UE_BUILD_TEST
 
 	// If debug viewmodes are not allowed, skip all of the debug viewmode handling.
 	if (!AllowDebugViewmodes())
diff --git a/Engine/Source/Runtime/Engine/Public/PrimitiveSceneProxy.h b/Engine/Source/Runtime/Engine/Public/PrimitiveSceneProxy.h
--- a/Engine/Source/Runtime/Engine/Public/PrimitiveSceneProxy.h
+++ b/Engine/Source/Runtime/Engine/Public/PrimitiveSceneProxy.h
@@ -745,7 +745,7 @@ public:
 	ENGINE_API void SetUsedMaterialForVerification(const TArray<UMaterialInterface*>& InUsedMaterialsForVerification);
 #endif
 
-#if !(UE_BUILD_SHIPPING || UE_BUILD_TEST)
+#if !UE_BUILD_TEST
 	inline FLinearColor GetWireframeColor() const { return WireframeColor; }
 	inline FLinearColor GetLevelColor() const { return LevelColor; }
 	inline FLinearColor GetPropertyColor() const { return PropertyColor; }
@@ -1060,7 +1060,7 @@ protected:
 	ENGINE_API void EnableGPUSceneSupportFlags();
 
 private:
-#if !(UE_BUILD_SHIPPING || UE_BUILD_TEST)
+#if !UE_BUILD_TEST
 	FLinearColor WireframeColor;
 	FLinearColor LevelColor;
 	FLinearColor PropertyColor;
