Revert crash fix in AnimSubsystem_SharedLinkedAnimLayers.cpp now that a tentative fix for the source issue has been made. Split the assert in two so that we know which part fails if it happens again

diff --git a/Engine/Source/Runtime/Engine/Private/Animation/AnimSubsystem_SharedLinkedAnimLayers.cpp b/Engine/Source/Runtime/Engine/Private/Animation/AnimSubsystem_SharedLinkedAnimLayers.cpp
--- a/Engine/Source/Runtime/Engine/Private/Animation/AnimSubsystem_SharedLinkedAnimLayers.cpp
+++ b/Engine/Source/Runtime/Engine/Private/Animation/AnimSubsystem_SharedLinkedAnimLayers.cpp
@@ -118,10 +118,7 @@ FAnimSubsystem_SharedLinkedAnimLayers* FAnimSubsystem_SharedLinkedAnimLayers::Ge
 		return nullptr;
 	}
 #endif
-	if (!SkelMesh)
-	{
-		return nullptr;
-	}
+	check(SkelMesh);
 	check(SkelMesh->GetAnimInstance());
 	return SkelMesh->GetAnimInstance()->FindSubsystem<FAnimSubsystem_SharedLinkedAnimLayers>();
 }
