Enable r.Shaders.RemoveDeadCode by default for all back-ends except Metal and OpenGL

* Significantly reduces the size of shader debug data and improves shader compilation time (reducing parsing cost in the native shader compiler)
* Remove "experimental" label for the feature

diff --git a/Engine/Source/Developer/Apple/MetalShaderFormat/Private/MetalShaderCompiler.cpp b/Engine/Source/Developer/Apple/MetalShaderFormat/Private/MetalShaderCompiler.cpp
--- a/Engine/Source/Developer/Apple/MetalShaderFormat/Private/MetalShaderCompiler.cpp
+++ b/Engine/Source/Developer/Apple/MetalShaderFormat/Private/MetalShaderCompiler.cpp
@@ -1104,7 +1104,7 @@ void CompileShader_Metal(const FShaderCompilerInput& _Input,FShaderCompilerOutpu
 	// Process TEXT macro.
 	TransformStringIntoCharacterArray(PreprocessedShader);
 
-	// Run the experimental shader minifier
+	// Run the shader minifier
 	#if UE_METAL_SHADER_COMPILER_ALLOW_DEAD_CODE_REMOVAL
 	if (Input.Environment.CompilerFlags.Contains(CFLAG_RemoveDeadCode))
 	{
diff --git a/Engine/Source/Developer/ShaderFormatOpenGL/Private/OpenGLShaderCompiler.cpp b/Engine/Source/Developer/ShaderFormatOpenGL/Private/OpenGLShaderCompiler.cpp
--- a/Engine/Source/Developer/ShaderFormatOpenGL/Private/OpenGLShaderCompiler.cpp
+++ b/Engine/Source/Developer/ShaderFormatOpenGL/Private/OpenGLShaderCompiler.cpp
@@ -3387,7 +3387,7 @@ void FOpenGLFrontend::CompileShader(const FShaderCompilerInput& Input, FShaderCo
 	// Process TEXT macro.
 	TransformStringIntoCharacterArray(PreprocessedShader);
 
-	// Run the experimental shader minifier
+	// Run the shader minifier
 #if UE_OPENGL_SHADER_COMPILER_ALLOW_DEAD_CODE_REMOVAL
 	if (Input.Environment.CompilerFlags.Contains(CFLAG_RemoveDeadCode))
 	{
diff --git a/Engine/Source/Developer/VulkanShaderFormat/Private/VulkanShaderCompiler.cpp b/Engine/Source/Developer/VulkanShaderFormat/Private/VulkanShaderCompiler.cpp
--- a/Engine/Source/Developer/VulkanShaderFormat/Private/VulkanShaderCompiler.cpp
+++ b/Engine/Source/Developer/VulkanShaderFormat/Private/VulkanShaderCompiler.cpp
@@ -2591,7 +2591,7 @@ void DoCompileVulkanShader(const FShaderCompilerInput& Input, FShaderCompilerOut
 	// Process TEXT macro.
 	TransformStringIntoCharacterArray(PreprocessedShaderSource);
 
-	// Run the experimental shader minifier
+	// Run the shader minifier
 	#if UE_VULKAN_SHADER_COMPILER_ALLOW_DEAD_CODE_REMOVAL
 	if (Input.Environment.CompilerFlags.Contains(CFLAG_RemoveDeadCode))
 	{
diff --git a/Engine/Source/Developer/Windows/ShaderFormatD3D/Private/D3DShaderCompiler.cpp b/Engine/Source/Developer/Windows/ShaderFormatD3D/Private/D3DShaderCompiler.cpp
--- a/Engine/Source/Developer/Windows/ShaderFormatD3D/Private/D3DShaderCompiler.cpp
+++ b/Engine/Source/Developer/Windows/ShaderFormatD3D/Private/D3DShaderCompiler.cpp
@@ -1231,7 +1231,7 @@ void CompileD3DShader(const FShaderCompilerInput& Input, FShaderCompilerOutput&
 
 	TArray<FString> FilteredErrors;
 
-	// Run the experimental shader minifier
+	// Run the shader minifier
 	#if UE_D3D_SHADER_COMPILER_ALLOW_DEAD_CODE_REMOVAL
 	if (Input.Environment.CompilerFlags.Contains(CFLAG_RemoveDeadCode))
 	{
diff --git a/Engine/Source/Runtime/Engine/Private/ShaderCompiler/ShaderCompiler.cpp b/Engine/Source/Runtime/Engine/Private/ShaderCompiler/ShaderCompiler.cpp
--- a/Engine/Source/Runtime/Engine/Private/ShaderCompiler/ShaderCompiler.cpp
+++ b/Engine/Source/Runtime/Engine/Private/ShaderCompiler/ShaderCompiler.cpp
@@ -1133,11 +1133,11 @@ static TAutoConsoleVariable<int32> CVarShadersValidation(
 
 static TAutoConsoleVariable<int32> CVarShadersRemoveDeadCode(
 	TEXT("r.Shaders.RemoveDeadCode"),
-	0,
-	TEXT("EXPERIMENTAL: Run a preprocessing step that removes unreferenced code before compiling shaders.\n")
+	1,
+	TEXT("Run a preprocessing step that removes unreferenced code before compiling shaders.\n")
 	TEXT("This can improve the compilation speed for shaders which include many large utility headers.\n")
-	TEXT("\t0: Keep all input source code (Default).\n")
-	TEXT("\t1: Remove unreferenced code before compilation\n"),
+	TEXT("\t0: Keep all input source code.\n")
+	TEXT("\t1: Remove unreferenced code before compilation (Default)\n"),
 	ECVF_ReadOnly);
 
 // note: this cvar is intended to be short-lived; when removing be sure to remove the associated code from
