Updates to IActiveSoundUpdateInterface
+ Adding OnNotifyAddActiveSound (to match existing OnNotifyPendingDelete)
+ Changing interface to have empty default implementations (instead of pure virtual)
+ Hooking new function into AudioDevice

diff --git a/Engine/Source/Runtime/Engine/Private/AudioDevice.cpp b/Engine/Source/Runtime/Engine/Private/AudioDevice.cpp
--- a/Engine/Source/Runtime/Engine/Private/AudioDevice.cpp
+++ b/Engine/Source/Runtime/Engine/Private/AudioDevice.cpp
@@ -3605,7 +3605,16 @@ void FAudioDevice::ApplyInteriorSettings(FActiveSound& ActiveSound, FSoundParseP
 	});
 }
 
-void FAudioDevice::NotifyPendingDelete(FActiveSound& ActiveSound) const
+void FAudioDevice::NotifyAddActiveSound(FActiveSound& ActiveSound) const
+{
+	SubsystemCollection.ForEachSubsystem<IActiveSoundUpdateInterface>([&ActiveSound](IActiveSoundUpdateInterface* ActiveSoundUpdate)
+	{
+		ActiveSoundUpdate->OnNotifyAddActiveSound(ActiveSound);
+		return true;
+	});
+}
+
+void FAudioDevice::NotifyPendingDeleteInternal(FActiveSound& ActiveSound) const
 {
 	SubsystemCollection.ForEachSubsystem<IActiveSoundUpdateInterface>([&ActiveSound](IActiveSoundUpdateInterface* ActiveSoundUpdate)
 	{
@@ -5206,6 +5215,7 @@ void FAudioDevice::AddNewActiveSoundInternal(const FActiveSound& InNewActiveSoun
 
 	InitSoundParams(*ActiveSound, MoveTemp(InDefaultParams));
 	ActiveSounds.Add(ActiveSound);
+	NotifyAddActiveSound(*ActiveSound);
 
 	if (ActiveSound->GetAudioComponentID() > 0)
 	{
@@ -5412,7 +5422,7 @@ void FAudioDevice::ProcessingPendingActiveSoundStops(bool bForceDelete)
 					Transmitter->Reset();
 				}
 
-				NotifyPendingDelete(*ActiveSound);
+				NotifyPendingDeleteInternal(*ActiveSound);
 				delete ActiveSound;
 			}
 		}
@@ -5472,7 +5482,7 @@ void FAudioDevice::ProcessingPendingActiveSoundStops(bool bForceDelete)
 					Transmitter->Reset();
 				}
 
-				NotifyPendingDelete(*ActiveSound);
+				NotifyPendingDeleteInternal(*ActiveSound);
 
 				// Remove from the list of pending sounds to stop
 				PendingSoundsToStop.Remove(ActiveSound);
diff --git a/Engine/Source/Runtime/Engine/Public/ActiveSoundUpdateInterface.h b/Engine/Source/Runtime/Engine/Public/ActiveSoundUpdateInterface.h
--- a/Engine/Source/Runtime/Engine/Public/ActiveSoundUpdateInterface.h
+++ b/Engine/Source/Runtime/Engine/Public/ActiveSoundUpdateInterface.h
@@ -24,24 +24,35 @@ public:
 
 	/**
 	 * Gathers interior data that can affect the active sound.  Non-const as this step can be used to track state about the sound on the implementing object
+	 * NOTE! Called on the AudioThread
 	 *
 	 * @param ActiveSound	The active sound affected
 	 * @param ParseParams	The parameters to apply to the wave instances
 	 */
-	virtual void GatherInteriorData(const FActiveSound& ActiveSound, FSoundParseParameters& ParseParams) = 0;
+	virtual void GatherInteriorData(const FActiveSound& ActiveSound, FSoundParseParameters& ParseParams) {}
 
 	/**
 	 * Applies interior data previously collected to the active sound and parse parameters.
+	 * NOTE! Called on the AudioThread
 	 *
 	 * @param ActiveSound	The active sound affected
 	 * @param ParseParams	The parameters to apply to the wave instances
 	 */
-	virtual void ApplyInteriorSettings(const FActiveSound& ActiveSound, FSoundParseParameters& ParseParams) = 0;
+	virtual void ApplyInteriorSettings(const FActiveSound& ActiveSound, FSoundParseParameters& ParseParams) {}
+
+	/**
+	 * Called when an active sound is being added to the audio engine
+	 * NOTE! Called on the AudioThread
+	 *
+	 * @param ActiveSound	The active sound being added
+	 */
+	virtual void OnNotifyAddActiveSound(FActiveSound& ActiveSound) {}
 
 	/**
 	 * Called when the active sound is being removed from the audio engine
+	 * NOTE! Called on the AudioThread
 	 *
 	 * @param ActiveSound	The active sound
 	 */
-	virtual void OnNotifyPendingDelete(const FActiveSound& ActiveSound) = 0;
+	virtual void OnNotifyPendingDelete(const FActiveSound& ActiveSound) {}
 };
diff --git a/Engine/Source/Runtime/Engine/Public/AudioDevice.h b/Engine/Source/Runtime/Engine/Public/AudioDevice.h
--- a/Engine/Source/Runtime/Engine/Public/AudioDevice.h
+++ b/Engine/Source/Runtime/Engine/Public/AudioDevice.h
@@ -841,10 +841,23 @@ public:
 	 */
 	void ApplyInteriorSettings(FActiveSound& ActiveSound, FSoundParseParameters& ParseParams) const;
 
+	/**
+	 * Notifies subsystems an active sound is about to be deleted (called on audio thread) - Deprecated, see NotifyPendingDeleteInternal
+	 */
+	UE_DEPRECATED(5.3, "NotifyPending is deprecated in public scope. Use IActiveSoundUpdateInterface::OnNotifyPendingDelete instead.")
+	void NotifyPendingDelete(FActiveSound& ActiveSound) const {}
+
+protected:
+
+	/**
+	 * Notifies subsystems an active sound has been added (called on audio thread)
+	 */
+	void NotifyAddActiveSound(FActiveSound& ActiveSound) const;
+
 	/**
 	 * Notifies subsystems an active sound is about to be deleted (called on audio thread)
 	 */
-	void NotifyPendingDelete(FActiveSound& ActiveSound) const;
+	void NotifyPendingDeleteInternal(FActiveSound& ActiveSound) const;
 
 public:
 
