Clamp quality scalar calculation, to allow for 0% for minimim settings.

diff --git a/Engine/Source/Runtime/Engine/Private/AudioDerivedData.cpp b/Engine/Source/Runtime/Engine/Private/AudioDerivedData.cpp
--- a/Engine/Source/Runtime/Engine/Private/AudioDerivedData.cpp
+++ b/Engine/Source/Runtime/Engine/Private/AudioDerivedData.cpp
@@ -1440,15 +1440,10 @@ static void CookSimpleWave(const FAudioCookInputs& Inputs, TArray<uint8>& Output
 	// Compression Quality
 	FSoundQualityInfo QualityInfo = { 0 };
 	float ModifiedCompressionQuality = (float)Inputs.CompressionQuality * Inputs.CompressionQualityModifier;
-	if (ModifiedCompressionQuality >= 1.0f)
-	{
-		QualityInfo.Quality = FMath::FloorToInt(ModifiedCompressionQuality);
-		UE_CLOG(Inputs.CompressionQuality != QualityInfo.Quality, LogAudioDerivedData, Display, TEXT("Compression Quality for %s will be modified from %d to %d."), *Inputs.SoundFullName, Inputs.CompressionQuality, QualityInfo.Quality);
-	}
-	else
-	{
-		QualityInfo.Quality = Inputs.CompressionQuality;
-	}
+
+	QualityInfo.Quality = FMath::Clamp<int32>(FMath::FloorToInt(ModifiedCompressionQuality), 1, 100);
+
+	UE_CLOG(Inputs.CompressionQuality != QualityInfo.Quality, LogAudioDerivedData, Display, TEXT("Compression Quality for %s will be modified from %d to %d."), *Inputs.SoundFullName, Inputs.CompressionQuality, QualityInfo.Quality);
 
 	QualityInfo.NumChannels = NumChannels;
 	QualityInfo.SampleRate = WaveSampleRate;
