Added a way to upgrade existing sparse class data

Objects using sparse data can now override the ConformSparseClassDataStruct function to perform a custom upgrade step when loading a BPGC using a sparse data struct that no longer matches its archetype. If ConformSparseClassDataStruct doesn't explicitly handle the conversion then we fallback to trying to copy common properties if the structs are related.

diff --git a/Engine/Source/Editor/KismetCompiler/Private/KismetCompiler.cpp b/Engine/Source/Editor/KismetCompiler/Private/KismetCompiler.cpp
--- a/Engine/Source/Editor/KismetCompiler/Private/KismetCompiler.cpp
+++ b/Engine/Source/Editor/KismetCompiler/Private/KismetCompiler.cpp
@@ -5263,6 +5263,9 @@ void FKismetCompilerContext::PostCDOCompiled(const UObject::FPostCDOCompiledCont
 		return;
 	}
 
+	// Conform any sparse class data now that the CDO is ready
+	NewClass->ConformSparseClassData(NewClass->ClassDefaultObject);
+
 	// Notify the CDO that it has finished compiling
 	NewClass->ClassDefaultObject->PostCDOCompiled(Context);
 	if (!Context.bIsSkeletonOnly)
diff --git a/Engine/Source/Runtime/CoreUObject/Public/UObject/Object.h b/Engine/Source/Runtime/CoreUObject/Public/UObject/Object.h
--- a/Engine/Source/Runtime/CoreUObject/Public/UObject/Object.h
+++ b/Engine/Source/Runtime/CoreUObject/Public/UObject/Object.h
@@ -685,6 +685,7 @@ public:
 
 #if WITH_EDITOR
 	virtual void MoveDataToSparseClassDataStruct() const {}
+	virtual bool ConformSparseClassDataStruct(UScriptStruct* SourceSparseClassDataStruct, void* SourceSparseClassData) const { return false; }
 #endif
 
 #if WITH_ENGINE
diff --git a/Engine/Source/Runtime/Engine/Classes/Animation/AnimBlueprintGeneratedClass.h b/Engine/Source/Runtime/Engine/Classes/Animation/AnimBlueprintGeneratedClass.h
--- a/Engine/Source/Runtime/Engine/Classes/Animation/AnimBlueprintGeneratedClass.h
+++ b/Engine/Source/Runtime/Engine/Classes/Animation/AnimBlueprintGeneratedClass.h
@@ -630,6 +630,11 @@ public:
 	static FName GetConstantsStructName();
 	static FName GetMutablesStructName();
 	
+#if WITH_EDITOR
+	virtual void PrepareToConformSparseClassData() override;
+	virtual void ConformSparseClassData(UObject* Object) override;
+#endif
+
 	// UObject interface
 	virtual void Serialize(FArchive& Ar) override;
 	// End of UObject interface
diff --git a/Engine/Source/Runtime/Engine/Classes/Engine/BlueprintGeneratedClass.h b/Engine/Source/Runtime/Engine/Classes/Engine/BlueprintGeneratedClass.h
--- a/Engine/Source/Runtime/Engine/Classes/Engine/BlueprintGeneratedClass.h
+++ b/Engine/Source/Runtime/Engine/Classes/Engine/BlueprintGeneratedClass.h
@@ -739,6 +739,20 @@ public:
 	static void CreateComponentsForActor(const UClass* ThisClass, AActor* Actor);
 	static void CreateTimelineComponent(AActor* Actor, const UTimelineTemplate* TimelineTemplate);
 
+#if WITH_EDITOR
+	/**
+	 * Called during serialization to allow the class to stash any sparse class data that needs to 
+	 * be conformed against the archetype once the CDO is available.
+	 */
+	virtual void PrepareToConformSparseClassData();
+
+	/**
+	 * Called to conform any pending sparse class data stashed by PrepareToConformSparseClassData.
+	 * @note Called either prior to generating the CDO (compile-on-load) or after loading the CDO.
+	 */
+	virtual void ConformSparseClassData(UObject* Object);
+#endif
+
 	// IBlueprintPropertyGuidProvider interface
 	virtual FName FindBlueprintPropertyNameFromGuid(const FGuid& PropertyGuid) const override final;
 	virtual FGuid FindBlueprintPropertyGuidFromName(const FName PropertyName) const override final;
@@ -924,4 +938,9 @@ private:
 	UPROPERTY()
 	TObjectPtr<UClassCookedMetaData> CachedCookedMetaDataPtr;
 #endif // WITH_EDITORONLY_DATA
+
+#if WITH_EDITORONLY_DATA
+	TWeakObjectPtr<UScriptStruct> SparseClassDataPendingConformStruct = nullptr;
+	void* SparseClassDataPendingConform = nullptr;
+#endif // WITH_EDITORONLY_DATA
 };
diff --git a/Engine/Source/Runtime/Engine/Private/Animation/AnimBlueprintGeneratedClass.cpp b/Engine/Source/Runtime/Engine/Private/Animation/AnimBlueprintGeneratedClass.cpp
--- a/Engine/Source/Runtime/Engine/Private/Animation/AnimBlueprintGeneratedClass.cpp
+++ b/Engine/Source/Runtime/Engine/Private/Animation/AnimBlueprintGeneratedClass.cpp
@@ -616,6 +616,18 @@ uint8* UAnimBlueprintGeneratedClass::GetPersistentUberGraphFrame(UObject* Obj, U
 	return Super::GetPersistentUberGraphFrame(Obj, FuncToCheck);
 }
 
+#if WITH_EDITOR
+void UAnimBlueprintGeneratedClass::PrepareToConformSparseClassData()
+{
+	// ABPGC sparse data is generated on compile, so there is nothing to conform
+}
+
+void UAnimBlueprintGeneratedClass::ConformSparseClassData(UObject* Object)
+{
+	// ABPGC sparse data is generated on compile, so there is nothing to conform
+}
+#endif
+
 void UAnimBlueprintGeneratedClass::PostLoadDefaultObject(UObject* Object)
 {
 	Super::PostLoadDefaultObject(Object);
diff --git a/Engine/Source/Runtime/Engine/Private/BlueprintGeneratedClass.cpp b/Engine/Source/Runtime/Engine/Private/BlueprintGeneratedClass.cpp
--- a/Engine/Source/Runtime/Engine/Private/BlueprintGeneratedClass.cpp
+++ b/Engine/Source/Runtime/Engine/Private/BlueprintGeneratedClass.cpp
@@ -549,15 +549,14 @@ void UBlueprintGeneratedClass::SerializeDefaultObject(UObject* Object, FStructur
 
 	if (UnderlyingArchive.IsLoading())
 	{
-		if (SparseClassDataStruct)
-		{
-			// TODO: We should make sure that this instance conforms to its archetype (keeping data where possible, potentially via 
-			// an explicit upgrade path for sparse data), as it may have changed since this BPGC was saved (UE-127121)
-		}
-		else
+		if (!SparseClassDataStruct)
 		{
 			SparseClassDataStruct = GetSparseClassDataArchetypeStruct();
 		}
+
+#if WITH_EDITOR
+		PrepareToConformSparseClassData();
+#endif
 	}
 }
 
@@ -580,7 +579,6 @@ void UBlueprintGeneratedClass::PostLoadDefaultObject(UObject* Object)
 	}
 
 #if WITH_EDITOR
-#if WITH_EDITORONLY_DATA
 	Object->MoveDataToSparseClassDataStruct();
 
 	if (Object->GetSparseClassDataStruct())
@@ -588,10 +586,77 @@ void UBlueprintGeneratedClass::PostLoadDefaultObject(UObject* Object)
 		// now that any data has been moved into the sparse data structure we can safely serialize it
 		bIsSparseClassDataSerializable = true;
 	}
+
+	ConformSparseClassData(Object);
 #endif
-#endif
 }
 
+#if WITH_EDITOR
+void UBlueprintGeneratedClass::PrepareToConformSparseClassData()
+{
+	checkf(SparseClassDataPendingConformStruct.IsExplicitlyNull() && SparseClassDataPendingConform == nullptr, TEXT("PrepareToConformSparseClassData was called while data was already pending conform!"));
+
+	if (SparseClassDataStruct)
+	{
+		UScriptStruct* SparseClassDataArchetypeStruct = GetSparseClassDataArchetypeStruct();
+
+		if (SparseClassDataStruct != SparseClassDataArchetypeStruct)
+		{
+			if (SparseClassDataArchetypeStruct)
+			{
+				SparseClassDataPendingConformStruct = SparseClassDataStruct;
+				SparseClassDataPendingConform = SparseClassData;
+
+				SparseClassDataStruct = SparseClassDataArchetypeStruct;
+				SparseClassData = nullptr;
+			}
+			else
+			{
+				CleanupSparseClassData();
+				SparseClassDataStruct = nullptr;
+			}
+		}
+	}
+}
+
+void UBlueprintGeneratedClass::ConformSparseClassData(UObject* Object)
+{
+	if (UScriptStruct* SparseClassDataPendingConformStructPtr = SparseClassDataPendingConformStruct.Get();
+		SparseClassDataPendingConformStructPtr && SparseClassDataPendingConform)
+	{
+		// Always allow the CDO first refusal at handling the conversion
+		if (!Object->ConformSparseClassDataStruct(SparseClassDataPendingConformStructPtr, SparseClassDataPendingConform))
+		{
+			UScriptStruct* SparseClassDataArchetypeStruct = GetSparseClassDataArchetypeStruct();
+
+			// Copy common properties if the structs are related types
+			UScriptStruct* SparseClassDataStructToCopy = nullptr;
+			if (SparseClassDataArchetypeStruct->IsChildOf(SparseClassDataPendingConformStructPtr))
+			{
+				SparseClassDataStructToCopy = SparseClassDataPendingConformStructPtr;
+			}
+			else if (SparseClassDataPendingConformStructPtr->IsChildOf(SparseClassDataArchetypeStruct))
+			{
+				SparseClassDataStructToCopy = SparseClassDataArchetypeStruct;
+			}
+			if (SparseClassDataStructToCopy)
+			{
+				SparseClassDataStructToCopy->CopyScriptStruct(GetOrCreateSparseClassData(), SparseClassDataPendingConform);
+			}
+		}
+
+		SparseClassDataPendingConformStructPtr->DestroyStruct(SparseClassDataPendingConform);
+	}
+
+	SparseClassDataPendingConformStruct = nullptr;
+	if (SparseClassDataPendingConform)
+	{
+		FMemory::Free(SparseClassDataPendingConform);
+		SparseClassDataPendingConform = nullptr;
+	}
+}
+#endif
+
 bool UBlueprintGeneratedClass::BuildCustomPropertyListForPostConstruction(FCustomPropertyListNode*& InPropertyList, UStruct* InStruct, const uint8* DataPtr, const uint8* DefaultDataPtr)
 {
 	const UClass* OwnerClass = Cast<UClass>(InStruct);
