Added an optional source string to FindTextInLocalizationTable

Like with FText::FindText, this can be used to provide an extra level of validation that the found text is what you expected

diff --git a/Engine/Source/Runtime/Engine/Classes/Kismet/KismetTextLibrary.h b/Engine/Source/Runtime/Engine/Classes/Kismet/KismetTextLibrary.h
--- a/Engine/Source/Runtime/Engine/Classes/Kismet/KismetTextLibrary.h
+++ b/Engine/Source/Runtime/Engine/Classes/Kismet/KismetTextLibrary.h
@@ -167,9 +167,14 @@ class ENGINE_API UKismetTextLibrary : public UBlueprintFunctionLibrary
 	UFUNCTION(BlueprintPure, Category="Utilities|Text")
 	static FText GetEmptyText();
 
-	/** Attempts to find existing Text using the representation found in the loc tables for the specified namespace and key. */
-	UFUNCTION(BlueprintPure, Category="Utilities|Text")
-	static bool FindTextInLocalizationTable(const FString& Namespace, const FString& Key, FText& OutText);
+	/**
+	 * Attempts to find existing Text using the representation found in the loc tables for the specified namespace and key.
+	 * @param Namespace The namespace of the text to find (if any).
+	 * @param Key The key of the text to find.
+	 * @param SourceString If set (not empty) then the found text must also have been created from this source string.
+	 */
+	UFUNCTION(BlueprintPure, Category="Utilities|Text", meta = (AdvancedDisplay = "SourceString"))
+	static bool FindTextInLocalizationTable(const FString& Namespace, const FString& Key, FText& OutText, const FString& SourceString = TEXT(""));
 
 	/** Returns true if A and B are linguistically equal (A == B). */
 	UFUNCTION(BlueprintPure, meta=(DisplayName = "Equal Exactly (Text)", CompactNodeTitle = "==="), Category="Utilities|Text")
diff --git a/Engine/Source/Runtime/Engine/Private/KismetTextLibrary.cpp b/Engine/Source/Runtime/Engine/Private/KismetTextLibrary.cpp
--- a/Engine/Source/Runtime/Engine/Private/KismetTextLibrary.cpp
+++ b/Engine/Source/Runtime/Engine/Private/KismetTextLibrary.cpp
@@ -135,9 +135,9 @@ FText UKismetTextLibrary::GetEmptyText()
 	return FText::GetEmpty();
 }
 
-bool UKismetTextLibrary::FindTextInLocalizationTable(const FString& Namespace, const FString& Key, FText& OutText)
+bool UKismetTextLibrary::FindTextInLocalizationTable(const FString& Namespace, const FString& Key, FText& OutText, const FString& SourceString)
 {
-	return FText::FindText(Namespace, Key, OutText);
+	return FText::FindText(Namespace, Key, OutText, SourceString.IsEmpty() ? nullptr : &SourceString);
 }
 
 bool UKismetTextLibrary::EqualEqual_IgnoreCase_TextText(const FText& A, const FText& B)
