Use correct radius and half height when drawing capsule component

diff --git a/Engine/Source/Runtime/Engine/Private/Components/CapsuleComponent.cpp b/Engine/Source/Runtime/Engine/Private/Components/CapsuleComponent.cpp
--- a/Engine/Source/Runtime/Engine/Private/Components/CapsuleComponent.cpp
+++ b/Engine/Source/Runtime/Engine/Private/Components/CapsuleComponent.cpp
@@ -41,8 +41,8 @@ FPrimitiveSceneProxy* UCapsuleComponent::CreateSceneProxy()
 		FDrawCylinderSceneProxy(const UCapsuleComponent* InComponent)
 			:	FPrimitiveSceneProxy(InComponent)
 			,	bDrawOnlyIfSelected( InComponent->bDrawOnlyIfSelected )
-			,	CapsuleRadius( InComponent->CapsuleRadius )
-			,	CapsuleHalfHeight( InComponent->CapsuleHalfHeight )
+			,	CapsuleRadius( InComponent->GetScaledCapsuleRadius() )
+			,	CapsuleHalfHeight( InComponent->GetScaledCapsuleHalfHeight() )
 			,	ShapeColor( InComponent->ShapeColor )
 		{
 			bWillEverBeLit = false;
@@ -65,7 +65,7 @@ FPrimitiveSceneProxy* UCapsuleComponent::CreateSceneProxy()
 					const FLinearColor DrawCapsuleColor = GetViewSelectionColor(ShapeColor, *View, IsSelected(), IsHovered(), false, IsIndividuallySelected() );
 
 					FPrimitiveDrawInterface* PDI = Collector.GetPDI(ViewIndex);
-					DrawWireCapsule( PDI, LocalToWorld.GetOrigin(), LocalToWorld.GetScaledAxis( EAxis::X ), LocalToWorld.GetScaledAxis( EAxis::Y ), LocalToWorld.GetScaledAxis( EAxis::Z ), DrawCapsuleColor, CapsuleRadius, CapsuleHalfHeight, CapsuleSides, SDPG_World );
+					DrawWireCapsule( PDI, LocalToWorld.GetOrigin(), LocalToWorld.GetUnitAxis(EAxis::X), LocalToWorld.GetUnitAxis(EAxis::Y), LocalToWorld.GetUnitAxis(EAxis::Z), DrawCapsuleColor, CapsuleRadius, CapsuleHalfHeight, CapsuleSides, SDPG_World );
 				}
 			}
 		}
