Adds a stat unit for TSR's history convergence speed in miliseconds to converge to 1spp

This takes into account screen percentage (even with dynamic res enabled) which shows it requires more
frames to accumulate the history as the screen percentage goes down, and frame time to more clearly
show that the frame rate which may also be bound by CPU can affect TSR's history convergence speed
and therefore its quality too.

diff --git a/Engine/Source/Runtime/Engine/Private/UnrealClient.cpp b/Engine/Source/Runtime/Engine/Private/UnrealClient.cpp
--- a/Engine/Source/Runtime/Engine/Private/UnrealClient.cpp
+++ b/Engine/Source/Runtime/Engine/Private/UnrealClient.cpp
@@ -930,6 +930,33 @@ int32 FStatUnitData::DrawStat(FViewport* InViewport, FCanvas* InCanvas, int32 In
 
 			PushRows(/* RowsCount = */ 1);
 		}
+
+		// TSR history convergence rate
+		if (GPixelRenderCounters.GetPixelDisplayCount())
+		{
+			// Target 1 sample per pixel for the convergence speed measurment.
+			const float TargetSamplePerPixel = 1.0f;
+
+			// Ideal TSR uses it to render 1080p -> 4k at 60hz.
+			const float IdealConvergenceTime = TargetSamplePerPixel * FMath::Pow(1080.0f / 2160.0f, -2.0f) * (1000.0f / 60.0f);
+
+			// Compute the resolution fraction agregate.
+			uint32 PixelRenderCount = GPixelRenderCounters.GetPixelRenderCount();
+			uint32 PixelDisplayCount = GPixelRenderCounters.GetPixelDisplayCount();
+			float ResolutionFraction = FMath::Sqrt(float(PixelRenderCount) / float(PixelDisplayCount));
+			float ConvergenceSpeedMultiplier = FMath::Pow(ResolutionFraction, -2.0f);
+
+			// Compute how long it takes for the history to converge to 1spp.
+			float ConvergenceFrameCount = ConvergenceSpeedMultiplier * TargetSamplePerPixel;
+			float ConvergenceTime = FMath::Max(ConvergenceFrameCount, 1.0f) * FrameTime;
+
+			const FColor Color = (ConvergenceTime <= IdealConvergenceTime) ? StatGreen : ((ConvergenceTime < 2.0f * IdealConvergenceTime) ? StatOrange : StatRed);
+
+			DrawTitleString(TEXT("TSR -> 1spp"), NoUnitGraphColor);
+			DrawDefaultAvgCell(FString::Printf(STATUNIT_FORMAT_AVGTIME, ConvergenceTime), Color);
+
+			PushRows(/* RowsCount = */ 1);
+		}
 	}
 
 #if !UE_BUILD_SHIPPING
diff --git a/Engine/Source/Runtime/RenderCore/Private/RenderCore.cpp b/Engine/Source/Runtime/RenderCore/Private/RenderCore.cpp
--- a/Engine/Source/Runtime/RenderCore/Private/RenderCore.cpp
+++ b/Engine/Source/Runtime/RenderCore/Private/RenderCore.cpp
@@ -16,6 +16,7 @@
 void UpdateShaderDevelopmentMode();
 
 void InitRenderGraph();
+static void InitPixelRenderCounters();
 
 class FRenderCoreModule : public FDefaultModuleImpl
 {
@@ -27,6 +28,7 @@ public:
 		IConsoleManager::Get().RegisterConsoleVariableSink_Handle(FConsoleCommandDelegate::CreateStatic(&UpdateShaderDevelopmentMode));
 
 		InitRenderGraph();
+		InitPixelRenderCounters();
 	}
 };
 
@@ -252,6 +254,22 @@ void FInputLatencyTimer::GameThreadTick()
 #endif
 }
 
+FPixelRenderCounters GPixelRenderCounters;
+
+void TickPixelRenderCounters()
+{
+	check(IsInRenderingThread());
+	GPixelRenderCounters.PrevPixelRenderCount = GPixelRenderCounters.CurrentPixelRenderCount;
+	GPixelRenderCounters.PrevPixelDisplayCount = GPixelRenderCounters.CurrentPixelDisplayCount;
+	GPixelRenderCounters.CurrentPixelRenderCount = 0;
+	GPixelRenderCounters.CurrentPixelDisplayCount = 0;
+}
+
+void InitPixelRenderCounters()
+{
+	FCoreDelegates::OnBeginFrameRT.AddStatic(TickPixelRenderCounters);
+}
+
 // Can be optimized to avoid the virtual function call but it's compiled out for final release anyway
 RENDERCORE_API int32 GetCVarForceLOD()
 {
diff --git a/Engine/Source/Runtime/RenderCore/Public/RenderTimer.h b/Engine/Source/Runtime/RenderCore/Public/RenderTimer.h
--- a/Engine/Source/Runtime/RenderCore/Public/RenderTimer.h
+++ b/Engine/Source/Runtime/RenderCore/Public/RenderTimer.h
@@ -118,3 +118,35 @@ extern RENDERCORE_API uint32 GSwapBufferTime;
 
 /** How many cycles the renderthread used, including dependent wait time. */
 extern RENDERCORE_API uint32 GRenderThreadTimeCriticalPath;
+
+
+class FPixelRenderCounters
+{
+public:
+	uint32 GetPixelRenderCount() const
+	{
+		return PrevPixelRenderCount;
+	}
+
+	uint32 GetPixelDisplayCount() const
+	{
+		return PrevPixelDisplayCount;
+	}
+
+	void AddViewStatistics(uint32 PixelRenderCount, uint32 PixelDisplayCount)
+	{
+		check(IsInRenderingThread());
+		CurrentPixelRenderCount += PixelRenderCount;
+		CurrentPixelDisplayCount += PixelDisplayCount;
+	}
+
+private:
+	uint32 PrevPixelRenderCount = 0;
+	uint32 PrevPixelDisplayCount = 0;
+	uint32 CurrentPixelRenderCount = 0;
+	uint32 CurrentPixelDisplayCount = 0;
+
+	friend void TickPixelRenderCounters();
+};
+
+extern RENDERCORE_API FPixelRenderCounters GPixelRenderCounters;
diff --git a/Engine/Source/Runtime/Renderer/Private/SceneRendering.cpp b/Engine/Source/Runtime/Renderer/Private/SceneRendering.cpp
--- a/Engine/Source/Runtime/Renderer/Private/SceneRendering.cpp
+++ b/Engine/Source/Runtime/Renderer/Private/SceneRendering.cpp
@@ -2867,6 +2867,12 @@ void FSceneRenderer::PrepareViewRectsForRendering(FRHICommandListImmediate& RHIC
 
 		check(View.ViewRect.Area() != 0);
 		check(View.VerifyMembersChecks());
+
+		if (!ViewFamily.EngineShowFlags.HitProxies && View.AntiAliasingMethod == AAM_TSR)
+		{
+			FIntPoint SecondaryViewRect = View.GetSecondaryViewRectSize();
+			GPixelRenderCounters.AddViewStatistics(View.ViewRect.Area(), SecondaryViewRect.X * SecondaryViewRect.Y);
+		}
 	}
 
 	// Shifts all view rects layout to the top left corner of the buffers, since post processing will just output the final
