Preload component outers to ensure that subobjects of component templates are loaded before we compile and possibly reinstance a blueprint

diff --git a/Engine/Source/Runtime/Engine/Private/SCS_Node.cpp b/Engine/Source/Runtime/Engine/Private/SCS_Node.cpp
--- a/Engine/Source/Runtime/Engine/Private/SCS_Node.cpp
+++ b/Engine/Source/Runtime/Engine/Private/SCS_Node.cpp
@@ -345,6 +345,21 @@ void USCS_Node::PreloadChain()
 		if (ensureMsgf(ComponentTemplate->GetLinker(), TEXT("Failed to find linker for %s, likely a circular dependency"), *ComponentTemplate->GetPathName()))
 		{
 			ComponentTemplate->GetLinker()->Preload(ComponentTemplate);
+
+			TArray<UObject*> Children;
+			GetObjectsWithOuter(ComponentTemplate, Children, true, RF_LoadCompleted);
+			for (UObject* Obj : Children)
+			{
+				if (!Obj->HasAnyFlags(RF_WasLoaded))
+				{
+					continue;
+				}
+
+				if (FLinkerLoad* Linker = Obj->GetLinker())
+				{
+					Linker->Preload(Obj);
+				}
+			}
 		}
 	}
 
