Remove unused DynamicRayTracingGeometryVertexBuffers.

diff --git a/Engine/Source/Runtime/Engine/Private/StaticMeshRender.cpp b/Engine/Source/Runtime/Engine/Private/StaticMeshRender.cpp
--- a/Engine/Source/Runtime/Engine/Private/StaticMeshRender.cpp
+++ b/Engine/Source/Runtime/Engine/Private/StaticMeshRender.cpp
@@ -470,13 +470,6 @@ void FStaticMeshSceneProxy::SetEvaluateWorldPositionOffsetInRayTracing(bool NewV
 
 			DynamicRayTracingGeometries.Empty();
 
-			for (auto& Buffer : DynamicRayTracingGeometryVertexBuffers)
-			{
-				Buffer.Release();
-			}
-
-			DynamicRayTracingGeometryVertexBuffers.Empty();
-
 			if (GetPrimitiveSceneInfo())
 			{
 				GetPrimitiveSceneInfo()->bIsRayTracingStaticRelevant = IsRayTracingStaticRelevant();
@@ -495,11 +488,6 @@ int32 FStaticMeshSceneProxy::GetLightMapCoordinateIndex() const
 FStaticMeshSceneProxy::~FStaticMeshSceneProxy()
 {
 #if RHI_RAYTRACING
-	for (auto& Buffer : DynamicRayTracingGeometryVertexBuffers)
-	{
-		Buffer.Release();
-	}
-
 	for (auto& Geometry: DynamicRayTracingGeometries)
 	{
 		Geometry.ReleaseResource();
@@ -1915,12 +1903,8 @@ void FStaticMeshSceneProxy::GetDynamicRayTracingInstances(FRayTracingMaterialGat
 
 		if (bEvaluateWPO && RenderData->LODVertexFactories[LODIndex].VertexFactory.GetType()->SupportsRayTracingDynamicGeometry())
 		{
-			// Use the internal vertex buffer only when initialized otherwise used the shared vertex buffer - needs to be updated every frame
+			// Use the shared vertex buffer - needs to be updated every frame
 			FRWBuffer* VertexBuffer = nullptr;
-			if (DynamicRayTracingGeometryVertexBuffers.Num() > (int32)LODIndex && DynamicRayTracingGeometryVertexBuffers[LODIndex].NumBytes > 0)
-			{
-				VertexBuffer = &DynamicRayTracingGeometryVertexBuffers[LODIndex];
-			}
 
 			Context.DynamicRayTracingGeometriesToUpdate.Add(
 				FRayTracingDynamicGeometryUpdateParams
diff --git a/Engine/Source/Runtime/Engine/Public/StaticMeshSceneProxy.h b/Engine/Source/Runtime/Engine/Public/StaticMeshSceneProxy.h
--- a/Engine/Source/Runtime/Engine/Public/StaticMeshSceneProxy.h
+++ b/Engine/Source/Runtime/Engine/Public/StaticMeshSceneProxy.h
@@ -227,7 +227,6 @@ protected:
 	bool bSupportRayTracing;
 	bool bDynamicRayTracingGeometry;
 	TArray<FRayTracingGeometry, TInlineAllocator<MAX_MESH_LOD_COUNT>> DynamicRayTracingGeometries;
-	TArray<FRWBuffer, TInlineAllocator<MAX_MESH_LOD_COUNT>> DynamicRayTracingGeometryVertexBuffers;
 	TArray<FMeshBatch> CachedRayTracingMaterials;
 	int16 CachedRayTracingMaterialsLODIndex = INDEX_NONE;
 	FRayTracingMaskAndFlags CachedRayTracingInstanceMaskAndFlags;
