Fix multi-line cvar descriptions in StaticMeshRender.cpp

diff --git a/Engine/Source/Runtime/Engine/Private/StaticMeshRender.cpp b/Engine/Source/Runtime/Engine/Private/StaticMeshRender.cpp
--- a/Engine/Source/Runtime/Engine/Private/StaticMeshRender.cpp
+++ b/Engine/Source/Runtime/Engine/Private/StaticMeshRender.cpp
@@ -122,9 +122,9 @@ int32 GNaniteProxyRenderMode = 0;
 FAutoConsoleVariableRef CVarNaniteProxyRenderMode(
 	TEXT("r.Nanite.ProxyRenderMode"),
 	GNaniteProxyRenderMode,
-	TEXT("Render proxy meshes if Nanite is unsupported.")
-	TEXT(" 0: Fall back to rendering Nanite proxy meshes if Nanite is unsupported. (default)")
-	TEXT(" 1: Disable rendering if Nanite is enabled on a mesh but is unsupported.")
+	TEXT("Render proxy meshes if Nanite is unsupported.\n")
+	TEXT(" 0: Fall back to rendering Nanite proxy meshes if Nanite is unsupported. (default)\n")
+	TEXT(" 1: Disable rendering if Nanite is enabled on a mesh but is unsupported.\n")
 	TEXT(" 2: Disable rendering if Nanite is enabled on a mesh but is unsupported, except for static mesh editor toggle."),
 	FConsoleVariableDelegate::CreateLambda([](IConsoleVariable* InVariable)
 	{
@@ -157,10 +157,10 @@ static TAutoConsoleVariable<int32> CVarRayTracingStaticMeshes(
 static TAutoConsoleVariable<int32> CVarRayTracingStaticMeshesWPO(
 	TEXT("r.RayTracing.Geometry.StaticMeshes.WPO"),
 	1,
-	TEXT("World position offset evaluation for static meshes with EvaluateWPO enabled in ray tracing effects")
-	TEXT(" 0: static meshes with world position offset hidden in ray tracing")
-	TEXT(" 1: static meshes with world position offset visible in ray tracing, WPO evaluation enabled (default)")
-	TEXT(" 2: static meshes with world position offset visible in ray tracing, WPO evaluation disabled")
+	TEXT("World position offset evaluation for static meshes with EvaluateWPO enabled in ray tracing effects.\n")
+	TEXT(" 0: static meshes with world position offset hidden in ray tracing.\n")
+	TEXT(" 1: static meshes with world position offset visible in ray tracing, WPO evaluation enabled (default).\n")
+	TEXT(" 2: static meshes with world position offset visible in ray tracing, WPO evaluation disabled.")
 );
 
 FAutoConsoleVariableSink CVarRayTracingStaticMeshesWPOSink(FConsoleCommandDelegate::CreateLambda([]()
