Improves stat unit meaning with -VMGPU

diff --git a/Engine/Source/Runtime/Engine/Private/UnrealClient.cpp b/Engine/Source/Runtime/Engine/Private/UnrealClient.cpp
--- a/Engine/Source/Runtime/Engine/Private/UnrealClient.cpp
+++ b/Engine/Source/Runtime/Engine/Private/UnrealClient.cpp
@@ -736,7 +736,7 @@ int32 FStatUnitData::DrawStat(FViewport* InViewport, FCanvas* InCanvas, int32 In
 				{
 					DrawCell(
 						GPURowId++, GPUIndex * ColumnPerGPU + AvgUnitColumnId, GPUColumnCount,
-						FString::Printf(TEXT("GPU%d"), GPUIndex), NoUnitGraphColor);
+						GVirtualMGPU ? FString::Printf(TEXT("VGPU%d"), GPUIndex) : FString::Printf(TEXT("GPU%d"), GPUIndex), NoUnitGraphColor);
 				}
 
 				if (!bHaveGPUData[GPUIndex])
@@ -763,7 +763,12 @@ int32 FStatUnitData::DrawStat(FViewport* InViewport, FCanvas* InCanvas, int32 In
 					GPURowId++;
 				}
 
-				if (GRHISupportsGPUUsage)
+				// Only display the GPU usage of one GPU usage when using the -VMGPU 
+				if (GRHISupportsGPUUsage && GPUIndex > 0 && GVirtualMGPU)
+				{
+					GPURowId = GPURowCount;
+				}
+				else if (GRHISupportsGPUUsage)
 				{
 					{
 						const FColor Color = (GPUClockFraction[GPUIndex] < 0.5f) ? StatRed : ((GPUClockFraction[GPUIndex] < 0.75f) ? StatOrange : StatGreen);
