Fix for bus sends and submix sends on sound waves being mute with the default parameters

diff --git a/Engine/Source/Runtime/Engine/Classes/Sound/SoundSourceBusSend.h b/Engine/Source/Runtime/Engine/Classes/Sound/SoundSourceBusSend.h
--- a/Engine/Source/Runtime/Engine/Classes/Sound/SoundSourceBusSend.h
+++ b/Engine/Source/Runtime/Engine/Classes/Sound/SoundSourceBusSend.h
@@ -44,7 +44,7 @@ struct ENGINE_API FSoundSourceBusSendInfo
 	TObjectPtr<UAudioBus> AudioBus;
 
 	// The amount of audio to send to the bus.
-	UPROPERTY(EditAnywhere, Category = BusSend)
+	UPROPERTY(EditAnywhere, Category = BusSend, meta = (DisplayName = "Manual Send Level"))
 	float SendLevel;
 
 	// The amount to send to the bus when sound is located at a distance equal to value specified in the min send distance.
@@ -71,7 +71,7 @@ struct ENGINE_API FSoundSourceBusSendInfo
 		: SourceBusSendLevelControlMethod(ESourceBusSendLevelControlMethod::Manual)
 		, SoundSourceBus(nullptr)
 		, AudioBus(nullptr)
-		, SendLevel(0.0f)
+		, SendLevel(1.0f)
 		, MinSendLevel(0.0f)
 		, MaxSendLevel(1.0f)
 		, MinSendDistance(100.0f)
diff --git a/Engine/Source/Runtime/Engine/Classes/Sound/SoundSubmixSend.h b/Engine/Source/Runtime/Engine/Classes/Sound/SoundSubmixSend.h
--- a/Engine/Source/Runtime/Engine/Classes/Sound/SoundSubmixSend.h
+++ b/Engine/Source/Runtime/Engine/Classes/Sound/SoundSubmixSend.h
@@ -113,7 +113,7 @@ struct ENGINE_API FSoundSubmixSendInfo
 	TObjectPtr<USoundSubmixBase> SoundSubmix;
 
 	// The amount of audio to send
-	UPROPERTY(EditAnywhere, BlueprintReadWrite, Category = SubmixSend)
+	UPROPERTY(EditAnywhere, BlueprintReadWrite, Category = SubmixSend, meta = (DisplayName = "Manual Send Level"))
 	float SendLevel;
 
 	// Whether to disable the 0-1 clamp for manual SendLevel control
diff --git a/Engine/Source/Runtime/Engine/Private/SoundSubmixSend.cpp b/Engine/Source/Runtime/Engine/Private/SoundSubmixSend.cpp
--- a/Engine/Source/Runtime/Engine/Private/SoundSubmixSend.cpp
+++ b/Engine/Source/Runtime/Engine/Private/SoundSubmixSend.cpp
@@ -8,7 +8,7 @@ FSoundSubmixSendInfo::FSoundSubmixSendInfo()
 	: SendLevelControlMethod(ESendLevelControlMethod::Manual)
 	, SendStage(ESubmixSendStage::PostDistanceAttenuation)
 	, SoundSubmix(nullptr)
-	, SendLevel(0.0f)
+	, SendLevel(1.0f)
 	, DisableManualSendClamp(false)
 	, MinSendLevel(0.0f)
 	, MaxSendLevel(1.0f)
