WorldPartition::Generate Hlods now called on registered all containers (content bundles)
RuntimeSpatialHash::GenerateHlods now targets a specifc container instead of the WorldPartition main container
ASpatialHashRuntimeGridInfo now saved in contentbundle container when needed
Adjusted places where we loop on/retrieve Hlods actors from the world partition main container to world partition container collection

diff --git a/Engine/Source/Editor/UnrealEd/Private/WorldPartition/WorldPartitionHLODsBuilder.cpp b/Engine/Source/Editor/UnrealEd/Private/WorldPartition/WorldPartitionHLODsBuilder.cpp
--- a/Engine/Source/Editor/UnrealEd/Private/WorldPartition/WorldPartitionHLODsBuilder.cpp
+++ b/Engine/Source/Editor/UnrealEd/Private/WorldPartition/WorldPartitionHLODsBuilder.cpp
@@ -367,7 +367,7 @@ bool UWorldPartitionHLODsBuilder::SetupHLODActors()
 	UE_LOG(LogWorldPartitionHLODsBuilder, Display, TEXT("#### World HLOD actors ####"));
 
 	int32 NumActors = 0;
-	for (FActorDescList::TIterator<AWorldPartitionHLOD> HLODIterator(WorldPartition->GetActorDescContainer()); HLODIterator; ++HLODIterator)
+	for (FActorDescContainerCollection::TIterator<AWorldPartitionHLOD> HLODIterator(WorldPartition); HLODIterator; ++HLODIterator)
 	{
 		FWorldPartitionActorDesc* HLODActorDesc = *HLODIterator;
 		FString PackageName = HLODActorDesc->GetActorPackage().ToString();
@@ -549,7 +549,7 @@ bool UWorldPartitionHLODsBuilder::DeleteHLODActors()
 	UE_LOG(LogWorldPartitionHLODsBuilder, Display, TEXT("#### Deleting HLOD actors ####"));
 
 	TArray<FString> PackagesToDelete;
-	for (FActorDescList::TIterator<AWorldPartitionHLOD> HLODIterator(WorldPartition->GetActorDescContainer()); HLODIterator; ++HLODIterator)
+	for (FActorDescContainerCollection::TIterator<AWorldPartitionHLOD> HLODIterator(WorldPartition); HLODIterator; ++HLODIterator)
 	{
 		FWorldPartitionActorDesc* HLODActorDesc = *HLODIterator;
 		FString PackageName = HLODActorDesc->GetActorPackage().ToString();
@@ -645,8 +645,7 @@ bool UWorldPartitionHLODsBuilder::DumpStats()
 	};
 
 	TArray<FHLODStat> HLODStats;
-
-	for (UActorDescContainer::TIterator<AWorldPartitionHLOD> HLODIterator(WorldPartition->GetActorDescContainer()); HLODIterator; ++HLODIterator)
+	for (FActorDescContainerCollection::TIterator<AWorldPartitionHLOD> HLODIterator(WorldPartition); HLODIterator; ++HLODIterator)
 	{
 		FWorldPartitionReference HLODActorRef(WorldPartition, HLODIterator->GetGuid());
 		AWorldPartitionHLOD* HLODActor = CastChecked<AWorldPartitionHLOD>(HLODActorRef->GetActor());
@@ -790,7 +789,7 @@ TArray<TArray<FGuid>> UWorldPartitionHLODsBuilder::GetHLODWorkloads(int32 NumWor
 {
 	// Build a mapping of 1 HLOD[Level] -> N HLOD[Level - 1]
 	TMap<FGuid, TArray<FGuid>>	HLODParenting;
-	for (FActorDescList::TIterator<AWorldPartitionHLOD> HLODIterator(WorldPartition->GetActorDescContainer()); HLODIterator; ++HLODIterator)
+	for (FActorDescContainerCollection::TIterator<AWorldPartitionHLOD> HLODIterator(WorldPartition); HLODIterator; ++HLODIterator)
 	{
 		// Filter by HLOD actor
 		if (!HLODActorToBuild.IsNone() && HLODIterator->GetActorLabel() != HLODActorToBuild)
diff --git a/Engine/Source/Runtime/Engine/Private/WorldPartition/HLOD/HLODSubsystem.cpp b/Engine/Source/Runtime/Engine/Private/WorldPartition/HLOD/HLODSubsystem.cpp
--- a/Engine/Source/Runtime/Engine/Private/WorldPartition/HLOD/HLODSubsystem.cpp
+++ b/Engine/Source/Runtime/Engine/Private/WorldPartition/HLOD/HLODSubsystem.cpp
@@ -776,7 +776,7 @@ static void DumpHLODStats(UWorld* World)
 	Output.Logf(TEXT("%s" LINE_TERMINATOR_ANSI), *StatHeader);
 
 	// Write one line per HLOD actor desc
-	for (FActorDescList::TIterator<AWorldPartitionHLOD> HLODIterator(WorldPartition->GetActorDescContainer()); HLODIterator; ++HLODIterator)
+	for (FActorDescContainerCollection::TIterator<AWorldPartitionHLOD> HLODIterator(WorldPartition); HLODIterator; ++HLODIterator)
 	{
 		const FString StatLine = FString::JoinBy(StatsToWrite, TEXT(","), [&HLODIterator](const TPair<FName, FGetStatFunc>& Pair) { return Pair.Value(**HLODIterator); });
 		Output.Logf(TEXT("%s" LINE_TERMINATOR_ANSI), *StatLine);
diff --git a/Engine/Source/Runtime/Engine/Private/WorldPartition/RuntimeSpatialHash/RuntimeSpatialHashHLOD.cpp b/Engine/Source/Runtime/Engine/Private/WorldPartition/RuntimeSpatialHash/RuntimeSpatialHashHLOD.cpp
--- a/Engine/Source/Runtime/Engine/Private/WorldPartition/RuntimeSpatialHash/RuntimeSpatialHashHLOD.cpp
+++ b/Engine/Source/Runtime/Engine/Private/WorldPartition/RuntimeSpatialHash/RuntimeSpatialHashHLOD.cpp
@@ -18,6 +18,7 @@
 
 #include "WorldPartition/WorldPartition.h"
 #include "WorldPartition/WorldPartitionHelpers.h"
+#include "WorldPartition/ContentBundle/ContentBundleActivationScope.h"
 
 #include "UObject/GCObjectScopeGuard.h"
 #include "Engine/Engine.h"
@@ -89,7 +90,7 @@ static void DeletePackage(UWorldPartition* WorldPartition, FWorldPartitionActorD
 	if (ActorDesc->IsLoaded())
 	{
 		DeletePackage(ActorDesc->GetActor()->GetPackage(), SourceControlHelper);
-		WorldPartition->GetActorDescContainer()->OnPackageDeleted(ActorDesc->GetActor()->GetPackage());
+		WorldPartition->OnPackageDeleted(ActorDesc->GetActor()->GetPackage());
 	}
 	else
 	{
@@ -174,9 +175,6 @@ static TArray<FGuid> GenerateHLODsForGrid(UWorldPartition* WorldPartition, const
 			if (bShouldGenerateHLODs)
 			{
 				SlowTask.EnterProgressFrame(1);
-				
-				// todo_ow: Content Bundle Hlods generation is not enable currently. We should never get a valid ContentBundleGuid.
-				check(!GridCellDataChunk.GetContentBundleID().IsValid());
 
 				const FGuid CellGuid = UWorldPartitionRuntimeSpatialHash::GetCellGuid(RuntimeGrid.GridName, CellGlobalCoord, GridCellDataChunk.GetDataLayersID(), GridCellDataChunk.GetContentBundleID());
 				const FString CellName = UWorldPartitionRuntimeSpatialHash::GetCellNameString(OuterWorld->GetPackage()->GetFName(), RuntimeGrid.GridName, CellGlobalCoord, GridCellDataChunk.GetDataLayersID(), GridCellDataChunk.GetContentBundleID());
@@ -273,19 +271,19 @@ static TArray<FGuid> GenerateHLODsForGrid(UWorldPartition* WorldPartition, const
 }
 
 // Find all referenced HLODLayer assets
-static TMap<UHLODLayer*, int32> GatherHLODLayers(UWorldPartition* WorldPartition)
+static TMap<UHLODLayer*, int32> GatherHLODLayers(const UActorDescContainer* ActorDescContainer)
 {
 	// Gather up all HLODLayers referenced by the actors, along with the HLOD level at which it was used
 	TMap<UHLODLayer*, int32> HLODLayersLevel;
 
-	for (FActorDescList::TIterator<> ActorDescIterator(WorldPartition->GetActorDescContainer()); ActorDescIterator; ++ActorDescIterator)
+	for (FActorDescList::TConstIterator<> ActorDescIterator(ActorDescContainer); ActorDescIterator; ++ActorDescIterator)
 	{
 		const FWorldPartitionActorDesc& ActorDesc = **ActorDescIterator;
 		if (!ActorDesc.GetActorNativeClass()->IsChildOf<AWorldPartitionHLOD>())
 		{
 			if (ActorDesc.GetActorIsHLODRelevant())
 			{
-				UHLODLayer* HLODLayer = UHLODLayer::GetHLODLayer(ActorDesc, WorldPartition);
+				UHLODLayer* HLODLayer = UHLODLayer::GetHLODLayer(ActorDesc, ActorDescContainer->GetWorld()->GetWorldPartition());
 
 				// If layer was already encountered, no need to process it again
 				if (!HLODLayersLevel.Contains(HLODLayer))
@@ -355,28 +353,36 @@ static TMap<FName, FSpatialHashRuntimeGrid> CreateHLODGrids(TMap<UHLODLayer*, in
 }
 
 // Create/destroy HLOD grid actors
-static void UpdateHLODGridsActors(UWorld* World, const TMap<FName, FSpatialHashRuntimeGrid>& HLODGrids, ISourceControlHelper* SourceControlHelper)
+static void UpdateHLODGridsActors(const UActorDescContainer* Container, const TMap<FName, FSpatialHashRuntimeGrid>& HLODGrids, ISourceControlHelper* SourceControlHelper)
 {
 	static const FName HLODGridTag = TEXT("HLOD");
 	static const uint32 HLODGridTagLen = HLODGridTag.ToString().Len();
 
+	UWorld* World = Container->GetWorld();
+
 	// Gather all existing HLOD grid actors, see if some are unused and needs to be deleted
 	TMap<FName, ASpatialHashRuntimeGridInfo*> ExistingGridActors;
-	for (TActorIterator<ASpatialHashRuntimeGridInfo> ItRuntimeGridActor(World); ItRuntimeGridActor; ++ItRuntimeGridActor)
+	for (FActorDescList::TConstIterator<> ItActorDesc(Container); ItActorDesc; ++ItActorDesc)
 	{
-		ASpatialHashRuntimeGridInfo* GridActor = *ItRuntimeGridActor;
-		if (GridActor->ActorHasTag(HLODGridTag))
+		if(ItActorDesc->GetActorNativeClass()->IsChildOf<ASpatialHashRuntimeGridInfo>())
 		{
-			const FSpatialHashRuntimeGrid* HLODGrid = HLODGrids.Find(GridActor->GridSettings.GridName);
-			if (HLODGrid && GridActor->GridSettings.Priority && GridActor->GridSettings.HLODLayer)
-			{
-				ExistingGridActors.Emplace(GridActor->GridSettings.GridName, GridActor);
-			}
-			else
+			const FWorldPartitionActorDesc* GridActorDesc = *ItActorDesc;
+			ASpatialHashRuntimeGridInfo* GridActor = CastChecked<ASpatialHashRuntimeGridInfo>(GridActorDesc->GetActor());
+			if (GridActor->ActorHasTag(HLODGridTag))
 			{
-				World->DestroyActor(GridActor);
+				const FSpatialHashRuntimeGrid* HLODGrid = HLODGrids.Find(GridActor->GridSettings.GridName);
+				if (HLODGrid && GridActor->GridSettings.Priority && GridActor->GridSettings.HLODLayer)
+				{
+					check(GridActor->GetContentBundleGuid() == Container->GetContentBundleGuid());
+
+					ExistingGridActors.Emplace(GridActor->GridSettings.GridName, GridActor);
+				}
+				else
+				{
+					World->DestroyActor(GridActor);
 
-				DeletePackage(GridActor->GetPackage(), SourceControlHelper);
+					DeletePackage(GridActor->GetPackage(), SourceControlHelper);
+				}
 			}
 		}
 	}
@@ -392,11 +398,15 @@ static void UpdateHLODGridsActors(UWorld* World, const TMap<FName, FSpatialHashR
 		ASpatialHashRuntimeGridInfo* GridActor = ExistingGridActors.FindRef(GridName);
 		if (!GridActor)
 		{
+			FContentBundleActivationScope ContentBndleScope(Container->GetContentBundleGuid());
+
 			FActorSpawnParameters SpawnParams;
 			SpawnParams.bCreateActorPackage = true;
 			GridActor = World->SpawnActor<ASpatialHashRuntimeGridInfo>(SpawnParams);
 			GridActor->Tags.Add(HLODGridTag);
 			bDirty = true;
+
+			check(GridActor->GetContentBundleGuid() == Container->GetContentBundleGuid());
 		}
 
 		const FString ActorLabel = GridSettings.GridName.ToString();
@@ -430,9 +440,11 @@ bool UWorldPartitionRuntimeSpatialHash::GenerateHLOD(ISourceControlHelper* Sourc
 	}
 
 	UWorldPartition* WorldPartition = GetOuterUWorldPartition();
+	IStreamingGenerationContext::FActorSetContainer* MainActorSetContainer = const_cast<IStreamingGenerationContext::FActorSetContainer*>(StreamingGenerationContext->GetMainWorldContainer());
+	const UActorDescContainer* ActorDescContainer = MainActorSetContainer->ActorDescContainer;
 
 	// Find all used HLOD layers
-	TMap<UHLODLayer*, int32> HLODLayersLevels = GatherHLODLayers(WorldPartition);
+	TMap<UHLODLayer*, int32> HLODLayersLevels = GatherHLODLayers(ActorDescContainer);
 	TArray<UHLODLayer*> HLODLayers;
 	HLODLayersLevels.GetKeys(HLODLayers);
 
@@ -452,14 +464,12 @@ bool UWorldPartitionRuntimeSpatialHash::GenerateHLOD(ISourceControlHelper* Sourc
 
 	// HLOD creation context
 	FHLODCreationContext Context;
-	for (FActorDescList::TIterator<AWorldPartitionHLOD> HLODIterator(WorldPartition->GetActorDescContainer()); HLODIterator; ++HLODIterator)
+	for (FActorDescList::TConstIterator<AWorldPartitionHLOD> HLODIterator(ActorDescContainer); HLODIterator; ++HLODIterator)
 	{
 		FWorldPartitionHandle HLODActorHandle(WorldPartition, HLODIterator->GetGuid());
 		Context.HLODActorDescs.Emplace(HLODIterator->GetActorName(), MoveTemp(HLODActorHandle));
 	}
 
-	IStreamingGenerationContext::FActorSetContainer* MainActorSetContainer = const_cast<IStreamingGenerationContext::FActorSetContainer*>(StreamingGenerationContext->GetMainWorldContainer());
-
 	TArray<TArray<const IStreamingGenerationContext::FActorSetInstance*>> GridActorSetInstances;
 	GridActorSetInstances.InsertDefaulted(0, Grids.Num());
 
@@ -535,6 +545,7 @@ bool UWorldPartitionRuntimeSpatialHash::GenerateHLOD(ISourceControlHelper* Sourc
 			NewHLODActorSetInstance.RuntimeGrid = HLODActorDescView.GetRuntimeGrid();
 			NewHLODActorSetInstance.bIsSpatiallyLoaded = HLODActorDescView.GetIsSpatiallyLoaded();
 			NewHLODActorSetInstance.DataLayers = UDataLayerSubsystem::GetRuntimeDataLayerInstances(MainActorSetContainer->ActorDescContainer->GetWorld(), HLODActorDescView.GetRuntimeDataLayers());
+			NewHLODActorSetInstance.ContentBundleID = MainActorSetContainer->ActorDescContainer->GetContentBundleGuid();
 			NewHLODActorSetInstance.ContainerInstance = MainActorSetContainer;
 			NewHLODActorSetInstance.ActorSet = &MainActorSetContainer->ActorSets.AddDefaulted_GetRef();
 			const_cast<IStreamingGenerationContext::FActorSet*>(NewHLODActorSetInstance.ActorSet)->Actors.Add(HLODActorDescView.GetGuid());
@@ -562,7 +573,7 @@ bool UWorldPartitionRuntimeSpatialHash::GenerateHLOD(ISourceControlHelper* Sourc
 	}
 
 	// Create/destroy HLOD grid actors
-	UpdateHLODGridsActors(GetWorld(), HLODGrids, SourceControlHelper);
+	UpdateHLODGridsActors(ActorDescContainer, HLODGrids, SourceControlHelper);
 
 	return true;
 }
diff --git a/Engine/Source/Runtime/Engine/Private/WorldPartition/WorldPartitionStreamingGeneration.cpp b/Engine/Source/Runtime/Engine/Private/WorldPartition/WorldPartitionStreamingGeneration.cpp
--- a/Engine/Source/Runtime/Engine/Private/WorldPartition/WorldPartitionStreamingGeneration.cpp
+++ b/Engine/Source/Runtime/Engine/Private/WorldPartition/WorldPartitionStreamingGeneration.cpp
@@ -1023,22 +1023,24 @@ void UWorldPartition::FlushStreaming()
 
 void UWorldPartition::GenerateHLOD(ISourceControlHelper* SourceControlHelper, bool bCreateActorsOnly)
 {
-	FStreamingGenerationLogErrorHandler LogErrorHandler;
-
-	FWorldPartitionStreamingGenerator::FWorldPartitionStreamingGeneratorParams StreamingGeneratorParams;
-	StreamingGeneratorParams.ErrorHandler = &LogErrorHandler;
-	StreamingGeneratorParams.bEnableStreaming = IsStreamingEnabled();
-	StreamingGeneratorParams.FilteredClasses.Add(AWorldPartitionHLOD::StaticClass());
-	StreamingGeneratorParams.IsValidGrid = [this](FName GridName) { return RuntimeHash->IsValidGrid(GridName); };
+	ForEachActorDescContainer([this, &SourceControlHelper, bCreateActorsOnly](UActorDescContainer* ActorDescContainer)
+	{
+		FStreamingGenerationLogErrorHandler LogErrorHandler;
+		FWorldPartitionStreamingGenerator::FWorldPartitionStreamingGeneratorParams StreamingGeneratorParams;
+		StreamingGeneratorParams.ErrorHandler = &LogErrorHandler;
+		StreamingGeneratorParams.bEnableStreaming = IsStreamingEnabled();
+		StreamingGeneratorParams.FilteredClasses.Add(AWorldPartitionHLOD::StaticClass());
+		StreamingGeneratorParams.IsValidGrid = [this](FName GridName) { return RuntimeHash->IsValidGrid(GridName); };
 
-	FWorldPartitionStreamingGenerator StreamingGenerator(StreamingGeneratorParams);
-	StreamingGenerator.PreparationPhase(ActorDescContainer);
+		FWorldPartitionStreamingGenerator StreamingGenerator(StreamingGeneratorParams);
+		StreamingGenerator.PreparationPhase(ActorDescContainer);
 
-	TUniquePtr<FArchive> LogFileAr = FWorldPartitionStreamingGenerator::CreateDumpStateLogArchive(TEXT("HLOD"));
-	FHierarchicalLogArchive HierarchicalLogAr(*LogFileAr);
-	StreamingGenerator.DumpStateLog(HierarchicalLogAr);
+		TUniquePtr<FArchive> LogFileAr = FWorldPartitionStreamingGenerator::CreateDumpStateLogArchive(TEXT("HLOD"));
+		FHierarchicalLogArchive HierarchicalLogAr(*LogFileAr);
+		StreamingGenerator.DumpStateLog(HierarchicalLogAr);
 
-	RuntimeHash->GenerateHLOD(SourceControlHelper, StreamingGenerator.GetStreamingGenerationContext(), bCreateActorsOnly);
+		RuntimeHash->GenerateHLOD(SourceControlHelper, StreamingGenerator.GetStreamingGenerationContext(), bCreateActorsOnly);
+	});	
 }
 
 void UWorldPartition::CheckForErrors(IStreamingGenerationErrorHandler* ErrorHandler) const
