Fix shadow field in compilation

diff --git a/Engine/Source/Runtime/Engine/Private/WorldPartition/WorldPartitionStreamingGeneration.cpp b/Engine/Source/Runtime/Engine/Private/WorldPartition/WorldPartitionStreamingGeneration.cpp
--- a/Engine/Source/Runtime/Engine/Private/WorldPartition/WorldPartitionStreamingGeneration.cpp
+++ b/Engine/Source/Runtime/Engine/Private/WorldPartition/WorldPartitionStreamingGeneration.cpp
@@ -1023,7 +1023,7 @@ void UWorldPartition::FlushStreaming()
 
 void UWorldPartition::GenerateHLOD(ISourceControlHelper* SourceControlHelper, bool bCreateActorsOnly)
 {
-	ForEachActorDescContainer([this, &SourceControlHelper, bCreateActorsOnly](UActorDescContainer* ActorDescContainer)
+	ForEachActorDescContainer([this, &SourceControlHelper, bCreateActorsOnly](UActorDescContainer* InActorDescContainer)
 	{
 		FStreamingGenerationLogErrorHandler LogErrorHandler;
 		FWorldPartitionStreamingGenerator::FWorldPartitionStreamingGeneratorParams StreamingGeneratorParams;
@@ -1033,7 +1033,7 @@ void UWorldPartition::GenerateHLOD(ISourceControlHelper* SourceControlHelper, bo
 		StreamingGeneratorParams.IsValidGrid = [this](FName GridName) { return RuntimeHash->IsValidGrid(GridName); };
 
 		FWorldPartitionStreamingGenerator StreamingGenerator(StreamingGeneratorParams);
-		StreamingGenerator.PreparationPhase(ActorDescContainer);
+		StreamingGenerator.PreparationPhase(InActorDescContainer);
 
 		TUniquePtr<FArchive> LogFileAr = FWorldPartitionStreamingGenerator::CreateDumpStateLogArchive(TEXT("HLOD"));
 		FHierarchicalLogArchive HierarchicalLogAr(*LogFileAr);
