Made FActorGUIDGenerator available in non-editor builds and added FActorSpawnParameters.bForceGloballyUniqueName to force a spawned actor to get a unique name id at runtime.
Also, FActorGUIDGenerator is more resilient to unavailable network adress uid and will construct a unique id based on per-user settings at last resort.

diff --git a/Engine/Source/Runtime/Engine/Classes/Engine/World.h b/Engine/Source/Runtime/Engine/Classes/Engine/World.h
--- a/Engine/Source/Runtime/Engine/Classes/Engine/World.h
+++ b/Engine/Source/Runtime/Engine/Classes/Engine/World.h
@@ -494,6 +494,9 @@ public:
 	/* Determines whether or not the actor may be spawned when running a construction script. If true spawning will fail if a construction script is being run. */
 	uint8	bAllowDuringConstructionScript:1;
 
+	/* Force the spawned actor to use a globally unique name (provided name should be none). */
+	uint8	bForceGloballyUniqueName:1;
+
 #if WITH_EDITOR
 	/* Determines whether the begin play cycle will run on the spawned actor when in the editor. */
 	uint8	bTemporaryEditorActor:1;
diff --git a/Engine/Source/Runtime/Engine/Private/LevelActor.cpp b/Engine/Source/Runtime/Engine/Private/LevelActor.cpp
--- a/Engine/Source/Runtime/Engine/Private/LevelActor.cpp
+++ b/Engine/Source/Runtime/Engine/Private/LevelActor.cpp
@@ -242,10 +242,9 @@ void LineCheckTracker::CaptureLineCheck(int32 LineCheckFlags, const FVector* Ext
 	Level actor management.
 -----------------------------------------------------------------------------*/
 
-#if WITH_EDITOR
 /** 
  * Generates a 102-bits actor GUID:
- *	- Bits 101-54 hold the user MAC address.
+ *	- Bits 101-54 hold the user unique id.
  *	- Bits 53-0 hold a microseconds timestamp.
   *
  * Notes:
@@ -267,15 +266,30 @@ class FActorGUIDGenerator
 public:
 	FActorGUIDGenerator()
 		: Origin(FDateTime(2020, 1, 1))
-		, MacAddress(FPlatformMisc::GetMacAddress())
 		, Counter(0)
 	{
-		if (MacAddress.Num() != 6)
+		TArray<uint8> MACAddress = FPlatformMisc::GetMacAddress();
+
+		if (MACAddress.Num() == 6)
+		{
+			FMemory::Memcpy(UniqueID, MACAddress.GetData(), 6);
+		}
+		else if (IsRunningCommandlet())
+		{
+			FMemory::Memzero(UniqueID, 6);
+		}
+		else
 		{
-			// During cooking, some platform can't retrieve the MAC adress, so force a default one.
-			check(IsRunningCommandlet());
-			MacAddress.Empty(6);
-			MacAddress.AddZeroed(6);
+			const FString EpicAccountId = FPlatformMisc::GetEpicAccountId();
+			const uint32 UniqueUserId = GetTypeHash(EpicAccountId) ^ FPlatformProcess::GetCurrentProcessId();
+			const uint64 ElapsedUs = FDateTime::Now().GetTicks();
+
+			UniqueID[0] = UniqueUserId & 0xff;
+			UniqueID[1] = (UniqueUserId >> 8) & 0xff;
+			UniqueID[2] = (UniqueUserId >> 16) & 0xff;
+			UniqueID[3] = (UniqueUserId >> 24) & 0xff;
+			UniqueID[4] = ElapsedUs & 0xff;
+			UniqueID[5] = (ElapsedUs >> 8) & 0xff;
 		}
 	}
 
@@ -289,8 +303,8 @@ public:
 		const FTimespan Elapsed = FDateTime::Now() - Origin;
 		const uint64 ElapsedUs = (uint64)Elapsed.GetTotalMilliseconds() * 1000 + (Counter++ % 1000);
 
-		// Copy 48-bits MAC address
-		FMemory::Memcpy(HighPart, MacAddress.GetData(), 6);
+		// Copy 48-bits unique id
+		FMemory::Memcpy(HighPart, UniqueID, 6);
 		
 		// Append the high part of the timestamp (will change every ~17 minutes)
 		const uint64 ElapsedUsHighPart = ElapsedUs >> 30;
@@ -312,16 +326,14 @@ public:
 	
 private:
 	FDateTime Origin;
-	TArray<uint8> MacAddress;
+	uint8 UniqueID[6];
 	uint32 Counter;
 };
-#endif
 
 FName FActorSpawnUtils::MakeUniqueActorName(ULevel* Level, const UClass* Class, FName BaseName, bool bGloballyUnique)
 {
 	FName NewActorName;
 
-#if WITH_EDITOR
 	if (bGloballyUnique)
 	{
 		static FActorGUIDGenerator ActorGUIDGenerator;
@@ -333,7 +345,6 @@ FName FActorSpawnUtils::MakeUniqueActorName(ULevel* Level, const UClass* Class,
 		while (StaticFindObjectFast(nullptr, Level, NewActorName));
 	}
 	else
-#endif
 	{
 		NewActorName = MakeUniqueObjectName(Level, Class, BaseName);
 	}
@@ -489,6 +500,11 @@ AActor* UWorld::SpawnActor( UClass* Class, FTransform const* UserTransformPtr, c
 		UE_LOG(LogSpawn, Warning, TEXT("SpawnActor failed because the given transform (%s) is invalid"), *(UserTransformPtr->ToString()));
 		return NULL;
 	}
+	else if (SpawnParameters.bForceGloballyUniqueName && !SpawnParameters.Name.IsNone())
+	{
+		UE_LOG(LogSpawn, Warning, TEXT("SpawnActor failed because a globally unique name was requested and an actor name was provided (%s)"), *SpawnParameters.Name.ToString());
+		return NULL;
+	}
 	
 #if WITH_EDITOR
 	if (SpawnParameters.OverridePackage && SpawnParameters.bCreateActorPackage)
@@ -542,6 +558,9 @@ AActor* UWorld::SpawnActor( UClass* Class, FTransform const* UserTransformPtr, c
 	}
 #endif
 
+	// Override all previous logic if the spawn request a globally unique name
+	bNeedGloballyUniqueName |= SpawnParameters.bForceGloballyUniqueName;
+
 	if (NewActorName.IsNone())
 	{
 		// If we are using a template object and haven't specified a name, create a name relative to the template, otherwise let the default object naming behavior in Stat
diff --git a/Engine/Source/Runtime/Engine/Private/World.cpp b/Engine/Source/Runtime/Engine/Private/World.cpp
--- a/Engine/Source/Runtime/Engine/Private/World.cpp
+++ b/Engine/Source/Runtime/Engine/Private/World.cpp
@@ -170,6 +170,7 @@ FActorSpawnParameters::FActorSpawnParameters()
 , bNoFail(false)
 , bDeferConstruction(false)
 , bAllowDuringConstructionScript(false)
+, bForceGloballyUniqueName(false)
 #if WITH_EDITOR
 , bTemporaryEditorActor(false)
 , bHideFromSceneOutliner(false)
