Fix ActorDescContainerCollection missing some actors when iterating with a specific ActorDesc type
If a container in the middle of the collection did not have the specified an actor desc of the specified type, the iterator would end early.
Also removed useless parameter

diff --git a/Engine/Source/Runtime/Engine/Public/WorldPartition/ActorDescContainerCollection.h b/Engine/Source/Runtime/Engine/Public/WorldPartition/ActorDescContainerCollection.h
--- a/Engine/Source/Runtime/Engine/Public/WorldPartition/ActorDescContainerCollection.h
+++ b/Engine/Source/Runtime/Engine/Public/WorldPartition/ActorDescContainerCollection.h
@@ -82,12 +82,12 @@ public:
 
 	public:
 		template<class T>
-		TBaseIterator(T* Collection, UClass* InActorClass)
+		TBaseIterator(T* Collection)
 			: ContainerIterator(Collection->ActorDescContainerCollection)
 		{
 			if (ContainerIterator)
 			{
-				ActorsIterator = MakeUnique<ActDescIteratorType>(*ContainerIterator, InActorClass);
+				ActorsIterator = MakeUnique<ActDescIteratorType>(*ContainerIterator);
 			}
 		}
 
@@ -98,9 +98,9 @@ public:
 		{
 			++(*ActorsIterator);
 
-			if (!(*ActorsIterator))
+			if (!*ActorsIterator)
 			{
-				AdvanceToNextRegistryContainer();
+				AdvanceToRelevantActorInNextContainer();
 			}
 		}
 
@@ -135,15 +135,15 @@ public:
 		}
 
 	protected:
-		FORCEINLINE void AdvanceToNextRegistryContainer()
+		FORCEINLINE void AdvanceToRelevantActorInNextContainer()
 		{
-			if (++ContainerIterator)
+			while (!(*ActorsIterator) && ContainerIterator)
 			{
-				ActorsIterator = MakeUnique<ActDescIteratorType>(*ContainerIterator, ActorsIterator->GetActorClass());
-			}
-			else
-			{
-				ActorsIterator.Reset();
+				++ContainerIterator;
+				if (ContainerIterator)
+				{
+					ActorsIterator = MakeUnique<ActDescIteratorType>(*ContainerIterator);
+				}
 			}
 		}
 
@@ -158,8 +158,8 @@ public:
 
 	public:
 		template<class T>
-		TIterator(T* Collection, UClass* InActorClass = nullptr)
-			: BaseType(Collection, InActorClass ? InActorClass : ActorType::StaticClass())
+		TIterator(T* Collection)
+			: BaseType(Collection)
 		{}
 	};
 
@@ -170,8 +170,8 @@ public:
 
 	public:
 		template<class T>
-		TConstIterator(T* Collection, UClass* InActorClass = nullptr)
-			: BaseType(Collection, InActorClass ? InActorClass : ActorType::StaticClass())
+		TConstIterator(T* Collection)
+			: BaseType(Collection)
 		{}
 	};
 #endif
