Fix CIS compile error

diff --git a/Engine/Source/Runtime/Engine/Private/LevelActor.cpp b/Engine/Source/Runtime/Engine/Private/LevelActor.cpp
--- a/Engine/Source/Runtime/Engine/Private/LevelActor.cpp
+++ b/Engine/Source/Runtime/Engine/Private/LevelActor.cpp
@@ -268,7 +268,9 @@ public:
 		: Origin(FDateTime(2020, 1, 1))
 		, Counter(0)
 	{
+		PRAGMA_DISABLE_DEPRECATION_WARNINGS
 		TArray<uint8> MACAddress = FPlatformMisc::GetMacAddress();
+		PRAGMA_ENABLE_DEPRECATION_WARNINGS
 
 		if (MACAddress.Num() == 6)
 		{
