Clamp Surround Quality Modifier

diff --git a/Engine/Source/Runtime/Engine/Private/AudioDerivedData.cpp b/Engine/Source/Runtime/Engine/Private/AudioDerivedData.cpp
--- a/Engine/Source/Runtime/Engine/Private/AudioDerivedData.cpp
+++ b/Engine/Source/Runtime/Engine/Private/AudioDerivedData.cpp
@@ -1324,6 +1324,12 @@ struct FAudioCookInputs
 
 #if WITH_EDITORONLY_DATA
 
+static int32 CalculateModifiedCompressionQuality(int32 InQuality, float InQualityModifier)
+{
+	float ModifiedCompressionQuality = (float)InQuality * InQualityModifier;
+	return FMath::Clamp<int32>(FMath::FloorToInt(ModifiedCompressionQuality), 1, 100);
+}
+
 /**
  * Cook a simple mono or stereo wave
  */
@@ -1438,12 +1444,12 @@ static void CookSimpleWave(const FAudioCookInputs& Inputs, TArray<uint8>& Output
 	}
 
 	// Compression Quality
-	FSoundQualityInfo QualityInfo = { 0 };
-	float ModifiedCompressionQuality = (float)Inputs.CompressionQuality * Inputs.CompressionQualityModifier;
+	FSoundQualityInfo QualityInfo = { 0 };	
+	QualityInfo.Quality = CalculateModifiedCompressionQuality(Inputs.CompressionQuality, Inputs.CompressionQualityModifier);
 
-	QualityInfo.Quality = FMath::Clamp<int32>(FMath::FloorToInt(ModifiedCompressionQuality), 1, 100);
-
-	UE_CLOG(Inputs.CompressionQuality != QualityInfo.Quality, LogAudioDerivedData, Display, TEXT("Compression Quality for %s will be modified from %d to %d."), *Inputs.SoundFullName, Inputs.CompressionQuality, QualityInfo.Quality);
+	UE_CLOG(Inputs.CompressionQuality != QualityInfo.Quality, LogAudioDerivedData,
+		Display, TEXT("Compression Quality for %s will be modified from %d to %d, with modifier [%.2f] "),
+		*Inputs.SoundFullName, Inputs.CompressionQuality, QualityInfo.Quality, Inputs.CompressionQualityModifier);
 
 	QualityInfo.NumChannels = NumChannels;
 	QualityInfo.SampleRate = WaveSampleRate;
@@ -1688,22 +1694,11 @@ static void CookSurroundWave(const FAudioCookInputs& Inputs,  TArray<uint8>& Out
 	UE_LOG(LogAudioDerivedData, Display, TEXT("Cooking %d channels for: %s"), ChannelCount, *Inputs.SoundFullName);
 
 	FSoundQualityInfo QualityInfo = { 0 };
+	QualityInfo.Quality = CalculateModifiedCompressionQuality(Inputs.CompressionQuality,Inputs.CompressionQualityModifier);
 
-	float ModifiedCompressionQuality = (float)Inputs.CompressionQuality;
-
-	if (!FMath::IsNearlyEqual(Inputs.CompressionQualityModifier, 1.0f))
-	{
-		ModifiedCompressionQuality = (float)Inputs.CompressionQuality * Inputs.CompressionQualityModifier;
-	}
-	
-	if (ModifiedCompressionQuality >= 1.0f)
-	{
-		QualityInfo.Quality = FMath::FloorToInt(ModifiedCompressionQuality);
-	}
-	else
-	{
-		QualityInfo.Quality = Inputs.CompressionQuality;
-	}
+	UE_CLOG(Inputs.CompressionQuality != QualityInfo.Quality, LogAudioDerivedData,
+		Display, TEXT("Compression Quality for %s will be modified from %d to %d, with modifier [%.2f] "),
+		*Inputs.SoundFullName, Inputs.CompressionQuality, QualityInfo.Quality, Inputs.CompressionQualityModifier);
 
 	QualityInfo.NumChannels = ChannelCount;
 	QualityInfo.SampleRate = WaveSampleRate;
