Animated SparseVolumeTextures can now be sampled in material graphs. This is done by using a UAnimatedSparseVolumeTextureController to control animation playback. Every tick, updating this controller returns a UAnimatedSparseVolumeTextureFrame that contains the current frame to sample. This frame can then be set as a SVT parameter on a MID.

#rnx

diff --git a/Engine/Source/Editor/SparseVolumeTexture/Private/SparseVolumeTextureFactory.cpp b/Engine/Source/Editor/SparseVolumeTexture/Private/SparseVolumeTextureFactory.cpp
--- a/Engine/Source/Editor/SparseVolumeTexture/Private/SparseVolumeTextureFactory.cpp
+++ b/Engine/Source/Editor/SparseVolumeTexture/Private/SparseVolumeTextureFactory.cpp
@@ -454,7 +454,6 @@ UObject* USparseVolumeTextureFactory::FactoryCreateFile(UClass* InClass, UObject
 	}
 
 	// Utility function for computing the bounding box encompassing the bounds of all frames in the SVT.
-	// Accounts for padding caused by rounding up to a multiple of the page resolution.
 	auto ExpandVolumeBounds = [](const FOpenVDBImportOptions& ImportOptions, const TArray<FOpenVDBGridInfo>& GridInfoArray, FBox& VolumeBounds)
 	{
 		const FSparseVolumeRawSourcePackedData* PackedData[] = { &ImportOptions.PackedDataA, &ImportOptions.PackedDataB };
@@ -470,11 +469,9 @@ UObject* USparseVolumeTextureFactory::FactoryCreateFile(UClass* InClass, UObject
 					VolumeBounds.Min.Y = FMath::Min(VolumeBounds.Min.Y, GridInfo.VolumeActiveAABBMin.Y);
 					VolumeBounds.Min.Z = FMath::Min(VolumeBounds.Min.Z, GridInfo.VolumeActiveAABBMin.Z);
 
-					const FVector PageTableVolumeResolution = FMath::DivideAndRoundUp(GridInfo.VolumeActiveDim, FVector(SPARSE_VOLUME_TILE_RES));
-					const FVector Max = GridInfo.VolumeActiveAABBMin + (PageTableVolumeResolution * SPARSE_VOLUME_TILE_RES);
-					VolumeBounds.Max.X = FMath::Max(VolumeBounds.Max.X, Max.X);
-					VolumeBounds.Max.Y = FMath::Max(VolumeBounds.Max.Y, Max.Y);
-					VolumeBounds.Max.Z = FMath::Max(VolumeBounds.Max.Z, Max.Z);
+					VolumeBounds.Max.X = FMath::Max(VolumeBounds.Max.X, GridInfo.VolumeActiveAABBMax.X);
+					VolumeBounds.Max.Y = FMath::Max(VolumeBounds.Max.Y, GridInfo.VolumeActiveAABBMax.Y);
+					VolumeBounds.Max.Z = FMath::Max(VolumeBounds.Max.Z, GridInfo.VolumeActiveAABBMax.Z);
 				}
 			}
 		}
diff --git a/Engine/Source/Runtime/Engine/Classes/SparseVolumeTexture/SparseVolumeTexture.h b/Engine/Source/Runtime/Engine/Classes/SparseVolumeTexture/SparseVolumeTexture.h
--- a/Engine/Source/Runtime/Engine/Classes/SparseVolumeTexture/SparseVolumeTexture.h
+++ b/Engine/Source/Runtime/Engine/Classes/SparseVolumeTexture/SparseVolumeTexture.h
@@ -193,9 +193,10 @@ struct ENGINE_API FSparseVolumeTextureFrame
 
 enum ESparseVolumeTextureShaderUniform
 {
-	ESparseVolumeTexture_PhysicalUVToPageUV,
 	ESparseVolumeTexture_TileSize,
 	ESparseVolumeTexture_PageTableSize,
+	ESparseVolumeTexture_UVScale,
+	ESparseVolumeTexture_UVBias,
 	ESparseVolumeTexture_Count,
 };
 
@@ -211,19 +212,18 @@ public:
 	virtual ~USparseVolumeTexture() = default;
 
 	virtual int32 GetFrameCount() const { return 0; }
-	virtual const FSparseVolumeAssetHeader* GetSparseVolumeTextureHeader(int32 FrameIndex) const { return nullptr; }
-	virtual FSparseVolumeTextureSceneProxy* GetSparseVolumeTextureSceneProxy(int32 FrameIndex) { return nullptr; }
-	virtual const FSparseVolumeTextureSceneProxy* GetSparseVolumeTextureSceneProxy(int32 FrameIndex) const { return nullptr; }
+	virtual const FSparseVolumeAssetHeader* GetSparseVolumeTextureHeader() const { return nullptr; }
+	virtual FSparseVolumeTextureSceneProxy* GetSparseVolumeTextureSceneProxy() { return nullptr; }
+	virtual const FSparseVolumeTextureSceneProxy* GetSparseVolumeTextureSceneProxy() const { return nullptr; }
+	virtual FBox GetVolumeBounds() const { return FBox(); }
 
 	/** Getter for the shader uniform parameters with index as ESparseVolumeTextureShaderUniform. */
-	virtual FVector4 GetUniformParameter(int32 Index, int32 FrameIndex) const { return FVector4(ForceInitToZero); }
-
-	virtual FBox GetVolumeBounds() const { return FBox(); }
+	FVector4 GetUniformParameter(int32 Index) const;
 
 	/** In order to keep the contents of an animated SVT sequence stable in world space, we need to account for the fact that
 		different frames of the sequence have different AABBs. We solve this by scaling and biasing UVs that are relative to
 		the volume bounds into the UV space represented by the AABB of each animation frame.*/
-	void GetFrameUVScaleBias(int32 FrameIndex, FVector* OutScale, FVector* OutBias) const;
+	void GetFrameUVScaleBias(FVector* OutScale, FVector* OutBias) const;
 
 	/** Getter for the shader uniform parameter type with index as ESparseVolumeTextureShaderUniform. */
 	static UE::Shader::EValueType GetUniformParameterType(int32 Index);
@@ -259,12 +259,9 @@ public:
 
 	//~ Begin USparseVolumeTexture Interface.
 	int32 GetFrameCount() const override { return 1; }
-	const FSparseVolumeAssetHeader* GetSparseVolumeTextureHeader(int32 FrameIndex) const override;
-	FSparseVolumeTextureSceneProxy* GetSparseVolumeTextureSceneProxy(int32 FrameIndex) override;
-	const FSparseVolumeTextureSceneProxy* GetSparseVolumeTextureSceneProxy(int32 FrameIndex) const override;
-
-	/** Getter for the shader uniform parameters. */
-	FVector4 GetUniformParameter(int32 Index, int32 FrameIndex) const override;
+	const FSparseVolumeAssetHeader* GetSparseVolumeTextureHeader() const override;
+	FSparseVolumeTextureSceneProxy* GetSparseVolumeTextureSceneProxy() override;
+	const FSparseVolumeTextureSceneProxy* GetSparseVolumeTextureSceneProxy() const override;
 	FBox GetVolumeBounds() const override;
 	//~ End USparseVolumeTexture Interface.
 
@@ -313,17 +310,15 @@ public:
 
 	//~ Begin USparseVolumeTexture Interface.
 	int32 GetFrameCount() const override { return FrameCount; }
-	const FSparseVolumeAssetHeader* GetSparseVolumeTextureHeader(int32 FrameIndex) const override;
-	FSparseVolumeTextureSceneProxy* GetSparseVolumeTextureSceneProxy(int32 FrameIndex) override;
-	const FSparseVolumeTextureSceneProxy* GetSparseVolumeTextureSceneProxy(int32 FrameIndex) const override;
-
-	/** Getter for the shader uniform parameters. */
-	FVector4 GetUniformParameter(int32 Index, int32 FrameIndex) const override;
+	const FSparseVolumeAssetHeader* GetSparseVolumeTextureHeader() const override;
+	FSparseVolumeTextureSceneProxy* GetSparseVolumeTextureSceneProxy() override;
+	const FSparseVolumeTextureSceneProxy* GetSparseVolumeTextureSceneProxy() const override;
 	FBox GetVolumeBounds() const override;
 	//~ End USparseVolumeTexture Interface.
 
 	// Used for debugging a specific frame of an animated sequence.
-	const FSparseVolumeTextureSceneProxy* GetSparseVolumeTextureFrameSceneProxy(int32 FrameIndex) const;
+	FSparseVolumeTextureSceneProxy* GetSparseVolumeTextureFrameSceneProxy(int32 FrameIndex);
+	const FSparseVolumeAssetHeader* GetSparseVolumeTextureFrameHeader(int32 FrameIndex) const;
 
 private:
 
@@ -333,9 +328,75 @@ private:
 	
 	const bool bLoadAllFramesToProxies = true;	// SVT_TODO remove that once streaming is working
 	TArray<FSparseVolumeTextureFrame> AnimationFrames;
+	int32 PreviewFrameIndex;
 
 	int32 GetFrameCountToLoad() const;
 	void ConvertRawSourceDataToSparseVolumeTextureRuntime(int32 FrameIndex);
 	void GenerateOrLoadDDCRuntimeData(int32 FrameIndex);
 	void GenerateOrLoadDDCRuntimeDataAndCreateSceneProxy(int32 FrameIndex);
 };
+
+// USparseVolumeTextureFrame inherits from USparseVolumeTexture to be viewed using any given frame of a UAnimatedSparseVolumeTexture (or UStaticSparseVolumeTexture)
+UCLASS(ClassGroup = Rendering, BlueprintType)//, hidecategories = (Object))
+class ENGINE_API USparseVolumeTextureFrame : public USparseVolumeTexture
+{
+	GENERATED_UCLASS_BODY()
+
+public:
+	USparseVolumeTextureFrame();
+	virtual ~USparseVolumeTextureFrame() = default;
+
+	static USparseVolumeTextureFrame* CreateFrame(USparseVolumeTexture* Texture, int32 FrameIndex);
+
+	void Init(FSparseVolumeTextureSceneProxy* InSceneProxy, const FSparseVolumeAssetHeader* InAssetHeader, const FBox& InVolumeBounds);
+
+	//~ Begin USparseVolumeTexture Interface.
+	int32 GetFrameCount() const override { return 1; }
+	const FSparseVolumeAssetHeader* GetSparseVolumeTextureHeader() const override;
+	FSparseVolumeTextureSceneProxy* GetSparseVolumeTextureSceneProxy() override;
+	const FSparseVolumeTextureSceneProxy* GetSparseVolumeTextureSceneProxy() const override;
+	FBox GetVolumeBounds() const override;
+	//~ End USparseVolumeTexture Interface.
+
+private:
+	FSparseVolumeTextureSceneProxy* SceneProxy;
+	const FSparseVolumeAssetHeader* AssetHeader;
+	FBox VolumeBounds;
+};
+
+UCLASS(ClassGroup = Rendering, BlueprintType)//, hidecategories = (Object))
+class ENGINE_API UAnimatedSparseVolumeTextureController : public UObject
+{
+	GENERATED_UCLASS_BODY()
+
+public:
+
+	UPROPERTY(BlueprintReadWrite, Category = "Animation", meta = (DisplayName = "Frame Rate"))
+	float FrameRate = 24.0f;
+
+	UAnimatedSparseVolumeTextureController();
+	virtual ~UAnimatedSparseVolumeTextureController() = default;
+
+	UFUNCTION(BlueprintCallable, Category = "Animation")
+	void SetSparseVolumeTexture(USparseVolumeTexture* Texture);
+
+	UFUNCTION(BlueprintCallable, Category = "Animation")
+	void Play();
+
+	UFUNCTION(BlueprintCallable, Category = "Animation")
+	void Pause();
+
+	UFUNCTION(BlueprintCallable, Category = "Animation")
+	void Stop();
+
+	UFUNCTION(BlueprintCallable, Category = "Animation")
+	bool IsPlaying();
+
+	UFUNCTION(BlueprintCallable, Category = "Animation")
+	USparseVolumeTextureFrame* Update(float DeltaTime);
+
+private:
+	TObjectPtr<USparseVolumeTexture> SparseVolumeTexture;
+	float Time;
+	bool bIsPlaying;
+};
diff --git a/Engine/Source/Runtime/Engine/Private/Materials/MaterialExpressions.cpp b/Engine/Source/Runtime/Engine/Private/Materials/MaterialExpressions.cpp
--- a/Engine/Source/Runtime/Engine/Private/Materials/MaterialExpressions.cpp
+++ b/Engine/Source/Runtime/Engine/Private/Materials/MaterialExpressions.cpp
@@ -25859,7 +25859,8 @@ int32 UMaterialExpressionSparseVolumeTextureSample::Compile(class FMaterialCompi
 		// SVT_TODO Send parameters and do all the following math in shared functions (one to sample the page table and one to sample the attribute from the page table sample output)
 
 		int32 UVW = CompileWithDefaultFloat3(Compiler, Coordinates, 0.0f, 0.0f, 0.0f);
-		int32 UVWPageTable = Compiler->Mul(UVW, UniformCodeChunkIndex[ESparseVolumeTexture_PhysicalUVToPageUV]);
+		int32 UVWScaled = Compiler->Mul(UVW, UniformCodeChunkIndex[ESparseVolumeTexture_UVScale]);
+		int32 UVWPageTable = Compiler->Add(UVWScaled, UniformCodeChunkIndex[ESparseVolumeTexture_UVBias]);
 		int32 CoordPageTable = Compiler->Mul(UVWPageTable, UniformCodeChunkIndex[ESparseVolumeTexture_PageTableSize]);
 
 		// Code chunks have the same ID when they map to the same string of code. PackedPhysicalTileCoord will automatically end up being the same code chunk used
diff --git a/Engine/Source/Runtime/Engine/Private/Materials/MaterialUniformExpressions.cpp b/Engine/Source/Runtime/Engine/Private/Materials/MaterialUniformExpressions.cpp
--- a/Engine/Source/Runtime/Engine/Private/Materials/MaterialUniformExpressions.cpp
+++ b/Engine/Source/Runtime/Engine/Private/Materials/MaterialUniformExpressions.cpp
@@ -1203,7 +1203,7 @@ void FUniformExpressionSet::FillUniformBuffer(const FMaterialRenderContext& Mate
 					// 1 is tile data A
 					// 2 is tile data B
 
-					const FSparseVolumeTextureSceneProxy* SVTextureProxy = SVTexture->GetSparseVolumeTextureSceneProxy(0 /*SVT_TODO*/);
+					const FSparseVolumeTextureSceneProxy* SVTextureProxy = SVTexture->GetSparseVolumeTextureSceneProxy();
 					if (SVTextureProxy)
 					{
 						if (LayerIndex == 0)
diff --git a/Engine/Source/Runtime/Engine/Private/Shader/Preshader.cpp b/Engine/Source/Runtime/Engine/Private/Shader/Preshader.cpp
--- a/Engine/Source/Runtime/Engine/Private/Shader/Preshader.cpp
+++ b/Engine/Source/Runtime/Engine/Private/Shader/Preshader.cpp
@@ -823,7 +823,7 @@ static void EvaluateSparseTextureUniform(const FMaterialRenderContext& Context,
 	}
 	if (Texture != nullptr && VectorIndex != INDEX_NONE)
 	{
-		Stack.PushValue(FValue(Texture->GetUniformParameter(VectorIndex, 0 /*SVT_TODO*/)));
+		Stack.PushValue(FValue(Texture->GetUniformParameter(VectorIndex)));
 	}
 	else
 	{
diff --git a/Engine/Source/Runtime/Engine/Private/SparseVolumeTexture/SparseVolumeTexture.cpp b/Engine/Source/Runtime/Engine/Private/SparseVolumeTexture/SparseVolumeTexture.cpp
--- a/Engine/Source/Runtime/Engine/Private/SparseVolumeTexture/SparseVolumeTexture.cpp
+++ b/Engine/Source/Runtime/Engine/Private/SparseVolumeTexture/SparseVolumeTexture.cpp
@@ -277,19 +277,63 @@ USparseVolumeTexture::USparseVolumeTexture(const FObjectInitializer& ObjectIniti
 {
 }
 
-void USparseVolumeTexture::GetFrameUVScaleBias(int32 FrameIndex, FVector* OutScale, FVector* OutBias) const
+FVector4 USparseVolumeTexture::GetUniformParameter(int32 Index) const
+{
+	const FSparseVolumeTextureSceneProxy* Proxy = GetSparseVolumeTextureSceneProxy();
+	if (Proxy)
+	{
+		const FSparseVolumeAssetHeader& Header = Proxy->GetHeader();
+		switch (Index)
+		{
+		case ESparseVolumeTexture_TileSize:
+		{
+			return FVector4(float(SPARSE_VOLUME_TILE_RES), 0.0f, 0.0f, 0.0f);
+		}
+		case ESparseVolumeTexture_PageTableSize:
+		{
+			return FVector4(Header.PageTableVolumeResolution.X, Header.PageTableVolumeResolution.Y, Header.PageTableVolumeResolution.Z, 0.0f);
+		}
+		case ESparseVolumeTexture_UVScale: // fallthrough
+		case ESparseVolumeTexture_UVBias:
+		{
+			FVector Scale;
+			FVector Bias;
+			GetFrameUVScaleBias(&Scale, &Bias);
+			return (Index == ESparseVolumeTexture_UVScale) ? FVector4(Scale) : FVector4(Bias);
+		}
+		default:
+		{
+			break;
+		}
+		}
+		checkNoEntry();
+		return FVector4(ForceInitToZero);
+	}
+
+	// 0 while waiting for the proxy
+	return FVector4(ForceInitToZero);
+}
+
+void USparseVolumeTexture::GetFrameUVScaleBias(FVector* OutScale, FVector* OutBias) const
 {
 	*OutScale = FVector::One();
 	*OutBias = FVector::Zero();
-	const FSparseVolumeTextureSceneProxy* Proxy = GetSparseVolumeTextureSceneProxy(FrameIndex);
+	const FSparseVolumeTextureSceneProxy* Proxy = GetSparseVolumeTextureSceneProxy();
 	if (Proxy)
 	{
 		const FSparseVolumeAssetHeader& Header = Proxy->GetHeader();
 		const FBox VolumeBounds = GetVolumeBounds();
 		check(VolumeBounds.GetVolume() > 0.0);
-		const FBox FrameBounds = FBox(FVector(Header.SourceVolumeAABBMin), FVector(Header.SourceVolumeAABBMin + Header.PageTableVolumeResolution * SPARSE_VOLUME_TILE_RES)); // AABB of current frame
+		const FBox FrameBounds = FBox(FVector(Header.SourceVolumeAABBMin), FVector(Header.SourceVolumeAABBMin + Header.SourceVolumeResolution)); // AABB of current frame
+		const FBox FrameBoundsPadded = FBox(FVector(Header.SourceVolumeAABBMin), FVector(Header.SourceVolumeAABBMin + Header.PageTableVolumeResolution * SPARSE_VOLUME_TILE_RES)); // padded to multiple of page size
+		const FVector FrameExtent = FrameBounds.GetExtent();
+
+		// 3D UV coordinates are specified in [0, 1]. Before addressing the page tables which might have padding,
+		// since source volume resolution might not align to tile resolution, we have to rescale the uv so that [0,1] maps to the source texture boundaries.
+		const FVector UVToPhysical = FrameExtent / FrameBoundsPadded.GetExtent();
 
-		*OutScale = VolumeBounds.GetExtent() / FrameBounds.GetExtent();
+		*OutScale = VolumeBounds.GetExtent() / FrameExtent;
+		*OutScale *= UVToPhysical;
 		*OutBias = -((FrameBounds.Min - VolumeBounds.Min) / (VolumeBounds.Max - VolumeBounds.Min) * *OutScale);
 	}
 }
@@ -298,13 +342,14 @@ UE::Shader::EValueType USparseVolumeTexture::GetUniformParameterType(int32 Index
 {
 	switch (Index)
 	{
-	case ESparseVolumeTexture_PhysicalUVToPageUV:	return UE::Shader::EValueType::Float3;
 	case ESparseVolumeTexture_TileSize:				return UE::Shader::EValueType::Float1;
 	case ESparseVolumeTexture_PageTableSize:		return UE::Shader::EValueType::Float3;
+	case ESparseVolumeTexture_UVScale:				return UE::Shader::EValueType::Float3;
+	case ESparseVolumeTexture_UVBias:				return UE::Shader::EValueType::Float3;
 	default:
 		break;
 	}
-	check(0);
+	checkNoEntry();
 	return UE::Shader::EValueType::Float4;
 }
 
@@ -407,49 +452,19 @@ void UStaticSparseVolumeTexture::PostEditChangeProperty(FPropertyChangedEvent& P
 }
 #endif // WITH_EDITOR
 
-const FSparseVolumeAssetHeader* UStaticSparseVolumeTexture::GetSparseVolumeTextureHeader(int32 FrameIndex) const
+const FSparseVolumeAssetHeader* UStaticSparseVolumeTexture::GetSparseVolumeTextureHeader() const
 {
 	return &StaticFrame.SparseVolumeTextureRuntime.Header;
 }
-FSparseVolumeTextureSceneProxy* UStaticSparseVolumeTexture::GetSparseVolumeTextureSceneProxy(int32 FrameIndex)
+FSparseVolumeTextureSceneProxy* UStaticSparseVolumeTexture::GetSparseVolumeTextureSceneProxy()
 { 
 	return &StaticFrame.SparseVolumeTextureSceneProxy; 
 }
-const FSparseVolumeTextureSceneProxy* UStaticSparseVolumeTexture::GetSparseVolumeTextureSceneProxy(int32 FrameIndex) const
+const FSparseVolumeTextureSceneProxy* UStaticSparseVolumeTexture::GetSparseVolumeTextureSceneProxy() const
 { 
 	return &StaticFrame.SparseVolumeTextureSceneProxy; 
 }
 
-FVector4 UStaticSparseVolumeTexture::GetUniformParameter(int32 Index, int32 FrameIndex) const
-{
-	const FSparseVolumeTextureSceneProxy* Proxy = GetSparseVolumeTextureSceneProxy(0);
-	if (Proxy)
-	{
-		const FSparseVolumeAssetHeader& Header = Proxy->GetHeader();
-		switch (Index)
-		{
-		case ESparseVolumeTexture_PhysicalUVToPageUV:
-		{
-			// 3d uv coordinates are specified in [0, 1]. Before addressing the page tables which might have padding,
-			// since source volume resolution might not align to tile resolution, we have to rescale the uv so that [0,1] maps to the source texture boundaries.
-			FVector3f PhysicalUVToPageUV = FVector3f(Header.SourceVolumeResolution)/ FVector3f(Header.PageTableVolumeResolution * SPARSE_VOLUME_TILE_RES);
-			return FVector4(PhysicalUVToPageUV.X, PhysicalUVToPageUV.Y, PhysicalUVToPageUV.Z, 0.0f);
-		}
-		case ESparseVolumeTexture_PageTableSize:
-		{
-			return FVector4(Header.PageTableVolumeResolution.X, Header.PageTableVolumeResolution.Y, Header.PageTableVolumeResolution.Z, 0.0f);
-		}
-		case ESparseVolumeTexture_TileSize:				return FVector4(float(SPARSE_VOLUME_TILE_RES), 0.0f, 0.0f, 0.0f);
-		default:
-			break;
-		}
-		check(0);
-		return FVector4(ForceInitToZero);
-	}
-	// 0 while waiting for the proxy
-	return FVector4(ForceInitToZero);
-}
-
 FBox UStaticSparseVolumeTexture::GetVolumeBounds() const
 {
 	return VolumeBounds;
@@ -678,71 +693,39 @@ void UAnimatedSparseVolumeTexture::PostEditChangeProperty(FPropertyChangedEvent&
 }
 #endif // WITH_EDITOR
 
-const FSparseVolumeAssetHeader* UAnimatedSparseVolumeTexture::GetSparseVolumeTextureHeader(int32 FrameIndex) const
+const FSparseVolumeAssetHeader* UAnimatedSparseVolumeTexture::GetSparseVolumeTextureHeader() const
 {
 	// When an AnimatedSparseVolumeTexture is used as SparseVolumeTexture, it can only be previewed using its first frame.
 	check(AnimationFrames.Num() >= 1);
-	FrameIndex = FrameIndex % GetFrameCountToLoad();
+	const int32 FrameIndex = PreviewFrameIndex % GetFrameCountToLoad();
 	const FSparseVolumeTextureFrame& Frame = AnimationFrames[FrameIndex];
 	return &Frame.SparseVolumeTextureRuntime.Header;
 }
 
-FSparseVolumeTextureSceneProxy* UAnimatedSparseVolumeTexture::GetSparseVolumeTextureSceneProxy(int32 FrameIndex)
+FSparseVolumeTextureSceneProxy* UAnimatedSparseVolumeTexture::GetSparseVolumeTextureSceneProxy()
 {
 	// When an AnimatedSparseVolumeTexture is used as SparseVolumeTexture, it can only be previewed using its first frame.
 	check(AnimationFrames.Num() >= 1);
-	FrameIndex = FrameIndex % GetFrameCountToLoad();
+	const int32 FrameIndex = PreviewFrameIndex % GetFrameCountToLoad();
 	FSparseVolumeTextureFrame& Frame = AnimationFrames[FrameIndex];
 	return &Frame.SparseVolumeTextureSceneProxy;
 }
 
-const FSparseVolumeTextureSceneProxy* UAnimatedSparseVolumeTexture::GetSparseVolumeTextureSceneProxy(int32 FrameIndex) const
+const FSparseVolumeTextureSceneProxy* UAnimatedSparseVolumeTexture::GetSparseVolumeTextureSceneProxy() const
 {
 	// When an AnimatedSparseVolumeTexture is used as SparseVolumeTexture, it can only be previewed using its first frame.
 	check(AnimationFrames.Num() >= 1);
-	FrameIndex = FrameIndex % GetFrameCountToLoad();
+	const int32 FrameIndex = PreviewFrameIndex % GetFrameCountToLoad();
 	const FSparseVolumeTextureFrame& Frame = AnimationFrames[FrameIndex];
 	return &Frame.SparseVolumeTextureSceneProxy;
 }
 
-FVector4 UAnimatedSparseVolumeTexture::GetUniformParameter(int32 Index, int32 FrameIndex) const
-{
-	// When an AnimatedSparseVolumeTexture is used as SparseVolumeTexture, it can only be previewed using its first frame.
-	const FSparseVolumeTextureSceneProxy* Proxy = GetSparseVolumeTextureSceneProxy(FrameIndex);
-	if (Proxy)
-	{
-		const FSparseVolumeAssetHeader& Header = Proxy->GetHeader();
-		switch (Index)
-		{
-		case ESparseVolumeTexture_PhysicalUVToPageUV:
-		{
-			// 3d uv coordinates are specified in [0, 1]. Before addressing the page tables which might have padding,
-			// since source volume resolution might not align to tile resolution, we have to rescale the uv so that [0,1] maps to the source texture boundaries.
-			FVector3f PhysicalUVToPageUV = FVector3f(Header.SourceVolumeResolution) / FVector3f(Header.PageTableVolumeResolution * SPARSE_VOLUME_TILE_RES);
-			return FVector4(PhysicalUVToPageUV.X, PhysicalUVToPageUV.Y, PhysicalUVToPageUV.Z, 0.0f);
-		}
-		case ESparseVolumeTexture_PageTableSize:
-		{
-			return FVector4(Header.PageTableVolumeResolution.X, Header.PageTableVolumeResolution.Y, Header.PageTableVolumeResolution.Z, 0.0f);
-		}
-		case ESparseVolumeTexture_TileSize:				return FVector4(float(SPARSE_VOLUME_TILE_RES), 0.0f, 0.0f, 0.0f);
-		default:
-			break;
-		}
-		check(0);
-		return FVector4(ForceInitToZero);
-	}
-
-	// 0 while waiting for the proxy
-	return FVector4(ForceInitToZero);
-}
-
 FBox UAnimatedSparseVolumeTexture::GetVolumeBounds() const
 {
 	return VolumeBounds;
 }
 
-const FSparseVolumeTextureSceneProxy* UAnimatedSparseVolumeTexture::GetSparseVolumeTextureFrameSceneProxy(int32 FrameIndex) const
+FSparseVolumeTextureSceneProxy* UAnimatedSparseVolumeTexture::GetSparseVolumeTextureFrameSceneProxy(int32 FrameIndex)
 {
 	check(AnimationFrames.Num() >= 1);
 	FrameIndex = FrameIndex % GetFrameCountToLoad();
@@ -750,6 +733,14 @@ const FSparseVolumeTextureSceneProxy* UAnimatedSparseVolumeTexture::GetSparseVol
 	return &AnimationFrames[FrameIndex].SparseVolumeTextureSceneProxy;
 }
 
+const FSparseVolumeAssetHeader* UAnimatedSparseVolumeTexture::GetSparseVolumeTextureFrameHeader(int32 FrameIndex) const
+{
+	check(AnimationFrames.Num() >= 1);
+	FrameIndex = FrameIndex % GetFrameCountToLoad();
+	// SVT_TODO when streaming is enabled, this will likely change.
+	return &AnimationFrames[FrameIndex].SparseVolumeTextureRuntime.Header;
+}
+
 int32 UAnimatedSparseVolumeTexture::GetFrameCountToLoad() const
 {
 	if (FrameCount > 0)
@@ -854,4 +845,127 @@ void UAnimatedSparseVolumeTexture::GenerateOrLoadDDCRuntimeDataAndCreateScenePro
 
 ////////////////////////////////////////////////////////////////////////////////////////////////
 
+USparseVolumeTextureFrame::USparseVolumeTextureFrame(const FObjectInitializer& ObjectInitializer)
+	: Super(ObjectInitializer)
+{
+}
+
+USparseVolumeTextureFrame* USparseVolumeTextureFrame::CreateFrame(USparseVolumeTexture* Texture, int32 FrameIndex)
+{
+	USparseVolumeTextureFrame* Frame = NewObject<USparseVolumeTextureFrame>();
+	FSparseVolumeTextureSceneProxy* Proxy = nullptr;
+	if (Texture->IsA<UAnimatedSparseVolumeTexture>())
+	{
+		UAnimatedSparseVolumeTexture* AnimatedSVT = CastChecked<UAnimatedSparseVolumeTexture>(Texture);
+		check(AnimatedSVT);
+		Proxy = AnimatedSVT->GetSparseVolumeTextureFrameSceneProxy(FrameIndex);
+	}
+	else
+	{
+		Proxy = Texture->GetSparseVolumeTextureSceneProxy();
+	}
+
+	Frame->Init(Proxy, Texture->GetSparseVolumeTextureHeader(), Texture->GetVolumeBounds());
+
+	return Frame;
+}
+
+void USparseVolumeTextureFrame::Init(FSparseVolumeTextureSceneProxy* InSceneProxy, const FSparseVolumeAssetHeader* InAssetHeader, const FBox& InVolumeBounds)
+{
+	SceneProxy = InSceneProxy;
+	AssetHeader = InAssetHeader;
+	VolumeBounds = InVolumeBounds;
+}
+
+const FSparseVolumeAssetHeader* USparseVolumeTextureFrame::GetSparseVolumeTextureHeader() const
+{
+	return AssetHeader;
+}
+
+FSparseVolumeTextureSceneProxy* USparseVolumeTextureFrame::GetSparseVolumeTextureSceneProxy()
+{
+	return SceneProxy;
+}
+
+const FSparseVolumeTextureSceneProxy* USparseVolumeTextureFrame::GetSparseVolumeTextureSceneProxy() const
+{
+	return SceneProxy;
+}
+
+FBox USparseVolumeTextureFrame::GetVolumeBounds() const
+{
+	return VolumeBounds;
+}
+
+////////////////////////////////////////////////////////////////////////////////////////////////
+
+UAnimatedSparseVolumeTextureController::UAnimatedSparseVolumeTextureController(const FObjectInitializer& ObjectInitializer)
+	: Super(ObjectInitializer)
+{
+}
+
+void UAnimatedSparseVolumeTextureController::SetSparseVolumeTexture(USparseVolumeTexture* Texture)
+{
+	if (Texture == SparseVolumeTexture)
+	{
+		return;
+	}
+
+	SparseVolumeTexture = Texture;
+	bIsPlaying = bIsPlaying && (SparseVolumeTexture != nullptr);
+	Time = 0.0f;
+}
+
+void UAnimatedSparseVolumeTextureController::Play()
+{
+	bIsPlaying = true;
+}
+
+void UAnimatedSparseVolumeTextureController::Pause()
+{
+	bIsPlaying = false;
+}
+
+void UAnimatedSparseVolumeTextureController::Stop()
+{
+	if (bIsPlaying)
+	{
+		bIsPlaying = false;
+		Time = 0.0f;
+	}
+}
+
+bool UAnimatedSparseVolumeTextureController::IsPlaying()
+{
+	return bIsPlaying;
+}
+
+USparseVolumeTextureFrame* UAnimatedSparseVolumeTextureController::Update(float DeltaTime)
+{
+	if (!SparseVolumeTexture)
+	{
+		return nullptr;
+	}
+
+	const int32 FrameCount = SparseVolumeTexture->GetFrameCount();
+	const float AnimationDuration = FrameCount / (FrameRate + UE_SMALL_NUMBER);
+
+	// Update animation time
+	if (bIsPlaying)
+	{
+		Time = FMath::Fmod(Time + DeltaTime, AnimationDuration + UE_SMALL_NUMBER);
+	}
+
+	// Compute (fractional) index of frame to sample
+	const float FrameIndexF = FMath::Fmod(Time * FrameRate, (float)FrameCount);
+	const int32 FrameIndex = (int32)FrameIndexF;
+
+	// Create and initialize a USparseVolumeTextureFrame which holds the frame to sample and can be bound to shaders
+	USparseVolumeTextureFrame* Frame = USparseVolumeTextureFrame::CreateFrame(SparseVolumeTexture, FrameIndex);
+
+	return Frame;
+}
+
+////////////////////////////////////////////////////////////////////////////////////////////////
+
 #undef LOCTEXT_NAMESPACE
diff --git a/Engine/Source/Runtime/Renderer/Private/SparseVolumeTexture/SparseVolumeTextureViewerComponent.cpp b/Engine/Source/Runtime/Renderer/Private/SparseVolumeTexture/SparseVolumeTextureViewerComponent.cpp
--- a/Engine/Source/Runtime/Renderer/Private/SparseVolumeTexture/SparseVolumeTextureViewerComponent.cpp
+++ b/Engine/Source/Runtime/Renderer/Private/SparseVolumeTexture/SparseVolumeTextureViewerComponent.cpp
@@ -160,13 +160,13 @@ void USparseVolumeTextureViewerComponent::SendRenderTransformCommand()
 		FVector UVScale = FVector::One();
 		FVector UVBias = FVector::Zero();
 		
-		if (SparseVolumeTexturePreview)
+		if (SparseVolumeTextureFrame)
 		{
-			const FVector VolumeBoundsExtent = SparseVolumeTexturePreview->GetVolumeBounds().GetExtent();
+			const FVector VolumeBoundsExtent = SparseVolumeTextureFrame->GetVolumeBounds().GetExtent();
 			const double MaxBoundsDim = FMath::Max(FMath::Max(VolumeBoundsExtent.X, VolumeBoundsExtent.Y), VolumeBoundsExtent.Z);
 			VolumeExtent = VolumeBoundsExtent / MaxBoundsDim * SVTViewerDefaultVolumeExtent;
 
-			SparseVolumeTexturePreview->GetFrameUVScaleBias(FrameIndex, &UVScale, &UVBias);
+			SparseVolumeTextureFrame->GetFrameUVScaleBias(&UVScale, &UVBias);
 		}
 
 		const FTransform ToWorldTransform = GetComponentTransform();
@@ -199,18 +199,24 @@ void USparseVolumeTextureViewerComponent::SendRenderTransformCommand()
 
 void USparseVolumeTextureViewerComponent::TickComponent(float DeltaTime, enum ELevelTick TickType, FActorComponentTickFunction* ThisTickFunction)
 {
-	if (bAnimate)
-	{
-		const int32 NumFrames = SparseVolumeTexturePreview->GetFrameCount();
-		const float AnimationDuration = NumFrames / FrameRate;
-		FrameIndex = int32((AnimationTime / (AnimationDuration + SMALL_NUMBER)) * float(SparseVolumeTexturePreview->GetFrameCount()));
-		AnimationTime = FMath::Fmod(AnimationTime + DeltaTime, (AnimationDuration + SMALL_NUMBER));
-	}
-	else if(SparseVolumeTexturePreview)
+	if (SparseVolumeTexturePreview)
 	{
-		FrameIndex = int32(AnimationFrame * float(SparseVolumeTexturePreview->GetFrameCount()));
+		if (bAnimate)
+		{
+			const int32 NumFrames = SparseVolumeTexturePreview->GetFrameCount();
+			const float AnimationDuration = NumFrames / (FrameRate + UE_SMALL_NUMBER);
+			AnimationTime = FMath::Fmod(AnimationTime + DeltaTime, (AnimationDuration + UE_SMALL_NUMBER));
+			FrameIndex = int32(AnimationTime * FrameRate);
+		}
+		else
+		{
+			FrameIndex = int32(AnimationFrame * float(SparseVolumeTexturePreview->GetFrameCount()));
+		}
+
+		SparseVolumeTextureFrame = USparseVolumeTextureFrame::CreateFrame(SparseVolumeTexturePreview, FrameIndex);
+
+		MarkRenderStateDirty();
 	}
-	MarkRenderStateDirty();
 }
 
 /*=============================================================================
diff --git a/Engine/Source/Runtime/Renderer/Private/SparseVolumeTexture/SparseVolumeTextureViewerComponent.h b/Engine/Source/Runtime/Renderer/Private/SparseVolumeTexture/SparseVolumeTextureViewerComponent.h
--- a/Engine/Source/Runtime/Renderer/Private/SparseVolumeTexture/SparseVolumeTextureViewerComponent.h
+++ b/Engine/Source/Runtime/Renderer/Private/SparseVolumeTexture/SparseVolumeTextureViewerComponent.h
@@ -44,6 +44,8 @@ class USparseVolumeTextureViewerComponent : public UPrimitiveComponent
 	UPROPERTY(EditAnywhere, BlueprintReadWrite, Category = "Asset Preview", meta = (UIMin = 0, UIMax = 7, ClampMin = 0, ClampMax = 7))
 	int32 ComponentToVisualize;
 
+	TObjectPtr<class USparseVolumeTextureFrame> SparseVolumeTextureFrame;
+
 	//~ Begin UPrimitiveComponent Interface.
 	virtual bool SupportsStaticLighting() const override { return false; }
 	//~ End UPrimitiveComponent Interface.
diff --git a/Engine/Source/Runtime/Renderer/Private/SparseVolumeTexture/SparseVolumeTextureViewerSceneProxy.cpp b/Engine/Source/Runtime/Renderer/Private/SparseVolumeTexture/SparseVolumeTextureViewerSceneProxy.cpp
--- a/Engine/Source/Runtime/Renderer/Private/SparseVolumeTexture/SparseVolumeTextureViewerSceneProxy.cpp
+++ b/Engine/Source/Runtime/Renderer/Private/SparseVolumeTexture/SparseVolumeTextureViewerSceneProxy.cpp
@@ -40,19 +40,10 @@ FSparseVolumeTextureViewerSceneProxy::FSparseVolumeTextureViewerSceneProxy(const
 	: FPrimitiveSceneProxy((UPrimitiveComponent*)InComponent, ResourceName)
 	, SparseVolumeTextureSceneProxy(nullptr)
 {
-	if (InComponent->SparseVolumeTexturePreview)
+	if (InComponent->SparseVolumeTextureFrame)
 	{
-		check(InComponent->SparseVolumeTexturePreview);
-		if (InComponent->SparseVolumeTexturePreview->IsA(UAnimatedSparseVolumeTexture::StaticClass()))
-		{
-			const UAnimatedSparseVolumeTexture* AnimatedSparseVolumeTexture = CastChecked<const UAnimatedSparseVolumeTexture>(InComponent->SparseVolumeTexturePreview);
-			SparseVolumeTextureSceneProxy = AnimatedSparseVolumeTexture->GetSparseVolumeTextureFrameSceneProxy(FrameIndex);
-		}
-		else
-		{
-			// Regular texture from proxy
-			SparseVolumeTextureSceneProxy = InComponent->SparseVolumeTexturePreview->GetSparseVolumeTextureSceneProxy(FrameIndex);
-		}
+		check(InComponent->SparseVolumeTextureFrame);
+		SparseVolumeTextureSceneProxy = InComponent->SparseVolumeTextureFrame->GetSparseVolumeTextureSceneProxy();
 	}
 }
 
