Fix sub object re-instancing.

diff --git a/Engine/Source/Editor/UnrealEd/Private/Kismet2/KismetReinstanceUtilities.cpp b/Engine/Source/Editor/UnrealEd/Private/Kismet2/KismetReinstanceUtilities.cpp
--- a/Engine/Source/Editor/UnrealEd/Private/Kismet2/KismetReinstanceUtilities.cpp
+++ b/Engine/Source/Editor/UnrealEd/Private/Kismet2/KismetReinstanceUtilities.cpp
@@ -2296,6 +2296,13 @@ void FBlueprintCompileReinstancer::ReplaceInstancesOfClass_Inner(const TMap<UCla
 				for (int32 OldObjIndex = 0; OldObjIndex < ObjectsToReplace.Num(); ++OldObjIndex)
 				{
 					UObject* OldObject = ObjectsToReplace[OldObjIndex];
+
+					// Skipping any object that outer is going to be replaced
+					// This isn't needed for the actor loop as the only outer for actor is a level
+					if (InOldToNewClassMap.Contains(OldObject->GetOuter()->GetClass()))
+					{
+						continue;
+					}
 					
 					AActor* OldActor = Cast<AActor>(OldObject);
 
diff --git a/Engine/Source/Runtime/Engine/Private/InstancedReferenceSubobjectHelper.cpp b/Engine/Source/Runtime/Engine/Private/InstancedReferenceSubobjectHelper.cpp
--- a/Engine/Source/Runtime/Engine/Private/InstancedReferenceSubobjectHelper.cpp
+++ b/Engine/Source/Runtime/Engine/Private/InstancedReferenceSubobjectHelper.cpp
@@ -225,13 +225,23 @@ void FFindInstancedReferenceSubobjectHelper::ForEachInstancedSubObject(FInstance
 				ObjRefFunc(FInstancedSubObjRef(ObjectValue, PropertyPath), ContainerAddress);
 			}
 		}
+		else
+		{
+			// @todo figure out how to support verse property types
+			//UE_LOG(LogBlueprint, Error, TEXT("%s: Skipping unknow property type(%s)"), ANSI_TO_TCHAR(__FUNCTION__), *TargetProp->GetName());
+		}
 	}
 }
 
 template ENGINE_API void FFindInstancedReferenceSubobjectHelper::ForEachInstancedSubObject<void*>(FInstancedPropertyPath& PropertyPath, void* ContainerAddress, TFunctionRef<void(const FInstancedSubObjRef&, void*)> ObjRefFunc);
 template ENGINE_API void FFindInstancedReferenceSubobjectHelper::ForEachInstancedSubObject<const void*>(FInstancedPropertyPath& PropertyPath, const void* ContainerAddress, TFunctionRef<void(const FInstancedSubObjRef&, const void*)> ObjRefFunc);
 
-void FFindInstancedReferenceSubobjectHelper::Duplicate(UObject* OldObject, UObject* NewObject, TMap<UObject*, UObject*>& ReferenceReplacementMap, TArray<UObject*>& DuplicatedObjects)
+void FFindInstancedReferenceSubobjectHelper::Duplicate(
+	UObject* OldObject, 
+	UObject* NewObject, 
+	TMap<UObject*, UObject*>& ReferenceReplacementMap, 
+	TArray<UObject*>& DuplicatedObjects, 
+	const TMap<UObject*, UObject*>* OptionalMappings)
 {
 	if (OldObject->GetClass()->HasAnyClassFlags(CLASS_HasInstancedReference) &&
 		NewObject->GetClass()->HasAnyClassFlags(CLASS_HasInstancedReference))
@@ -248,7 +258,8 @@ void FFindInstancedReferenceSubobjectHelper::Duplicate(UObject* OldObject, UObje
 				if (bNewObjectHasOldOuter)
 				{
 					const bool bKeptByOld = OldInstancedSubObjects.Contains(Obj);
-					const bool bNotHandledYet = !ReferenceReplacementMap.Contains(Obj);
+					const bool bNotHandledYet = !ReferenceReplacementMap.Contains(Obj) &&
+						(!OptionalMappings || !OptionalMappings->Contains(Obj));
 					if (bKeptByOld)
 					{
 						if (bNotHandledYet)
@@ -280,7 +291,20 @@ void FFindInstancedReferenceSubobjectHelper::Duplicate(UObject* OldObject, UObje
 						{
 							// make sure the object is outered correctly:
 							check(bNewObjectHasOldOuter);
-							UObject* RealNewSubobject = ReferenceReplacementMap.FindChecked(Obj);
+							UObject* RealNewSubobject = nullptr;
+							if (OptionalMappings)
+							{
+								UObject* const* Entry = OptionalMappings->Find(Obj);
+								if (Entry)
+								{
+									RealNewSubobject = *Entry;
+								}
+							}
+
+							if (!RealNewSubobject)
+							{
+								RealNewSubobject = ReferenceReplacementMap.FindChecked(Obj);
+							}
 							if (RealNewSubobject->GetOuter() != NewObject)
 							{
 								UObject* ExistingObject = StaticFindObjectFast(UObject::StaticClass(), NewObject, RealNewSubobject->GetFName());
diff --git a/Engine/Source/Runtime/Engine/Private/UnrealEngine.cpp b/Engine/Source/Runtime/Engine/Private/UnrealEngine.cpp
--- a/Engine/Source/Runtime/Engine/Private/UnrealEngine.cpp
+++ b/Engine/Source/Runtime/Engine/Private/UnrealEngine.cpp
@@ -16758,10 +16758,6 @@ void UEngine::CopyPropertiesForUnrelatedObjects(UObject* OldObject, UObject* New
 		ReferenceReplacementMap.Add(OldObject->GetClass(), NewObject->GetClass());
 	}
 	ReferenceReplacementMap.Add(OldObject->GetClass()->GetDefaultObject(), NewObject->GetClass()->GetDefaultObject());
-	if (Params.OptionalReplacementMappings)
-	{
-		ReferenceReplacementMap.Append(*Params.OptionalReplacementMappings);
-	}
 
 	TArray<UObject*> ComponentsOnNewObject;
 	TMap<UObject*, UObject*> AggressiveReplaceReferences;
@@ -16857,13 +16853,13 @@ void UEngine::CopyPropertiesForUnrelatedObjects(UObject* OldObject, UObject* New
 				// Restore modified properties into the new instance
 				FInstancedObjectRecord& Record = SavedInstances[*pOldInstanceIndex];
 				FObjectReader Reader(NewInstance, Record.SavedProperties, true, true);
-				FFindInstancedReferenceSubobjectHelper::Duplicate(Record.OldInstance, NewInstance, ReferenceReplacementMap, EditInlineSubobjectsOfComponents);
+				FFindInstancedReferenceSubobjectHelper::Duplicate(Record.OldInstance, NewInstance, ReferenceReplacementMap, EditInlineSubobjectsOfComponents, Params.OptionalReplacementMappings);
 			}
 		}
 		ComponentsOnNewObject.Append(MoveTemp(EditInlineSubobjectsOfComponents));
 	}
 
-	FFindInstancedReferenceSubobjectHelper::Duplicate(OldObject, NewObject, ReferenceReplacementMap, ComponentsOnNewObject);
+	FFindInstancedReferenceSubobjectHelper::Duplicate(OldObject, NewObject, ReferenceReplacementMap, ComponentsOnNewObject, Params.OptionalReplacementMappings);
 
 	// Replace anything with an outer of the old object with NULL, unless it already has a replacement or is marked pending kill
 	// Also, if requested, leave pointers to instances of renewed classes intact, assuming we are in the midst of reinstancing, and those will be reinstanced
@@ -16886,6 +16882,10 @@ void UEngine::CopyPropertiesForUnrelatedObjects(UObject* OldObject, UObject* New
 		else
 		{
 			FArchiveReplaceOrClearExternalReferences<UObject> ReplaceInCDOAr(NewObject, ReferenceReplacementMap, NewPackage);
+			if (Params.OptionalReplacementMappings)
+			{
+				FArchiveReplaceOrClearExternalReferences<UObject> OptionalReplaceInCDOAr(NewObject, *Params.OptionalReplacementMappings, NewPackage);
+			}
 		}
 
 		// Replace references inside each individual component. This is always required because if something is in ReferenceReplacementMap, the above replace code will skip fixing child properties
@@ -16899,6 +16899,10 @@ void UEngine::CopyPropertiesForUnrelatedObjects(UObject* OldObject, UObject* New
 			else
 			{
 				FArchiveReplaceOrClearExternalReferences<UObject> ReplaceInComponentAr(NewComponent, ReferenceReplacementMap, NewPackage);
+				if (Params.OptionalReplacementMappings)
+				{
+					FArchiveReplaceOrClearExternalReferences<UObject> OptionalReplaceInCDOAr(NewComponent, *Params.OptionalReplacementMappings, NewPackage);
+				}
 			}
 		}
 	}
diff --git a/Engine/Source/Runtime/Engine/Public/InstancedReferenceSubobjectHelper.h b/Engine/Source/Runtime/Engine/Public/InstancedReferenceSubobjectHelper.h
--- a/Engine/Source/Runtime/Engine/Public/InstancedReferenceSubobjectHelper.h
+++ b/Engine/Source/Runtime/Engine/Public/InstancedReferenceSubobjectHelper.h
@@ -114,7 +114,12 @@ public:
 		}
 	}
 
-	static void Duplicate(UObject* OldObject, UObject* NewObject, TMap<UObject*, UObject*>& ReferenceReplacementMap, TArray<UObject*>& DuplicatedObjects);
+	static void Duplicate(
+		UObject* OldObject, 
+		UObject* NewObject, 
+		TMap<UObject*, UObject*>& ReferenceReplacementMap, 
+		TArray<UObject*>& DuplicatedObjects, 
+		const TMap<UObject*, UObject*>* OptionalMappings);
 
 	template<typename T>
 	static void ForEachInstancedSubObject(FInstancedPropertyPath& PropertyPath, T ContainerAddress, TFunctionRef<void(const FInstancedSubObjRef& Ref, T PropertyValueAddress)> ObjRefFunc);
