Add a way to check if a path point to content bundle content but is not an ExternalActor

diff --git a/Engine/Plugins/Editor/DataValidation/Source/DataValidation/Private/WorldPartitionChangelistValidator.cpp b/Engine/Plugins/Editor/DataValidation/Source/DataValidation/Private/WorldPartitionChangelistValidator.cpp
--- a/Engine/Plugins/Editor/DataValidation/Source/DataValidation/Private/WorldPartitionChangelistValidator.cpp
+++ b/Engine/Plugins/Editor/DataValidation/Source/DataValidation/Private/WorldPartitionChangelistValidator.cpp
@@ -184,13 +184,13 @@ EDataValidationResult UWorldPartitionChangelistValidator::ValidateActorsAndDataL
 		for (const FAssetData& ActorData : ActorsData)
 		{
 			FString ActorPackagePath = ActorData.PackagePath.ToString();
-			if (ContentBundlePaths::IsAContentBundlePackagePath(ActorPackagePath))
+			if (ContentBundlePaths::IsAContentBundleExternalActorPackagePath(ActorPackagePath))
 			{
 				FStringView ContentBundleMountPoint = FPathViews::GetMountPointNameFromPath(ActorPackagePath);
 				FGuid ContentBundleGuid = ContentBundlePaths::GetContentBundleGuidFromExternalActorPackagePath(ActorPackagePath);
 				
 				FString ContentBundleContainerPackagePath;
-				verify(ContentBundlePaths::BuildActorDescContainerPackgePath(FString(ContentBundleMountPoint), ContentBundleGuid, MapPath.GetPackageName().ToString(), ContentBundleContainerPackagePath));
+				verify(ContentBundlePaths::BuildActorDescContainerPackagePath(FString(ContentBundleMountPoint), ContentBundleGuid, MapPath.GetPackageName().ToString(), ContentBundleContainerPackagePath));
 
 				RegisterContainerToValidate(World, FName(*ContentBundleContainerPackagePath), ContainersToValidate);
 			}
diff --git a/Engine/Plugins/Experimental/GameFeatures/Source/GameFeatures/Private/GameFeatureDataAssetDependencyGatherer.cpp b/Engine/Plugins/Experimental/GameFeatures/Source/GameFeatures/Private/GameFeatureDataAssetDependencyGatherer.cpp
--- a/Engine/Plugins/Experimental/GameFeatures/Source/GameFeatures/Private/GameFeatureDataAssetDependencyGatherer.cpp
+++ b/Engine/Plugins/Experimental/GameFeatures/Source/GameFeatures/Private/GameFeatureDataAssetDependencyGatherer.cpp
@@ -26,10 +26,10 @@ void FGameFeatureDataAssetDependencyGatherer::GatherDependencies(const FAssetDat
 		FName MountPoint = FPackageName::GetPackageMountPoint(AssetData.PackagePath.ToString());
 		for (const FGuid& ContentBundleGuid : ContentBundleGuids)
 		{
-			FString ContentBundleRootPath;
-			if (ContentBundlePaths::BuildContentBundleRootPath(MountPoint.ToString(), ContentBundleGuid, ContentBundleRootPath))
+			FString ContentBundleExternalActorPath;
+			if (ContentBundlePaths::BuildContentBundleExternalActorPath(MountPoint.ToString(), ContentBundleGuid, ContentBundleExternalActorPath))
 			{
-				const FString ExternalActorsPath = ULevel::GetExternalActorsPath(ContentBundleRootPath);
+				const FString ExternalActorsPath = ULevel::GetExternalActorsPath(ContentBundleExternalActorPath);
 
 				OutDependencyDirectories.Add(ExternalActorsPath);
 				Filter.PackagePaths.Add(*ExternalActorsPath);
diff --git a/Engine/Source/Runtime/Engine/Private/WorldPartition/ContentBundle/ContentBundleEditor.cpp b/Engine/Source/Runtime/Engine/Private/WorldPartition/ContentBundle/ContentBundleEditor.cpp
--- a/Engine/Source/Runtime/Engine/Private/WorldPartition/ContentBundle/ContentBundleEditor.cpp
+++ b/Engine/Source/Runtime/Engine/Private/WorldPartition/ContentBundle/ContentBundleEditor.cpp
@@ -50,7 +50,7 @@ void FContentBundleEditor::DoUninitialize()
 void FContentBundleEditor::DoInjectContent()
 {
 	FString ActorDescContainerPackage;
-	bool bCreatedContainerPath = ContentBundlePaths::BuildActorDescContainerPackgePath(GetDescriptor()->GetPackageRoot(), GetDescriptor()->GetGuid(), GetInjectedWorld()->GetPackage()->GetName(), ActorDescContainerPackage);
+	bool bCreatedContainerPath = ContentBundlePaths::BuildActorDescContainerPackagePath(GetDescriptor()->GetPackageRoot(), GetDescriptor()->GetGuid(), GetInjectedWorld()->GetPackage()->GetName(), ActorDescContainerPackage);
 	if (bCreatedContainerPath)
 	{
 		UnsavedActorMonitor = NewObject<UContentBundleUnsavedActorMonitor>(GetTransientPackage(), NAME_None, RF_Transactional);
diff --git a/Engine/Source/Runtime/Engine/Private/WorldPartition/ContentBundle/ContentBundlePaths.cpp b/Engine/Source/Runtime/Engine/Private/WorldPartition/ContentBundle/ContentBundlePaths.cpp
--- a/Engine/Source/Runtime/Engine/Private/WorldPartition/ContentBundle/ContentBundlePaths.cpp
+++ b/Engine/Source/Runtime/Engine/Private/WorldPartition/ContentBundle/ContentBundlePaths.cpp
@@ -48,11 +48,11 @@ FString ContentBundlePaths::GetRelativeLevelFolder(const FContentBundleBase& Con
 FString ContentBundlePaths::MakeExternalActorPackagePath(const FString& ContentBundleExternalActorFolder, const FString& ActorName)
 {
 	const FString ContentBundleExternalActor = ULevel::GetActorPackageName(ContentBundleExternalActorFolder, EActorPackagingScheme::Reduced, *ActorName);
-	check(IsAContentBundlePackagePath(ContentBundleExternalActor));
+	check(IsAContentBundleExternalActorPackagePath(ContentBundleExternalActor));
 	return ContentBundleExternalActor;
 }
 
-bool ContentBundlePaths::IsAContentBundlePackagePath(FStringView InPackagePath)
+bool ContentBundlePaths::IsAContentBundleExternalActorPackagePath(FStringView InPackagePath)
 {
 	return GetContentBundleGuidFromExternalActorPackagePath(InPackagePath).IsValid();
 }
@@ -106,7 +106,7 @@ FStringView ContentBundlePaths::GetActorPathRelativeToExternalActors(FStringView
 	return FStringView();
 }
 
-bool ContentBundlePaths::BuildContentBundleRootPath(const FString& InContenBundleMountPoint, const FGuid& InContentBundleGuid, FString& OutContentBundleRootPath)
+bool ContentBundlePaths::BuildContentBundleExternalActorPath(const FString& InContenBundleMountPoint, const FGuid& InContentBundleGuid, FString& OutContentBundleRootPath)
 {
 	if (InContenBundleMountPoint.IsEmpty() || !InContentBundleGuid.IsValid())
 	{
@@ -125,10 +125,10 @@ bool ContentBundlePaths::BuildContentBundleRootPath(const FString& InContenBundl
 	return true;
 }
 
-bool ContentBundlePaths::BuildActorDescContainerPackgePath(const FString& InContenBundleMountPoint, const FGuid& InContentBundleGuid, const FString& InLevelPackagePath, FString& OutContainerPackagePath)
+bool ContentBundlePaths::BuildActorDescContainerPackagePath(const FString& InContenBundleMountPoint, const FGuid& InContentBundleGuid, const FString& InLevelPackagePath, FString& OutContainerPackagePath)
 {
 	FString ContentBundleRootPath;
-	if (!BuildContentBundleRootPath(InContenBundleMountPoint, InContentBundleGuid, ContentBundleRootPath))
+	if (!BuildContentBundleExternalActorPath(InContenBundleMountPoint, InContentBundleGuid, ContentBundleRootPath))
 	{
 		return false;
 	}
@@ -148,4 +148,24 @@ bool ContentBundlePaths::BuildActorDescContainerPackgePath(const FString& InCont
 	return false;
 }
 
+FStringView ContentBundlePaths::GetRelativePath(FStringView InContentBundlePath)
+{
+	uint32 ContentBundleFolderIdx = UE::String::FindFirst(InContentBundlePath, GetContentBundleFolder(), ESearchCase::IgnoreCase);
+	if (ContentBundleFolderIdx != INDEX_NONE)
+	{
+		FStringView RelativeContentBundlePath = InContentBundlePath.RightChop(ContentBundleFolderIdx + GetContentBundleFolder().Len());
+		uint32 ContentBundleGuidEndIdx = UE::String::FindFirst(RelativeContentBundlePath, TEXT("/"), ESearchCase::IgnoreCase);
+		if (ContentBundleGuidEndIdx != INDEX_NONE)
+		{
+			return RelativeContentBundlePath.RightChop(ContentBundleGuidEndIdx + 1); // + 1 to remove the "/"
+		}
+	}
+	return FStringView();
+}
+
+bool ContentBundlePaths::IsAContentBundlePath(FStringView InContentBundlePath)
+{
+	return GetRelativePath(InContentBundlePath).Len() > 0;
+}
+
 #endif
diff --git a/Engine/Source/Runtime/Engine/Public/WorldPartition/ContentBundle/ContentBundlePaths.h b/Engine/Source/Runtime/Engine/Public/WorldPartition/ContentBundle/ContentBundlePaths.h
--- a/Engine/Source/Runtime/Engine/Public/WorldPartition/ContentBundle/ContentBundlePaths.h
+++ b/Engine/Source/Runtime/Engine/Public/WorldPartition/ContentBundle/ContentBundlePaths.h
@@ -39,17 +39,24 @@ namespace ContentBundlePaths
 	FString GetRelativeLevelFolder(const FContentBundleBase& ContentBundle);
 
 #if WITH_EDITOR
-	// return format is /{MountPoint}/{ExternalActorFolder}/{ContentBundleFolder}/{ContentBundleUID}/
-	ENGINE_API bool BuildContentBundleRootPath(const FString& InContenBundleMountPoint, const FGuid& InContentBundleGuid, FString& OutContentBundleRootPath);
+	// return format is /{MountPoint}/{ContentBundleFolder}/{ContentBundleUID}/
+	ENGINE_API bool BuildContentBundleExternalActorPath(const FString& InContenBundleMountPoint, const FGuid& InContentBundleGuid, FString& OutContentBundleRootPath);
 
-	// return format is /{MountPoint}/{ExternalActorFolder}/{ContentBundleFolder}/{ContentBundleUID}/{LevelPath}/
-	ENGINE_API bool BuildActorDescContainerPackgePath(const FString& InContenBundleMountPoint, const FGuid& InContentBundleGuid, const FString& InLevelPackagePath, FString& OutContainerPackagePath);
+	// return format is /{MountPoint}/{ContentBundleFolder}/{ContentBundleUID}/{LevelPath}/
+	ENGINE_API bool BuildActorDescContainerPackagePath(const FString& InContenBundleMountPoint, const FGuid& InContentBundleGuid, const FString& InLevelPackagePath, FString& OutContainerPackagePath);
+
+	// InContentBundlePath format is : */{ContentBundleFolder}/{ContentBundleUID}/{RelativePath}
+	// return format is {RelativePath}, empty otherwise
+	ENGINE_API FStringView GetRelativePath(FStringView InContentBundlePath);
+
+	// return true if InPackage follow format : */{ContentBundleFolder}/{ContentBundleUID}/*, false otherwise
+	ENGINE_API bool IsAContentBundlePath(FStringView InContentBundlePath);
 
 	// return an ExternalActor path following the format : /{MountPoint}/{ExternalActorFolder}/{ContentBundleFolder}/{ContentBundleUID}/{LevelPath}/{ExternalActorPackagePath}
 	ENGINE_API FString MakeExternalActorPackagePath(const FString& ContentBundleExternalActorFolder, const FString& ActorName);
 
-	// return true if InPackage follow format : //{MountPoint}/{ExternalActorFolder}/{ContentBundleFolder}/{ContentBundleUID}/, false otherwise
-	ENGINE_API bool IsAContentBundlePackagePath(FStringView InPackagePath);
+	// return true if InPackage follow format : //{MountPoint}/{ExternalActorFolder}/{ContentBundleFolder}/{ContentBundleUID}/*, false otherwise
+	ENGINE_API bool IsAContentBundleExternalActorPackagePath(FStringView InPackagePath);
 
 	// InContentBundleExternalActorPackagePath format is : /{MountPoint}/{ExternalActorFolder}/{ContentBundleFolder}/{ContentBundleUID}/{LevelPath}/{ExternalActorPackagePath}
 	// return format is /{LevelPath}/{ExternalActorPackagePath}, empty otherwise
