WorldPartitionLevelHelper: Fix SoftObjectPath embedded level instance mapping regression (CL 23781741)
Disable SoftObjectPath remapping through FLinkerInstancingContext of the embedded level instance actors.
This simplifies the SoftObjectPath remappings as simple two step process:
- Remap on actor load by adding container id to subpathstring to distinguish multiple instances of the same Level Instance
- Then call to FWorldPartitionLevelHelper::RemapLevelSoftObjectPaths will remap from persistent level to grid cel package

diff --git a/Engine/Source/Runtime/Engine/Private/WorldPartition/WorldPartition.cpp b/Engine/Source/Runtime/Engine/Private/WorldPartition/WorldPartition.cpp
--- a/Engine/Source/Runtime/Engine/Private/WorldPartition/WorldPartition.cpp
+++ b/Engine/Source/Runtime/Engine/Private/WorldPartition/WorldPartition.cpp
@@ -652,6 +652,7 @@ void UWorldPartition::Initialize(UWorld* InWorld, const FTransform& InTransform)
 		}
 
 		// Apply remapping of Persistent Level's SoftObjectPaths
+		// Here we remap SoftObjectPaths so that they are mapped from the PersistentLevel Package to the Cell Packages using the mapping built by the policy
 		FWorldPartitionLevelHelper::RemapLevelSoftObjectPaths(OuterWorld->PersistentLevel, this);
 	}
 #endif
diff --git a/Engine/Source/Runtime/Engine/Private/WorldPartition/WorldPartitionLevelHelper.cpp b/Engine/Source/Runtime/Engine/Private/WorldPartition/WorldPartitionLevelHelper.cpp
--- a/Engine/Source/Runtime/Engine/Private/WorldPartition/WorldPartitionLevelHelper.cpp
+++ b/Engine/Source/Runtime/Engine/Private/WorldPartition/WorldPartitionLevelHelper.cpp
@@ -19,6 +19,7 @@
 #include "WorldPartition/WorldPartitionPackageHelper.h"
 #include "LevelUtils.h"
 #include "ActorFolder.h"
+#include "Containers/StringFwd.h"
 
 FWorldPartitionLevelHelper& FWorldPartitionLevelHelper::Get()
 {
@@ -205,30 +206,35 @@ FString FWorldPartitionLevelHelper::AddActorContainerIDToSubPathString(const FAc
 
 FString FWorldPartitionLevelHelper::GetContainerPackage(const FActorContainerID& InContainerID, const FString& InPackageName, const FString& InDestLevelPackageName)
 {
+	// Generate a unique name to load a Level Instance embedded actor if there are multiple instances of this Level Instance and possibly across multiple instances of the World Partition world
+	// InContainerID will distinguish between instances of the same Level Instance
+	// InDestLevelPackageName will distinguish between instances of the same top level World Partition world (Only needed in PIE, In Cook we always cook the source WP and not an instance and Actor packages no longer exist at runtime)
 	uint64 DestLevelID = 0;
+	TStringBuilder<512> PackageNameBuilder;
+	PackageNameBuilder.Appendf(TEXT("/Temp%s_%s"), *InPackageName, *InContainerID.ToString());
+		
 	if (!InDestLevelPackageName.IsEmpty())
 	{
 		DestLevelID = CityHash64((const char*)*InDestLevelPackageName, InDestLevelPackageName.Len() * sizeof(TCHAR));
+		PackageNameBuilder.Appendf(TEXT("_%016llx"), DestLevelID);
 	}
 
-	return FString::Printf(TEXT("/Temp%s_%s_%016llx"), *InPackageName, *InContainerID.ToString(), DestLevelID);
+	return PackageNameBuilder.ToString();
 }
 
 bool FWorldPartitionLevelHelper::RemapActorPath(const FActorContainerID& InContainerID, const FString& InActorPath, FString& OutActorPath)
 {
+	// This gets called by UWorldPartitionLevelStreamingPolicy::PrepareActorToCellRemapping while it is creating its mappings from PersistentLevel package to Cell package
+	// so that it can later remap SoftObjectPaths. So here we remap the SubPathString of the provided ActorPath to add a ContainerID suffix the same way we did in FWorldPartitionLevelHelper::LoadActors
+	//
+	// At this point we are remapping the SubPathString and adding a ContainerID suffix so
+	// '/Game/SomePath/WorldName.WorldName:PersistentLevel.ActorA' becomes
+	// '/Game/SomePath/WorldName.WorldName:PersistentLevel.ActorA_{ContainerID}'
 	if (!InContainerID.IsMainContainer())
 	{
-		const FSoftObjectPath SoftObjectPath(InActorPath);
-		const FString LongPackageName = SoftObjectPath.GetLongPackageName();
-			
-		const FString ContainerPackageString = GetContainerPackage(InContainerID, LongPackageName);	
-		const FString NewSubPathString = FWorldPartitionLevelHelper::AddActorContainerIDToSubPathString(InContainerID, SoftObjectPath.GetSubPathString());
-
-		FNameBuilder NewAssetPathBuilder;
-		SoftObjectPath.GetAssetPath().AppendString(NewAssetPathBuilder);
-		NewAssetPathBuilder.ReplaceAt(0, LongPackageName.Len(), ContainerPackageString);
-
-		OutActorPath = FSoftObjectPath(FTopLevelAssetPath(NewAssetPathBuilder.ToView()), NewSubPathString).ToString();
+		FSoftObjectPath SoftObjectPath(InActorPath);
+		FString RemappedSubPathString = FWorldPartitionLevelHelper::AddActorContainerIDToSubPathString(InContainerID, SoftObjectPath.GetSubPathString());
+		OutActorPath = FSoftObjectPath(SoftObjectPath.GetAssetPath(), RemappedSubPathString).ToString();
 		return true;
 	}
 
@@ -331,6 +337,11 @@ bool FWorldPartitionLevelHelper::LoadActors(UWorld* InOwningWorld, ULevel* InDes
 			const FName ContainerPackageInstanceName(GetContainerPackage(PackageObjectMapping.ContainerID, PackageObjectMapping.ContainerPackage.ToString(), DestLevelPackageName));
 
 			FLinkerInstancingContext& NewContext = LinkerInstancingContexts.Add(PackageObjectMapping.ContainerID);
+
+			// Make sure here we don't remap the SoftObjectPaths through the linker when loading the embedded actor packages. 
+			// A remapping will happen in the packaged loaded callback later in this method.
+			NewContext.SetSoftObjectPathRemappingEnabled(false); 
+			
 			NewContext.AddTag(ULevel::DontLoadExternalObjectsTag);
 			NewContext.AddPackageMapping(PackageObjectMapping.ContainerPackage, ContainerPackageInstanceName);
 			Context = &NewContext;
@@ -424,10 +435,14 @@ bool FWorldPartitionLevelHelper::LoadActors(UWorld* InOwningWorld, ULevel* InDes
 					PackageObjectMapping->LoadedPath = *Actor->GetPathName();
 
 					// Fixup any FSoftObjectPath from this Actor (and its SubObjects) in this container to another object in the same container with a ContainerID suffix that can be remapped to
-					// to a Cell in the StreamingPolicy (this relies on the fact that the _DUP package doesn't get fixed up)
+					// a Cell package in the StreamingPolicy.
+					// 
+					// At  this point we are remapping the SubPathString and adding a ContainerID suffix so
+					// '/Game/SomePath/WorldName.WorldName:PersistentLevel.ActorA' becomes
+					// '/Game/SomePath/WorldName.WorldName:PersistentLevel.ActorA_{ContainerID}'
 					FSoftObjectPathFixupArchive FixupArchive([&](FSoftObjectPath& Value)
 					{
-						if (!Value.IsNull() && Value.GetAssetPathString().Equals(RemappedWorldPath, ESearchCase::IgnoreCase))
+						if (!Value.IsNull() && Value.GetAssetPathString().Equals(SourceWorldPath, ESearchCase::IgnoreCase))
 						{
 							Value.SetSubPathString(AddActorContainerIDToSubPathString(PackageObjectMapping->ContainerID, Value.GetSubPathString()));
 						}
diff --git a/Engine/Source/Runtime/Engine/Private/WorldPartition/WorldPartitionLevelStreamingDynamic.cpp b/Engine/Source/Runtime/Engine/Private/WorldPartition/WorldPartitionLevelStreamingDynamic.cpp
--- a/Engine/Source/Runtime/Engine/Private/WorldPartition/WorldPartitionLevelStreamingDynamic.cpp
+++ b/Engine/Source/Runtime/Engine/Private/WorldPartition/WorldPartitionLevelStreamingDynamic.cpp
@@ -442,6 +442,7 @@ void UWorldPartitionLevelStreamingDynamic::FinalizeRuntimeLevel()
 	{
 		TRACE_CPUPROFILER_EVENT_SCOPE(FixupIDs);
 
+		// Here we remap SoftObjectPaths so that they are mapped from the PersistentLevel Package to the Cell Packages using the mapping built by the policy
 		if (OuterWorld->IsPlayInEditor())
 		{
 			int32 PIEInstanceID = GetPackage()->GetPIEInstanceID();
