Adjust paths checks which are no longer valid in Hlod generation
During Hlods generations child packages can recieve a ContentBundleGuid.
However they point to a transient level does not follow a classic content bundle cell.

diff --git a/Engine/Source/Runtime/Engine/Private/WorldPartition/WorldPartitionLevelStreamingDynamic.cpp b/Engine/Source/Runtime/Engine/Private/WorldPartition/WorldPartitionLevelStreamingDynamic.cpp
--- a/Engine/Source/Runtime/Engine/Private/WorldPartition/WorldPartitionLevelStreamingDynamic.cpp
+++ b/Engine/Source/Runtime/Engine/Private/WorldPartition/WorldPartitionLevelStreamingDynamic.cpp
@@ -13,6 +13,7 @@
 #include "Engine/Level.h"
 #include "Misc/PathViews.h"
 #include "ContentStreaming.h"
+#include "WorldPartition/ContentBundle/ContentBundlePaths.h"
 #endif
 
 #define LOCTEXT_NAMESPACE "World"
@@ -310,32 +311,29 @@ bool UWorldPartitionLevelStreamingDynamic::IssueLoadRequests()
 
 		for (const FWorldPartitionRuntimeCellObjectMapping& CellObjectMapping : ChildPackages)
 		{
-			if (CellObjectMapping.ContentBundleGuid.IsValid())
+			if (ContentBundlePaths::IsAContentBundlePath(CellObjectMapping.ContainerPackage.ToString()))
 			{
+				check(CellObjectMapping.ContentBundleGuid.IsValid());
 				check(CellObjectMapping.ContainerPackage != CellObjectMapping.WorldPackage);
+
 				bool bIsContainerPackageAlreadyRemapped = OutLinkInstancingContext.RemapPackage(CellObjectMapping.ContainerPackage) != CellObjectMapping.ContainerPackage;
 				if (!bIsContainerPackageAlreadyRemapped)
 				{
-					FString ContainerPackage = CellObjectMapping.ContainerPackage.ToString();
-					FString WorldPackage = CellObjectMapping.WorldPackage.ToString();
+					// Get Content Bundle World
+					FString ActorDescContainerPackage = CellObjectMapping.ContainerPackage.ToString();
+					FStringView WorldFromContainerPackage = ContentBundlePaths::GetRelativePath(ActorDescContainerPackage);
 
-					bool bWithoutSlashes = false;
-					FName ContainerMountPoint = FPackageName::GetPackageMountPoint(ContainerPackage, bWithoutSlashes);
+					// Get ChildPackage World
+					FStringView RelativeWorldPackage;
+					const bool bWithoutSlashes = false;
+					FString WorldPackage = CellObjectMapping.WorldPackage.ToString();
 					FName WorldPackageMountPoint = FPackageName::GetPackageMountPoint(WorldPackage, bWithoutSlashes);
-					if (ContainerMountPoint != WorldPackageMountPoint)
-					{
-#if DO_CHECK
-						//  Validate that while the container mounting points are different, they point to the same world package.
-						FStringView RelativeWorldPackage;
-						ensure(FPathViews::TryMakeChildPathRelativeTo(MakeStringView(WorldPackage), MakeStringView(WorldPackageMountPoint.ToString()), RelativeWorldPackage));
-
-						uint32 WorldPackageInContainerPackageStartIdx = ContainerPackage.Len() - RelativeWorldPackage.Len();
-						FStringView WorldPackageInContainerPackage = MakeStringView(ContainerPackage).SubStr(WorldPackageInContainerPackageStartIdx, RelativeWorldPackage.Len());
-						check(WorldPackageInContainerPackage.Equals(RelativeWorldPackage));
-#endif // DO_CHECK
+					ensure(FPathViews::TryMakeChildPathRelativeTo(MakeStringView(WorldPackage), MakeStringView(WorldPackageMountPoint.ToString()), RelativeWorldPackage));
 
-						OutLinkInstancingContext.AddPackageMapping(CellObjectMapping.ContainerPackage, RuntimePackage->GetFName());
-					}
+					check(WorldFromContainerPackage.Equals(RelativeWorldPackage));
+					
+					OutLinkInstancingContext.AddPackageMapping(CellObjectMapping.ContainerPackage, RuntimePackage->GetFName());
+					
 				}
 			}
 		}
