Add a ToString function to FHardwareDeviceIdentifier for easier logging and display of input devices. Create a "GetHardwareDeviceForClassName" function which will return the full hardware device id for the platform given only the input device name.

#rnx

diff --git a/Engine/Source/Runtime/Engine/Classes/GameFramework/InputSettings.h b/Engine/Source/Runtime/Engine/Classes/GameFramework/InputSettings.h
--- a/Engine/Source/Runtime/Engine/Classes/GameFramework/InputSettings.h
+++ b/Engine/Source/Runtime/Engine/Classes/GameFramework/InputSettings.h
@@ -319,6 +319,12 @@ struct ENGINE_API FHardwareDeviceIdentifier
 
 	/** Returns true if this hardware device Identifier has valid names */
 	bool IsValid() const;
+
+	/**
+	 * Returns a string containing the Input Class Name and HardwareDeviceIdentifier properties
+	 * concatenated together.
+	 */
+	FString ToString() const;
 	
 	/** An Invalid Hardware Device Identifier. */
 	static FHardwareDeviceIdentifier Invalid;
@@ -362,6 +368,10 @@ public:
 	static const TArray<FName>& GetAllHardwareDeviceNames();
 #endif	// WITH_EDITOR
 
+
+	/** Returns the first matching FHardwareDeviceIdentifier that has the given HardwareDeviceIdentifier.  */
+	const FHardwareDeviceIdentifier* GetHardwareDeviceForClassName(const FName InHardwareDeviceIdentifier) const;
+	
 	/** Add the given hardware device identifier to this platform's settings. */
 	void AddHardwareDeviceIdentifier(const FHardwareDeviceIdentifier& InHardwareDevice);
 
diff --git a/Engine/Source/Runtime/Engine/Private/UserInterface/InputSettings.cpp b/Engine/Source/Runtime/Engine/Private/UserInterface/InputSettings.cpp
--- a/Engine/Source/Runtime/Engine/Private/UserInterface/InputSettings.cpp
+++ b/Engine/Source/Runtime/Engine/Private/UserInterface/InputSettings.cpp
@@ -526,6 +526,11 @@ bool FHardwareDeviceIdentifier::IsValid() const
 	return InputClassName.IsValid() && HardwareDeviceIdentifier.IsValid();
 }
 
+FString FHardwareDeviceIdentifier::ToString() const
+{
+	return FString::Printf(TEXT("%s::%s"), *InputClassName.ToString(), *HardwareDeviceIdentifier.ToString());
+}
+
 //////////////////////////////////////////////////////////////////
 // UInputPlatformSettings
 
@@ -542,6 +547,15 @@ UInputPlatformSettings* UInputPlatformSettings::Get()
 	return UPlatformSettingsManager::Get().GetSettingsForPlatform<UInputPlatformSettings>();
 }
 
+const FHardwareDeviceIdentifier* UInputPlatformSettings::GetHardwareDeviceForClassName(const FName InHardwareDeviceIdentifier) const
+{
+	return HardwareDevices.FindByPredicate(
+		[InHardwareDeviceIdentifier](const FHardwareDeviceIdentifier& Hardware)
+		{
+			return Hardware.HardwareDeviceIdentifier == InHardwareDeviceIdentifier;
+		});
+}
+
 void UInputPlatformSettings::AddHardwareDeviceIdentifier(const FHardwareDeviceIdentifier& InHardwareDevice)
 {
 	if (ensure(InHardwareDevice.IsValid()))
