Removed duplicate stat scope from the FuncDynDelegate case in FTimerUnifiedDelegate::Execute.

diff --git a/Engine/Source/Runtime/Engine/Public/TimerManager.h b/Engine/Source/Runtime/Engine/Public/TimerManager.h
--- a/Engine/Source/Runtime/Engine/Public/TimerManager.h
+++ b/Engine/Source/Runtime/Engine/Public/TimerManager.h
@@ -49,7 +49,7 @@ struct FTimerUnifiedDelegate
 		}
 		else if (FuncDynDelegate.IsBound())
 		{
-			FScopeCycleCounterUObject Context(FuncDynDelegate.GetUObject());
+			// stat scope is handled by UObject::ProcessEvent for the UFunction.
 			FuncDynDelegate.ProcessDelegate<UObject>(nullptr);
 		}
 		else if ( FuncCallback )
