Avoiding network version requests from demo driver cdo

diff --git a/Engine/Source/Runtime/Engine/Private/DemoNetDriver.cpp b/Engine/Source/Runtime/Engine/Private/DemoNetDriver.cpp
--- a/Engine/Source/Runtime/Engine/Private/DemoNetDriver.cpp
+++ b/Engine/Source/Runtime/Engine/Private/DemoNetDriver.cpp
@@ -701,6 +701,8 @@ void UDemoNetDriver::InitDefaults()
 	if (!HasAnyFlags(RF_ClassDefaultObject))
 	{
 		LevelIntervals.Reserve(512);
+
+		ReplayHelper.PlaybackDemoHeader.SetDefaultNetworkVersions();
 	}
 
 	RecordBuildConsiderAndPrioritizeTimeSlice = CVarDemoMaximumRepPrioritizeTime.GetValueOnGameThread();
diff --git a/Engine/Source/Runtime/Engine/Private/ReplayHelper.cpp b/Engine/Source/Runtime/Engine/Private/ReplayHelper.cpp
--- a/Engine/Source/Runtime/Engine/Private/ReplayHelper.cpp
+++ b/Engine/Source/Runtime/Engine/Private/ReplayHelper.cpp
@@ -208,6 +208,7 @@ void FReplayHelper::WriteNetworkDemoHeader(UNetConnection* Connection)
 	if (FArchive* FileAr = ReplayStreamer->GetHeaderArchive())
 	{
 		FNetworkDemoHeader DemoHeader;
+		DemoHeader.SetDefaultNetworkVersions();
 
 		DemoHeader.LevelNamesAndTimes = LevelNamesAndTimes;
 
@@ -409,6 +410,7 @@ bool FReplayHelper::ReadPlaybackDemoHeader(FString& Error)
 
 	PRAGMA_DISABLE_DEPRECATION_WARNINGS
 	PlaybackDemoHeader = FNetworkDemoHeader();
+	PlaybackDemoHeader.SetDefaultNetworkVersions();
 	PRAGMA_ENABLE_DEPRECATION_WARNINGS
 
 	FArchive* FileAr = ReplayStreamer->GetHeaderArchive();
diff --git a/Engine/Source/Runtime/Engine/Private/ReplayTypes.cpp b/Engine/Source/Runtime/Engine/Private/ReplayTypes.cpp
--- a/Engine/Source/Runtime/Engine/Private/ReplayTypes.cpp
+++ b/Engine/Source/Runtime/Engine/Private/ReplayTypes.cpp
@@ -21,9 +21,9 @@ PRAGMA_DISABLE_DEPRECATION_WARNINGS
 FNetworkDemoHeader::FNetworkDemoHeader() :
 	Magic(NETWORK_DEMO_MAGIC),
 	Version((uint32)FReplayCustomVersion::CustomVersions),
-	NetworkChecksum(FNetworkVersion::GetLocalNetworkVersion()),
-	EngineNetworkProtocolVersion(FNetworkVersion::GetEngineNetworkProtocolVersion()),
-	GameNetworkProtocolVersion(FNetworkVersion::GetGameNetworkProtocolVersion()),
+	NetworkChecksum(0),
+	EngineNetworkProtocolVersion(FEngineNetworkCustomVersion::LatestVersion),
+	GameNetworkProtocolVersion(FGameNetworkCustomVersion::LatestVersion),
 	Guid(),
 	MinRecordHz(0.0f),
 	MaxRecordHz(0.0f),
@@ -36,8 +36,15 @@ FNetworkDemoHeader::FNetworkDemoHeader() :
 	PackageVersionUE(GPackageFileUEVersion),
 	PackageVersionLicenseeUE(GPackageFileLicenseeUEVersion)
 {
-	CustomVersions = FNetworkVersion::GetNetworkCustomVersions();
+}
 
+void FNetworkDemoHeader::SetDefaultNetworkVersions()
+{
+	NetworkChecksum = FNetworkVersion::GetLocalNetworkVersion();
+	EngineNetworkProtocolVersion = FNetworkVersion::GetEngineNetworkProtocolVersion();
+	GameNetworkProtocolVersion = FNetworkVersion::GetGameNetworkProtocolVersion();
+
+	CustomVersions = FNetworkVersion::GetNetworkCustomVersions();
 	CustomVersions.SetVersion(FReplayCustomVersion::Guid, FReplayCustomVersion::LatestVersion, TEXT("Replay"));
 }
 PRAGMA_ENABLE_DEPRECATION_WARNINGS
diff --git a/Engine/Source/Runtime/Engine/Public/ReplayTypes.h b/Engine/Source/Runtime/Engine/Public/ReplayTypes.h
--- a/Engine/Source/Runtime/Engine/Public/ReplayTypes.h
+++ b/Engine/Source/Runtime/Engine/Public/ReplayTypes.h
@@ -210,6 +210,8 @@ struct ENGINE_API FNetworkDemoHeader
 
 	FNetworkDemoHeader();
 
+	void SetDefaultNetworkVersions();
+
 	PRAGMA_DISABLE_DEPRECATION_WARNINGS
 	FNetworkDemoHeader(const FNetworkDemoHeader&) = default;
 	FNetworkDemoHeader& operator=(const FNetworkDemoHeader&) = default;
