Temporarily disable hlods generations for plugins
An ASpatialHashRuntimeGridInfo actor is missing in a plugin ActorDescContainer
This lead to errors during streaming generation when trying to matching grids

diff --git a/Engine/Source/Runtime/Engine/Private/WorldPartition/WorldPartitionStreamingGeneration.cpp b/Engine/Source/Runtime/Engine/Private/WorldPartition/WorldPartitionStreamingGeneration.cpp
--- a/Engine/Source/Runtime/Engine/Private/WorldPartition/WorldPartitionStreamingGeneration.cpp
+++ b/Engine/Source/Runtime/Engine/Private/WorldPartition/WorldPartitionStreamingGeneration.cpp
@@ -1025,6 +1025,11 @@ void UWorldPartition::GenerateHLOD(ISourceControlHelper* SourceControlHelper, bo
 {
 	ForEachActorDescContainer([this, &SourceControlHelper, bCreateActorsOnly](UActorDescContainer* InActorDescContainer)
 	{
+		if (!InActorDescContainer->IsMainPartitionContainer())
+		{
+			return;
+		}
+
 		FStreamingGenerationLogErrorHandler LogErrorHandler;
 		FWorldPartitionStreamingGenerator::FWorldPartitionStreamingGeneratorParams StreamingGeneratorParams;
 		StreamingGeneratorParams.ErrorHandler = &LogErrorHandler;
