World Partition: get rid of IWorldPartitionRuntimeCellOwner.

diff --git a/Engine/Source/Runtime/Engine/Private/WorldPartition/ContentBundle/ContentBundle.cpp b/Engine/Source/Runtime/Engine/Private/WorldPartition/ContentBundle/ContentBundle.cpp
--- a/Engine/Source/Runtime/Engine/Private/WorldPartition/ContentBundle/ContentBundle.cpp
+++ b/Engine/Source/Runtime/Engine/Private/WorldPartition/ContentBundle/ContentBundle.cpp
@@ -32,7 +32,7 @@ void FContentBundle::DoInitialize()
 		if (UObject* Object = StaticFindObjectFast(URuntimeHashExternalStreamingObjectBase::StaticClass(), ExternalStreamingObjectPackage, *GetExternalStreamingObjectName()))
 		{
 			ExternalStreamingObject = CastChecked<URuntimeHashExternalStreamingObjectBase>(Object);
-			ExternalStreamingObject->OnStreamingObjectLoaded();
+			ExternalStreamingObject->OnStreamingObjectLoaded(GetInjectedWorld());
 		}
 		else
 		{
diff --git a/Engine/Source/Runtime/Engine/Private/WorldPartition/HLOD/HLODSubsystem.cpp b/Engine/Source/Runtime/Engine/Private/WorldPartition/HLOD/HLODSubsystem.cpp
--- a/Engine/Source/Runtime/Engine/Private/WorldPartition/HLOD/HLODSubsystem.cpp
+++ b/Engine/Source/Runtime/Engine/Private/WorldPartition/HLOD/HLODSubsystem.cpp
@@ -105,7 +105,7 @@ namespace FHLODSubsystem
 	    } // If not find it through the cell
 	    else if (const UWorldPartitionRuntimeCell* RuntimeCell = GetActorRuntimeCell(InWorldPartitionHLOD))
 	    {
-		    return RuntimeCell->GetCellOwner()->GetOuterWorld()->GetWorldPartition();
+		    return RuntimeCell->GetTypedOuter<UWorldPartition>();
 	    }
     
 	    return nullptr;
@@ -290,7 +290,7 @@ UHLODSubsystem::FCellData* UHLODSubsystem::GetCellData(const UWorldPartitionRunt
 		return nullptr;
 	}
 
-	const UWorldPartition* WorldPartition = InCell->GetCellOwner()->GetOuterWorld()->GetWorldPartition();
+	const UWorldPartition* WorldPartition = InCell->GetTypedOuter<UWorldPartition>();
 	check(WorldPartition);
 	FWorldPartitionHLODRuntimeData* WorldPartitionHLODRuntimeData = WorldPartitionsHLODRuntimeData.Find(WorldPartition);
 	if (WorldPartitionHLODRuntimeData)
diff --git a/Engine/Source/Runtime/Engine/Private/WorldPartition/WorldPartitionHelpers.cpp b/Engine/Source/Runtime/Engine/Private/WorldPartition/WorldPartitionHelpers.cpp
--- a/Engine/Source/Runtime/Engine/Private/WorldPartition/WorldPartitionHelpers.cpp
+++ b/Engine/Source/Runtime/Engine/Private/WorldPartition/WorldPartitionHelpers.cpp
@@ -5,7 +5,6 @@
 #include "WorldPartition/WorldPartition.h"
 #include "WorldPartition/WorldPartitionEditorHash.h"
 #include "WorldPartition/WorldPartitionRuntimeCellInterface.h"
-#include "WorldPartition/WorldPartitionRuntimeCellOwner.h"
 #include "Engine/World.h"
 #include "Engine/Level.h"
 #include "Algo/AnyOf.h"
@@ -35,17 +34,6 @@ UWorldPartition* FWorldPartitionHelpers::GetWorldPartition(const UObject* InObje
 			return GetWorldPartition(Cast<UObject>(Cell));
 		}
 	}
-	else if (const IWorldPartitionCell* Cell = Cast<const IWorldPartitionCell>(InObject))
-	{
-		if (const IWorldPartitionRuntimeCellOwner* CellOwner = Cell->GetCellOwner())
-		{
-			return GetWorldPartition(Cast<UObject>(CellOwner));
-		}
-	}
-	else if (const IWorldPartitionRuntimeCellOwner* CellOwner = Cast<const IWorldPartitionRuntimeCellOwner>(InObject))
-	{
-		OuterWorld = CellOwner->GetOuterWorld();
-	}
 	else if (const ULevel* OuterLevel = InObject->GetTypedOuter<ULevel>())
 	{
 		return GetWorldPartition(OuterLevel);
diff --git a/Engine/Source/Runtime/Engine/Private/WorldPartition/WorldPartitionLevelStreamingDynamic.cpp b/Engine/Source/Runtime/Engine/Private/WorldPartition/WorldPartitionLevelStreamingDynamic.cpp
--- a/Engine/Source/Runtime/Engine/Private/WorldPartition/WorldPartitionLevelStreamingDynamic.cpp
+++ b/Engine/Source/Runtime/Engine/Private/WorldPartition/WorldPartitionLevelStreamingDynamic.cpp
@@ -55,11 +55,9 @@ void UWorldPartitionLevelStreamingDynamic::Initialize(const UWorldPartitionRunti
 
 	UnsavedActorsContainer = InCell.UnsavedActorsContainer;
 
-	IWorldPartitionRuntimeCellOwner* CellOwner = InCell.GetCellOwner();
-	Initialize(CellOwner->GetOuterWorld(), InCell.GetPackages());
+	Initialize(GetTypedOuter<UWorld>(), InCell.GetPackages());
 #else
-	IWorldPartitionRuntimeCellOwner* CellOwner = InCell.GetCellOwner();
-	OuterWorldPartition = CellOwner->GetOuterWorld()->GetWorldPartition();
+	OuterWorldPartition = GetTypedOuter<UWorld>()->GetWorldPartition();
 #endif
 
 	UpdateShouldSkipMakingVisibilityTransactionRequest();
@@ -645,7 +643,7 @@ bool UWorldPartitionLevelStreamingDynamic::CanChangeVisibility(bool bMakeVisible
 
 			if (!bAlwaysAllowVisibilityChange)
 			{
-				if (const UWorldPartition* WorldPartition = RuntimeLevelStreamingCell->GetCellOwner()->GetWorldPartition())
+				if (const UWorldPartition* WorldPartition = RuntimeLevelStreamingCell->GetTypedOuter<UWorldPartition>())
 				{
 					if (UHLODSubsystem* HLODSubsystem = GetWorld()->GetSubsystem<UHLODSubsystem>())
 					{
diff --git a/Engine/Source/Runtime/Engine/Private/WorldPartition/WorldPartitionRuntimeCell.cpp b/Engine/Source/Runtime/Engine/Private/WorldPartition/WorldPartitionRuntimeCell.cpp
--- a/Engine/Source/Runtime/Engine/Private/WorldPartition/WorldPartitionRuntimeCell.cpp
+++ b/Engine/Source/Runtime/Engine/Private/WorldPartition/WorldPartitionRuntimeCell.cpp
@@ -21,9 +21,7 @@ UWorldPartitionRuntimeCell::UWorldPartitionRuntimeCell(const FObjectInitializer&
 	, DebugStreamingPriority(-1.f)
 #endif
 	, RuntimeCellData(nullptr)
-{
-	check(HasAnyFlags(RF_ClassDefaultObject) || GetOuter()->Implements<UWorldPartitionRuntimeCellOwner>());
-}
+{}
 
 #if WITH_EDITOR
 void UWorldPartitionRuntimeCell::PostDuplicate(bool bDuplicateForPIE)
diff --git a/Engine/Source/Runtime/Engine/Private/WorldPartition/WorldPartitionRuntimeCellOwner.cpp b/Engine/Source/Runtime/Engine/Private/WorldPartition/WorldPartitionRuntimeCellOwner.cpp
deleted file mode 100644
diff --git a/Engine/Source/Runtime/Engine/Private/WorldPartition/WorldPartitionRuntimeHash.cpp b/Engine/Source/Runtime/Engine/Private/WorldPartition/WorldPartitionRuntimeHash.cpp
--- a/Engine/Source/Runtime/Engine/Private/WorldPartition/WorldPartitionRuntimeHash.cpp
+++ b/Engine/Source/Runtime/Engine/Private/WorldPartition/WorldPartitionRuntimeHash.cpp
@@ -30,19 +30,20 @@ void URuntimeHashExternalStreamingObjectBase::ForEachStreamingCells(TFunctionRef
 	}
 }
 
-void URuntimeHashExternalStreamingObjectBase::OnStreamingObjectLoaded()
+void URuntimeHashExternalStreamingObjectBase::OnStreamingObjectLoaded(UWorld* InjectedWorld)
 {
 	bool bIsACookedObject = !CellToLevelStreamingPackage.IsEmpty();
 	if (bIsACookedObject)
 	{
-		// Cooked streaming object's Cells do not have LevelStreaming and are outered to the streaming object.
-		// Create their level streaming now and re-outered to the runtime hash.
+		// Outer the streaming object back to the runtime hash.
+		Rename(nullptr, InjectedWorld->GetWorldPartition()->RuntimeHash);
+
+		// Cooked streaming object's Cells do not have LevelStreaming.
 		ForEachStreamingCells([this](UWorldPartitionRuntimeCell& Cell)
 		{
 			UWorldPartitionRuntimeLevelStreamingCell* RuntimeCell = CastChecked<UWorldPartitionRuntimeLevelStreamingCell>(&Cell);
 
 			FName LevelStreamingPackage = CellToLevelStreamingPackage.FindChecked(RuntimeCell->GetFName());
-			RuntimeCell->Rename(nullptr, GetOuterWorld()->GetWorldPartition()->RuntimeHash);
 			RuntimeCell->CreateAndSetLevelStreaming(*LevelStreamingPackage.ToString());
 		});
 	}
diff --git a/Engine/Source/Runtime/Engine/Private/WorldPartition/WorldPartitionRuntimeLevelStreamingCell.cpp b/Engine/Source/Runtime/Engine/Private/WorldPartition/WorldPartitionRuntimeLevelStreamingCell.cpp
--- a/Engine/Source/Runtime/Engine/Private/WorldPartition/WorldPartitionRuntimeLevelStreamingCell.cpp
+++ b/Engine/Source/Runtime/Engine/Private/WorldPartition/WorldPartitionRuntimeLevelStreamingCell.cpp
@@ -68,7 +68,7 @@ TSet<FName> UWorldPartitionRuntimeLevelStreamingCell::GetActorPackageNames() con
 FName UWorldPartitionRuntimeLevelStreamingCell::GetLevelPackageName() const
 {
 #if WITH_EDITOR
-	UWorld* World = GetCellOwner()->GetOwningWorld();
+	UWorld* World = GetOwningWorld();
 	if (World->IsPlayInEditor())
 	{
 		return FName(UWorldPartitionLevelStreamingPolicy::GetCellPackagePath(GetFName(), World));
@@ -106,10 +106,9 @@ UWorldPartitionLevelStreamingDynamic* UWorldPartitionRuntimeLevelStreamingCell::
 {
 	if (HasActors())
 	{
-		const UWorldPartition* WorldPartition = GetCellOwner()->GetWorldPartition();
-		UWorld* OuterWorld = GetCellOwner()->GetOuterWorld();
-		UWorld* OwningWorld = GetCellOwner()->GetOwningWorld();
-
+		UWorld* OuterWorld = GetOuterWorld();
+		UWorld* OwningWorld = GetOwningWorld();
+		const UWorldPartition* WorldPartition = OuterWorld->GetWorldPartition();
 		const FName LevelStreamingName = FName(*FString::Printf(TEXT("WorldPartitionLevelStreaming_%s"), *GetName()));
 
 		// When called by Commandlet (PopulateGeneratedPackageForCook), LevelStreaming's outer is set to Cell/WorldPartition's outer to prevent warnings when saving Cell Levels (Warning: Obj in another map). 
@@ -170,6 +169,24 @@ bool UWorldPartitionRuntimeLevelStreamingCell::IsLoading() const
 	return Super::IsLoading();
 }
 
+UWorld* UWorldPartitionRuntimeLevelStreamingCell::GetOwningWorld() const
+{
+	if (URuntimeHashExternalStreamingObjectBase* StreamingObjectOuter = GetTypedOuter<URuntimeHashExternalStreamingObjectBase>())
+	{
+		return StreamingObjectOuter->GetOwningWorld();
+	}
+	return GetTypedOuter<UWorldPartition>()->GetWorld();
+}
+
+UWorld* UWorldPartitionRuntimeLevelStreamingCell::GetOuterWorld() const
+{
+	if (URuntimeHashExternalStreamingObjectBase* StreamingObjectOuter = GetTypedOuter<URuntimeHashExternalStreamingObjectBase>())
+	{
+		return StreamingObjectOuter->GetOuterWorld();
+	}
+	return GetTypedOuter<UWorld>();
+}
+
 FLinearColor UWorldPartitionRuntimeLevelStreamingCell::GetDebugColor(EWorldPartitionRuntimeCellVisualizeMode VisualizeMode) const
 {
 	switch (VisualizeMode)
@@ -219,7 +236,7 @@ bool UWorldPartitionRuntimeLevelStreamingCell::PopulateGeneratorPackageForCook(T
 	{
 		FWorldPartitionLevelHelper::FPackageReferencer PackageReferencer;
 		const bool bLoadAsync = false;
-		UWorld* OuterWorld = GetCellOwner()->GetOuterWorld();
+		UWorld* OuterWorld = GetOuterWorld();
 		UWorldPartition* WorldPartition = OuterWorld->GetWorldPartition();
 
 		// Don't do SoftObjectPath remapping for PersistentLevel actors because references can end up in different cells
@@ -269,7 +286,7 @@ bool UWorldPartitionRuntimeLevelStreamingCell::PopulateGeneratedPackageForCook(U
 			return false;
 		}
 
-		UWorld* OuterWorld = GetCellOwner()->GetOuterWorld();
+		UWorld* OuterWorld = GetOuterWorld();
 		UWorldPartition* WorldPartition = OuterWorld->GetWorldPartition();
 
 		// Load cell Actors
@@ -298,7 +315,7 @@ int32 UWorldPartitionRuntimeLevelStreamingCell::GetActorCount() const
 
 FString UWorldPartitionRuntimeLevelStreamingCell::GetPackageNameToCreate() const
 {
-	return UWorldPartitionLevelStreamingPolicy::GetCellPackagePath(GetFName(), GetCellOwner()->GetOuterWorld());
+	return UWorldPartitionLevelStreamingPolicy::GetCellPackagePath(GetFName(), GetTypedOuter<UWorld>());
 }
 
 void UWorldPartitionRuntimeLevelStreamingCell::DumpStateLog(FHierarchicalLogArchive& Ar)
@@ -336,8 +353,8 @@ UWorldPartitionLevelStreamingDynamic* UWorldPartitionRuntimeLevelStreamingCell::
 	if (LevelStreaming)
 	{
 		// Setup pre-created LevelStreaming's outer to the WorldPartition owning world
-		const UWorldPartition* WorldPartition = GetCellOwner()->GetWorldPartition();
-		UWorld* OwningWorld = GetCellOwner()->GetOwningWorld();
+		const UWorldPartition* WorldPartition = GetTypedOuter<UWorldPartition>();
+		UWorld* OwningWorld = GetOwningWorld();
 		if (LevelStreaming->GetWorld() != OwningWorld)
 		{
 			LevelStreaming->Rename(nullptr, OwningWorld);
@@ -477,8 +494,7 @@ void UWorldPartitionRuntimeLevelStreamingCell::Deactivate() const
 
 void UWorldPartitionRuntimeLevelStreamingCell::OnLevelShown()
 {
-	UWorldPartition* WorldPartition = GetCellOwner()->GetWorldPartition();
-	if (WorldPartition && WorldPartition->IsInitialized())
+	if (UWorldPartition* WorldPartition = GetTypedOuter<UWorldPartition>(); WorldPartition && WorldPartition->IsInitialized())
 	{
 		WorldPartition->OnCellShown(this);
 	}
@@ -486,8 +502,7 @@ void UWorldPartitionRuntimeLevelStreamingCell::OnLevelShown()
 
 void UWorldPartitionRuntimeLevelStreamingCell::OnLevelHidden()
 {
-	UWorldPartition* WorldPartition = GetCellOwner()->GetWorldPartition();
-	if (WorldPartition && WorldPartition->IsInitialized())
+	if (UWorldPartition* WorldPartition = GetTypedOuter<UWorldPartition>(); WorldPartition && WorldPartition->IsInitialized())
 	{
 		WorldPartition->OnCellHidden(this);
 	}
diff --git a/Engine/Source/Runtime/Engine/Public/WorldPartition/WorldPartitionRuntimeCell.h b/Engine/Source/Runtime/Engine/Public/WorldPartition/WorldPartitionRuntimeCell.h
--- a/Engine/Source/Runtime/Engine/Public/WorldPartition/WorldPartitionRuntimeCell.h
+++ b/Engine/Source/Runtime/Engine/Public/WorldPartition/WorldPartitionRuntimeCell.h
@@ -9,7 +9,6 @@
 #include "WorldPartition/WorldPartitionActorContainerID.h"
 #include "WorldPartition/WorldPartitionRuntimeCellData.h"
 #include "WorldPartition/WorldPartitionRuntimeCellInterface.h"
-#include "WorldPartition/WorldPartitionRuntimeCellOwner.h"
 #include "ProfilingDebugging/ProfilingHelpers.h"
 #include "Misc/HierarchicalLogArchive.h"
 #include "Algo/AnyOf.h"
@@ -189,7 +188,6 @@ class ENGINE_API UWorldPartitionRuntimeCell : public UObject, public IWorldParti
 		return Algo::AnyOf(DataLayers, [&InDataLayers](const FName& DataLayer) { return InDataLayers.Contains(DataLayer); });
 	}
 	virtual FName GetLevelPackageName() const override;
-	virtual IWorldPartitionRuntimeCellOwner* GetCellOwner() const override { return CastChecked<IWorldPartitionRuntimeCellOwner>(GetOuter()); }
 	//~End IWorldPartitionCell Interface
 
 	bool GetBlockOnSlowLoading() const { return bBlockOnSlowLoading; }
diff --git a/Engine/Source/Runtime/Engine/Public/WorldPartition/WorldPartitionRuntimeCellInterface.h b/Engine/Source/Runtime/Engine/Public/WorldPartition/WorldPartitionRuntimeCellInterface.h
--- a/Engine/Source/Runtime/Engine/Public/WorldPartition/WorldPartitionRuntimeCellInterface.h
+++ b/Engine/Source/Runtime/Engine/Public/WorldPartition/WorldPartitionRuntimeCellInterface.h
@@ -7,7 +7,6 @@
 
 class UDataLayerInstance;
 class UDataLayerAsset;
-class IWorldPartitionRuntimeCellOwner;
 
 UINTERFACE()
 class ENGINE_API UWorldPartitionCell : public UInterface
@@ -38,8 +37,6 @@ public:
 	virtual FBox GetCellBounds() const = 0;
 	/** Returns the associated level package name. */
 	virtual FName GetLevelPackageName() const = 0;
-	/** Returns cell owner */
-	virtual IWorldPartitionRuntimeCellOwner* GetCellOwner() const = 0;
 
 #if WITH_EDITOR
 	/** Returns the referenced actor packages. */
diff --git a/Engine/Source/Runtime/Engine/Public/WorldPartition/WorldPartitionRuntimeCellOwner.h b/Engine/Source/Runtime/Engine/Public/WorldPartition/WorldPartitionRuntimeCellOwner.h
deleted file mode 100644
diff --git a/Engine/Source/Runtime/Engine/Public/WorldPartition/WorldPartitionRuntimeHash.h b/Engine/Source/Runtime/Engine/Public/WorldPartition/WorldPartitionRuntimeHash.h
--- a/Engine/Source/Runtime/Engine/Public/WorldPartition/WorldPartitionRuntimeHash.h
+++ b/Engine/Source/Runtime/Engine/Public/WorldPartition/WorldPartitionRuntimeHash.h
@@ -9,7 +9,6 @@
 #include "WorldPartition/ActorDescList.h"
 #include "WorldPartition/WorldPartitionHandle.h"
 #include "WorldPartition/WorldPartitionRuntimeCell.h"
-#include "WorldPartition/WorldPartitionRuntimeCellOwner.h"
 #include "WorldPartition/WorldPartitionActorDescViewProxy.h"
 #include "WorldPartition/WorldPartitionStreamingGeneration.h"
 #include "WorldPartition/WorldPartitionStreamingGenerationContext.h"
@@ -31,7 +30,7 @@ enum class EWorldPartitionStreamingPerformance : uint8
 };
 
 UCLASS(Abstract)
-class ENGINE_API URuntimeHashExternalStreamingObjectBase : public UObject, public IWorldPartitionRuntimeCellOwner
+class ENGINE_API URuntimeHashExternalStreamingObjectBase : public UObject
 {
 	GENERATED_BODY()
 
@@ -42,16 +41,14 @@ public:
 		OuterWorld = InOuterWorld;
 	}
 
-	//~ Begin IWorldPartitionRuntimeCellOwner Interface
-	virtual UWorld* GetOwningWorld() const override final { return OwningWorld.Get(); }
-	virtual UWorld* GetOuterWorld() const override final { return OuterWorld.Get(); }
-	//~ End IWorldPartitionRuntimeCellOwner Interface
+	UWorld* GetOwningWorld() const { return OwningWorld.Get(); }
+	UWorld* GetOuterWorld() const { return OuterWorld.Get(); }
 
 	virtual class UWorld* GetWorld() const override final { return GetOwningWorld(); }
 
 	void ForEachStreamingCells(TFunctionRef<void(UWorldPartitionRuntimeCell&)> Func);
 	
-	void OnStreamingObjectLoaded();
+	void OnStreamingObjectLoaded(UWorld* InjectedWorld);
 
 #if WITH_EDITOR
 	void PopulateGeneratorPackageForCook();
@@ -78,8 +75,8 @@ private:
 	TMap<const UWorldPartitionRuntimeCell*, double> CellToSourceMinSqrDistances;
 };
 
-UCLASS(Abstract, Config=Engine, AutoExpandCategories=(WorldPartition), Within = WorldPartition)
-class ENGINE_API UWorldPartitionRuntimeHash : public UObject, public IWorldPartitionRuntimeCellOwner
+UCLASS(Abstract, Config=Engine, AutoExpandCategories=(WorldPartition), Within=WorldPartition)
+class ENGINE_API UWorldPartitionRuntimeHash : public UObject
 {
 	GENERATED_UCLASS_BODY()
 
@@ -141,11 +138,6 @@ public:
 	UE_DEPRECATED(5.1, "GetStreamingCells is deprecated, use ForEachStreamingCells instead.")
 	bool GetStreamingCells(const TArray<FWorldPartitionStreamingSource>& Sources, UWorldPartitionRuntimeHash::FStreamingSourceCells& OutActivateCells, UWorldPartitionRuntimeHash::FStreamingSourceCells& OutLoadCells) const;
 
-	//~ Begin IWorldPartitionRuntimeCellOwner Interface
-	virtual UWorld* GetOwningWorld() const override { return GetWorld(); }
-	virtual UWorld* GetOuterWorld() const override { return GetTypedOuter<UWorld>(); }
-	//~ End IWorldPartitionRuntimeCellOwner Interface
-
 	// Streaming interface
 	virtual void ForEachStreamingCells(TFunctionRef<bool(const UWorldPartitionRuntimeCell*)> Func) const {}
 	virtual void ForEachStreamingCellsQuery(const FWorldPartitionStreamingQuerySource& QuerySource, TFunctionRef<bool(const UWorldPartitionRuntimeCell*)> Func, FWorldPartitionQueryCache* QueryCache = nullptr) const {}
diff --git a/Engine/Source/Runtime/Engine/Public/WorldPartition/WorldPartitionRuntimeLevelStreamingCell.h b/Engine/Source/Runtime/Engine/Public/WorldPartition/WorldPartitionRuntimeLevelStreamingCell.h
--- a/Engine/Source/Runtime/Engine/Public/WorldPartition/WorldPartitionRuntimeLevelStreamingCell.h
+++ b/Engine/Source/Runtime/Engine/Public/WorldPartition/WorldPartitionRuntimeLevelStreamingCell.h
@@ -30,6 +30,9 @@ class ENGINE_API UWorldPartitionRuntimeLevelStreamingCell : public UWorldPartiti
 	virtual bool IsLoading() const override;
 	//~End UWorldPartitionRuntimeCell Interface
 
+	UWorld* GetOwningWorld() const;
+	UWorld* GetOuterWorld() const;
+
 	//~Begin IWorldPartitionCell Interface
 	FName GetLevelPackageName() const override;
 	//~End IWorldPartitionCell Interface
