Fix non-unity build.

diff --git a/Engine/Source/Runtime/Engine/Private/WorldPartition/ContentBundle/ContentBundle.cpp b/Engine/Source/Runtime/Engine/Private/WorldPartition/ContentBundle/ContentBundle.cpp
--- a/Engine/Source/Runtime/Engine/Private/WorldPartition/ContentBundle/ContentBundle.cpp
+++ b/Engine/Source/Runtime/Engine/Private/WorldPartition/ContentBundle/ContentBundle.cpp
@@ -6,6 +6,7 @@
 #include "WorldPartition/ContentBundle/ContentBundleDescriptor.h"
 #include "WorldPartition/WorldPartitionRuntimeHash.h"
 #include "WorldPartition/ContentBundle/ContentBundleLog.h"
+#include "Engine/World.h"
 
 #if WITH_EDITOR
 #include "WorldPartition/ContentBundle/ContentBundleWorldSubsystem.h"
