Fix non-unity build.

diff --git a/Engine/Source/Runtime/Engine/Private/WorldPartition/WorldPartitionUtils.cpp b/Engine/Source/Runtime/Engine/Private/WorldPartition/WorldPartitionUtils.cpp
--- a/Engine/Source/Runtime/Engine/Private/WorldPartition/WorldPartitionUtils.cpp
+++ b/Engine/Source/Runtime/Engine/Private/WorldPartition/WorldPartitionUtils.cpp
@@ -9,6 +9,7 @@
 #include "WorldPartition/ContentBundle/ContentBundleEditor.h"
 #include "WorldPartition/ContentBundle/ContentBundleWorldSubsystem.h"
 #include "WorldPartition/WorldPartitionLog.h"
+#include "Engine/World.h"
 
 #define LOCTEXT_NAMESPACE "WorldPartition"
 
