[Backout] - CL23920665
Original CL Desc
-----------------------------------------------------------------
World Partition: get rid of IWorldPartitionRuntimeCellOwner.

diff --git a/Engine/Source/Runtime/Engine/Private/WorldPartition/ContentBundle/ContentBundle.cpp b/Engine/Source/Runtime/Engine/Private/WorldPartition/ContentBundle/ContentBundle.cpp
--- a/Engine/Source/Runtime/Engine/Private/WorldPartition/ContentBundle/ContentBundle.cpp
+++ b/Engine/Source/Runtime/Engine/Private/WorldPartition/ContentBundle/ContentBundle.cpp
@@ -33,7 +33,7 @@ void FContentBundle::DoInitialize()
 		if (UObject* Object = StaticFindObjectFast(URuntimeHashExternalStreamingObjectBase::StaticClass(), ExternalStreamingObjectPackage, *GetExternalStreamingObjectName()))
 		{
 			ExternalStreamingObject = CastChecked<URuntimeHashExternalStreamingObjectBase>(Object);
-			ExternalStreamingObject->OnStreamingObjectLoaded(GetInjectedWorld());
+			ExternalStreamingObject->OnStreamingObjectLoaded();
 		}
 		else
 		{
diff --git a/Engine/Source/Runtime/Engine/Private/WorldPartition/HLOD/HLODSubsystem.cpp b/Engine/Source/Runtime/Engine/Private/WorldPartition/HLOD/HLODSubsystem.cpp
--- a/Engine/Source/Runtime/Engine/Private/WorldPartition/HLOD/HLODSubsystem.cpp
+++ b/Engine/Source/Runtime/Engine/Private/WorldPartition/HLOD/HLODSubsystem.cpp
@@ -105,7 +105,7 @@ namespace FHLODSubsystem
 	    } // If not find it through the cell
 	    else if (const UWorldPartitionRuntimeCell* RuntimeCell = GetActorRuntimeCell(InWorldPartitionHLOD))
 	    {
-		    return RuntimeCell->GetTypedOuter<UWorldPartition>();
+		    return RuntimeCell->GetCellOwner()->GetOuterWorld()->GetWorldPartition();
 	    }
     
 	    return nullptr;
@@ -290,7 +290,7 @@ UHLODSubsystem::FCellData* UHLODSubsystem::GetCellData(const UWorldPartitionRunt
 		return nullptr;
 	}
 
-	const UWorldPartition* WorldPartition = InCell->GetTypedOuter<UWorldPartition>();
+	const UWorldPartition* WorldPartition = InCell->GetCellOwner()->GetOuterWorld()->GetWorldPartition();
 	check(WorldPartition);
 	FWorldPartitionHLODRuntimeData* WorldPartitionHLODRuntimeData = WorldPartitionsHLODRuntimeData.Find(WorldPartition);
 	if (WorldPartitionHLODRuntimeData)
diff --git a/Engine/Source/Runtime/Engine/Private/WorldPartition/WorldPartitionHelpers.cpp b/Engine/Source/Runtime/Engine/Private/WorldPartition/WorldPartitionHelpers.cpp
--- a/Engine/Source/Runtime/Engine/Private/WorldPartition/WorldPartitionHelpers.cpp
+++ b/Engine/Source/Runtime/Engine/Private/WorldPartition/WorldPartitionHelpers.cpp
@@ -5,6 +5,7 @@
 #include "WorldPartition/WorldPartition.h"
 #include "WorldPartition/WorldPartitionEditorHash.h"
 #include "WorldPartition/WorldPartitionRuntimeCellInterface.h"
+#include "WorldPartition/WorldPartitionRuntimeCellOwner.h"
 #include "Engine/World.h"
 #include "Engine/Level.h"
 #include "Algo/AnyOf.h"
@@ -34,6 +35,17 @@ UWorldPartition* FWorldPartitionHelpers::GetWorldPartition(const UObject* InObje
 			return GetWorldPartition(Cast<UObject>(Cell));
 		}
 	}
+	else if (const IWorldPartitionCell* Cell = Cast<const IWorldPartitionCell>(InObject))
+	{
+		if (const IWorldPartitionRuntimeCellOwner* CellOwner = Cell->GetCellOwner())
+		{
+			return GetWorldPartition(Cast<UObject>(CellOwner));
+		}
+	}
+	else if (const IWorldPartitionRuntimeCellOwner* CellOwner = Cast<const IWorldPartitionRuntimeCellOwner>(InObject))
+	{
+		OuterWorld = CellOwner->GetOuterWorld();
+	}
 	else if (const ULevel* OuterLevel = InObject->GetTypedOuter<ULevel>())
 	{
 		return GetWorldPartition(OuterLevel);
diff --git a/Engine/Source/Runtime/Engine/Private/WorldPartition/WorldPartitionLevelStreamingDynamic.cpp b/Engine/Source/Runtime/Engine/Private/WorldPartition/WorldPartitionLevelStreamingDynamic.cpp
--- a/Engine/Source/Runtime/Engine/Private/WorldPartition/WorldPartitionLevelStreamingDynamic.cpp
+++ b/Engine/Source/Runtime/Engine/Private/WorldPartition/WorldPartitionLevelStreamingDynamic.cpp
@@ -55,9 +55,11 @@ void UWorldPartitionLevelStreamingDynamic::Initialize(const UWorldPartitionRunti
 
 	UnsavedActorsContainer = InCell.UnsavedActorsContainer;
 
-	Initialize(GetTypedOuter<UWorld>(), InCell.GetPackages());
+	IWorldPartitionRuntimeCellOwner* CellOwner = InCell.GetCellOwner();
+	Initialize(CellOwner->GetOuterWorld(), InCell.GetPackages());
 #else
-	OuterWorldPartition = GetTypedOuter<UWorld>()->GetWorldPartition();
+	IWorldPartitionRuntimeCellOwner* CellOwner = InCell.GetCellOwner();
+	OuterWorldPartition = CellOwner->GetOuterWorld()->GetWorldPartition();
 #endif
 
 	UpdateShouldSkipMakingVisibilityTransactionRequest();
@@ -643,7 +645,7 @@ bool UWorldPartitionLevelStreamingDynamic::CanChangeVisibility(bool bMakeVisible
 
 			if (!bAlwaysAllowVisibilityChange)
 			{
-				if (const UWorldPartition* WorldPartition = RuntimeLevelStreamingCell->GetTypedOuter<UWorldPartition>())
+				if (const UWorldPartition* WorldPartition = RuntimeLevelStreamingCell->GetCellOwner()->GetWorldPartition())
 				{
 					if (UHLODSubsystem* HLODSubsystem = GetWorld()->GetSubsystem<UHLODSubsystem>())
 					{
diff --git a/Engine/Source/Runtime/Engine/Private/WorldPartition/WorldPartitionRuntimeCell.cpp b/Engine/Source/Runtime/Engine/Private/WorldPartition/WorldPartitionRuntimeCell.cpp
--- a/Engine/Source/Runtime/Engine/Private/WorldPartition/WorldPartitionRuntimeCell.cpp
+++ b/Engine/Source/Runtime/Engine/Private/WorldPartition/WorldPartitionRuntimeCell.cpp
@@ -21,7 +21,9 @@ UWorldPartitionRuntimeCell::UWorldPartitionRuntimeCell(const FObjectInitializer&
 	, DebugStreamingPriority(-1.f)
 #endif
 	, RuntimeCellData(nullptr)
-{}
+{
+	check(HasAnyFlags(RF_ClassDefaultObject) || GetOuter()->Implements<UWorldPartitionRuntimeCellOwner>());
+}
 
 #if WITH_EDITOR
 void UWorldPartitionRuntimeCell::PostDuplicate(bool bDuplicateForPIE)
diff --git a/Engine/Source/Runtime/Engine/Private/WorldPartition/WorldPartitionRuntimeCellOwner.cpp b/Engine/Source/Runtime/Engine/Private/WorldPartition/WorldPartitionRuntimeCellOwner.cpp
new file mode 100644
--- /dev/null
+++ b/Engine/Source/Runtime/Engine/Private/WorldPartition/WorldPartitionRuntimeCellOwner.cpp
@@ -0,0 +1,9 @@
+// Copyright Epic Games, Inc. All Rights Reserved.
+
+#include "WorldPartition/WorldPartitionRuntimeCellOwner.h"
+
+#include UE_INLINE_GENERATED_CPP_BY_NAME(WorldPartitionRuntimeCellOwner)
+
+UWorldPartitionRuntimeCellOwner::UWorldPartitionRuntimeCellOwner(const FObjectInitializer& ObjectInitializer)
+	: Super(ObjectInitializer)
+{}
diff --git a/Engine/Source/Runtime/Engine/Private/WorldPartition/WorldPartitionRuntimeHash.cpp b/Engine/Source/Runtime/Engine/Private/WorldPartition/WorldPartitionRuntimeHash.cpp
--- a/Engine/Source/Runtime/Engine/Private/WorldPartition/WorldPartitionRuntimeHash.cpp
+++ b/Engine/Source/Runtime/Engine/Private/WorldPartition/WorldPartitionRuntimeHash.cpp
@@ -30,20 +30,19 @@ void URuntimeHashExternalStreamingObjectBase::ForEachStreamingCells(TFunctionRef
 	}
 }
 
-void URuntimeHashExternalStreamingObjectBase::OnStreamingObjectLoaded(UWorld* InjectedWorld)
+void URuntimeHashExternalStreamingObjectBase::OnStreamingObjectLoaded()
 {
 	bool bIsACookedObject = !CellToLevelStreamingPackage.IsEmpty();
 	if (bIsACookedObject)
 	{
-		// Outer the streaming object back to the runtime hash.
-		Rename(nullptr, InjectedWorld->GetWorldPartition()->RuntimeHash);
-
-		// Cooked streaming object's Cells do not have LevelStreaming.
+		// Cooked streaming object's Cells do not have LevelStreaming and are outered to the streaming object.
+		// Create their level streaming now and re-outered to the runtime hash.
 		ForEachStreamingCells([this](UWorldPartitionRuntimeCell& Cell)
 		{
 			UWorldPartitionRuntimeLevelStreamingCell* RuntimeCell = CastChecked<UWorldPartitionRuntimeLevelStreamingCell>(&Cell);
 
 			FName LevelStreamingPackage = CellToLevelStreamingPackage.FindChecked(RuntimeCell->GetFName());
+			RuntimeCell->Rename(nullptr, GetOuterWorld()->GetWorldPartition()->RuntimeHash);
 			RuntimeCell->CreateAndSetLevelStreaming(*LevelStreamingPackage.ToString());
 		});
 	}
diff --git a/Engine/Source/Runtime/Engine/Private/WorldPartition/WorldPartitionRuntimeLevelStreamingCell.cpp b/Engine/Source/Runtime/Engine/Private/WorldPartition/WorldPartitionRuntimeLevelStreamingCell.cpp
--- a/Engine/Source/Runtime/Engine/Private/WorldPartition/WorldPartitionRuntimeLevelStreamingCell.cpp
+++ b/Engine/Source/Runtime/Engine/Private/WorldPartition/WorldPartitionRuntimeLevelStreamingCell.cpp
@@ -68,7 +68,7 @@ TSet<FName> UWorldPartitionRuntimeLevelStreamingCell::GetActorPackageNames() con
 FName UWorldPartitionRuntimeLevelStreamingCell::GetLevelPackageName() const
 {
 #if WITH_EDITOR
-	UWorld* World = GetOwningWorld();
+	UWorld* World = GetCellOwner()->GetOwningWorld();
 	if (World->IsPlayInEditor())
 	{
 		return FName(UWorldPartitionLevelStreamingPolicy::GetCellPackagePath(GetFName(), World));
@@ -106,9 +106,10 @@ UWorldPartitionLevelStreamingDynamic* UWorldPartitionRuntimeLevelStreamingCell::
 {
 	if (HasActors())
 	{
-		UWorld* OuterWorld = GetOuterWorld();
-		UWorld* OwningWorld = GetOwningWorld();
-		const UWorldPartition* WorldPartition = OuterWorld->GetWorldPartition();
+		const UWorldPartition* WorldPartition = GetCellOwner()->GetWorldPartition();
+		UWorld* OuterWorld = GetCellOwner()->GetOuterWorld();
+		UWorld* OwningWorld = GetCellOwner()->GetOwningWorld();
+
 		const FName LevelStreamingName = FName(*FString::Printf(TEXT("WorldPartitionLevelStreaming_%s"), *GetName()));
 
 		// When called by Commandlet (PopulateGeneratedPackageForCook), LevelStreaming's outer is set to Cell/WorldPartition's outer to prevent warnings when saving Cell Levels (Warning: Obj in another map). 
@@ -169,24 +170,6 @@ bool UWorldPartitionRuntimeLevelStreamingCell::IsLoading() const
 	return Super::IsLoading();
 }
 
-UWorld* UWorldPartitionRuntimeLevelStreamingCell::GetOwningWorld() const
-{
-	if (URuntimeHashExternalStreamingObjectBase* StreamingObjectOuter = GetTypedOuter<URuntimeHashExternalStreamingObjectBase>())
-	{
-		return StreamingObjectOuter->GetOwningWorld();
-	}
-	return GetTypedOuter<UWorldPartition>()->GetWorld();
-}
-
-UWorld* UWorldPartitionRuntimeLevelStreamingCell::GetOuterWorld() const
-{
-	if (URuntimeHashExternalStreamingObjectBase* StreamingObjectOuter = GetTypedOuter<URuntimeHashExternalStreamingObjectBase>())
-	{
-		return StreamingObjectOuter->GetOuterWorld();
-	}
-	return GetTypedOuter<UWorld>();
-}
-
 FLinearColor UWorldPartitionRuntimeLevelStreamingCell::GetDebugColor(EWorldPartitionRuntimeCellVisualizeMode VisualizeMode) const
 {
 	switch (VisualizeMode)
@@ -236,7 +219,7 @@ bool UWorldPartitionRuntimeLevelStreamingCell::PopulateGeneratorPackageForCook(T
 	{
 		FWorldPartitionLevelHelper::FPackageReferencer PackageReferencer;
 		const bool bLoadAsync = false;
-		UWorld* OuterWorld = GetOuterWorld();
+		UWorld* OuterWorld = GetCellOwner()->GetOuterWorld();
 		UWorldPartition* WorldPartition = OuterWorld->GetWorldPartition();
 
 		// Don't do SoftObjectPath remapping for PersistentLevel actors because references can end up in different cells
@@ -286,7 +269,7 @@ bool UWorldPartitionRuntimeLevelStreamingCell::PopulateGeneratedPackageForCook(U
 			return false;
 		}
 
-		UWorld* OuterWorld = GetOuterWorld();
+		UWorld* OuterWorld = GetCellOwner()->GetOuterWorld();
 		UWorldPartition* WorldPartition = OuterWorld->GetWorldPartition();
 
 		// Load cell Actors
@@ -315,7 +298,7 @@ int32 UWorldPartitionRuntimeLevelStreamingCell::GetActorCount() const
 
 FString UWorldPartitionRuntimeLevelStreamingCell::GetPackageNameToCreate() const
 {
-	return UWorldPartitionLevelStreamingPolicy::GetCellPackagePath(GetFName(), GetTypedOuter<UWorld>());
+	return UWorldPartitionLevelStreamingPolicy::GetCellPackagePath(GetFName(), GetCellOwner()->GetOuterWorld());
 }
 
 void UWorldPartitionRuntimeLevelStreamingCell::DumpStateLog(FHierarchicalLogArchive& Ar)
@@ -353,8 +336,8 @@ UWorldPartitionLevelStreamingDynamic* UWorldPartitionRuntimeLevelStreamingCell::
 	if (LevelStreaming)
 	{
 		// Setup pre-created LevelStreaming's outer to the WorldPartition owning world
-		const UWorldPartition* WorldPartition = GetTypedOuter<UWorldPartition>();
-		UWorld* OwningWorld = GetOwningWorld();
+		const UWorldPartition* WorldPartition = GetCellOwner()->GetWorldPartition();
+		UWorld* OwningWorld = GetCellOwner()->GetOwningWorld();
 		if (LevelStreaming->GetWorld() != OwningWorld)
 		{
 			LevelStreaming->Rename(nullptr, OwningWorld);
@@ -494,7 +477,8 @@ void UWorldPartitionRuntimeLevelStreamingCell::Deactivate() const
 
 void UWorldPartitionRuntimeLevelStreamingCell::OnLevelShown()
 {
-	if (UWorldPartition* WorldPartition = GetTypedOuter<UWorldPartition>(); WorldPartition && WorldPartition->IsInitialized())
+	UWorldPartition* WorldPartition = GetCellOwner()->GetWorldPartition();
+	if (WorldPartition && WorldPartition->IsInitialized())
 	{
 		WorldPartition->OnCellShown(this);
 	}
@@ -502,7 +486,8 @@ void UWorldPartitionRuntimeLevelStreamingCell::OnLevelShown()
 
 void UWorldPartitionRuntimeLevelStreamingCell::OnLevelHidden()
 {
-	if (UWorldPartition* WorldPartition = GetTypedOuter<UWorldPartition>(); WorldPartition && WorldPartition->IsInitialized())
+	UWorldPartition* WorldPartition = GetCellOwner()->GetWorldPartition();
+	if (WorldPartition && WorldPartition->IsInitialized())
 	{
 		WorldPartition->OnCellHidden(this);
 	}
diff --git a/Engine/Source/Runtime/Engine/Public/WorldPartition/WorldPartitionRuntimeCell.h b/Engine/Source/Runtime/Engine/Public/WorldPartition/WorldPartitionRuntimeCell.h
--- a/Engine/Source/Runtime/Engine/Public/WorldPartition/WorldPartitionRuntimeCell.h
+++ b/Engine/Source/Runtime/Engine/Public/WorldPartition/WorldPartitionRuntimeCell.h
@@ -9,6 +9,7 @@
 #include "WorldPartition/WorldPartitionActorContainerID.h"
 #include "WorldPartition/WorldPartitionRuntimeCellData.h"
 #include "WorldPartition/WorldPartitionRuntimeCellInterface.h"
+#include "WorldPartition/WorldPartitionRuntimeCellOwner.h"
 #include "ProfilingDebugging/ProfilingHelpers.h"
 #include "Misc/HierarchicalLogArchive.h"
 #include "Algo/AnyOf.h"
@@ -188,6 +189,7 @@ class ENGINE_API UWorldPartitionRuntimeCell : public UObject, public IWorldParti
 		return Algo::AnyOf(DataLayers, [&InDataLayers](const FName& DataLayer) { return InDataLayers.Contains(DataLayer); });
 	}
 	virtual FName GetLevelPackageName() const override;
+	virtual IWorldPartitionRuntimeCellOwner* GetCellOwner() const override { return CastChecked<IWorldPartitionRuntimeCellOwner>(GetOuter()); }
 	//~End IWorldPartitionCell Interface
 
 	bool GetBlockOnSlowLoading() const { return bBlockOnSlowLoading; }
diff --git a/Engine/Source/Runtime/Engine/Public/WorldPartition/WorldPartitionRuntimeCellInterface.h b/Engine/Source/Runtime/Engine/Public/WorldPartition/WorldPartitionRuntimeCellInterface.h
--- a/Engine/Source/Runtime/Engine/Public/WorldPartition/WorldPartitionRuntimeCellInterface.h
+++ b/Engine/Source/Runtime/Engine/Public/WorldPartition/WorldPartitionRuntimeCellInterface.h
@@ -7,6 +7,7 @@
 
 class UDataLayerInstance;
 class UDataLayerAsset;
+class IWorldPartitionRuntimeCellOwner;
 
 UINTERFACE()
 class ENGINE_API UWorldPartitionCell : public UInterface
@@ -37,6 +38,8 @@ public:
 	virtual FBox GetCellBounds() const = 0;
 	/** Returns the associated level package name. */
 	virtual FName GetLevelPackageName() const = 0;
+	/** Returns cell owner */
+	virtual IWorldPartitionRuntimeCellOwner* GetCellOwner() const = 0;
 
 #if WITH_EDITOR
 	/** Returns the referenced actor packages. */
diff --git a/Engine/Source/Runtime/Engine/Public/WorldPartition/WorldPartitionRuntimeCellOwner.h b/Engine/Source/Runtime/Engine/Public/WorldPartition/WorldPartitionRuntimeCellOwner.h
new file mode 100644
--- /dev/null
+++ b/Engine/Source/Runtime/Engine/Public/WorldPartition/WorldPartitionRuntimeCellOwner.h
@@ -0,0 +1,30 @@
+// Copyright Epic Games, Inc. All Rights Reserved.
+#pragma once
+
+#include "CoreMinimal.h"
+#include "UObject/ObjectMacros.h"
+#include "UObject/Interface.h"
+
+#include "Engine/World.h"
+
+#include "WorldPartitionRuntimeCellOwner.generated.h"
+
+class UWorld;
+class UWorldPartition;
+
+UINTERFACE()
+class ENGINE_API UWorldPartitionRuntimeCellOwner : public UInterface
+{
+	GENERATED_UINTERFACE_BODY()
+};
+
+class IWorldPartitionRuntimeCellOwner
+{
+	GENERATED_IINTERFACE_BODY()
+
+public:
+	virtual UWorld* GetOwningWorld() const = 0;
+	virtual UWorld* GetOuterWorld() const = 0;
+
+	UWorldPartition* GetWorldPartition() const { return GetOuterWorld()->GetWorldPartition(); }
+};
diff --git a/Engine/Source/Runtime/Engine/Public/WorldPartition/WorldPartitionRuntimeHash.h b/Engine/Source/Runtime/Engine/Public/WorldPartition/WorldPartitionRuntimeHash.h
--- a/Engine/Source/Runtime/Engine/Public/WorldPartition/WorldPartitionRuntimeHash.h
+++ b/Engine/Source/Runtime/Engine/Public/WorldPartition/WorldPartitionRuntimeHash.h
@@ -9,6 +9,7 @@
 #include "WorldPartition/ActorDescList.h"
 #include "WorldPartition/WorldPartitionHandle.h"
 #include "WorldPartition/WorldPartitionRuntimeCell.h"
+#include "WorldPartition/WorldPartitionRuntimeCellOwner.h"
 #include "WorldPartition/WorldPartitionActorDescViewProxy.h"
 #include "WorldPartition/WorldPartitionStreamingGeneration.h"
 #include "WorldPartition/WorldPartitionStreamingGenerationContext.h"
@@ -30,7 +31,7 @@ enum class EWorldPartitionStreamingPerformance : uint8
 };
 
 UCLASS(Abstract)
-class ENGINE_API URuntimeHashExternalStreamingObjectBase : public UObject
+class ENGINE_API URuntimeHashExternalStreamingObjectBase : public UObject, public IWorldPartitionRuntimeCellOwner
 {
 	GENERATED_BODY()
 
@@ -41,14 +42,16 @@ public:
 		OuterWorld = InOuterWorld;
 	}
 
-	UWorld* GetOwningWorld() const { return OwningWorld.Get(); }
-	UWorld* GetOuterWorld() const { return OuterWorld.Get(); }
+	//~ Begin IWorldPartitionRuntimeCellOwner Interface
+	virtual UWorld* GetOwningWorld() const override final { return OwningWorld.Get(); }
+	virtual UWorld* GetOuterWorld() const override final { return OuterWorld.Get(); }
+	//~ End IWorldPartitionRuntimeCellOwner Interface
 
 	virtual class UWorld* GetWorld() const override final { return GetOwningWorld(); }
 
 	void ForEachStreamingCells(TFunctionRef<void(UWorldPartitionRuntimeCell&)> Func);
 	
-	void OnStreamingObjectLoaded(UWorld* InjectedWorld);
+	void OnStreamingObjectLoaded();
 
 #if WITH_EDITOR
 	void PopulateGeneratorPackageForCook();
@@ -75,8 +78,8 @@ private:
 	TMap<const UWorldPartitionRuntimeCell*, double> CellToSourceMinSqrDistances;
 };
 
-UCLASS(Abstract, Config=Engine, AutoExpandCategories=(WorldPartition), Within=WorldPartition)
-class ENGINE_API UWorldPartitionRuntimeHash : public UObject
+UCLASS(Abstract, Config=Engine, AutoExpandCategories=(WorldPartition), Within = WorldPartition)
+class ENGINE_API UWorldPartitionRuntimeHash : public UObject, public IWorldPartitionRuntimeCellOwner
 {
 	GENERATED_UCLASS_BODY()
 
@@ -138,6 +141,11 @@ public:
 	UE_DEPRECATED(5.1, "GetStreamingCells is deprecated, use ForEachStreamingCells instead.")
 	bool GetStreamingCells(const TArray<FWorldPartitionStreamingSource>& Sources, UWorldPartitionRuntimeHash::FStreamingSourceCells& OutActivateCells, UWorldPartitionRuntimeHash::FStreamingSourceCells& OutLoadCells) const;
 
+	//~ Begin IWorldPartitionRuntimeCellOwner Interface
+	virtual UWorld* GetOwningWorld() const override { return GetWorld(); }
+	virtual UWorld* GetOuterWorld() const override { return GetTypedOuter<UWorld>(); }
+	//~ End IWorldPartitionRuntimeCellOwner Interface
+
 	// Streaming interface
 	virtual void ForEachStreamingCells(TFunctionRef<bool(const UWorldPartitionRuntimeCell*)> Func) const {}
 	virtual void ForEachStreamingCellsQuery(const FWorldPartitionStreamingQuerySource& QuerySource, TFunctionRef<bool(const UWorldPartitionRuntimeCell*)> Func, FWorldPartitionQueryCache* QueryCache = nullptr) const {}
diff --git a/Engine/Source/Runtime/Engine/Public/WorldPartition/WorldPartitionRuntimeLevelStreamingCell.h b/Engine/Source/Runtime/Engine/Public/WorldPartition/WorldPartitionRuntimeLevelStreamingCell.h
--- a/Engine/Source/Runtime/Engine/Public/WorldPartition/WorldPartitionRuntimeLevelStreamingCell.h
+++ b/Engine/Source/Runtime/Engine/Public/WorldPartition/WorldPartitionRuntimeLevelStreamingCell.h
@@ -30,9 +30,6 @@ class ENGINE_API UWorldPartitionRuntimeLevelStreamingCell : public UWorldPartiti
 	virtual bool IsLoading() const override;
 	//~End UWorldPartitionRuntimeCell Interface
 
-	UWorld* GetOwningWorld() const;
-	UWorld* GetOuterWorld() const;
-
 	//~Begin IWorldPartitionCell Interface
 	FName GetLevelPackageName() const override;
 	//~End IWorldPartitionCell Interface
