Fix DEPRECATED_CHARACTER_MOVEMENT_RPC argument (Contributed by [REDACTED_PERSON])
#10058

diff --git a/Engine/Source/Runtime/Engine/Classes/GameFramework/Character.h b/Engine/Source/Runtime/Engine/Classes/GameFramework/Character.h
--- a/Engine/Source/Runtime/Engine/Classes/GameFramework/Character.h
+++ b/Engine/Source/Runtime/Engine/Classes/GameFramework/Character.h
@@ -345,7 +345,7 @@ public:
 	void ClientAckGoodMove_Implementation(float TimeStamp);
 
 	/** Replicate position correction to client, associated with a timestamped servermove.  Client will replay subsequent moves after applying adjustment.  */
-	DEPRECATED_CHARACTER_MOVEMENT_RPC(ClientAckGoodMove, ClientMoveResponsePacked)
+	DEPRECATED_CHARACTER_MOVEMENT_RPC(ClientAdjustPosition, ClientMoveResponsePacked)
 	UFUNCTION(unreliable, client)
 	void ClientAdjustPosition(float TimeStamp, FVector NewLoc, FVector NewVel, UPrimitiveComponent* NewBase, FName NewBaseBoneName, bool bHasBase, bool bBaseRelativePosition, uint8 ServerMovementMode);
 	DEPRECATED_CHARACTER_MOVEMENT_RPC(ClientAdjustPosition_Implementation, ClientMoveResponsePacked_Implementation)
