UWorldPartitionLevelStreamingDynamic::LoadInEditor() - Create the temp world in the same mount point to avoid any possible issue with referenced assets that are in packages marked as not externally referenceable

diff --git a/Engine/Source/Runtime/Engine/Private/WorldPartition/WorldPartitionLevelStreamingDynamic.cpp b/Engine/Source/Runtime/Engine/Private/WorldPartition/WorldPartitionLevelStreamingDynamic.cpp
--- a/Engine/Source/Runtime/Engine/Private/WorldPartition/WorldPartitionLevelStreamingDynamic.cpp
+++ b/Engine/Source/Runtime/Engine/Private/WorldPartition/WorldPartitionLevelStreamingDynamic.cpp
@@ -70,8 +70,11 @@ UWorldPartitionLevelStreamingDynamic* UWorldPartitionLevelStreamingDynamic::Load
 	check(World->WorldType == EWorldType::Editor);
 	UWorldPartitionLevelStreamingDynamic* LevelStreaming = NewObject<UWorldPartitionLevelStreamingDynamic>(World, LevelStreamingName, RF_Transient);
 	
-	FString PackageName = FString::Printf(TEXT("/Temp/%s"), *LevelStreamingName.ToString());
-	TSoftObjectPtr<UWorld> WorldAsset(FSoftObjectPath(FString::Printf(TEXT("%s.%s"), *PackageName, *World->GetName())));
+	const FString WorldPackageName = World->GetPackage()->GetPathName();
+	const FStringView WorldMountPointName = FPathViews::GetMountPointNameFromPath(WorldPackageName);
+
+	FString PackageName = FString::Printf(TEXT("Temp/%s"), *LevelStreamingName.ToString());
+	TSoftObjectPtr<UWorld> WorldAsset(FSoftObjectPath(FString::Printf(TEXT("/%.*s/%s.%s"), WorldMountPointName.Len(), WorldMountPointName.GetData(), *PackageName, *World->GetName())));
 	LevelStreaming->SetWorldAsset(WorldAsset);
 	
 	LevelStreaming->LevelTransform = FTransform::Identity;
