LinearColor: expose ToHex to blueprint
- Use Case:  feed dynamic colors into rich text

diff --git a/Engine/Source/Runtime/Engine/Classes/Kismet/KismetMathLibrary.h b/Engine/Source/Runtime/Engine/Classes/Kismet/KismetMathLibrary.h
--- a/Engine/Source/Runtime/Engine/Classes/Kismet/KismetMathLibrary.h
+++ b/Engine/Source/Runtime/Engine/Classes/Kismet/KismetMathLibrary.h
@@ -3241,7 +3241,10 @@ class ENGINE_API UKismetMathLibrary : public UBlueprintFunctionLibrary
 	/** Element-wise multiplication of two linear colors (R/R, G/G, B/B, A/A) */
 	UFUNCTION(BlueprintPure, meta=(DisplayName = "LinearColor / LinearColor", CompactNodeTitle = "/", ScriptMethod = "Divide", ScriptOperator = "/;/=", Keywords = "/ divide"), Category="Math|Color")
 	static FLinearColor Divide_LinearColorLinearColor(FLinearColor A, FLinearColor B);
-
+	
+	/** Converts this color value to a hexadecimal string. The format of the string is RRGGBBAA. */
+	UFUNCTION(BlueprintPure, meta = (DisplayName = "ToHex"), Category = "Math|Color")
+	static FString ToHex_LinearColor(FLinearColor InColor);
 
 	//
 	// Plane functions.
diff --git a/Engine/Source/Runtime/Engine/Classes/Kismet/KismetMathLibrary.inl b/Engine/Source/Runtime/Engine/Classes/Kismet/KismetMathLibrary.inl
--- a/Engine/Source/Runtime/Engine/Classes/Kismet/KismetMathLibrary.inl
+++ b/Engine/Source/Runtime/Engine/Classes/Kismet/KismetMathLibrary.inl
@@ -3821,6 +3821,12 @@ FLinearColor UKismetMathLibrary::Divide_LinearColorLinearColor(FLinearColor A, F
 	return A / B;
 }
 
+KISMET_MATH_FORCEINLINE
+FString UKismetMathLibrary::ToHex_LinearColor(FLinearColor InColor)
+{
+	return InColor.ToFColor(true).ToHex();
+}
+
 
 KISMET_MATH_FORCEINLINE
 FString UKismetMathLibrary::SelectString(const FString& A, const FString& B, bool bSelectA)
