Fix a crash in FBatchedElements::Draw.
At boot, a default texture (because the right one was still compiling) was used to request a text draw, but the texture was released at the end of the compilation, resulting in a dangling pointer being used to draw.
It was only affecting fonts with an offline texture.
We know check that the texture is ready (compilation is over) before trying to use it for draw request.

Tests:
- repro steps from the jira, it no longer crash.
- forced display of text using an offline font, still works like before.
- ran packaged version, including QAGame, using a HUD with canvas drawing.

diff --git a/Engine/Source/Runtime/Engine/Private/UserInterface/CanvasItem.cpp b/Engine/Source/Runtime/Engine/Private/UserInterface/CanvasItem.cpp
--- a/Engine/Source/Runtime/Engine/Private/UserInterface/CanvasItem.cpp
+++ b/Engine/Source/Runtime/Engine/Private/UserInterface/CanvasItem.cpp
@@ -1016,7 +1016,7 @@ void FCanvasTextItem::DrawStringInternal_OfflineCache(FCanvas* InCanvas, const F
 	FVector2D CurrentPos = FVector2D(ForceInitToZero);
 	FHitProxyId HitProxyId = InCanvas->GetHitProxyId();
 	FTexture* LastTexture = nullptr;
-	UTexture* Tex = nullptr;
+	UTexture* Texture = nullptr;
 	FVector2D InvTextureSize(1.0f,1.0f);
 
 	const float CharIncrement = ( (float)Font->Kerning + HorizSpacingAdjust ) * Scale.X;
@@ -1056,25 +1056,29 @@ void FCanvasTextItem::DrawStringInternal_OfflineCache(FCanvas* InCanvas, const F
 		}
 
 		if( Font->Textures.IsValidIndex(Char.TextureIndex) && 
-			(Tex=Font->Textures[Char.TextureIndex])!=nullptr && 
-			Tex->GetResource() != nullptr )
+			(Texture=Font->Textures[Char.TextureIndex])!=nullptr && 
+			Texture->GetResource() != nullptr
+#if WITH_EDITOR
+			&& !Texture->IsCompiling()
+#endif
+			)
 		{
-			if( LastTexture != Tex->GetResource() || BatchedElements == nullptr )
+			if( LastTexture != Texture->GetResource() || BatchedElements == nullptr )
 			{
 				FBatchedElementParameters* BatchedElementParams = nullptr;
-				BatchedElements = InCanvas->GetBatchedElements(FCanvas::ET_Triangle, BatchedElementParams, Tex->GetResource(), BlendMode, FontRenderInfo.GlowInfo, false);
+				BatchedElements = InCanvas->GetBatchedElements(FCanvas::ET_Triangle, BatchedElementParams, Texture->GetResource(), BlendMode, FontRenderInfo.GlowInfo, false);
 				check(BatchedElements != nullptr);
 				// Trade-off to use memory for performance by pre-allocating more reserved space 
 				// for the triangles/vertices of the batched elements used to render the text tiles
 				// Only reserve initial batch, allow growth afterwards in case there are multiple repeated calls.
 				// Reserving exactly the added amount each time would essentially force an alloc each time on subsequent calls.
-				BatchedElements->ReserveTriangles(TextLen*2,Tex->GetResource(),BlendMode);
+				BatchedElements->ReserveTriangles(TextLen*2,Texture->GetResource(),BlendMode);
 				BatchedElements->ReserveVertices(TextLen*4);
 
-				InvTextureSize.X = 1.0f / Tex->GetSurfaceWidth();
-				InvTextureSize.Y = 1.0f / Tex->GetSurfaceHeight();
+				InvTextureSize.X = 1.0f / Texture->GetSurfaceWidth();
+				InvTextureSize.Y = 1.0f / Texture->GetSurfaceHeight();
 			}
-			LastTexture = Tex->GetResource();
+			LastTexture = Texture->GetResource();
 
 			float SizeX			= Char.USize * Scale.X;
 			const float SizeY	= Char.VSize * Scale.Y;
@@ -1114,8 +1118,8 @@ void FCanvasTextItem::DrawStringInternal_OfflineCache(FCanvas* InCanvas, const F
 					InColor,
 					HitProxyId);
 
-				BatchedElements->AddTriangle(V00, V10, V11, Tex->GetResource(), BlendMode, FontRenderInfo.GlowInfo);
-				BatchedElements->AddTriangle(V00, V11, V01, Tex->GetResource(), BlendMode, FontRenderInfo.GlowInfo);
+				BatchedElements->AddTriangle(V00, V10, V11, Texture->GetResource(), BlendMode, FontRenderInfo.GlowInfo);
+				BatchedElements->AddTriangle(V00, V11, V01, Texture->GetResource(), BlendMode, FontRenderInfo.GlowInfo);
 			};
 
 			// Render all the effects first that appear under the text
