Fix crash when passing a null DataLayerAsset to GetDataLayerInstance in cook builds

diff --git a/Engine/Source/Runtime/Engine/Private/WorldPartition/DataLayer/WorldDataLayers.cpp b/Engine/Source/Runtime/Engine/Private/WorldPartition/DataLayer/WorldDataLayers.cpp
--- a/Engine/Source/Runtime/Engine/Private/WorldPartition/DataLayer/WorldDataLayers.cpp
+++ b/Engine/Source/Runtime/Engine/Private/WorldPartition/DataLayer/WorldDataLayers.cpp
@@ -707,9 +707,12 @@ const UDataLayerInstance* AWorldDataLayers::GetDataLayerInstance(const UDataLaye
 		}
 	}
 #else
-	if (const UDataLayerInstance* const* FoundDataLayerInstance = AssetNameToInstance.Find(InDataLayerAsset->GetFullName()))
+	if (InDataLayerAsset != nullptr)
 	{
-		return *FoundDataLayerInstance;
+		if (const UDataLayerInstance* const* FoundDataLayerInstance = AssetNameToInstance.Find(InDataLayerAsset->GetFullName()))
+		{
+			return *FoundDataLayerInstance;
+		}
 	}
 #endif
 
