LevelInstance: Call DetachLoader (release package file handle) allows saving WP LevelInstance level package (Pivot edit)

diff --git a/Engine/Source/Runtime/Engine/Private/LevelInstance/LevelInstanceLevelStreaming.cpp b/Engine/Source/Runtime/Engine/Private/LevelInstance/LevelInstanceLevelStreaming.cpp
--- a/Engine/Source/Runtime/Engine/Private/LevelInstance/LevelInstanceLevelStreaming.cpp
+++ b/Engine/Source/Runtime/Engine/Private/LevelInstance/LevelInstanceLevelStreaming.cpp
@@ -22,6 +22,7 @@
 #include "LevelInstance/LevelInstanceEditorInstanceActor.h"
 #include "LevelUtils.h"
 #include "ActorFolder.h"
+#include "UObject/LinkerLoad.h"
 #endif
 
 ULevelStreamingLevelInstance::ULevelStreamingLevelInstance(const FObjectInitializer& ObjectInitializer)
@@ -88,10 +89,7 @@ void ULevelStreamingLevelInstance::OnLoadedActorAddedToLevel(AActor& InActor)
 
 void ULevelStreamingLevelInstance::ResetLevelInstanceLoaders()
 {
-	// @todo_ow: Resetting the load will prevent any OFPA package to properly load since their import level will fail to resolve.
-	//           This is a temporary workaround. The downside of this is that the level can't be saved. 
-	//           Most of the changes will only affect OFPA packages. One problematic use case is when changing the pivot of a Level Instance.
-
+	// @todo_ow: Ideally at some point it is no longer needed to ResetLoaders at all and Linker would not lock the package files preventing saves.
 	if (bResetLoadersCalled)
 	{
 		return;
@@ -106,6 +104,12 @@ void ULevelStreamingLevelInstance::ResetLevelInstanceLoaders()
 		{
 			ResetLoaders(OuterWorld->GetPackage());
 		}
+		else if(FLinkerLoad* LinkerLoad = OuterWorld->GetPackage()->GetLinker())
+		{
+			// Resetting the loader will prevent any OFPA package to properly re-load since their import level will fail to resolve.
+			// DetachLoader allows releasing the lock on the file handle so level package can be saved.
+			LinkerLoad->DetachLoader();
+		}
 
 		for (AActor* Actor : LoadedLevel->Actors)
 		{
