World Partition: always include landscape actors into world size calculations for enable streaming warninng, even if they are non-spatially loaded.

diff --git a/Engine/Source/Runtime/Engine/Private/WorldPartition/WorldPartition.cpp b/Engine/Source/Runtime/Engine/Private/WorldPartition/WorldPartition.cpp
--- a/Engine/Source/Runtime/Engine/Private/WorldPartition/WorldPartition.cpp
+++ b/Engine/Source/Runtime/Engine/Private/WorldPartition/WorldPartition.cpp
@@ -15,6 +15,7 @@
 #include "GameFramework/WorldSettings.h"
 #include "ProfilingDebugging/ScopedTimers.h"
 #include "UObject/UE5MainStreamObjectVersion.h"
+#include "LandscapeProxy.h"
 
 #include UE_INLINE_GENERATED_CPP_BY_NAME(WorldPartition)
 
@@ -1406,7 +1407,7 @@ void UWorldPartition::Tick(float DeltaSeconds)
 			FBox AllActorsBounds(ForceInit);
 			for (FActorDescContainerCollection::TConstIterator<> ActorDescIterator(this); ActorDescIterator; ++ActorDescIterator)
 			{
-				if (ActorDescIterator->GetIsSpatiallyLoadedRaw())
+				if (ActorDescIterator->GetIsSpatiallyLoadedRaw() || ActorDescIterator->GetActorNativeClass()->IsChildOf<ALandscapeProxy>())
 				{
 					AllActorsBounds += ActorDescIterator->GetEditorBounds();
 
