Use DynamicShadowDistanceMovableLight for Stationary lights when static lighting is disable in project settings
- Currently it was incorrectly using r.Shadow.UnbuiltWholeSceneDynamicShadowRadius.
- Hide "Dynamic Shadow Distance StationaryLight" from editor.
- Rename "Dynamic Shadow Distance MovableLight" to "Dynamic Shadow Distance".

diff --git a/Engine/Source/Editor/DetailCustomizations/Private/DirectionalLightComponentDetails.cpp b/Engine/Source/Editor/DetailCustomizations/Private/DirectionalLightComponentDetails.cpp
--- a/Engine/Source/Editor/DetailCustomizations/Private/DirectionalLightComponentDetails.cpp
+++ b/Engine/Source/Editor/DetailCustomizations/Private/DirectionalLightComponentDetails.cpp
@@ -26,21 +26,25 @@ TSharedRef<IDetailCustomization> FDirectionalLightComponentDetails::MakeInstance
 
 void FDirectionalLightComponentDetails::CustomizeDetails( IDetailLayoutBuilder& DetailBuilder )
 {
-	// Grab the Mobility property from SceneComponent
-	MobilityProperty = DetailBuilder.GetProperty("Mobility", USceneComponent::StaticClass());
-
 	// Get cascaded shadow map category
 	IDetailCategoryBuilder& ShadowMapCategory = DetailBuilder.EditCategory("CascadedShadowMaps", FText::GetEmpty(), ECategoryPriority::Default );
 
-	// Add DynamicShadowDistanceMovableLight
-	TSharedPtr<IPropertyHandle> MovableShadowRadiusProperty = DetailBuilder.GetProperty("DynamicShadowDistanceMovableLight");
-	ShadowMapCategory.AddProperty( MovableShadowRadiusProperty)
-		.IsEnabled( TAttribute<bool>( this, &FDirectionalLightComponentDetails::IsLightMovable ) );
+	TSharedPtr<IPropertyHandle> MovableShadowRadiusPropertyHandle = DetailBuilder.GetProperty("DynamicShadowDistanceMovableLight");
+	TSharedPtr<IPropertyHandle> StationaryShadowRadiusPropertyHandle = DetailBuilder.GetProperty("DynamicShadowDistanceStationaryLight");
+
+	static const auto AllowStaticLightingVar = IConsoleManager::Get().FindTConsoleVariableDataInt(TEXT("r.AllowStaticLighting"));
+	const bool bAllowStaticLighting = (!AllowStaticLightingVar || AllowStaticLightingVar->GetValueOnGameThread() != 0);
+
+	if(!bAllowStaticLighting)
+	{
+		// If static lighting is not allowed, hide DynamicShadowDistanceStationaryLight and rename DynamicShadowDistanceMovableLight to "Dynamic Shadow Distance"
 
-	// Add DynamicShadowDistanceStationaryLight
-	TSharedPtr<IPropertyHandle> StationaryShadowRadiusProperty = DetailBuilder.GetProperty("DynamicShadowDistanceStationaryLight");
-	ShadowMapCategory.AddProperty( StationaryShadowRadiusProperty)
-		.IsEnabled( TAttribute<bool>( this, &FDirectionalLightComponentDetails::IsLightStationary ) );
+		FProperty* MovableShadowRadiusProperty = MovableShadowRadiusPropertyHandle->GetProperty();
+		MovableShadowRadiusProperty->SetMetaData(TEXT("DisplayName"), TEXT("Dynamic Shadow Distance"));
+
+		ShadowMapCategory.AddProperty(StationaryShadowRadiusPropertyHandle)
+			.Visibility(EVisibility::Hidden);
+	}
 
 	TSharedPtr<IPropertyHandle> LightIntensityProperty = DetailBuilder.GetProperty("Intensity", ULightComponentBase::StaticClass());
 	// Point lights need to override the ui min and max for units of lumens, so we have to undo that
@@ -52,18 +56,4 @@ void FDirectionalLightComponentDetails::CustomizeDetails( IDetailLayoutBuilder&
 
 }
 
-bool FDirectionalLightComponentDetails::IsLightMovable() const
-{
-	uint8 Mobility;
-	MobilityProperty->GetValue(Mobility);
-	return (Mobility == EComponentMobility::Movable);
-}
-
-bool FDirectionalLightComponentDetails::IsLightStationary() const
-{
-	uint8 Mobility;
-	MobilityProperty->GetValue(Mobility);
-	return (Mobility == EComponentMobility::Stationary);
-}
-
 #undef LOCTEXT_NAMESPACE
diff --git a/Engine/Source/Editor/DetailCustomizations/Private/DirectionalLightComponentDetails.h b/Engine/Source/Editor/DetailCustomizations/Private/DirectionalLightComponentDetails.h
--- a/Engine/Source/Editor/DetailCustomizations/Private/DirectionalLightComponentDetails.h
+++ b/Engine/Source/Editor/DetailCustomizations/Private/DirectionalLightComponentDetails.h
@@ -16,12 +16,4 @@ public:
 
 	/** IDetailCustomization interface */
 	virtual void CustomizeDetails( IDetailLayoutBuilder& DetailBuilder ) override;
-
-private:
-
-	bool IsLightMovable() const;
-	bool IsLightStationary() const;
-
-	TSharedPtr<IPropertyHandle> MobilityProperty;
-
 };
diff --git a/Engine/Source/Runtime/Engine/Private/Components/DirectionalLightComponent.cpp b/Engine/Source/Runtime/Engine/Private/Components/DirectionalLightComponent.cpp
--- a/Engine/Source/Runtime/Engine/Private/Components/DirectionalLightComponent.cpp
+++ b/Engine/Source/Runtime/Engine/Private/Components/DirectionalLightComponent.cpp
@@ -188,6 +188,21 @@ void ComputeShadowCullingVolume(bool bReverseCulling, const FVector* CascadeFrus
 	ConvexVolumeOut = FConvexVolume(Planes);
 }
 
+float ComputeWholeSceneDynamicShadowRadius(EComponentMobility::Type Mobility, float DynamicShadowDistanceMovableLight, float DynamicShadowDistanceStationaryLight)
+{
+	static const auto AllowStaticLightingVar = IConsoleManager::Get().FindTConsoleVariableDataInt(TEXT("r.AllowStaticLighting"));
+	const bool bAllowStaticLighting = (!AllowStaticLightingVar || AllowStaticLightingVar->GetValueOnGameThread() != 0);
+
+	if (Mobility == EComponentMobility::Movable || !bAllowStaticLighting)
+	{
+		return DynamicShadowDistanceMovableLight;
+	}
+	else
+	{
+		return DynamicShadowDistanceStationaryLight;
+	}
+}
+
 /**
  * The scene info for a directional light.
  */
@@ -326,14 +341,7 @@ public:
 	{
 		LightShaftOverrideDirection.Normalize();
 
-		if(Component->Mobility == EComponentMobility::Movable)
-		{
-			WholeSceneDynamicShadowRadius = Component->DynamicShadowDistanceMovableLight;
-		}
-		else
-		{
-			WholeSceneDynamicShadowRadius = Component->DynamicShadowDistanceStationaryLight;
-		}
+		WholeSceneDynamicShadowRadius = ComputeWholeSceneDynamicShadowRadius(Component->Mobility, Component->DynamicShadowDistanceMovableLight, Component->DynamicShadowDistanceStationaryLight);
 
 		const float FarCascadeSize = Component->FarShadowDistance - WholeSceneDynamicShadowRadius;
 
@@ -1065,12 +1073,17 @@ bool UDirectionalLightComponent::CanEditChange(const FProperty* InProperty) cons
 {
 	if (InProperty)
 	{
+		static const auto AllowStaticLightingVar = IConsoleManager::Get().FindTConsoleVariableDataInt(TEXT("r.AllowStaticLighting"));
+		const bool bAllowStaticLighting = (!AllowStaticLightingVar || AllowStaticLightingVar->GetValueOnGameThread() != 0);
+
 		FString PropertyName = InProperty->GetName();
-		
-		bool bShadowCascades = CastShadows
+
+		const float WholeSceneDynamicShadowRadius = ComputeWholeSceneDynamicShadowRadius(Mobility, DynamicShadowDistanceMovableLight, DynamicShadowDistanceStationaryLight);
+
+		const bool bShadowCascades = CastShadows
 			&& CastDynamicShadows 
-			&& ((DynamicShadowDistanceMovableLight > 0 && Mobility == EComponentMobility::Movable)
-			|| (DynamicShadowDistanceStationaryLight > 0 && Mobility == EComponentMobility::Stationary));
+			&& WholeSceneDynamicShadowRadius > 0
+			&& Mobility != EComponentMobility::Static;
 
 		if (PropertyName == GET_MEMBER_NAME_STRING_CHECKED(UDirectionalLightComponent, bUseInsetShadowsForMovableObjects))
 		{
@@ -1079,14 +1092,20 @@ bool UDirectionalLightComponent::CanEditChange(const FProperty* InProperty) cons
 
 		if (PropertyName == GET_MEMBER_NAME_STRING_CHECKED(UDirectionalLightComponent, DynamicShadowDistanceMovableLight))
 		{
-			return CastShadows && CastDynamicShadows;
+			const bool bUseMovableLightDistance = (Mobility == EComponentMobility::Movable) || (Mobility == EComponentMobility::Stationary && !bAllowStaticLighting);
+			return CastShadows && CastDynamicShadows && bUseMovableLightDistance;
+		}
+
+		if (PropertyName == GET_MEMBER_NAME_STRING_CHECKED(UDirectionalLightComponent, DynamicShadowDistanceStationaryLight))
+		{
+			const bool bUseStationaryLightDistance = (Mobility == EComponentMobility::Stationary) && bAllowStaticLighting;
+			return CastShadows && CastDynamicShadows && bUseStationaryLightDistance;
 		}
 
 		if (PropertyName == GET_MEMBER_NAME_STRING_CHECKED(UDirectionalLightComponent, DynamicShadowCascades)
 			|| PropertyName == GET_MEMBER_NAME_STRING_CHECKED(UDirectionalLightComponent, CascadeDistributionExponent)
 			|| PropertyName == GET_MEMBER_NAME_STRING_CHECKED(UDirectionalLightComponent, CascadeTransitionFraction)
 			|| PropertyName == GET_MEMBER_NAME_STRING_CHECKED(UDirectionalLightComponent, ShadowDistanceFadeoutFraction)
-			|| PropertyName == GET_MEMBER_NAME_STRING_CHECKED(UDirectionalLightComponent, bUseInsetShadowsForMovableObjects)
 			|| PropertyName == GET_MEMBER_NAME_STRING_CHECKED(UDirectionalLightComponent, FarShadowCascadeCount))
 		{
 			return bShadowCascades;
