Better display name and tooltip for USceneCapture2D::bConsiderUnrenderedOpaquePixelAsFullyTranslucent

diff --git a/Engine/Source/Runtime/Engine/Classes/Components/SceneCaptureComponent2D.h b/Engine/Source/Runtime/Engine/Classes/Components/SceneCaptureComponent2D.h
--- a/Engine/Source/Runtime/Engine/Classes/Components/SceneCaptureComponent2D.h
+++ b/Engine/Source/Runtime/Engine/Classes/Components/SceneCaptureComponent2D.h
@@ -106,8 +106,8 @@ public:
 	UPROPERTY(Transient, BlueprintReadWrite, Category = SceneCapture)
 	uint32 bCameraCutThisFrame : 1;
 
-	/** Treat unrendered opaque pixels as fully translucent. This is important for effects such as exponential weight fog, so it does not get applied on unrendered opaque pixels. */
-	UPROPERTY(EditAnywhere, BlueprintReadWrite, Category = SceneCapture)
+	/** Whether to only render exponential height fog on opaque pixels which were rendered by the scene capture. */
+	UPROPERTY(EditAnywhere, BlueprintReadWrite, Category = SceneCapture, meta = (DisplayName = "Fog only on rendered pixels"))
 	uint32 bConsiderUnrenderedOpaquePixelAsFullyTranslucent : 1;
 
 	/** Array of scene view extensions specifically to apply to this scene capture */
