Fixed UWorldPartitionLevelStreamingDynamic's outer world partition : was invalid for instanced world partition.

diff --git a/Engine/Source/Runtime/Engine/Private/WorldPartition/WorldPartitionLevelStreamingDynamic.cpp b/Engine/Source/Runtime/Engine/Private/WorldPartition/WorldPartitionLevelStreamingDynamic.cpp
--- a/Engine/Source/Runtime/Engine/Private/WorldPartition/WorldPartitionLevelStreamingDynamic.cpp
+++ b/Engine/Source/Runtime/Engine/Private/WorldPartition/WorldPartitionLevelStreamingDynamic.cpp
@@ -50,14 +50,15 @@ void UWorldPartitionLevelStreamingDynamic::Initialize(const UWorldPartitionRunti
 	bShouldBeAlwaysLoaded = InCell.IsAlwaysLoaded();
 	StreamingPriority = 0;
 
+	UWorld* CellOuterWorld = InCell.GetTypedOuter<UWorld>();
 #if WITH_EDITOR
 	check(ChildPackages.Num() == 0);
 
 	UnsavedActorsContainer = InCell.UnsavedActorsContainer;
 
-	Initialize(GetTypedOuter<UWorld>(), InCell.GetPackages());
+	Initialize(CellOuterWorld, InCell.GetPackages());
 #else
-	OuterWorldPartition = GetTypedOuter<UWorld>()->GetWorldPartition();
+	OuterWorldPartition = CellOuterWorld->GetWorldPartition();
 #endif
 
 	UpdateShouldSkipMakingVisibilityTransactionRequest();
