Moves stat unit's TSR convergence speed to its own stat tsr and adds TSR input feed in MP/s

Also moves GPixelRenderCounters to new RenderCounters.h

diff --git a/Engine/Source/Runtime/Engine/Classes/Engine/Engine.h b/Engine/Source/Runtime/Engine/Classes/Engine/Engine.h
--- a/Engine/Source/Runtime/Engine/Classes/Engine/Engine.h
+++ b/Engine/Source/Runtime/Engine/Classes/Engine/Engine.h
@@ -3708,6 +3708,7 @@ private:
 	bool ToggleStatUnitTime(UWorld* World, FCommonViewportClient* ViewportClient, const TCHAR* Stream = nullptr);
 	bool ToggleStatRaw(UWorld* World, FCommonViewportClient* ViewportClient, const TCHAR* Stream = nullptr);
 	bool ToggleStatParticlePerf(UWorld* World, FCommonViewportClient* ViewportClient, const TCHAR* Stream = nullptr);
+	bool ToggleStatTSR(UWorld* World, FCommonViewportClient* ViewportClient, const TCHAR* Stream = nullptr);
 #endif
 
 	/**
diff --git a/Engine/Source/Runtime/Engine/Private/UnrealClient.cpp b/Engine/Source/Runtime/Engine/Private/UnrealClient.cpp
--- a/Engine/Source/Runtime/Engine/Private/UnrealClient.cpp
+++ b/Engine/Source/Runtime/Engine/Private/UnrealClient.cpp
@@ -28,6 +28,7 @@
 #include "Performance/EnginePerformanceTargets.h"
 #include "Elements/Framework/TypedElementList.h"
 #include "EngineUtils.h"
+#include "RenderCounters.h"
 #include "RenderGraphUtils.h"
 #include "DynamicResolutionState.h"
 #include "Stats/StatsTrace.h"
@@ -511,6 +512,7 @@ int32 FStatUnitData::DrawStat(FViewport* InViewport, FCanvas* InCanvas, int32 In
 	float Max_InputLatencyTime = 0.0f;
 
 	const bool bShowUnitMaxTimes = InViewport->GetClient() ? InViewport->GetClient()->IsStatEnabled(TEXT("UnitMax")) : false;
+	const bool bShowTSRStatistics = InViewport->GetClient() ? InViewport->GetClient()->IsStatEnabled(TEXT("TSR")) : false;
 #if !UE_BUILD_SHIPPING
 	const bool bShowRawUnitTimes = InViewport->GetClient() ? InViewport->GetClient()->IsStatEnabled(TEXT("Raw")) : false;
 	RenderThreadTimes[CurrentIndex] = bShowRawUnitTimes ? RawRenderThreadTime : RenderThreadTime;
@@ -936,31 +938,46 @@ int32 FStatUnitData::DrawStat(FViewport* InViewport, FCanvas* InCanvas, int32 In
 			PushRows(/* RowsCount = */ 1);
 		}
 
-		// TSR history convergence rate
-		if (GPixelRenderCounters.GetPixelDisplayCount())
+		if (bShowTSRStatistics && GPixelRenderCounters.GetPixelDisplayCount())
 		{
-			// Target 1 sample per pixel for the convergence speed measurment.
-			const float TargetSamplePerPixel = 1.0f;
+			const uint32 PixelRenderCount = GPixelRenderCounters.GetPixelRenderCount();
+			const uint32 PixelDisplayCount = GPixelRenderCounters.GetPixelDisplayCount();
 
-			// Ideal TSR uses it to render 1080p -> 4k at 60hz.
-			const float IdealConvergenceTime = TargetSamplePerPixel * FMath::Pow(1080.0f / 2160.0f, -2.0f) * (1000.0f / 60.0f);
+			// TSR input feed in pixel/s
+			{
+				const FColor Color = (PixelRenderCount < 1280 * 720) ? StatRed : ((PixelRenderCount < 1920 * 1080) ? StatOrange : StatGreen);
+
+				float TSRFeed = PixelRenderCount * (1000.0f / FrameTime);
 
-			// Compute the resolution fraction agregate.
-			uint32 PixelRenderCount = GPixelRenderCounters.GetPixelRenderCount();
-			uint32 PixelDisplayCount = GPixelRenderCounters.GetPixelDisplayCount();
-			float ResolutionFraction = FMath::Sqrt(float(PixelRenderCount) / float(PixelDisplayCount));
-			float ConvergenceSpeedMultiplier = FMath::Pow(ResolutionFraction, -2.0f);
+				DrawTitleString(TEXT("TSR feed"), NoUnitGraphColor);
+				DrawDefaultAvgCell(FString::Printf(TEXT("%.2f MP/s"), float(TSRFeed) / 1000000.0f), Color);
 
-			// Compute how long it takes for the history to converge to 1spp.
-			float ConvergenceFrameCount = ConvergenceSpeedMultiplier * TargetSamplePerPixel;
-			float ConvergenceTime = FMath::Max(ConvergenceFrameCount, 1.0f) * FrameTime;
+				PushRows(/* RowsCount = */ 1);
+			}
 
-			const FColor Color = (ConvergenceTime <= IdealConvergenceTime) ? StatGreen : ((ConvergenceTime < 2.0f * IdealConvergenceTime) ? StatOrange : StatRed);
+			// TSR history convergence rate
+			{
+				// Target 1 sample per pixel for the convergence speed measurment.
+				const float TargetSamplePerPixel = 1.0f;
 
-			DrawTitleString(TEXT("TSR -> 1spp"), NoUnitGraphColor);
-			DrawDefaultAvgCell(FString::Printf(STATUNIT_FORMAT_AVGTIME, ConvergenceTime), Color);
+				// Ideal TSR uses it to render 1080p -> 4k at 60hz.
+				const float IdealConvergenceTime = TargetSamplePerPixel * FMath::Pow(1080.0f / 2160.0f, -2.0f) * (1000.0f / 60.0f);
 
-			PushRows(/* RowsCount = */ 1);
+				// Compute the resolution fraction agregate.
+				float ResolutionFraction = FMath::Sqrt(float(PixelRenderCount) / float(PixelDisplayCount));
+				float ConvergenceSpeedMultiplier = FMath::Pow(ResolutionFraction, -2.0f);
+
+				// Compute how long it takes for the history to converge to 1spp.
+				float ConvergenceFrameCount = ConvergenceSpeedMultiplier * TargetSamplePerPixel;
+				float ConvergenceTime = FMath::Max(ConvergenceFrameCount, 1.0f) * FrameTime;
+
+				const FColor Color = (ConvergenceTime <= IdealConvergenceTime) ? StatGreen : ((ConvergenceTime < 2.0f * IdealConvergenceTime) ? StatOrange : StatRed);
+
+				DrawTitleString(TEXT("TSR 1spp"), NoUnitGraphColor);
+				DrawDefaultAvgCell(FString::Printf(STATUNIT_FORMAT_AVGTIME, ConvergenceTime), Color);
+
+				PushRows(/* RowsCount = */ 1);
+			}
 		}
 	}
 
diff --git a/Engine/Source/Runtime/Engine/Private/UnrealEngine.cpp b/Engine/Source/Runtime/Engine/Private/UnrealEngine.cpp
--- a/Engine/Source/Runtime/Engine/Private/UnrealEngine.cpp
+++ b/Engine/Source/Runtime/Engine/Private/UnrealEngine.cpp
@@ -2120,6 +2120,7 @@ void UEngine::Init(IEngineLoop* InEngineLoop)
 	EngineStats.Add(FEngineStatFuncs(TEXT("STAT_UnitTime"), TEXT("STATCAT_Engine"), FText::GetEmpty(), FEngineStatRender(), FEngineStatToggle::CreateUObject(this, &UEngine::ToggleStatUnitTime)));
 	EngineStats.Add(FEngineStatFuncs(TEXT("STAT_Raw"), TEXT("STATCAT_Engine"), FText::GetEmpty(), FEngineStatRender(), FEngineStatToggle::CreateUObject(this, &UEngine::ToggleStatRaw)));
 	EngineStats.Add(FEngineStatFuncs(TEXT("STAT_ParticlePerf"), TEXT("STATCAT_Engine"), FText::GetEmpty(), FEngineStatRender::CreateUObject(this, &UEngine::RenderStatParticlePerf), FEngineStatToggle::CreateUObject(this, &UEngine::ToggleStatParticlePerf), bIsRHS));
+	EngineStats.Add(FEngineStatFuncs(TEXT("STAT_TSR"), TEXT("STATCAT_Engine"), FText::FromString(TEXT("Same as stat unit, but with additional TSR settings.")), FEngineStatRender(), FEngineStatToggle::CreateUObject(this, &UEngine::ToggleStatTSR)));
 #endif // !UE_BUILD_SHIPPING
 
 	// Let any listeners know about the new stats
@@ -17191,6 +17192,7 @@ bool UEngine::ToggleStatDetailed(UWorld* World, FCommonViewportClient* ViewportC
 		DetailedStats.Add(TEXT("UnitGraph"));
 		DetailedStats.Add(TEXT("Raw"));
 		DetailedStats.Add(TEXT("Version"));
+		DetailedStats.Add(TEXT("TSR"));
 	}
 
 	// If any of the detailed stats are inactive, take this as enabling all, unless 'Skip' is specifically specified
@@ -17795,6 +17797,36 @@ bool UEngine::ToggleStatRaw(UWorld* World, FCommonViewportClient* ViewportClient
 	}
 	return true;
 }
+
+// TSR statistics
+bool UEngine::ToggleStatTSR(UWorld* World, FCommonViewportClient* ViewportClient, const TCHAR* Stream)
+{
+	if (ViewportClient == nullptr)
+	{
+		// Ignore if all Viewports are closed.
+		return false;
+	}
+	const bool bShowTSR = ViewportClient->IsStatEnabled(TEXT("TSR"));
+	if (bShowTSR)
+	{
+		// Force Unit to Active
+		SetEngineStat(World, ViewportClient, TEXT("Unit"), true);
+
+		// Force TSR to true as Unit will have Toggled it back to false
+		SetEngineStat(World, ViewportClient, TEXT("TSR"), true);
+	}
+	else
+	{
+		const bool bShowDetailed = ViewportClient->IsStatEnabled(TEXT("Detailed"));
+		if (bShowDetailed)
+		{
+			// Since we're turning this off, we also need to toggle off detailed too
+			ExecEngineStat(World, ViewportClient, TEXT("Detailed -Skip"));
+		}
+	}
+	return true;
+}
+
 #endif
 
 #if !UE_BUILD_SHIPPING
diff --git a/Engine/Source/Runtime/RenderCore/Private/RenderCore.cpp b/Engine/Source/Runtime/RenderCore/Private/RenderCore.cpp
--- a/Engine/Source/Runtime/RenderCore/Private/RenderCore.cpp
+++ b/Engine/Source/Runtime/RenderCore/Private/RenderCore.cpp
@@ -12,6 +12,7 @@
 #include "HDRHelper.h"
 #include "RHI.h"
 #include "RenderTimer.h"
+#include "RenderCounters.h"
 
 void UpdateShaderDevelopmentMode();
 
diff --git a/Engine/Source/Runtime/RenderCore/Public/RenderCounters.h b/Engine/Source/Runtime/RenderCore/Public/RenderCounters.h
new file mode 100644
--- /dev/null
+++ b/Engine/Source/Runtime/RenderCore/Public/RenderCounters.h
@@ -0,0 +1,38 @@
+// Copyright Epic Games, Inc. All Rights Reserved.
+
+#pragma once
+
+#include "CoreGlobals.h"
+
+
+
+class FPixelRenderCounters
+{
+public:
+	uint32 GetPixelRenderCount() const
+	{
+		return PrevPixelRenderCount;
+	}
+
+	uint32 GetPixelDisplayCount() const
+	{
+		return PrevPixelDisplayCount;
+	}
+
+	void AddViewStatistics(uint32 PixelRenderCount, uint32 PixelDisplayCount)
+	{
+		check(IsInRenderingThread());
+		CurrentPixelRenderCount += PixelRenderCount;
+		CurrentPixelDisplayCount += PixelDisplayCount;
+	}
+
+private:
+	uint32 PrevPixelRenderCount = 0;
+	uint32 PrevPixelDisplayCount = 0;
+	uint32 CurrentPixelRenderCount = 0;
+	uint32 CurrentPixelDisplayCount = 0;
+
+	friend void TickPixelRenderCounters();
+};
+
+extern RENDERCORE_API FPixelRenderCounters GPixelRenderCounters;
diff --git a/Engine/Source/Runtime/RenderCore/Public/RenderTimer.h b/Engine/Source/Runtime/RenderCore/Public/RenderTimer.h
--- a/Engine/Source/Runtime/RenderCore/Public/RenderTimer.h
+++ b/Engine/Source/Runtime/RenderCore/Public/RenderTimer.h
@@ -2,7 +2,6 @@
 
 #pragma once
 
-#include "CoreGlobals.h"
 #include "HAL/PlatformTime.h"
 
 /** Timer helper class. **/
@@ -119,35 +118,3 @@ extern RENDERCORE_API uint32 GSwapBufferTime;
 
 /** How many cycles the renderthread used, including dependent wait time. */
 extern RENDERCORE_API uint32 GRenderThreadTimeCriticalPath;
-
-
-class FPixelRenderCounters
-{
-public:
-	uint32 GetPixelRenderCount() const
-	{
-		return PrevPixelRenderCount;
-	}
-
-	uint32 GetPixelDisplayCount() const
-	{
-		return PrevPixelDisplayCount;
-	}
-
-	void AddViewStatistics(uint32 PixelRenderCount, uint32 PixelDisplayCount)
-	{
-		check(IsInRenderingThread());
-		CurrentPixelRenderCount += PixelRenderCount;
-		CurrentPixelDisplayCount += PixelDisplayCount;
-	}
-
-private:
-	uint32 PrevPixelRenderCount = 0;
-	uint32 PrevPixelDisplayCount = 0;
-	uint32 CurrentPixelRenderCount = 0;
-	uint32 CurrentPixelDisplayCount = 0;
-
-	friend void TickPixelRenderCounters();
-};
-
-extern RENDERCORE_API FPixelRenderCounters GPixelRenderCounters;
diff --git a/Engine/Source/Runtime/Renderer/Private/SceneRendering.cpp b/Engine/Source/Runtime/Renderer/Private/SceneRendering.cpp
--- a/Engine/Source/Runtime/Renderer/Private/SceneRendering.cpp
+++ b/Engine/Source/Runtime/Renderer/Private/SceneRendering.cpp
@@ -83,6 +83,7 @@
 #include "NaniteVisualizationData.h"
 #include "Shadows/ShadowSceneRenderer.h"
 #include "TextureResource.h"
+#include "RenderCounters.h"
 #include "RenderCore.h"
 #include "SkyAtmosphereRendering.h"
 #include "VolumetricCloudRendering.h"
