Hide COUNT enum from Destination selection

diff --git a/Engine/Source/Runtime/Engine/Classes/Components/AudioComponent.h b/Engine/Source/Runtime/Engine/Classes/Components/AudioComponent.h
--- a/Engine/Source/Runtime/Engine/Classes/Components/AudioComponent.h
+++ b/Engine/Source/Runtime/Engine/Classes/Components/AudioComponent.h
@@ -142,8 +142,7 @@ enum class EModulationDestination : uint8
 	/* Cutoff Frequency of a highpass filter */
 	Highpass,
 
-	
-	COUNT
+	Count UMETA(Hidden)
 };
 
 /**
diff --git a/Engine/Source/Runtime/Engine/Private/Components/AudioComponent.cpp b/Engine/Source/Runtime/Engine/Private/Components/AudioComponent.cpp
--- a/Engine/Source/Runtime/Engine/Private/Components/AudioComponent.cpp
+++ b/Engine/Source/Runtime/Engine/Private/Components/AudioComponent.cpp
@@ -1758,7 +1758,7 @@ void UAudioComponent::SetModulationRouting(const TSet<USoundModulatorBase*>& Mod
 		break;
 	default:
 	{
-		static_assert(static_cast<int32>(EModulationDestination::COUNT) == 4, "Possible missing ELiteralType case coverage.");
+		static_assert(static_cast<int32>(EModulationDestination::Count) == 4, "Possible missing ELiteralType case coverage.");
 		ensureMsgf(false, TEXT("Failed to set input node default: Literal type not supported"));
 		return;
 	}
@@ -1798,7 +1798,7 @@ TSet<USoundModulatorBase*> UAudioComponent::GetModulators(const EModulationDesti
 		break;
 	default:
 	{
-		static_assert(static_cast<int32>(EModulationDestination::COUNT) == 4, "Possible missing ELiteralType case coverage.");
+		static_assert(static_cast<int32>(EModulationDestination::Count) == 4, "Possible missing ELiteralType case coverage.");
 		ensureMsgf(false, TEXT("Failed to set input node default: Literal type not supported"));
 		return TSet<USoundModulatorBase*>();
 	}
