- Added UDataLayerManager which is local to each WorldPartition.

diff --git a/Engine/Source/Runtime/Engine/Private/WorldPartition/DataLayer/DataLayerInstanceWithAsset.cpp b/Engine/Source/Runtime/Engine/Private/WorldPartition/DataLayer/DataLayerInstanceWithAsset.cpp
--- a/Engine/Source/Runtime/Engine/Private/WorldPartition/DataLayer/DataLayerInstanceWithAsset.cpp
+++ b/Engine/Source/Runtime/Engine/Private/WorldPartition/DataLayer/DataLayerInstanceWithAsset.cpp
@@ -1,11 +1,11 @@
 // Copyright Epic Games, Inc. All Rights Reserved.
 
 #include "WorldPartition/DataLayer/DataLayerInstanceWithAsset.h"
+#include "WorldPartition/DataLayer/WorldDataLayers.h"
+#include "WorldPartition/DataLayer/DataLayerManager.h"
 #include "Engine/Level.h"
 #include "Misc/StringFormatArg.h"
-#include "WorldPartition/DataLayer/WorldDataLayers.h"
 #include "UObject/UnrealType.h"
-#include "WorldPartition/DataLayer/DataLayerSubsystem.h"
 #include "WorldPartition/ErrorHandling/WorldPartitionStreamingGenerationErrorHandler.h"
 
 #include UE_INLINE_GENERATED_CPP_BY_NAME(DataLayerInstanceWithAsset)
@@ -36,6 +36,10 @@ void UDataLayerInstanceWithAsset::OnCreated(const UDataLayerAsset* Asset)
 
 bool UDataLayerInstanceWithAsset::IsReadOnly() const
 {
+	if (Super::IsReadOnly())
+	{
+		return true;
+	}
 	return GetOuterAWorldDataLayers()->IsSubWorldDataLayers();
 }
 
@@ -50,8 +54,8 @@ bool UDataLayerInstanceWithAsset::IsLocked() const
 
 bool UDataLayerInstanceWithAsset::AddActor(AActor* Actor) const
 {	
-	check(GetWorld()->GetSubsystem<UDataLayerSubsystem>()->GetDataLayerInstance(DataLayerAsset) != nullptr);
 	check(GetTypedOuter<ULevel>() == Actor->GetLevel()); // Make sure the instance is part of the same world as the actor.
+	check(UDataLayerManager::GetDataLayerManager(Actor)->GetDataLayerInstance(DataLayerAsset) != nullptr); // Make sure the DataLayerInstance exists for this level
 	return Actor->AddDataLayer(DataLayerAsset);
 }
 
@@ -70,24 +74,28 @@ bool UDataLayerInstanceWithAsset::Validate(IStreamingGenerationErrorHandler* Err
 		return false;
 	}
 
-	UDataLayerSubsystem* DataLayerSubsystem = UWorld::GetSubsystem<UDataLayerSubsystem>(GetWorld());
-	DataLayerSubsystem->ForEachDataLayer([&bIsValid, this, ErrorHandler](UDataLayerInstance* DataLayerInstance)
+	// Get the DataLayerManager for this DataLayerInstance which will be the one of its outer world
+	UDataLayerManager* DataLayerManager = UDataLayerManager::GetDataLayerManager(this);
+	if (ensure(DataLayerManager))
 	{
-		if (DataLayerInstance != this)
+		DataLayerManager->ForEachDataLayerInstance([&bIsValid, this, ErrorHandler](UDataLayerInstance* DataLayerInstance)
 		{
-			if (UDataLayerInstanceWithAsset* DataLayerInstanceWithAsset = Cast<UDataLayerInstanceWithAsset>(DataLayerInstance))
+			if (DataLayerInstance != this)
 			{
-				if (DataLayerInstanceWithAsset->GetAsset() == GetAsset())
+				if (UDataLayerInstanceWithAsset* DataLayerInstanceWithAsset = Cast<UDataLayerInstanceWithAsset>(DataLayerInstance))
 				{
-					ErrorHandler->OnDataLayerAssetConflict(this, DataLayerInstanceWithAsset);
-					bIsValid = false;
-					return false;
+					if (DataLayerInstanceWithAsset->GetAsset() == GetAsset())
+					{
+						ErrorHandler->OnDataLayerAssetConflict(this, DataLayerInstanceWithAsset);
+						bIsValid = false;
+						return false;
+					}
 				}
 			}
-		}
-
-		return true;
-	}, GetOuterAWorldDataLayers()->GetLevel()); // Resolve DataLayerInstances based on outer level
+	
+			return true;
+		});
+	}
 
 	bIsValid &= Super::Validate(ErrorHandler);
 
diff --git a/Engine/Source/Runtime/Engine/Private/WorldPartition/DataLayer/DataLayerSubsystem.cpp b/Engine/Source/Runtime/Engine/Private/WorldPartition/DataLayer/DataLayerSubsystem.cpp
--- a/Engine/Source/Runtime/Engine/Private/WorldPartition/DataLayer/DataLayerSubsystem.cpp
+++ b/Engine/Source/Runtime/Engine/Private/WorldPartition/DataLayer/DataLayerSubsystem.cpp
@@ -2,11 +2,12 @@
 
 #include "WorldPartition/DataLayer/DataLayerSubsystem.h"
 #include "WorldPartition/ActorDescContainer.h"
-#include "WorldPartition/DataLayer/DataLayerLoadingPolicy.h"
 #include "WorldPartition/DataLayer/DataLayer.h"
 #include "WorldPartition/DataLayer/WorldDataLayers.h"
 #include "WorldPartition/DataLayer/DataLayerEditorContext.h"
 #include "WorldPartition/DataLayer/WorldDataLayersActorDesc.h"
+#include "WorldPartition/DataLayer/DataLayerManager.h"
+#include "WorldPartition/DataLayer/DataLayerInstance.h"
 #include "WorldPartition/WorldPartitionDebugHelper.h"
 #include "WorldPartition/WorldPartitionLog.h"
 #include "WorldPartition/WorldPartitionSubsystem.h"
@@ -27,818 +28,122 @@
 #include "Engine/Engine.h"
 #endif
 
-extern int32 GDrawDataLayersLoadTime;
-
-#if WITH_EDITOR
-FOnWorldDataLayersPostRegister UDataLayerSubsystem::OnWorldDataLayerPostRegister;
-FOnWorldDataLayersPreUnregister UDataLayerSubsystem::OnWorldDataLayerPreUnregister;
-#endif
-
-static FAutoConsoleCommandWithOutputDevice GDumpDataLayersCmd(
-	TEXT("wp.DumpDataLayers"),
-	TEXT("Dumps data layers to the log"),
-	FConsoleCommandWithOutputDeviceDelegate::CreateStatic([](FOutputDevice& OutputDevice)
-	{
-		for (const FWorldContext& Context : GEngine->GetWorldContexts())
-		{
-			UWorld* World = Context.World();
-			if (World && World->IsGameWorld())
-			{
-				if (const UDataLayerSubsystem* DataLayerSubsystem = World->GetSubsystem<UDataLayerSubsystem>())
-				{
-					DataLayerSubsystem->DumpDataLayers(OutputDevice);
-				}
-			}
-		}
-	})
-);
-
-#if WITH_EDITOR
-FDataLayersEditorBroadcast& FDataLayersEditorBroadcast::Get()
-{
-	static FDataLayersEditorBroadcast DataLayersEditorBroadcast;
-	return DataLayersEditorBroadcast;
-}
-
-void FDataLayersEditorBroadcast::StaticOnActorDataLayersEditorLoadingStateChanged(bool bIsFromUserChange)
-{
-	Get().DataLayerEditorLoadingStateChanged.Broadcast(bIsFromUserChange);
-}
-#endif
-
-UDataLayerSubsystem::UDataLayerSubsystem()
-	: bIsCachedEffectiveStateDirty(true)
-{
-#if WITH_EDITOR
-	DataLayerActorEditorContextID = 0;
-#endif
-}
-
-void UDataLayerSubsystem::Initialize(FSubsystemCollectionBase& Collection)
-{
-	Super::Initialize(Collection);
-
-	bIsCachedEffectiveStateDirty = true;
-
-#if WITH_EDITOR
-	UClass* DataLayerLoadingPolicyClassValue = DataLayerLoadingPolicyClass.Get();
-	if (!DataLayerLoadingPolicyClassValue)
-	{
-		DataLayerLoadingPolicyClassValue = UDataLayerLoadingPolicy::StaticClass();
-	}
-	DataLayerLoadingPolicy = NewObject<UDataLayerLoadingPolicy>(this, DataLayerLoadingPolicyClassValue);
-	check(DataLayerLoadingPolicy);
-
-	if (GEditor)
-	{
-		FModuleManager::LoadModuleChecked<IDataLayerEditorModule>("DataLayerEditor");
-	}
-
-	RegisterEarlyLoadedWorldDataLayers();
-
-	UActorDescContainer::OnActorDescContainerInitialized.AddUObject(this, &UDataLayerSubsystem::OnActorDescContainerInitialized);
-#endif
-}
-
-void UDataLayerSubsystem::Deinitialize()
-{
-	Super::Deinitialize();
-
-	WorldDataLayerCollection.Empty();
-
-#if WITH_EDITOR
-	UActorDescContainer::OnActorDescContainerInitialized.RemoveAll(this);
-#endif
-}
-
 bool UDataLayerSubsystem::DoesSupportWorldType(const EWorldType::Type WorldType) const
 {
 	return Super::DoesSupportWorldType(WorldType) || WorldType == EWorldType::Inactive;
 }
 
-
 #if WITH_EDITOR
 
-bool UDataLayerSubsystem::ResolveIsLoadedInEditor(const TArray<FName>& InDataLayerInstanceNames) const
-{
-	TArray<const UDataLayerInstance*> DataLayerInstances;
-	for (const FName& DataLayerInstanceName : InDataLayerInstanceNames)
-	{
-		if (const UDataLayerInstance* DataLayerInstance = GetDataLayerInstance(DataLayerInstanceName))
-		{
-			DataLayerInstances.Add(DataLayerInstance);
-		}
-	}
-	if (!DataLayerInstances.IsEmpty())
-	{
-		return DataLayerLoadingPolicy->ResolveIsLoadedInEditor(DataLayerInstances);
-	}
-	return true;
-}
-
-bool UDataLayerSubsystem::CanResolveDataLayers() const
-{
-	return !WorldDataLayerCollection.IsEmpty();
-}
-
-bool UDataLayerSubsystem::RemoveDataLayer(const UDataLayerInstance* InDataLayerInstance)
-{
-	bool bIsDataLayerRemoved = false;
-	WorldDataLayerCollection.ForEachWorldDataLayers([InDataLayerInstance, &bIsDataLayerRemoved](AWorldDataLayers* WorldDataLayer)
-	{
-		bIsDataLayerRemoved = WorldDataLayer->RemoveDataLayer(InDataLayerInstance);
-		return !bIsDataLayerRemoved;
-	});
-
-	return bIsDataLayerRemoved;
-}
-
-int32 UDataLayerSubsystem::RemoveDataLayers(const TArray<UDataLayerInstance*>& InDataLayerInstances)
-{
-	int32 DataLayerRemovedCount = 0;
-	WorldDataLayerCollection.ForEachWorldDataLayers([InDataLayerInstances, &DataLayerRemovedCount](AWorldDataLayers* WorldDataLayer)
-	{
-		DataLayerRemovedCount += WorldDataLayer->RemoveDataLayers(InDataLayerInstances);
-		return DataLayerRemovedCount != InDataLayerInstances.Num();
-	});
-
-	return DataLayerRemovedCount;
-}
-
-void UDataLayerSubsystem::UpdateDataLayerEditorPerProjectUserSettings() const
-{
-	TArray<FName> DataLayersNotLoadedInEditor;
-	TArray<FName> DataLayersLoadedInEditor;
-
-	GetUserLoadedInEditorStates(DataLayersLoadedInEditor, DataLayersNotLoadedInEditor);
-
-	GetMutableDefault<UWorldPartitionEditorPerProjectUserSettings>()->SetWorldDataLayersNonDefaultEditorLoadStates(GetWorld(), DataLayersLoadedInEditor, DataLayersNotLoadedInEditor);
-}
-
-void UDataLayerSubsystem::GetUserLoadedInEditorStates(TArray<FName>& OutDataLayersLoadedInEditor, TArray<FName>& OutDataLayersNotLoadedInEditor) const
+TArray<const UDataLayerInstance*> UDataLayerSubsystem::GetRuntimeDataLayerInstances(UWorld* InWorld, const TArray<FName>& InDataLayerInstanceNames)
 {
-	OutDataLayersLoadedInEditor.Empty();
-	OutDataLayersNotLoadedInEditor.Empty();
-
-	const TArray<FName>& SettingsDataLayersNotLoadedInEditor = GetMutableDefault<UWorldPartitionEditorPerProjectUserSettings>()->GetWorldDataLayersNotLoadedInEditor(GetWorld());
-	const TArray<FName>& SettingsDataLayersLoadedInEditor = GetMutableDefault<UWorldPartitionEditorPerProjectUserSettings>()->GetWorldDataLayersLoadedInEditor(GetWorld());
-
-	ForEachDataLayer([&OutDataLayersLoadedInEditor, &OutDataLayersNotLoadedInEditor, &SettingsDataLayersNotLoadedInEditor, &SettingsDataLayersLoadedInEditor](UDataLayerInstance* DataLayerInstance)
-	{
-		if (DataLayerInstance->IsLoadedInEditorChangedByUserOperation())
-		{
-			if (!DataLayerInstance->IsLoadedInEditor() && DataLayerInstance->IsInitiallyLoadedInEditor())
-			{
-				OutDataLayersNotLoadedInEditor.Add(DataLayerInstance->GetDataLayerFName());
-			}
-			else if (DataLayerInstance->IsLoadedInEditor() && !DataLayerInstance->IsInitiallyLoadedInEditor())
-			{
-				OutDataLayersLoadedInEditor.Add(DataLayerInstance->GetDataLayerFName());
-			}
-
-			DataLayerInstance->ClearLoadedInEditorChangedByUserOperation();
-		}
-		else
-		{
-			if (SettingsDataLayersNotLoadedInEditor.Contains(DataLayerInstance->GetDataLayerFName()))
-			{
-				OutDataLayersNotLoadedInEditor.Add(DataLayerInstance->GetDataLayerFName());
-			}
-			else if (SettingsDataLayersLoadedInEditor.Contains(DataLayerInstance->GetDataLayerFName()))
-			{
-				OutDataLayersLoadedInEditor.Add(DataLayerInstance->GetDataLayerFName());
-			}
-		}
-
-		return true;
-	});
-}
-
-TArray<const UDataLayerInstance*> UDataLayerSubsystem::GetRuntimeDataLayerInstances(UWorld* InWorld, const TArray<FName>& DataLayerInstanceNames)
-{
-	TArray<const UDataLayerInstance*> DataLayerInstances;
-	if (const UDataLayerSubsystem* DataLayerSubsystem = UWorld::GetSubsystem<UDataLayerSubsystem>(InWorld))
-	{
-		for (const FName& DataLayerInstanceName : DataLayerInstanceNames)
-		{
-			if (const UDataLayerInstance* DataLayerInstance = DataLayerSubsystem->GetDataLayerInstance(DataLayerInstanceName))
-			{
-				if (ensure(DataLayerInstance->IsRuntime()))
-				{
-					DataLayerInstances.Add(DataLayerInstance);
-				}
-			}
-		}
-	}
-	return DataLayerInstances;
+	static TArray<const UDataLayerInstance*> Empty;
+	UDataLayerManager* DataLayerManager = UDataLayerManager::GetDataLayerManager(InWorld);
+	return DataLayerManager ? DataLayerManager->GetRuntimeDataLayerInstances(InDataLayerInstanceNames) : Empty;
 }
 
 #endif
 
 const TSet<FName>& UDataLayerSubsystem::GetEffectiveActiveDataLayerNames() const
 {
-	UpdateCachedEffectiveRuntimeStates();
-	return CachedEffectiveActiveDataLayerNames;
+	static TSet<FName> Empty;
+	const UDataLayerManager* DataLayerManager = UDataLayerManager::GetDataLayerManager(GetWorld());
+	return DataLayerManager ? DataLayerManager->GetEffectiveActiveDataLayerNames() : Empty;
 }
 
 const TSet<FName>& UDataLayerSubsystem::GetEffectiveLoadedDataLayerNames() const
 {
-	UpdateCachedEffectiveRuntimeStates();
-	return CachedEffectiveLoadedDataLayerNames;
-}
-
-void UDataLayerSubsystem::OnEffectiveRuntimeDataLayerStatesChanged(AWorldDataLayers* InWorldDataLayers)
-{
-	bIsCachedEffectiveStateDirty = true;
-}
-
-void UDataLayerSubsystem::UpdateCachedEffectiveRuntimeStates() const
-{
-	if (bIsCachedEffectiveStateDirty)
-	{
-		bIsCachedEffectiveStateDirty = false;
-
-		CachedEffectiveActiveDataLayerNames.Reset();
-		WorldDataLayerCollection.ForEachWorldDataLayers([this](AWorldDataLayers* WorldDataLayers)
-		{
-			CachedEffectiveActiveDataLayerNames.Append(WorldDataLayers->GetEffectiveActiveDataLayerNames());
-		});
-
-		CachedEffectiveLoadedDataLayerNames.Reset();
-		WorldDataLayerCollection.ForEachWorldDataLayers([this](AWorldDataLayers* WorldDataLayers)
-		{
-			CachedEffectiveLoadedDataLayerNames.Append(WorldDataLayers->GetEffectiveLoadedDataLayerNames());
-		});
-	}
-}
-
-void UDataLayerSubsystem::RegisterWorldDataLayer(AWorldDataLayers* WorldDataLayers)
-{
-	// Don't register WorldDataLayers in game if it's not relevant at runtime
-	if (GetWorld()->IsGameWorld() && WorldDataLayers && !WorldDataLayers->IsRuntimeRelevant())
-	{
-		return;
-	}
-
-	if (WorldDataLayerCollection.RegisterWorldDataLayer(WorldDataLayers))
-	{
-		bIsCachedEffectiveStateDirty = true;
-#if WITH_EDITOR
-		OnWorldDataLayerPostRegister.Broadcast(WorldDataLayers);
-#endif
-	}
-}
-
-void UDataLayerSubsystem::UnregisterWorldDataLayer(AWorldDataLayers* WorldDataLayers)
-{
-	if (WorldDataLayerCollection.Contains(WorldDataLayers))
-	{
-#if WITH_EDITOR
-		OnWorldDataLayerPreUnregister.Broadcast(WorldDataLayers);
-#endif
-		bIsCachedEffectiveStateDirty = true;
-		WorldDataLayerCollection.UnregisterWorldDataLayer(WorldDataLayers);
-	}
-}
-
-void UDataLayerSubsystem::ForEachWorldDataLayer(TFunctionRef<bool(AWorldDataLayers*)> Func)
-{
-	WorldDataLayerCollection.ForEachWorldDataLayersBreakable(Func);
-}
-
-void UDataLayerSubsystem::ForEachWorldDataLayer(TFunctionRef<bool(AWorldDataLayers*)> Func) const
-{
-	const_cast<UDataLayerSubsystem*>(this)->ForEachWorldDataLayer(Func);
-}
-
-UDataLayerInstance* UDataLayerSubsystem::GetDataLayerInstanceFromAsset(const UDataLayerAsset* InDataLayerAsset) const
-{
-	return GetDataLayerInstance(InDataLayerAsset);
+	static TSet<FName> Empty;
+	const UDataLayerManager* DataLayerManager = UDataLayerManager::GetDataLayerManager(GetWorld());
+	return DataLayerManager ? DataLayerManager->GetEffectiveLoadedDataLayerNames() : Empty;
 }
 
 void UDataLayerSubsystem::SetDataLayerInstanceRuntimeState(const UDataLayerAsset* InDataLayerAsset, EDataLayerRuntimeState InState, bool bInIsRecursive)
 {
-	if (UDataLayerInstance* DataLayerInstance = GetDataLayerInstanceFromAsset(InDataLayerAsset))
+	if (UDataLayerManager* DataLayerManager = UDataLayerManager::GetDataLayerManager(GetWorld()))
 	{
-		SetDataLayerRuntimeState(DataLayerInstance, InState, bInIsRecursive);
-	}
-	else
-	{
-		UE_LOG(LogWorldPartition, Warning, TEXT("UDataLayerSubsystem::SetDataLayerInstanceRuntimeState called with a Data Layer Asset that does not have a Data Layer Instance in this world : %s"), *GetWorld()->GetName());
+		DataLayerManager->SetDataLayerInstanceRuntimeState(DataLayerManager->GetDataLayerInstance(InDataLayerAsset), InState, bInIsRecursive);
 	}
 }
 
 EDataLayerRuntimeState UDataLayerSubsystem::GetDataLayerInstanceRuntimeState(const UDataLayerAsset* InDataLayerAsset) const
 {
-	if (UDataLayerInstance* DataLayerInstance = GetDataLayerInstanceFromAsset(InDataLayerAsset))
-	{
-		return GetDataLayerRuntimeState(DataLayerInstance);
-	}
-	
-	UE_LOG(LogWorldPartition, Warning, TEXT("UDataLayerSubsystem::GetDataLayerInstanceRuntimeState called with a Data Layer Asset that does not have a Data Layer Instance in this world : %s"), *GetWorld()->GetName());
-	return EDataLayerRuntimeState::Unloaded;
+	const UDataLayerManager* DataLayerManager = UDataLayerManager::GetDataLayerManager(GetWorld());
+	return DataLayerManager ? DataLayerManager->GetDataLayerInstanceRuntimeState(DataLayerManager->GetDataLayerInstance(InDataLayerAsset)) : EDataLayerRuntimeState::Unloaded;
 }
 
 EDataLayerRuntimeState UDataLayerSubsystem::GetDataLayerInstanceEffectiveRuntimeState(const UDataLayerAsset* InDataLayerAsset) const
 {
-	if (UDataLayerInstance* DataLayerInstance = GetDataLayerInstanceFromAsset(InDataLayerAsset))
-	{
-		return GetDataLayerEffectiveRuntimeState(DataLayerInstance);
-	}
-
-	UE_LOG(LogWorldPartition, Warning, TEXT("UDataLayerSubsystem::GetDataLayerInstanceEffectiveRuntimeState called with a Data Layer Asset that does not have a Data Layer Instance in this world : %s"), *GetWorld()->GetName());
-	return EDataLayerRuntimeState::Unloaded;
+	const UDataLayerManager* DataLayerManager = UDataLayerManager::GetDataLayerManager(GetWorld());
+	return DataLayerManager ? DataLayerManager->GetDataLayerInstanceEffectiveRuntimeState(DataLayerManager->GetDataLayerInstance(InDataLayerAsset)) : EDataLayerRuntimeState::Unloaded;
 }
 
 void UDataLayerSubsystem::SetDataLayerRuntimeState(const UDataLayerInstance* InDataLayerInstance, EDataLayerRuntimeState InState, bool bInIsRecursive)
 {
-	if (InDataLayerInstance)
-	{
-		if (AWorldDataLayers* WorldDataLayers = InDataLayerInstance->GetOuterAWorldDataLayers())
-		{
-			WorldDataLayers->SetDataLayerRuntimeState(InDataLayerInstance, InState, bInIsRecursive);
-		}
-	}
-	else
+	if (UDataLayerManager* DataLayerManager = UDataLayerManager::GetDataLayerManager(GetWorld()))
 	{
-		UE_LOG(LogWorldPartition, Warning, TEXT("UDataLayerSubsystem::SetDataLayerRuntimeState called with null Data Layer"));
+		DataLayerManager->SetDataLayerInstanceRuntimeState(InDataLayerInstance, InState, bInIsRecursive);
 	}
 }
 
 EDataLayerRuntimeState UDataLayerSubsystem::GetDataLayerRuntimeState(const UDataLayerInstance* InDataLayerInstance) const
 {
-	AWorldDataLayers* WorldDataLayers = InDataLayerInstance->GetOuterAWorldDataLayers();
-	return WorldDataLayers->GetDataLayerRuntimeStateByName(InDataLayerInstance->GetDataLayerFName());
+	const UDataLayerManager* DataLayerManager = UDataLayerManager::GetDataLayerManager(GetWorld());
+	return DataLayerManager ? DataLayerManager->GetDataLayerInstanceRuntimeState(InDataLayerInstance) : EDataLayerRuntimeState::Unloaded;
 }
 
 EDataLayerRuntimeState UDataLayerSubsystem::GetDataLayerRuntimeStateByName(const FName& InDataLayerName) const
 {
-	if (UDataLayerInstance* DataLayerInstance = GetDataLayerInstance(InDataLayerName))
-	{		
-		return GetDataLayerRuntimeState(DataLayerInstance);
-	}
-	
-	return EDataLayerRuntimeState::Unloaded;
+	const UDataLayerManager* DataLayerManager = UDataLayerManager::GetDataLayerManager(GetWorld());
+	return DataLayerManager ? DataLayerManager->GetDataLayerInstanceRuntimeState(DataLayerManager->GetDataLayerInstanceFromName(InDataLayerName)) : EDataLayerRuntimeState::Unloaded;
 }
 
 EDataLayerRuntimeState UDataLayerSubsystem::GetDataLayerEffectiveRuntimeState(const UDataLayerInstance* InDataLayerInstance) const
 {
-	AWorldDataLayers* WorldDataLayers = InDataLayerInstance->GetOuterAWorldDataLayers();
-	return WorldDataLayers->GetDataLayerEffectiveRuntimeStateByName(InDataLayerInstance->GetDataLayerFName());
+	const UDataLayerManager* DataLayerManager = UDataLayerManager::GetDataLayerManager(GetWorld());
+	return DataLayerManager ? DataLayerManager->GetDataLayerInstanceEffectiveRuntimeState(InDataLayerInstance) : EDataLayerRuntimeState::Unloaded;
 }
 
 EDataLayerRuntimeState UDataLayerSubsystem::GetDataLayerEffectiveRuntimeStateByName(const FName& InDataLayerName) const
 {
-	if (UDataLayerInstance* DataLayerInstance = GetDataLayerInstance(InDataLayerName))
-	{
-		return GetDataLayerEffectiveRuntimeState(DataLayerInstance);
-	}
-
-	return EDataLayerRuntimeState::Unloaded;
+	const UDataLayerManager* DataLayerManager = UDataLayerManager::GetDataLayerManager(GetWorld());
+	return DataLayerManager ? DataLayerManager->GetDataLayerInstanceEffectiveRuntimeState(DataLayerManager->GetDataLayerInstanceFromName(InDataLayerName)) : EDataLayerRuntimeState::Unloaded;
 }
 
 bool UDataLayerSubsystem::IsAnyDataLayerInEffectiveRuntimeState(const TArray<FName>& InDataLayerNames, EDataLayerRuntimeState InState) const
 {
-	if (InState == EDataLayerRuntimeState::Activated)
-	{
-		const TSet<FName>& Activated = GetEffectiveActiveDataLayerNames();
-		for (const FName& DataLayerName : InDataLayerNames)
-		{
-			if (Activated.Contains(DataLayerName))
-			{
-				return true;
-			}
-		}
-	}
-	else if (InState == EDataLayerRuntimeState::Loaded)
-	{
-		const TSet<FName>& Loaded = GetEffectiveLoadedDataLayerNames();
-		for (const FName& DataLayerName : InDataLayerNames)
-		{
-			if (Loaded.Contains(DataLayerName))
-			{
-				return true;
-			}
-		}
-	}
-	return false;
+	const UDataLayerManager* DataLayerManager = UDataLayerManager::GetDataLayerManager(GetWorld());
+	return DataLayerManager ? DataLayerManager->IsAnyDataLayerInEffectiveRuntimeState(InDataLayerNames, InState) : false;
 }
 
 void UDataLayerSubsystem::DrawDataLayersStatus(UCanvas* Canvas, FVector2D& Offset) const
 {
-	TRACE_CPUPROFILER_EVENT_SCOPE(UDataLayerSubsystem::DrawDataLayersStatus);
-
-	if (!Canvas || !Canvas->SceneView)
-	{
-		return;
-	}
-	
-	FVector2D Pos = Offset;
-	float MaxTextWidth = 0.f;
-
-	UWorldPartitionSubsystem* WorldPartitionSubsystem = GetWorld()->GetSubsystem<UWorldPartitionSubsystem>();
-
-	auto DrawLayerNames = [this, Canvas, &Pos, &MaxTextWidth](const FString& Title, FColor HeaderColor, FColor TextColor, const TSet<FName>& LayerNames)
-	{
-		if (LayerNames.Num() > 0)
-		{
-			FWorldPartitionDebugHelper::DrawText(Canvas, Title, GEngine->GetSmallFont(), HeaderColor, Pos, &MaxTextWidth);
-
-			TArray<const UDataLayerInstance*> DataLayers;
-			DataLayers.Reserve(LayerNames.Num());
-			for (const FName& DataLayerName : LayerNames)
-			{
-				if (const UDataLayerInstance* DataLayer = GetDataLayerInstance(DataLayerName))
-				{
-					DataLayers.Add(DataLayer);
-				}
-			}
-
-			DataLayers.Sort([](const UDataLayerInstance& A, const UDataLayerInstance& B) { return A.GetDataLayerFullName() < B.GetDataLayerFullName(); });
-
-			UFont* DataLayerFont = GEngine->GetSmallFont();
-			for (const UDataLayerInstance* DataLayer : DataLayers)
-			{
-				FString DataLayerString = DataLayer->GetDataLayerShortName();
-
-				if (GDrawDataLayersLoadTime)
-				{
-					if (double* DataLayerLoadTime = ActiveDataLayersLoadTime.Find(DataLayer))
-					{
-						if (*DataLayerLoadTime < 0)
-						{
-							DataLayerString += FString::Printf(TEXT(" (streaming %s)"), *FPlatformTime::PrettyTime(FPlatformTime::Seconds() + *DataLayerLoadTime));
-						}
-						else
-						{
-							DataLayerString += FString::Printf(TEXT(" (took %s)"), *FPlatformTime::PrettyTime(*DataLayerLoadTime));
-						}
-					}
-				}
-
-				FWorldPartitionDebugHelper::DrawLegendItem(Canvas, *DataLayerString, DataLayerFont, DataLayer->GetDebugColor(), TextColor, Pos, &MaxTextWidth);
-			}
-		}
-	};
-
-	const TSet<FName> LoadedDataLayers = GetEffectiveLoadedDataLayerNames();
-	const TSet<FName> ActiveDataLayers = GetEffectiveActiveDataLayerNames();
-
-	DrawLayerNames(TEXT("Loaded Data Layers"), FColor::Cyan, FColor::White, LoadedDataLayers);
-	DrawLayerNames(TEXT("Active Data Layers"), FColor::Green, FColor::White, ActiveDataLayers);
-
-	TSet<FName> UnloadedDataLayers;
-	ForEachDataLayer([&LoadedDataLayers, &ActiveDataLayers, &UnloadedDataLayers](UDataLayerInstance* DataLayer)
-	{
-		if (DataLayer->IsRuntime())
-		{
-			const FName DataLayerName = DataLayer->GetDataLayerFName();
-			if (!LoadedDataLayers.Contains(DataLayerName) && !ActiveDataLayers.Contains(DataLayerName))
-			{
-				UnloadedDataLayers.Add(DataLayerName);
-			}
-		}
-		return true;
-	});
-	DrawLayerNames(TEXT("Unloaded Data Layers"), FColor::Silver, FColor(192, 192, 192), UnloadedDataLayers);
-
-	Offset.X += MaxTextWidth + 10;
-
-	// Update data layers load times
-	if (GDrawDataLayersLoadTime)
-	{
-		for (FName DataLayerName : UnloadedDataLayers)
-		{
-			if (const UDataLayerInstance* DataLayer = GetDataLayerInstance(DataLayerName))
-			{
-				ActiveDataLayersLoadTime.Remove(DataLayer);
-			}
-		}
-
-		TArray<const UDataLayerInstance*> LoadingDataLayers;
-		LoadingDataLayers.Reserve(LoadedDataLayers.Num() + ActiveDataLayers.Num());
-		auto CopyLambda = [this](FName DataLayerName) { return GetDataLayerInstance(DataLayerName); };
-		Algo::Transform(LoadedDataLayers, LoadingDataLayers, CopyLambda);
-		Algo::Transform(ActiveDataLayers, LoadingDataLayers, CopyLambda);
-
-		for (const UDataLayerInstance* DataLayer : LoadingDataLayers)
-		{
-			double* DataLayerLoadTime = ActiveDataLayersLoadTime.Find(DataLayer);
-
-			auto IsDataLayerReady = [WorldPartitionSubsystem](const UDataLayerInstance* DataLayer, EWorldPartitionRuntimeCellState TargetState)
-			{
-				FWorldPartitionStreamingQuerySource QuerySource;
-				QuerySource.bDataLayersOnly = true;
-				QuerySource.bSpatialQuery = false;
-				QuerySource.DataLayers.Add(DataLayer->GetDataLayerFName());
-				return WorldPartitionSubsystem->IsStreamingCompleted(TargetState, { QuerySource }, true);
-			};
-
-			const EWorldPartitionRuntimeCellState TargetState = ActiveDataLayers.Contains(DataLayer->GetDataLayerFName()) ? EWorldPartitionRuntimeCellState::Activated : EWorldPartitionRuntimeCellState::Loaded;
-
-			if (!DataLayerLoadTime)
-			{
-				if (!IsDataLayerReady(DataLayer, TargetState))
-				{
-					DataLayerLoadTime = &ActiveDataLayersLoadTime.Add(DataLayer, -FPlatformTime::Seconds());
-				}
-			}
-
-			if (DataLayerLoadTime && (*DataLayerLoadTime < 0))
-			{
-				if (IsDataLayerReady(DataLayer, TargetState))
-				{
-					*DataLayerLoadTime = FPlatformTime::Seconds() + *DataLayerLoadTime;
-				}
-			}
-		}
-	}
-	else
-	{
-		ActiveDataLayersLoadTime.Empty();
-	}
-}
-
-TArray<UDataLayerInstance*> UDataLayerSubsystem::ConvertArgsToDataLayers(UWorld* World, const TArray<FString>& InArgs)
-{
-	TSet<UDataLayerInstance*> OutDataLayers;
-
-	const TCHAR* QuoteChar = TEXT("\"");
-	bool bQuoteStarted = false;
-	TStringBuilder<512> Builder;
-	TArray<FString> Args;
-	for (const FString& Arg : InArgs)
-	{
-		if (!bQuoteStarted && Arg.StartsWith(QuoteChar))
-		{
-			Builder.Append(Arg.Replace(QuoteChar, TEXT("")));
-			if (Arg.EndsWith(QuoteChar) && Arg.Len() > 1)
-			{
-				Args.Add(Builder.ToString());
-				Builder.Reset();
-			}
-			else
-			{
-				bQuoteStarted = true;
-			}
-		}
-		else if (bQuoteStarted)
-		{
-			Builder.Append(TEXT(" "));
-			Builder.Append(Arg.Replace(QuoteChar, TEXT("")));
-			if (Arg.EndsWith(QuoteChar))
-			{
-				bQuoteStarted = false;
-				Args.Add(Builder.ToString());
-				Builder.Reset();
-			}
-		}
-		else
-		{
-			Args.Add(Arg);
-		}
-	}
-	if (bQuoteStarted)
+	if (const UDataLayerManager* DataLayerManager = UDataLayerManager::GetDataLayerManager(GetWorld()))
 	{
-		Args.Add(Builder.ToString());
+		DataLayerManager->DrawDataLayersStatus(Canvas, Offset);
 	}
-
-	
-	UDataLayerSubsystem* DataLayerSubsystem = UWorld::GetSubsystem<UDataLayerSubsystem>(World);
-	for (const FString& Arg : Args)
-	{
-		FName DataLayerName = FName(Arg);
-		bool bShortNameFound = false;
-		FString SanitizedDataLayerName = DataLayerName.ToString().Replace(TEXT(" "), TEXT(""));
-		DataLayerSubsystem->ForEachDataLayer([&OutDataLayers, &SanitizedDataLayerName, &bShortNameFound](UDataLayerInstance* DataLayerInstance)
-		{
-			if (DataLayerInstance->GetDataLayerShortName().Compare(SanitizedDataLayerName, ESearchCase::IgnoreCase) == 0)
-			{
-				if (bShortNameFound)
-				{
-					UE_LOG(LogWorldPartition, Error, TEXT("Found 2 data layers with the ShortName %s when converting arguments. Consider using the data layers FullName or renaming one of the two."), *SanitizedDataLayerName);
-					return false;
-				}
-	
-				OutDataLayers.Add(DataLayerInstance);
-				bShortNameFound = true;
-			}
-			else if (DataLayerInstance->GetDataLayerFullName().Find(SanitizedDataLayerName, ESearchCase::IgnoreCase) == 0)
-			{
-				OutDataLayers.Add(DataLayerInstance);
-				return false;
-			}
-	
-			return true;
-		});
-	}
-
-	return OutDataLayers.Array();
 }
 
 void UDataLayerSubsystem::DumpDataLayers(FOutputDevice& OutputDevice) const
 {
-	WorldDataLayerCollection.ForEachWorldDataLayers([&OutputDevice](AWorldDataLayers* WorldDataLayers)
-	{
-		WorldDataLayers->DumpDataLayers(OutputDevice);
-	});
-}
-
-const UDataLayerInstance* UDataLayerSubsystem::GetDataLayerInstanceFromAssetName(const FName& InDataLayerAssetFullName) const
-{
-	const UDataLayerInstance* DataLayerInstance = nullptr;
-
-	WorldDataLayerCollection.ForEachWorldDataLayers([&DataLayerInstance, &InDataLayerAssetFullName](const AWorldDataLayers* WorldDataLayers)
+	if (const UDataLayerManager* DataLayerManager = UDataLayerManager::GetDataLayerManager(GetWorld()))
 	{
-		DataLayerInstance = WorldDataLayers->GetDataLayerInstanceFromAssetName(InDataLayerAssetFullName);
-		return DataLayerInstance == nullptr;
-	});
-
-	return const_cast<UDataLayerInstance*>(DataLayerInstance);
-}
-
-void UDataLayerSubsystem::ForEachDataLayer(TFunctionRef<bool(UDataLayerInstance*)> Func, const ULevel* InLevelContext /* = nullptr */)
-{
-	WorldDataLayerCollection.ForEachDataLayerBreakable(Func, InLevelContext);
-}
-
-void UDataLayerSubsystem::ForEachDataLayer(TFunctionRef<bool(UDataLayerInstance*)> Func, const ULevel* InLevelContext /* = nullptr */) const
-{
-	const_cast<UDataLayerSubsystem*>(this)->ForEachDataLayer(Func, InLevelContext);
-}
-
-FAutoConsoleCommand UDataLayerSubsystem::ToggleDataLayerActivation(
-	TEXT("wp.Runtime.ToggleDataLayerActivation"),
-	TEXT("Toggles DataLayers active state. Args [DataLayerNames]"),
-	FConsoleCommandWithArgsDelegate::CreateLambda([](const TArray<FString>& InArgs)
-	{
-		for (const FWorldContext& Context : GEngine->GetWorldContexts())
-		{
-			UWorld* World = Context.World();
-			if (World && World->IsGameWorld() && !World->IsNetMode(NM_Client))
-			{
-				if (UDataLayerSubsystem* DataLayerSubsystem = World->GetSubsystem<UDataLayerSubsystem>())
-				{
-					TArray<UDataLayerInstance*> DataLayerInstances = UDataLayerSubsystem::ConvertArgsToDataLayers(World, InArgs);
-					for (UDataLayerInstance* DataLayerInstance : DataLayerInstances)
-					{
-						DataLayerSubsystem->SetDataLayerRuntimeState(DataLayerInstance, DataLayerSubsystem->GetDataLayerRuntimeState(DataLayerInstance) == EDataLayerRuntimeState::Activated ? EDataLayerRuntimeState::Unloaded : EDataLayerRuntimeState::Activated);
-					}
-				}
-			}
-		}
-	})
-);
-
-FAutoConsoleCommand UDataLayerSubsystem::SetDataLayerRuntimeStateCommand(
-	TEXT("wp.Runtime.SetDataLayerRuntimeState"),
-	TEXT("Sets Runtime DataLayers state. Args [State = Unloaded, Loaded, Activated] [DataLayerNames]"),
-	FConsoleCommandWithArgsDelegate::CreateLambda([](const TArray<FString>& InArgs)
-	{
-		if (InArgs.Num() < 2)
-		{
-			UE_LOG(LogWorldPartition, Warning, TEXT("wp.Runtime.SetDataLayerRuntimeState : Requires at least 2 arguments. First argument should be the target state and the next ones should be the list of DataLayers."));
-			return;
-		}
-
-		TArray<FString> Args = InArgs;
-		FString StatetStr;
-		Args.HeapPop(StatetStr);
-		EDataLayerRuntimeState State;
-		if (!GetDataLayerRuntimeStateFromName(StatetStr, State))
-		{
-			UE_LOG(LogWorldPartition, Warning, TEXT("wp.Runtime.SetDataLayerRuntimeState : Invalid first argument, expected one of these values : Unloaded, Loaded, Activated."));
-			return;
-		}
-
-		for (const FWorldContext& Context : GEngine->GetWorldContexts())
-		{
-			UWorld* World = Context.World();
-			if (World && World->IsGameWorld() && !World->IsNetMode(NM_Client))
-			{
-				if (UDataLayerSubsystem* DataLayerSubsystem = World->GetSubsystem<UDataLayerSubsystem>())
-				{
-					TArray<UDataLayerInstance*> DataLayerInstances = UDataLayerSubsystem::ConvertArgsToDataLayers(World, Args);
-					for (UDataLayerInstance* DataLayerInstance : DataLayerInstances)
-					{
-						DataLayerSubsystem->SetDataLayerRuntimeState(DataLayerInstance, State);
-					}
-				}
-			}
-		}
-	})
-);
-
-void UDataLayerSubsystem::GetDataLayerDebugColors(TMap<FName, FColor>& OutMapping) const
-{
-	OutMapping.Reset();
-	
-	ForEachDataLayer([&OutMapping](UDataLayerInstance* DataLayer)
-	{
-		OutMapping.Add(DataLayer->GetDataLayerFName(), DataLayer->GetDebugColor());
-		return true;
-	});
-}
-
-#if WITH_EDITOR
-
-void UDataLayerSubsystem::PushActorEditorContext() const
-{
-	++DataLayerActorEditorContextID;
-	WorldDataLayerCollection.ForEachWorldDataLayers([this](AWorldDataLayers* WorldDataLayers)
-	{
-		WorldDataLayers->PushActorEditorContext(DataLayerActorEditorContextID);
-	});
-}
-
-void UDataLayerSubsystem::PopActorEditorContext() const
-{
-	check(DataLayerActorEditorContextID > 0);
-	WorldDataLayerCollection.ForEachWorldDataLayers([this](AWorldDataLayers* WorldDataLayers)
-	{
-		WorldDataLayers->PopActorEditorContext(DataLayerActorEditorContextID);
-	});
-	--DataLayerActorEditorContextID;
-}
-
-TArray<UDataLayerInstance*> UDataLayerSubsystem::GetActorEditorContextDataLayers() const
-{
-	TArray<UDataLayerInstance*> DataLayersInstances;
-	WorldDataLayerCollection.ForEachWorldDataLayers([&DataLayersInstances](AWorldDataLayers* WorldDataLayers)
-	{
-		DataLayersInstances += WorldDataLayers->GetActorEditorContextDataLayers();
-	});
-
-	return DataLayersInstances;
-}
-
-uint32 UDataLayerSubsystem::GetDataLayerEditorContextHash() const
-{
-	TArray<FName> DataLayerInstanceNames;
-	for (UDataLayerInstance* DataLayerInstance : GetActorEditorContextDataLayers())
-	{
-		DataLayerInstanceNames.Add(DataLayerInstance->GetDataLayerFName());
+		DataLayerManager->DumpDataLayers(OutputDevice);
 	}
-	return FDataLayerEditorContext(GetWorld(), DataLayerInstanceNames).GetHash();
 }
 
-void UDataLayerSubsystem::OnActorDescContainerInitialized(UActorDescContainer* InActorDescContainer)
-{
-	ResolveActorDescContainer(InActorDescContainer);
-}
-
-void UDataLayerSubsystem::ResolveActorDescContainer(UActorDescContainer* InActorDescContainer)
-{
-	check(InActorDescContainer);
-
-	auto GetWorldDataLayersActorDesc = [](const UActorDescContainer* InContainer) -> TArray<const FWorldDataLayersActorDesc*>
-	{
-		TArray<const FWorldDataLayersActorDesc*> WorldDataLayersActorDecs;
-		if (InContainer)
-		{
-			for (FActorDescList::TIterator<AWorldDataLayers> Iterator(const_cast<UActorDescContainer*>(InContainer)); Iterator; ++Iterator)
-			{
-				// IsValid will return false for maps that don't have a valid WorldDataLayers actors
-				if (Iterator->IsValid())
-				{
-					WorldDataLayersActorDecs.Add(*Iterator);
-				}
-			}
-		}
-
-		return WorldDataLayersActorDecs;
-	};
-
-	const TArray<const FWorldDataLayersActorDesc*> WorldDataLayersActorDescs = GetWorldDataLayersActorDesc(InActorDescContainer);
-	for (FActorDescList::TIterator<> Iterator(InActorDescContainer); Iterator; ++Iterator)
-	{
-		FWorldPartitionActorDesc* ActorDesc = *Iterator;
-		check(ActorDesc->GetContainer() == InActorDescContainer);
-		ActorDesc->SetDataLayerInstanceNames(FDataLayerUtils::ResolvedDataLayerInstanceNames(ActorDesc, WorldDataLayersActorDescs));
-	}
-}
-
-void UDataLayerSubsystem::ResolveActorDescContainers()
+UDataLayerInstance* UDataLayerSubsystem::GetDataLayerInstanceFromAsset(const UDataLayerAsset* InDataLayerAsset) const
 {
-	for (TObjectIterator<UActorDescContainer> ContainerIt; ContainerIt; ++ContainerIt)
-	{
-		if (UActorDescContainer* ActorDescContainer = *ContainerIt; ActorDescContainer && GetWorld() == ActorDescContainer->GetWorld())
-		{
-			ResolveActorDescContainer(ActorDescContainer);
-		}
-	}
+	const UDataLayerManager* DataLayerManager = UDataLayerManager::GetDataLayerManager(GetWorld());
+	return DataLayerManager ? const_cast<UDataLayerInstance*>(DataLayerManager->GetDataLayerInstance(InDataLayerAsset)) : nullptr;
 }
 
-void UDataLayerSubsystem::RegisterEarlyLoadedWorldDataLayers()
+const UDataLayerInstance* UDataLayerSubsystem::GetDataLayerInstanceFromAssetName(const FName& InDataLayerAssetFullName) const
 {
-	// ULevel::WorldDatalayers is loaded before the subsystem is created and miss their registration.
-	RegisterWorldDataLayer(GetWorld()->GetWorldDataLayers());
+	const UDataLayerManager* DataLayerManager = UDataLayerManager::GetDataLayerManager(GetWorld());
+	return DataLayerManager ? DataLayerManager->GetDataLayerInstanceFromAssetName(InDataLayerAssetFullName) : nullptr;
 }
 
-#endif
-
-//~ Begin Deprecated
-
 PRAGMA_DISABLE_DEPRECATION_WARNINGS
 
 UDataLayerInstance* UDataLayerSubsystem::GetDataLayer(const FActorDataLayer& InDataLayer) const
@@ -869,13 +174,16 @@ EDataLayerRuntimeState UDataLayerSubsystem::GetDataLayerEffectiveRuntimeStateByL
 
 void UDataLayerSubsystem::SetDataLayerRuntimeState(const FActorDataLayer& InDataLayer, EDataLayerRuntimeState InState, bool bInIsRecursive)
 {
-	if (UDataLayerInstance* DataLayerInstance = GetDataLayerFromName(InDataLayer.Name))
+	if (UDataLayerManager* DataLayerManager = UDataLayerManager::GetDataLayerManager(GetWorld()))
 	{
-		SetDataLayerRuntimeState(DataLayerInstance, InState, bInIsRecursive);
-	}
-	else
-	{
-		UE_LOG(LogWorldPartition, Warning, TEXT("UDataLayerSubsystem::SetDataLayerRuntimeState unknown Data Layer: '%s'"), *InDataLayer.Name.ToString());
+		if (const UDataLayerInstance* DataLayerInstance = DataLayerManager->GetDataLayerInstanceFromName(InDataLayer.Name))
+		{
+			SetDataLayerRuntimeState(DataLayerInstance, InState, bInIsRecursive);
+		}
+		else
+		{
+			UE_LOG(LogWorldPartition, Warning, TEXT("UDataLayerSubsystem::SetDataLayerRuntimeState unknown Data Layer: '%s'"), *InDataLayer.Name.ToString());
+		}
 	}
 }
 
@@ -891,8 +199,6 @@ void UDataLayerSubsystem::SetDataLayerRuntimeStateByLabel(const FName& InDataLay
 	}
 }
 
-PRAGMA_ENABLE_DEPRECATION_WARNINGS
-
 UDataLayerInstance* UDataLayerSubsystem::GetDataLayerFromName(FName InDataLayerName) const
 {
 	return GetDataLayerInstance(InDataLayerName);
@@ -908,4 +214,4 @@ EDataLayerRuntimeState UDataLayerSubsystem::GetDataLayerEffectiveRuntimeState(co
 	return GetDataLayerEffectiveRuntimeStateByName(InDataLayer.Name);
 }
 
-//~ End Deprecated
+PRAGMA_ENABLE_DEPRECATION_WARNINGS
\ No newline at end of file
diff --git a/Engine/Source/Runtime/Engine/Private/WorldPartition/DataLayer/DataLayerUtils.cpp b/Engine/Source/Runtime/Engine/Private/WorldPartition/DataLayer/DataLayerUtils.cpp
--- a/Engine/Source/Runtime/Engine/Private/WorldPartition/DataLayer/DataLayerUtils.cpp
+++ b/Engine/Source/Runtime/Engine/Private/WorldPartition/DataLayer/DataLayerUtils.cpp
@@ -1,112 +1,92 @@
 // Copyright Epic Games, Inc. All Rights Reserved.
 
 #include "WorldPartition/DataLayer/DataLayerUtils.h"
-#include "WorldPartition/WorldPartitionStreamingGeneration.h"
 
 #if WITH_EDITOR
+#include "WorldPartition/ActorDescContainer.h"
+#include "WorldPartition/WorldPartitionStreamingGeneration.h"
 #include "WorldPartition/DataLayer/WorldDataLayers.h"
 #include "WorldPartition/DataLayer/WorldDataLayersActorDesc.h"
 #include "WorldPartition/DataLayer/DataLayerInstanceWithAsset.h"
-#include "WorldPartition/ActorDescContainer.h"
-#include "WorldPartition/DataLayer/DataLayerSubsystem.h"
+#include "WorldPartition/DataLayer/DataLayerManager.h"
 #include "Algo/AllOf.h"
 
-TArray<FName> FDataLayerUtils::ResolvedDataLayerInstanceNames(const FWorldPartitionActorDesc* InActorDesc, const TArray<const FWorldDataLayersActorDesc*>& InWorldDataLayersActorDescs, UWorld* InWorld, bool* bOutIsResultValid)
+TArray<FName> FDataLayerUtils::ResolvedDataLayerInstanceNames(const UDataLayerManager* InDataLayerManager, const FWorldPartitionActorDesc* InActorDesc, const TArray<const FWorldDataLayersActorDesc*>& InWorldDataLayersActorDescs)
 {
-	bool bLocalIsSuccess = true;
-	bool& bIsSuccess = bOutIsResultValid ? *bOutIsResultValid : bLocalIsSuccess;
-	bIsSuccess = true;
-
-	// Prioritize in-memory AWorldDataLayers
-	UWorld* World = InWorld;
-	if (!World)
-	{
-		World = InActorDesc->GetContainer() ? InActorDesc->GetContainer()->GetWorld() : nullptr;
-	}
+	const TArray<FName>& ActorDescDataLayers = InActorDesc->GetDataLayers();
 
-	const UDataLayerSubsystem* DataLayerSubsystem = UWorld::GetSubsystem<UDataLayerSubsystem>(World);
-
-	// DataLayers not using DataLayer Assets represent DataLayerInstanceNames
-	if (!InActorDesc->IsUsingDataLayerAsset())
+	if (!ActorDescDataLayers.IsEmpty())
 	{
-		if (DataLayerSubsystem && DataLayerSubsystem->CanResolveDataLayers())
+		// DataLayers not using DataLayer Assets represent DataLayerInstanceNames
+		if (!InActorDesc->IsUsingDataLayerAsset())
 		{
-			TArray<FName> Result;
-			for (const FName& DataLayerInstanceName : InActorDesc->GetDataLayers())
+			if (InDataLayerManager && InDataLayerManager->CanResolveDataLayers())
 			{
-				if (const UDataLayerInstance* DataLayerInstance = DataLayerSubsystem->GetDataLayerInstance(DataLayerInstanceName))
+				TArray<FName> Result;
+				for (const FName& DataLayerInstanceName : ActorDescDataLayers)
 				{
-					Result.Add(DataLayerInstanceName);
+					if (const UDataLayerInstance* DataLayerInstance = InDataLayerManager->GetDataLayerInstance(DataLayerInstanceName))
+					{
+						Result.Add(DataLayerInstanceName);
+					}
 				}
+				return Result;
 			}
-			return Result;
-		}
-		// Fallback on FWorldDataLayersActorDesc
-		else if (!InWorldDataLayersActorDescs.IsEmpty() && AreWorldDataLayersActorDescsSane(InWorldDataLayersActorDescs))
-		{
-			TArray<FName> Result;
-			for (const FName& DataLayerInstanceName : InActorDesc->GetDataLayers())
+			// Fallback on FWorldDataLayersActorDesc
+			else if (!InWorldDataLayersActorDescs.IsEmpty() && AreWorldDataLayersActorDescsSane(InWorldDataLayersActorDescs))
 			{
-				if (GetDataLayerInstanceDescFromInstanceName(InWorldDataLayersActorDescs, DataLayerInstanceName) != nullptr)
+				TArray<FName> Result;
+				for (const FName& DataLayerInstanceName : ActorDescDataLayers)
 				{
-					Result.Add(DataLayerInstanceName);
+					if (GetDataLayerInstanceDescFromInstanceName(InWorldDataLayersActorDescs, DataLayerInstanceName) != nullptr)
+					{
+						Result.Add(DataLayerInstanceName);
+					}
 				}
+				return Result;
 			}
-			return Result;
 		}
-	}
-	// ActorDesc DataLayers represents DataLayer Asset Paths
-	else
-	{
-		if (DataLayerSubsystem && DataLayerSubsystem->CanResolveDataLayers())
+		// ActorDesc DataLayers represents DataLayer Asset Paths
+		else
 		{
-			TArray<FName> Result;
-			for (const FName& DataLayerAssetPath : InActorDesc->GetDataLayers())
+			if (InDataLayerManager && InDataLayerManager->CanResolveDataLayers())
 			{
-				DataLayerSubsystem->ForEachDataLayer([DataLayerAssetPath, &Result](UDataLayerInstance* DataLayerInstance)
+				TArray<FName> Result;
+				for (const FName& DataLayerAssetPath : ActorDescDataLayers)
 				{
-					const UDataLayerInstanceWithAsset* DataLayerInstanceWithAsset = Cast<UDataLayerInstanceWithAsset>(DataLayerInstance);
-					const UDataLayerAsset* DataLayerAsset = DataLayerInstanceWithAsset ? DataLayerInstanceWithAsset->GetAsset() : nullptr;
-					if (DataLayerAsset && (FName(DataLayerAsset->GetPathName()) == DataLayerAssetPath))
+					if (const UDataLayerInstance* DataLayerInstance = InDataLayerManager->GetDataLayerInstanceFromAssetName(DataLayerAssetPath))
 					{
 						Result.Add(DataLayerInstance->GetDataLayerFName());
-						return false;
 					}
-					return true;
-				});
+				}
+				return Result;
 			}
-			return Result;
-		}
-		// Fallback on FWorldDataLayersActorDesc
-		else if (!InWorldDataLayersActorDescs.IsEmpty() && AreWorldDataLayersActorDescsSane(InWorldDataLayersActorDescs))
-		{
-			TArray<FName> Result;
-			for (const FName& DataLayerAssetPath : InActorDesc->GetDataLayers())
+			// Fallback on FWorldDataLayersActorDesc
+			else if (!InWorldDataLayersActorDescs.IsEmpty() && AreWorldDataLayersActorDescsSane(InWorldDataLayersActorDescs))
 			{
-				if (const FDataLayerInstanceDesc* DataLayerInstanceDesc = GetDataLayerInstanceDescFromAssetPath(InWorldDataLayersActorDescs, DataLayerAssetPath))
+				TArray<FName> Result;
+				for (const FName& DataLayerAssetPath : ActorDescDataLayers)
 				{
-					Result.Add(DataLayerInstanceDesc->GetName());
+					if (const FDataLayerInstanceDesc* DataLayerInstanceDesc = GetDataLayerInstanceDescFromAssetPath(InWorldDataLayersActorDescs, DataLayerAssetPath))
+					{
+						Result.Add(DataLayerInstanceDesc->GetName());
+					}
 				}
+				return Result;
 			}
-			return Result;
 		}
 	}
-
-	bIsSuccess = false;
 	return InActorDesc->GetDataLayers();
 }
 
 // For performance reasons, this function assumes that InActorDesc's DataLayerInstanceNames was already resolved.
-bool FDataLayerUtils::ResolveRuntimeDataLayerInstanceNames(const FWorldPartitionActorDescView& InActorDescView, const FActorDescViewMap& ActorDescViewMap, TArray<FName>& OutRuntimeDataLayerInstanceNames)
-{
-	UWorld* World = InActorDescView.GetActorDesc()->GetContainer()->GetWorld();
-	const UDataLayerSubsystem* DataLayerSubsystem = World ? World->GetSubsystem<UDataLayerSubsystem>() : nullptr;
-
-	if (DataLayerSubsystem && DataLayerSubsystem->CanResolveDataLayers())
+bool FDataLayerUtils::ResolveRuntimeDataLayerInstanceNames(const UDataLayerManager* InDataLayerManager, const FWorldPartitionActorDescView& InActorDescView, const FActorDescViewMap& ActorDescViewMap, TArray<FName>& OutRuntimeDataLayerInstanceNames)
+{	
+	if (InDataLayerManager && InDataLayerManager->CanResolveDataLayers())
 	{
 		for (FName DataLayerInstanceName : InActorDescView.GetDataLayerInstanceNames())
 		{
-			const UDataLayerInstance* DataLayerInstance = DataLayerSubsystem->GetDataLayerInstance(DataLayerInstanceName);
+			const UDataLayerInstance* DataLayerInstance = InDataLayerManager->GetDataLayerInstanceFromName(DataLayerInstanceName);
 			if (DataLayerInstance && DataLayerInstance->IsRuntime())
 			{
 				OutRuntimeDataLayerInstanceNames.Add(DataLayerInstanceName);
@@ -117,12 +97,12 @@ bool FDataLayerUtils::ResolveRuntimeDataLayerInstanceNames(const FWorldPartition
 	}
 	else
 	{
-		TArray<const FWorldDataLayersActorDesc*> WorldDataLayersActorDescs;
-		FindWorldDataLayerActorDescs(ActorDescViewMap, WorldDataLayersActorDescs);
+		TArray<const FWorldDataLayersActorDesc*> WorldDataLayersActorDescs = FindWorldDataLayerActorDescs(ActorDescViewMap);
 
 		// Fallback on FWorldDataLayersActorDesc
-		if (WorldDataLayersActorDescs.Num() && AreWorldDataLayersActorDescsSane(WorldDataLayersActorDescs))
+		if (WorldDataLayersActorDescs.Num())
 		{
+			check(AreWorldDataLayersActorDescsSane(WorldDataLayersActorDescs));
 			for (FName DataLayerInstanceName : InActorDescView.GetDataLayerInstanceNames())
 			{
 				if (const FDataLayerInstanceDesc* DataLayerInstanceDesc = GetDataLayerInstanceDescFromInstanceName(WorldDataLayersActorDescs, DataLayerInstanceName))
@@ -167,11 +147,14 @@ const FDataLayerInstanceDesc* FDataLayerUtils::GetDataLayerInstanceDescFromAsset
 	return nullptr;
 }
 
-bool FDataLayerUtils::FindWorldDataLayerActorDescs(const FActorDescViewMap& ActorDescViewMap, TArray<const FWorldDataLayersActorDesc*>& OutWorldDataLayersActorDescs)
+TArray<const FWorldDataLayersActorDesc*> FDataLayerUtils::FindWorldDataLayerActorDescs(const FActorDescViewMap& ActorDescViewMap)
 {
+	TArray<const FWorldDataLayersActorDesc*> WorldDataLayersActorDescs;
 	TArray<const FWorldPartitionActorDescView*> WorldDataLayerViews = ActorDescViewMap.FindByExactNativeClass<AWorldDataLayers>();
-	Algo::Transform(WorldDataLayerViews, OutWorldDataLayersActorDescs, [](const FWorldPartitionActorDescView* WorldDataLayersActorDescView) { return  (FWorldDataLayersActorDesc*)WorldDataLayersActorDescView->GetActorDesc(); });
-	return !OutWorldDataLayersActorDescs.IsEmpty();
+	Algo::TransformIf(WorldDataLayerViews, WorldDataLayersActorDescs, 
+		[](const FWorldPartitionActorDescView* WorldDataLayersActorDescView) { return ((FWorldDataLayersActorDesc*)WorldDataLayersActorDescView->GetActorDesc())->IsValid(); },
+		[](const FWorldPartitionActorDescView* WorldDataLayersActorDescView) { return (FWorldDataLayersActorDesc*)WorldDataLayersActorDescView->GetActorDesc(); });
+	return WorldDataLayersActorDescs;
 }
 
 bool FDataLayerUtils::AreWorldDataLayersActorDescsSane(const TArray<const FWorldDataLayersActorDesc*>& InWorldDataLayersActorDescs)
diff --git a/Engine/Source/Runtime/Engine/Private/WorldPartition/DataLayer/WorldDataLayers.cpp b/Engine/Source/Runtime/Engine/Private/WorldPartition/DataLayer/WorldDataLayers.cpp
--- a/Engine/Source/Runtime/Engine/Private/WorldPartition/DataLayer/WorldDataLayers.cpp
+++ b/Engine/Source/Runtime/Engine/Private/WorldPartition/DataLayer/WorldDataLayers.cpp
@@ -5,21 +5,24 @@
 =============================================================================*/
 
 #include "WorldPartition/DataLayer/WorldDataLayers.h"
-#include "Engine/Level.h"
-#include "ProfilingDebugging/CsvProfiler.h"
-#include "WorldPartition/DataLayer/DataLayerSubsystem.h"
+#include "WorldPartition/DataLayer/DataLayerManager.h"
 #include "WorldPartition/DataLayer/DataLayer.h"
 #include "WorldPartition/DataLayer/DataLayerInstanceWithAsset.h"
 #include "WorldPartition/DataLayer/DeprecatedDataLayerInstance.h"
 #include "WorldPartition/DataLayer/DataLayerUtils.h"
-#include "Net/UnrealNetwork.h"
+#include "WorldPartition/WorldPartition.h"
 #include "WorldPartition/WorldPartitionLog.h"
+#include "WorldPartition/WorldPartitionRuntimeCellInterface.h"
+#include "Engine/Level.h"
+#include "ProfilingDebugging/CsvProfiler.h"
+#include "Net/UnrealNetwork.h"
 
 #include UE_INLINE_GENERATED_CPP_BY_NAME(WorldDataLayers)
 
 #if WITH_EDITOR
 #include "WorldPartition/WorldPartitionEditorPerProjectUserSettings.h"
 #include "WorldPartition/DataLayer/WorldDataLayersActorDesc.h"
+#include "LevelInstance/LevelInstanceSubsystem.h"
 #include "Algo/Find.h"
 #endif
 
@@ -93,8 +96,14 @@ void AWorldDataLayers::InitializeDataLayerRuntimeStates()
 {
 	check(ActiveDataLayerNames.IsEmpty() && LoadedDataLayerNames.IsEmpty());
 
-	if (GetWorld()->IsGameWorld() && IsRuntimeRelevant())
+	if (GetWorld()->IsGameWorld())
 	{
+#if WITH_EDITOR
+		if (!IsRuntimeRelevant())
+		{
+			return;
+		}
+#endif
 		ForEachDataLayer([this](class UDataLayerInstance* DataLayer)
 		{
 			if (DataLayer && DataLayer->IsRuntime())
@@ -216,8 +225,10 @@ void AWorldDataLayers::SetDataLayerRuntimeState(const UDataLayerInstance* InData
 
 void AWorldDataLayers::OnDataLayerRuntimeStateChanged_Implementation(const UDataLayerInstance* InDataLayer, EDataLayerRuntimeState InState)
 {
-	UDataLayerSubsystem* DataLayerSubsystem = GetWorld()->GetSubsystem<UDataLayerSubsystem>();
-	DataLayerSubsystem->OnDataLayerRuntimeStateChanged.Broadcast(InDataLayer, InState);
+	if (UDataLayerManager* DataLayerManager = UDataLayerManager::GetDataLayerManager(this))
+	{
+		DataLayerManager->BroadcastOnDataLayerInstanceRuntimeStateChanged(InDataLayer, InState);
+	}
 }
 
 void AWorldDataLayers::OnRep_ActiveDataLayerNames()
@@ -255,10 +266,6 @@ void AWorldDataLayers::OnRep_EffectiveActiveDataLayerNames()
 	++DataLayersStateEpoch;
 	EffectiveActiveDataLayerNames.Reset();
 	EffectiveActiveDataLayerNames.Append(RepEffectiveActiveDataLayerNames);
-	if (UDataLayerSubsystem* DataLayerSubsystem = UWorld::GetSubsystem<UDataLayerSubsystem>(GetWorld()))
-	{
-		DataLayerSubsystem->OnEffectiveRuntimeDataLayerStatesChanged(this);
-	}
 }
 
 void AWorldDataLayers::OnRep_EffectiveLoadedDataLayerNames()
@@ -266,10 +273,6 @@ void AWorldDataLayers::OnRep_EffectiveLoadedDataLayerNames()
 	++DataLayersStateEpoch;
 	EffectiveLoadedDataLayerNames.Reset();
 	EffectiveLoadedDataLayerNames.Append(RepEffectiveLoadedDataLayerNames);
-	if (UDataLayerSubsystem* DataLayerSubsystem = UWorld::GetSubsystem<UDataLayerSubsystem>(GetWorld()))
-	{
-		DataLayerSubsystem->OnEffectiveRuntimeDataLayerStatesChanged(this);
-	}
 }
 
 EDataLayerRuntimeState AWorldDataLayers::GetDataLayerEffectiveRuntimeStateByName(FName InDataLayerName) const
@@ -322,12 +325,6 @@ void AWorldDataLayers::ResolveEffectiveRuntimeState(const UDataLayerInstance* In
 		// Update Replicated Properties
 		RepEffectiveActiveDataLayerNames = EffectiveActiveDataLayerNames.Array();
 		RepEffectiveLoadedDataLayerNames = EffectiveLoadedDataLayerNames.Array();
-
-		if (UDataLayerSubsystem* DataLayerSubsystem = UWorld::GetSubsystem<UDataLayerSubsystem>(GetWorld()))
-		{
-			DataLayerSubsystem->OnEffectiveRuntimeDataLayerStatesChanged(this);
-		}
-
 		++DataLayersStateEpoch;
 
 		if (bInNotifyChange)
@@ -771,6 +768,8 @@ void AWorldDataLayers::PostLoad()
 	}
 
 #if WITH_EDITOR
+	bListedInSceneOutliner = true;
+
 	ConvertDataLayerToInstancces();
 
 	// Remove all Editor Data Layers when cooking or when in a game world
@@ -794,12 +793,11 @@ void AWorldDataLayers::PostLoad()
 		RemoveDataLayers(EditorDataLayers);
 	}
 
-	// Sub-WorldDataLayers are not supported at runtime, clear data layers
-	if (GetWorld()->IsGameWorld() && !IsRuntimeRelevant() && !IsEmpty())
+	// Empty Data Layers for non-runtime relevant AWorldDataLayers
+	if (GetWorld()->IsGameWorld() && !IsEmpty() && !IsRuntimeRelevant())
 	{
 		UE_LOG(LogWorldPartition, Warning, TEXT("WorldDataLayers %s is not runtime relevant, Data Layers of this level will be ignored."), *GetPathName());
-		TArray<UDataLayerInstance*> AllDataLayers = DataLayerInstances.Array();
-		RemoveDataLayers(AllDataLayers);
+		RemoveDataLayers(DataLayerInstances.Array());
 	}
 
 	// Setup defaults before overriding with user settings
@@ -826,8 +824,6 @@ void AWorldDataLayers::PostLoad()
 			DataLayerInstance->SetIsLoadedInEditor(true, /*bFromUserChange*/false);
 		}
 	}
-
-	bListedInSceneOutliner = true;
 #else
 	// Build acceleration tables
 	for (const UDataLayerInstance* DataLayerInstance : DataLayerInstances)
@@ -843,52 +839,34 @@ void AWorldDataLayers::PostLoad()
 #endif
 
 	InitializeDataLayerRuntimeStates();
-
-	if (UDataLayerSubsystem* DataLayerSubsystem = UWorld::GetSubsystem<UDataLayerSubsystem>(GetWorld()))
-	{
-		DataLayerSubsystem->RegisterWorldDataLayer(this);
-	}
-}
-
-void AWorldDataLayers::PreRegisterAllComponents()
-{
-	Super::PreRegisterAllComponents();
-
-	if (UDataLayerSubsystem* DataLayerSubsystem = UWorld::GetSubsystem<UDataLayerSubsystem>(GetWorld()))
-	{
-		DataLayerSubsystem->RegisterWorldDataLayer(this);
-	}
-}
-
-void AWorldDataLayers::PostUnregisterAllComponents()
-{
-	if (UDataLayerSubsystem* DataLayerSubsystem = UWorld::GetSubsystem<UDataLayerSubsystem>(GetWorld()))
-	{
-		DataLayerSubsystem->UnregisterWorldDataLayer(this);
-	}
-
-	Super::PostUnregisterAllComponents();
-}
-
-bool AWorldDataLayers::IsRuntimeRelevant() const
-{
-	return !IsSubWorldDataLayers();
 }
 
+#if WITH_EDITOR
 bool AWorldDataLayers::IsSubWorldDataLayers() const
 {
 	UWorld* ActorWorld = GetWorld();
 	UWorld* OuterWorld = GetTypedOuter<UWorld>();
-	return ActorWorld != nullptr && OuterWorld != nullptr && OuterWorld->GetFName() != ActorWorld->GetFName();
+	return ActorWorld != nullptr && OuterWorld != nullptr && (OuterWorld->GetFName() != ActorWorld->GetFName());
 }
 
-bool AWorldDataLayers::IsTheMainWorldDataLayers() const
+bool AWorldDataLayers::IsRuntimeRelevant() const
 {
+	// This function is only meant to be called on Game worlds
+	check(GetWorld()->IsGameWorld());
+
+	// WorldDataLayers referenced by LevelInstances are not relevant at runtime
 	UWorld* OuterWorld = GetTypedOuter<UWorld>();
-	return OuterWorld && OuterWorld->GetWorldDataLayers() == this;
-}
+	if (OuterWorld && (OuterWorld->GetWorldDataLayers() == this))
+	{
+		check(OuterWorld->IsPartitionedWorld());
+		ULevelInstanceSubsystem* LevelInstanceSubsystem = UWorld::GetSubsystem<ULevelInstanceSubsystem>(GetWorld());
+		ILevelInstanceInterface* LevelInstance = LevelInstanceSubsystem ? LevelInstanceSubsystem->GetOwningLevelInstance(OuterWorld->PersistentLevel) : nullptr;
+		return (LevelInstance == nullptr);
+	}
 
-#if WITH_EDITOR
+	// @todo_ow: revisit that logic for content bundle data layers
+	return false;
+}
 
 void AWorldDataLayers::ConvertDataLayerToInstancces()
 {
@@ -944,9 +922,11 @@ void AWorldDataLayers::UpdateContainsDeprecatedDataLayers()
 
 void AWorldDataLayers::ResolveActorDescContainers()
 {
-	if (UDataLayerSubsystem* DataLayerSubsystem = GetWorld()->GetSubsystem<UDataLayerSubsystem>())
+	// Always use actor owning world to find DataLayerManager for resolving of data layers as partitioned 
+	// LevelInstance DataLayerManager can't resolve.
+	if (UDataLayerManager* DataLayerManager = UDataLayerManager::GetDataLayerManager(GetWorld()))
 	{
-		DataLayerSubsystem->ResolveActorDescContainers();
+		DataLayerManager->ResolveActorDescContainersDataLayers();
 	}
 }
 
diff --git a/Engine/Source/Runtime/Engine/Private/WorldPartition/DataLayer/WorldDataLayersActorDesc.cpp b/Engine/Source/Runtime/Engine/Private/WorldPartition/DataLayer/WorldDataLayersActorDesc.cpp
--- a/Engine/Source/Runtime/Engine/Private/WorldPartition/DataLayer/WorldDataLayersActorDesc.cpp
+++ b/Engine/Source/Runtime/Engine/Private/WorldPartition/DataLayer/WorldDataLayersActorDesc.cpp
@@ -171,7 +171,8 @@ void FWorldDataLayersActorDesc::Serialize(FArchive& Ar)
 
 bool FWorldDataLayersActorDesc::IsRuntimeRelevant(const FActorContainerID& InContainerID) const
 {
-	return InContainerID.IsMainContainer();
+	// @todo_ow : For now we don't consider WorldDataLayers from content bundles
+	return InContainerID.IsMainContainer() && !ContentBundleGuid.IsValid();
 }
 
 const FDataLayerInstanceDesc* FWorldDataLayersActorDesc::GetDataLayerInstanceFromInstanceName(FName InDataLayerInstanceName) const
diff --git a/Engine/Source/Runtime/Engine/Private/WorldPartition/HLOD/HLODSubsystem.cpp b/Engine/Source/Runtime/Engine/Private/WorldPartition/HLOD/HLODSubsystem.cpp
--- a/Engine/Source/Runtime/Engine/Private/WorldPartition/HLOD/HLODSubsystem.cpp
+++ b/Engine/Source/Runtime/Engine/Private/WorldPartition/HLOD/HLODSubsystem.cpp
@@ -7,7 +7,7 @@
 #include "Engine/World.h"
 #include "WorldPartition/HLOD/HLODActorDesc.h"
 #include "Materials/MaterialInterface.h"
-#include "WorldPartition/DataLayer/DataLayerSubsystem.h"
+#include "WorldPartition/DataLayer/DataLayerManager.h"
 #include "Misc/Paths.h"
 #include "WorldPartition/WorldPartition.h"
 #include "RenderUtils.h"
@@ -729,11 +729,11 @@ static void DumpHLODStats(UWorld* World)
 		});
 	};
 
-	const UDataLayerSubsystem* DataLayerSubsystem = UWorld::GetSubsystem<UDataLayerSubsystem>(World);
+	const UDataLayerManager* DataLayerManager = WorldPartition->GetDataLayerManager();
 
-	auto GetDataLayerShortName = [DataLayerSubsystem](FName DataLayerInstanceName)
+	auto GetDataLayerShortName = [DataLayerManager](FName DataLayerInstanceName)
 	{ 
-		UDataLayerInstance* DataLayerInstance = DataLayerSubsystem ? DataLayerSubsystem->GetDataLayerInstance(DataLayerInstanceName) : nullptr;
+		const UDataLayerInstance* DataLayerInstance = DataLayerManager ? DataLayerManager->GetDataLayerInstance(DataLayerInstanceName) : nullptr;
 		return DataLayerInstance ? DataLayerInstance->GetDataLayerShortName() : DataLayerInstanceName.ToString();
 	};
 
diff --git a/Engine/Source/Runtime/Engine/Private/WorldPartition/LevelInstance/LevelInstanceActorDesc.cpp b/Engine/Source/Runtime/Engine/Private/WorldPartition/LevelInstance/LevelInstanceActorDesc.cpp
--- a/Engine/Source/Runtime/Engine/Private/WorldPartition/LevelInstance/LevelInstanceActorDesc.cpp
+++ b/Engine/Source/Runtime/Engine/Private/WorldPartition/LevelInstance/LevelInstanceActorDesc.cpp
@@ -22,13 +22,14 @@ static FAutoConsoleVariableRef CVarForceLevelStreaming(
 
 FLevelInstanceActorDesc::FLevelInstanceActorDesc()
 	: DesiredRuntimeBehavior(ELevelInstanceRuntimeBehavior::Partitioned)
-	, LevelInstanceContainer(nullptr)	
+	, LevelInstanceContainer(nullptr)
 {}
 
 FLevelInstanceActorDesc::~FLevelInstanceActorDesc()
 {
 	UnregisterContainerInstance();
 	check(!LevelInstanceContainer.IsValid());
+	check(!LevelInstanceContainerWorldContext.IsValid());
 }
 
 void FLevelInstanceActorDesc::Init(const AActor* InActor)
@@ -67,7 +68,9 @@ bool FLevelInstanceActorDesc::Equals(const FWorldPartitionActorDesc* Other) cons
 
 void FLevelInstanceActorDesc::UpdateBounds()
 {
-	ULevelInstanceSubsystem* LevelInstanceSubsystem = UWorld::GetSubsystem<ULevelInstanceSubsystem>(LevelInstanceContainer->GetWorld());
+	check(LevelInstanceContainerWorldContext.IsValid());
+
+	ULevelInstanceSubsystem* LevelInstanceSubsystem = UWorld::GetSubsystem<ULevelInstanceSubsystem>(LevelInstanceContainerWorldContext.Get());
 	check(LevelInstanceSubsystem);
 
 	FTransform LevelInstancePivotOffsetTransform = FTransform(ULevel::GetLevelInstancePivotOffsetFromPackage(LevelPackage));
@@ -77,15 +80,19 @@ void FLevelInstanceActorDesc::UpdateBounds()
 	ContainerBounds.GetCenterAndExtents(BoundsLocation, BoundsExtent);
 }
 
-void FLevelInstanceActorDesc::RegisterContainerInstance(UWorld* InWorld)
+void FLevelInstanceActorDesc::RegisterContainerInstance(UWorld* InWorldContext)
 {
-	if (InWorld)
+	if (InWorldContext)
 	{
 		check(!LevelInstanceContainer.IsValid());
+		check(!LevelInstanceContainerWorldContext.IsValid());
 
 		if (IsContainerInstance())
 		{
-			ULevelInstanceSubsystem* LevelInstanceSubsystem = UWorld::GetSubsystem<ULevelInstanceSubsystem>(InWorld);
+			check(InWorldContext);
+			LevelInstanceContainerWorldContext = InWorldContext;
+
+			ULevelInstanceSubsystem* LevelInstanceSubsystem = UWorld::GetSubsystem<ULevelInstanceSubsystem>(InWorldContext);
 			check(LevelInstanceSubsystem);
 
 			LevelInstanceContainer = LevelInstanceSubsystem->RegisterContainer(LevelPackage);
@@ -101,21 +108,24 @@ void FLevelInstanceActorDesc::UnregisterContainerInstance()
 {
 	if (LevelInstanceContainer.IsValid())
 	{
-		ULevelInstanceSubsystem* LevelInstanceSubsystem = UWorld::GetSubsystem<ULevelInstanceSubsystem>(LevelInstanceContainer->GetWorld());
+		check(LevelInstanceContainerWorldContext.IsValid());
+
+		ULevelInstanceSubsystem* LevelInstanceSubsystem = UWorld::GetSubsystem<ULevelInstanceSubsystem>(LevelInstanceContainerWorldContext.Get());
 		check(LevelInstanceSubsystem);
 
 		LevelInstanceSubsystem->UnregisterContainer(LevelInstanceContainer.Get());
 		LevelInstanceContainer.Reset();
+		LevelInstanceContainerWorldContext.Reset();
 	}
 }
 
-void FLevelInstanceActorDesc::SetContainer(UActorDescContainer* InContainer)
+void FLevelInstanceActorDesc::SetContainer(UActorDescContainer* InContainer, UWorld* InWorldContext)
 {
-	FWorldPartitionActorDesc::SetContainer(InContainer);
+	FWorldPartitionActorDesc::SetContainer(InContainer, InWorldContext);
 
 	if (Container)
 	{
-		RegisterContainerInstance(Container->GetWorld());
+		RegisterContainerInstance(InWorldContext);
 	}
 	else
 	{
@@ -199,7 +209,8 @@ void FLevelInstanceActorDesc::TransferFrom(const FWorldPartitionActorDesc* From)
 	// Use the Register/Unregister so callbacks are added/removed
 	if (FromLevelInstanceActorDesc->LevelInstanceContainer.IsValid())
 	{
-		RegisterContainerInstance(Container->World);
+		check(FromLevelInstanceActorDesc->LevelInstanceContainerWorldContext.IsValid());
+		RegisterContainerInstance(FromLevelInstanceActorDesc->LevelInstanceContainerWorldContext.Get());
 		FromLevelInstanceActorDesc->UnregisterContainerInstance();
 	}
 }
diff --git a/Engine/Source/Runtime/Engine/Private/WorldPartition/LoaderAdapter/LoaderAdapterPinnedActors.cpp b/Engine/Source/Runtime/Engine/Private/WorldPartition/LoaderAdapter/LoaderAdapterPinnedActors.cpp
--- a/Engine/Source/Runtime/Engine/Private/WorldPartition/LoaderAdapter/LoaderAdapterPinnedActors.cpp
+++ b/Engine/Source/Runtime/Engine/Private/WorldPartition/LoaderAdapter/LoaderAdapterPinnedActors.cpp
@@ -21,8 +21,7 @@ FText FLoaderAdapterPinnedActors::GetUnloadedReason(FWorldPartitionActorDesc* In
 	if (InActorDesc)
 	{
 		UActorDescContainer* ActorDescContainer = InActorDesc->GetContainer();
-		UWorld* World = ActorDescContainer != nullptr ? ActorDescContainer->GetWorld() : nullptr;
-		UWorldPartition* WorldPartition = World != nullptr ? World->GetWorldPartition() : nullptr;
+		UWorldPartition* WorldPartition = ActorDescContainer != nullptr ? ActorDescContainer->GetWorldPartition() : nullptr;
 		bool bShouldBeLoaded = !InActorDesc->GetIsSpatiallyLoaded() && !InActorDesc->GetActorIsRuntimeOnly();
 		if (WorldPartition && (bShouldBeLoaded || WorldPartition->IsActorPinned(InActorDesc->GetGuid())))
 		{
@@ -54,11 +53,10 @@ bool FLoaderAdapterPinnedActors::SupportsPinning(FWorldPartitionActorDesc* InAct
 		}
 		else if (InActorDesc->GetContentBundleGuid().IsValid())
 		{
-			const UWorld* ContainerWorld = Container->GetWorld();
-			const UWorldPartition* ContainerWorldPartition = ContainerWorld ? ContainerWorld->GetWorldPartition() : nullptr;
+			const UWorldPartition* ContainerWorldPartition = Container->GetWorldPartition();
 			if (ContainerWorldPartition && ContainerWorldPartition->IsMainWorldPartition())
 			{
-				return InActorDesc->GetActorSoftPath().GetAssetPath().GetPackageName() == ContainerWorld->GetPackage()->GetFName();
+				return InActorDesc->GetActorSoftPath().GetAssetPath().GetPackageName() == ContainerWorldPartition->GetTypedOuter<UWorld>()->GetPackage()->GetFName();
 			}
 		}
 	}
diff --git a/Engine/Source/Runtime/Engine/Private/WorldPartition/RuntimeSpatialHash/RuntimeSpatialHashHLOD.cpp b/Engine/Source/Runtime/Engine/Private/WorldPartition/RuntimeSpatialHash/RuntimeSpatialHashHLOD.cpp
--- a/Engine/Source/Runtime/Engine/Private/WorldPartition/RuntimeSpatialHash/RuntimeSpatialHashHLOD.cpp
+++ b/Engine/Source/Runtime/Engine/Private/WorldPartition/RuntimeSpatialHash/RuntimeSpatialHashHLOD.cpp
@@ -14,7 +14,8 @@
 #include "WorldPartition/HLOD/HLODModifier.h"
 #include "WorldPartition/HLOD/IWorldPartitionHLODUtilities.h"
 #include "WorldPartition/HLOD/IWorldPartitionHLODUtilitiesModule.h"
-#include "WorldPartition/DataLayer/DataLayerSubsystem.h"
+#include "WorldPartition/DataLayer/DataLayerManager.h"
+#include "WorldPartition/DataLayer/DataLayerUtils.h"
 
 #include "WorldPartition/WorldPartition.h"
 #include "WorldPartition/WorldPartitionHelpers.h"
@@ -283,7 +284,7 @@ static TMap<UHLODLayer*, int32> GatherHLODLayers(const UActorDescContainer* Acto
 		{
 			if (ActorDesc.GetActorIsHLODRelevant())
 			{
-				UHLODLayer* HLODLayer = UHLODLayer::GetHLODLayer(ActorDesc, ActorDescContainer->GetWorld()->GetWorldPartition());
+				UHLODLayer* HLODLayer = UHLODLayer::GetHLODLayer(ActorDesc, ActorDescContainer->GetWorldPartition());
 
 				// If layer was already encountered, no need to process it again
 				if (!HLODLayersLevel.Contains(HLODLayer))
@@ -353,13 +354,11 @@ static TMap<FName, FSpatialHashRuntimeGrid> CreateHLODGrids(TMap<UHLODLayer*, in
 }
 
 // Create/destroy HLOD grid actors
-static void UpdateHLODGridsActors(const UActorDescContainer* Container, const TMap<FName, FSpatialHashRuntimeGrid>& HLODGrids, ISourceControlHelper* SourceControlHelper)
+static void UpdateHLODGridsActors(UWorld* World, const UActorDescContainer* Container, const TMap<FName, FSpatialHashRuntimeGrid>& HLODGrids, ISourceControlHelper* SourceControlHelper)
 {
 	static const FName HLODGridTag = TEXT("HLOD");
 	static const uint32 HLODGridTagLen = HLODGridTag.ToString().Len();
 
-	UWorld* World = Container->GetWorld();
-
 	// Gather all existing HLOD grid actors, see if some are unused and needs to be deleted
 	TMap<FName, ASpatialHashRuntimeGridInfo*> ExistingGridActors;
 	for (FActorDescList::TConstIterator<> ItActorDesc(Container); ItActorDesc; ++ItActorDesc)
@@ -490,7 +489,9 @@ bool UWorldPartitionRuntimeSpatialHash::GenerateHLOD(ISourceControlHelper* Sourc
 	TMap<FName, TArray<FGuid>> GridsHLODActors;
 	TArray<TUniquePtr<FWorldPartitionActorDescView>> ActorDescViews;
 
-	auto GenerateHLODs = [&GridsHLODActors, &ActorDescViews, &MainActorSetContainer, StreamingGenerationContext, WorldPartition, &Context, SourceControlHelper, bCreateActorsOnly]
+	const UDataLayerManager* DataLayerManager = WorldPartition->GetDataLayerManager();
+
+	auto GenerateHLODs = [&GridsHLODActors, &ActorDescViews, &MainActorSetContainer, StreamingGenerationContext, WorldPartition, DataLayerManager, &Context, SourceControlHelper, bCreateActorsOnly]
 	(const FSpatialHashRuntimeGrid& RuntimeGrid, uint32 HLODLevel, const TArray<const IStreamingGenerationContext::FActorSetInstance*>& ActorSetInstances)
 	{
 		// Generate HLODs for this grid
@@ -500,8 +501,12 @@ bool UWorldPartitionRuntimeSpatialHash::GenerateHLOD(ISourceControlHelper* Sourc
 		{
 			FWorldPartitionActorDesc* HLODActorDesc = WorldPartition->GetActorDesc(HLODActorGuid);
 			check(HLODActorDesc);
-
 			FWorldPartitionActorDescView* ActorDescView = ActorDescViews.Emplace_GetRef(MakeUnique<FWorldPartitionActorDescView>(HLODActorDesc)).Get();
+			TArray<FName> RuntimeDataLayerInstanceNames;
+			if (FDataLayerUtils::ResolveRuntimeDataLayerInstanceNames(DataLayerManager, *ActorDescView, *MainActorSetContainer->ActorDescViewMap, RuntimeDataLayerInstanceNames))
+			{
+				ActorDescView->SetRuntimeDataLayerInstanceNames(RuntimeDataLayerInstanceNames);
+			}
 			FActorDescViewMap* NonConstActorDescViewMap = const_cast<FActorDescViewMap*>(MainActorSetContainer->ActorDescViewMap);			
 			NonConstActorDescViewMap->Emplace(HLODActorGuid, *ActorDescView);
 
@@ -544,7 +549,10 @@ bool UWorldPartitionRuntimeSpatialHash::GenerateHLOD(ISourceControlHelper* Sourc
 			NewHLODActorSetInstance.Bounds = HLODActorDescView.GetRuntimeBounds();
 			NewHLODActorSetInstance.RuntimeGrid = HLODActorDescView.GetRuntimeGrid();
 			NewHLODActorSetInstance.bIsSpatiallyLoaded = HLODActorDescView.GetIsSpatiallyLoaded();
-			NewHLODActorSetInstance.DataLayers = UDataLayerSubsystem::GetRuntimeDataLayerInstances(MainActorSetContainer->ActorDescContainer->GetWorld(), HLODActorDescView.GetRuntimeDataLayers());
+			if (DataLayerManager)
+			{
+				NewHLODActorSetInstance.DataLayers = DataLayerManager->GetRuntimeDataLayerInstances(HLODActorDescView.GetRuntimeDataLayerInstanceNames());
+			}
 			NewHLODActorSetInstance.ContentBundleID = MainActorSetContainer->ActorDescContainer->GetContentBundleGuid();
 			NewHLODActorSetInstance.ContainerInstance = MainActorSetContainer;
 			NewHLODActorSetInstance.ActorSet = &MainActorSetContainer->ActorSets.AddDefaulted_GetRef();
@@ -573,7 +581,7 @@ bool UWorldPartitionRuntimeSpatialHash::GenerateHLOD(ISourceControlHelper* Sourc
 	}
 
 	// Create/destroy HLOD grid actors
-	UpdateHLODGridsActors(ActorDescContainer, HLODGrids, SourceControlHelper);
+	UpdateHLODGridsActors(GetWorld(), ActorDescContainer, HLODGrids, SourceControlHelper);
 
 	return true;
 }
diff --git a/Engine/Source/Runtime/Engine/Private/WorldPartition/WorldPartition.cpp b/Engine/Source/Runtime/Engine/Private/WorldPartition/WorldPartition.cpp
--- a/Engine/Source/Runtime/Engine/Private/WorldPartition/WorldPartition.cpp
+++ b/Engine/Source/Runtime/Engine/Private/WorldPartition/WorldPartition.cpp
@@ -12,6 +12,7 @@
 #include "WorldPartition/WorldPartitionLevelStreamingPolicy.h"
 #include "WorldPartition/WorldPartitionReplay.h"
 #include "WorldPartition/HLOD/HLODSubsystem.h"
+#include "WorldPartition/DataLayer/DataLayerManager.h"
 #include "GameFramework/WorldSettings.h"
 #include "ProfilingDebugging/ScopedTimers.h"
 #include "UObject/UE5MainStreamObjectVersion.h"
@@ -34,8 +35,6 @@
 #include "WorldPartition/WorldPartitionRuntimeHash.h"
 #include "WorldPartition/WorldPartitionMiniMap.h"
 #include "WorldPartition/WorldPartitionMiniMapHelper.h"
-#include "WorldPartition/DataLayer/DataLayerSubsystem.h"
-#include "WorldPartition/DataLayer/DataLayerSubsystem.h"
 #include "WorldPartition/DataLayer/WorldDataLayers.h"
 #include "WorldPartition/LoaderAdapter/LoaderAdapterShape.h"
 #include "WorldPartition/LoaderAdapter/LoaderAdapterPinnedActors.h"
@@ -124,10 +123,10 @@ FAutoConsoleVariableRef UWorldPartition::CVarUseMakingInvisibleTransactionReques
 TMap<FName, FString> GetDataLayersDumpString(const UWorldPartition* WorldPartition)
 {
 	TMap<FName, FString> DataLayersDumpString;
-	const UDataLayerSubsystem* DataLayerSubsystem = UWorld::GetSubsystem<UDataLayerSubsystem>(WorldPartition->GetWorld());
-	DataLayerSubsystem->ForEachDataLayer([&DataLayersDumpString](const UDataLayerInstance* DataLayer)
+	const UDataLayerManager* DataLayerManager = WorldPartition->GetDataLayerManager();
+	DataLayerManager->ForEachDataLayerInstance([&DataLayersDumpString](const UDataLayerInstance* DataLayerInstance)
 	{
-		DataLayersDumpString.FindOrAdd(DataLayer->GetDataLayerFName()) = FString::Format(TEXT("{0}{1})"), { DataLayer->GetDataLayerShortName(), DataLayer->GetDataLayerFName().ToString() });
+		DataLayersDumpString.FindOrAdd(DataLayerInstance->GetDataLayerFName()) = FString::Format(TEXT("{0}{1})"), { DataLayerInstance->GetDataLayerShortName(), DataLayerInstance->GetDataLayerFName().ToString() });
 		return true;
 	});
 	
@@ -277,6 +276,7 @@ UWorldPartition::UWorldPartition(const FObjectInitializer& ObjectInitializer)
 #endif
 	, InitState(EWorldPartitionInitState::Uninitialized)
 	, bStreamingInEnabled(true)
+	, DataLayerManager(nullptr)
 	, StreamingPolicy(nullptr)
 	, Replay(nullptr)
 {
@@ -435,6 +435,9 @@ void UWorldPartition::Initialize(UWorld* InWorld, const FTransform& InTransform)
 	check(InWorld);
 	World = InWorld;
 
+	DataLayerManager = NewObject<UDataLayerManager>(this, TEXT("DataLayerManager"), RF_Transient);
+	DataLayerManager->Initialize();
+
 	if (!InTransform.Equals(FTransform::Identity))
 	{
 		InstanceTransform = InTransform;
@@ -729,6 +732,12 @@ void UWorldPartition::Uninitialize()
 		EditorHash = nullptr;
 #endif		
 
+		if (DataLayerManager)
+		{
+			DataLayerManager->DeInitialize();
+			DataLayerManager = nullptr;
+		}
+
 		InitState = EWorldPartitionInitState::Uninitialized;
 
 		FWorldPartitionEvents::BroadcastWorldPartitionUninitialized(World, this);
@@ -740,6 +749,11 @@ PRAGMA_ENABLE_DEPRECATION_WARNINGS
 	}
 }
 
+UDataLayerManager* UWorldPartition::GetDataLayerManager() const
+{
+	return DataLayerManager;
+}
+
 bool UWorldPartition::IsInitialized() const
 {
 	return InitState == EWorldPartitionInitState::Initialized;
diff --git a/Engine/Source/Runtime/Engine/Private/WorldPartition/WorldPartitionActorDesc.cpp b/Engine/Source/Runtime/Engine/Private/WorldPartition/WorldPartitionActorDesc.cpp
--- a/Engine/Source/Runtime/Engine/Private/WorldPartition/WorldPartitionActorDesc.cpp
+++ b/Engine/Source/Runtime/Engine/Private/WorldPartition/WorldPartitionActorDesc.cpp
@@ -18,7 +18,7 @@
 #include "WorldPartition/ActorDescContainer.h"
 #include "WorldPartition/HLOD/HLODLayer.h"
 #include "WorldPartition/DataLayer/WorldDataLayers.h"
-#include "WorldPartition/DataLayer/DataLayerSubsystem.h"
+#include "WorldPartition/DataLayer/DataLayerManager.h"
 #include "WorldPartition/DataLayer/DataLayerAsset.h"
 #include "WorldPartition/DataLayer/DataLayerUtils.h"
 #include "WorldPartition/ContentBundle/ContentBundlePaths.h"
@@ -68,14 +68,15 @@ void FWorldPartitionActorDesc::Init(const AActor* InActor)
 	{
 		TArray<FName> LocalDataLayerAssetPaths;
 		TArray<FName> LocalDataLayerInstanceNames;
-		ULevel* Level = InActor->GetLevel();
-		if (UDataLayerSubsystem* DataLayerSubsystem = UWorld::GetSubsystem<UDataLayerSubsystem>(InActor->GetWorld()))
+		// Use Actor's DataLayerManager since the fixup is relative to this level
+		UDataLayerManager* DataLayerManager = UDataLayerManager::GetDataLayerManager(InActor);
+		if (ensure(DataLayerManager))
 		{
 			LocalDataLayerAssetPaths.Reserve(InActor->GetDataLayerAssets().Num());
 			for (const TObjectPtr<const UDataLayerAsset>& DataLayerAsset : InActor->GetDataLayerAssets())
 			{
 				// Pass Actor Level when resolving the DataLayerInstance as FWorldPartitionActorDesc always represents the state of the actor local to its outer level
-				if (DataLayerAsset && DataLayerSubsystem->GetDataLayerInstance(DataLayerAsset, Level))
+				if (DataLayerAsset && DataLayerManager->GetDataLayerInstance(DataLayerAsset))
 				{
 					LocalDataLayerAssetPaths.Add(*DataLayerAsset->GetPathName());
 				}
@@ -83,7 +84,7 @@ void FWorldPartitionActorDesc::Init(const AActor* InActor)
 
 			PRAGMA_DISABLE_DEPRECATION_WARNINGS
 			// Pass Actor Level when resolving the DataLayerInstance as FWorldPartitionActorDesc always represents the state of the actor local to its outer level
-			LocalDataLayerInstanceNames = DataLayerSubsystem->GetDataLayerInstanceNames(InActor->GetActorDataLayers(), Level);
+			LocalDataLayerInstanceNames = DataLayerManager->GetDataLayerInstanceNames(InActor->GetActorDataLayers());
 			PRAGMA_ENABLE_DEPRECATION_WARNINGS
 		}
 
@@ -97,7 +98,21 @@ void FWorldPartitionActorDesc::Init(const AActor* InActor)
 		DataLayers = bIsUsingDataLayerAsset ? MoveTemp(LocalDataLayerAssetPaths) : MoveTemp(LocalDataLayerInstanceNames);
 
 		// Init DataLayers transient info
-		DataLayerInstanceNames = FDataLayerUtils::ResolvedDataLayerInstanceNames(this, TArray<const FWorldDataLayersActorDesc*>(), InActor->GetWorld());
+		if (ensure(DataLayerManager))
+		{
+			// There are valid cases where the DataLayerManager can't resolve transient data layer instance names.
+			// One case is when committing changes on a level instance, FWorldPartitionActorDesc::Init is called when saving 
+			// the child actors.
+			if (DataLayerManager->CanResolveDataLayers())
+			{
+				// Here we process a loaded actor, so resolving makes sense as long as the actordesc is not reused as a template
+				ResolvedDataLayerInstanceNames = FDataLayerUtils::ResolvedDataLayerInstanceNames(DataLayerManager, this);
+			}
+			else
+			{
+				check(!HasResolvedDataLayerInstanceNames());
+			}
+		}
 	}
 
 	Tags = InActor->Tags;
@@ -216,6 +231,24 @@ void FWorldPartitionActorDesc::SerializeTo(TArray<uint8>& OutData)
 	OutData.Append(PayloadData);
 }
 
+const TArray<FName>& FWorldPartitionActorDesc::GetDataLayerInstanceNames() const
+{
+	static TArray<FName> EmptyDataLayers;
+	if (ensure(HasResolvedDataLayerInstanceNames()))
+	{
+		return ResolvedDataLayerInstanceNames.GetValue();
+	}
+	return EmptyDataLayers;
+}
+
+void FWorldPartitionActorDesc::TransferFrom(const FWorldPartitionActorDesc* From)
+{
+	Container = From->Container;
+	SoftRefCount = From->SoftRefCount;
+	HardRefCount = From->HardRefCount;
+	bIsForcedNonSpatiallyLoaded = From->bIsForcedNonSpatiallyLoaded;
+}
+
 void FWorldPartitionActorDesc::RegisterActorDescDeprecator(TSubclassOf<AActor> ActorClass, const FActorDescDeprecator& Deprecator)
 {
 	check(!Deprecators.Contains(ActorClass));
diff --git a/Engine/Source/Runtime/Engine/Private/WorldPartition/WorldPartitionActorDescView.cpp b/Engine/Source/Runtime/Engine/Private/WorldPartition/WorldPartitionActorDescView.cpp
--- a/Engine/Source/Runtime/Engine/Private/WorldPartition/WorldPartitionActorDescView.cpp
+++ b/Engine/Source/Runtime/Engine/Private/WorldPartition/WorldPartitionActorDescView.cpp
@@ -79,20 +79,37 @@ FName FWorldPartitionActorDescView::GetHLODLayer() const
 	return ActorDesc->GetHLODLayer();
 }
 
-const TArray<FName>& FWorldPartitionActorDescView::GetDataLayers() const
+void FWorldPartitionActorDescView::SetDataLayerInstanceNames(const TArray<FName>& InDataLayerInstanceNames)
 {
-	static TArray<FName> EmptyDataLayers;
-	return bInvalidDataLayers ? EmptyDataLayers : ActorDesc->GetDataLayerInstanceNames();
+	check(!ActorDesc->HasResolvedDataLayerInstanceNames());
+	ResolvedDataLayerInstanceNames = InDataLayerInstanceNames;
 }
 
 const TArray<FName>& FWorldPartitionActorDescView::GetDataLayerInstanceNames() const
 {
+	static TArray<FName> EmptyDataLayers;
+	if (bInvalidDataLayers)
+	{
+		return EmptyDataLayers;
+	}
+	else if (ResolvedDataLayerInstanceNames.IsSet())
+	{
+		return ResolvedDataLayerInstanceNames.GetValue();
+	}
 	return ActorDesc->GetDataLayerInstanceNames();
 }
-const TArray<FName>& FWorldPartitionActorDescView::GetRuntimeDataLayers() const
+
+const TArray<FName>& FWorldPartitionActorDescView::GetRuntimeDataLayerInstanceNames() const
 {
 	static TArray<FName> EmptyDataLayers;
-	return (bInvalidDataLayers || !RuntimeDataLayers.IsSet()) ? EmptyDataLayers : RuntimeDataLayers.GetValue();
+	if (!bInvalidDataLayers)
+	{
+		if (ensure(RuntimeDataLayerInstanceNames.IsSet()))
+		{
+			return RuntimeDataLayerInstanceNames.GetValue();
+		}
+	}
+	return EmptyDataLayers;
 }
 
 const TArray<FName>& FWorldPartitionActorDescView::GetTags() const
@@ -217,9 +234,9 @@ void FWorldPartitionActorDescView::SetInvalidDataLayers()
 	}
 }
 
-void FWorldPartitionActorDescView::SetRuntimeDataLayers(TArray<FName>& InRuntimeDataLayers)
+void FWorldPartitionActorDescView::SetRuntimeDataLayerInstanceNames(TArray<FName>& InRuntimeDataLayerInstanceNames)
 {
-	RuntimeDataLayers = InRuntimeDataLayers;
+	RuntimeDataLayerInstanceNames = InRuntimeDataLayerInstanceNames;
 }
 
 void FWorldPartitionActorDescView::SetRuntimeReferences(TArray<FGuid>& InRuntimeReferences)
diff --git a/Engine/Source/Runtime/Engine/Private/WorldPartition/WorldPartitionActorLoaderInterface.cpp b/Engine/Source/Runtime/Engine/Private/WorldPartition/WorldPartitionActorLoaderInterface.cpp
--- a/Engine/Source/Runtime/Engine/Private/WorldPartition/WorldPartitionActorLoaderInterface.cpp
+++ b/Engine/Source/Runtime/Engine/Private/WorldPartition/WorldPartitionActorLoaderInterface.cpp
@@ -11,7 +11,7 @@
 #include "WorldPartition/WorldPartition.h"
 #include "WorldPartition/WorldPartitionActorDescViewProxy.h"
 #include "WorldPartition/DataLayer/WorldDataLayers.h"
-#include "WorldPartition/DataLayer/DataLayerSubsystem.h"
+#include "WorldPartition/DataLayer/DataLayerManager.h"
 #endif
 
 #define LOCTEXT_NAMESPACE "WorldPartition"
@@ -119,19 +119,27 @@ bool IWorldPartitionActorLoaderInterface::ILoaderAdapter::PassActorDescFilter(co
 	return !Actor->GetActorIsRuntimeOnly();
 }
 
+// Helper added to since the nested class IWorldPartitionActorLoaderInterface::ILoaderAdapter can't be friend of UDataLayerManager.
+bool IWorldPartitionActorLoaderInterface::PassDataLayersFilter(UWorld* InWorld, const TArray<FName>& InDataLayerInstanceNames)
+{
+	check(!IsRunningCookCommandlet());
+	UDataLayerManager* DataLayerManager = UDataLayerManager::GetDataLayerManager(InWorld);
+	return DataLayerManager ? DataLayerManager->ResolveIsLoadedInEditor(InDataLayerInstanceNames) : true;
+}
+
 bool IWorldPartitionActorLoaderInterface::ILoaderAdapter::PassDataLayersFilter(const FWorldPartitionHandle& Actor) const
 {
 	UWorld* OwningWorld = World->PersistentLevel->GetWorld();
-	if (UDataLayerSubsystem* DataLayerSubsystem = UWorld::GetSubsystem<UDataLayerSubsystem>(OwningWorld))
+	if (UDataLayerManager* DataLayerManager = UDataLayerManager::GetDataLayerManager(OwningWorld))
 	{
 		FWorldPartitionActorViewProxy ActorDescProxy(*Actor);
 
 		if (IsRunningCookCommandlet())
 		{
 			// When running cook commandlet, dont allow loading of actors with runtime loaded data layers
-			for (const FName& DataLayerInstanceName : ActorDescProxy.GetDataLayers())
+			for (const FName& DataLayerInstanceName : ActorDescProxy.GetDataLayerInstanceNames())
 			{
-				const UDataLayerInstance* DataLayerInstance = DataLayerSubsystem->GetDataLayerInstance(DataLayerInstanceName);
+				const UDataLayerInstance* DataLayerInstance = DataLayerManager->GetDataLayerInstance(DataLayerInstanceName);
 				if (DataLayerInstance && DataLayerInstance->IsRuntime())
 				{
 					return false;
@@ -140,7 +148,7 @@ bool IWorldPartitionActorLoaderInterface::ILoaderAdapter::PassDataLayersFilter(c
 		}
 		else
 		{
-			return DataLayerSubsystem->ResolveIsLoadedInEditor(ActorDescProxy.GetDataLayers());
+			return IWorldPartitionActorLoaderInterface::PassDataLayersFilter(OwningWorld, ActorDescProxy.GetDataLayerInstanceNames());
 		}
 	}
 
diff --git a/Engine/Source/Runtime/Engine/Private/WorldPartition/WorldPartitionDebugHelper.cpp b/Engine/Source/Runtime/Engine/Private/WorldPartition/WorldPartitionDebugHelper.cpp
--- a/Engine/Source/Runtime/Engine/Private/WorldPartition/WorldPartitionDebugHelper.cpp
+++ b/Engine/Source/Runtime/Engine/Private/WorldPartition/WorldPartitionDebugHelper.cpp
@@ -5,7 +5,8 @@
 #include "GlobalRenderResources.h"
 #include "WorldPartition/WorldPartitionRuntimeHash.h"
 #include "WorldPartition/WorldPartition.h"
-#include "WorldPartition/DataLayer/DataLayerSubsystem.h"
+#include "WorldPartition/DataLayer/DataLayerManager.h"
+#include "WorldPartition/WorldPartitionSubsystem.h"
 #include "Engine/Canvas.h"
 #include "Algo/AllOf.h"
 
@@ -60,10 +61,20 @@ FAutoConsoleCommand FWorldPartitionDebugHelper::DebugFilterByDataLayerCommand(
 			UWorld* World = Context.World();
 			if (World && World->IsGameWorld())
 			{
-				TArray<UDataLayerInstance*> DataLayers = UDataLayerSubsystem::ConvertArgsToDataLayers(World, InArgs);
-				for (UDataLayerInstance* DataLayer : DataLayers)
+				if (UWorldPartitionSubsystem* WorldPartitionSubsystem = World->GetSubsystem<UWorldPartitionSubsystem>())
 				{
-					DebugDataLayerFilter.Add(DataLayer->GetDataLayerFName());
+					WorldPartitionSubsystem->ForEachWorldPartition([&InArgs](UWorldPartition* WorldPartition)
+					{
+						if (UDataLayerManager* DataLayerManager = WorldPartition->GetDataLayerManager())
+						{
+							TArray<UDataLayerInstance*> DataLayerInstances = DataLayerManager->ConvertArgsToDataLayers(InArgs);
+							for (UDataLayerInstance* DataLayerInstance : DataLayerInstances)
+							{
+								DebugDataLayerFilter.Add(DataLayerInstance->GetDataLayerFName());
+							}
+						}
+						return true;
+					});
 				}
 			}
 		}
@@ -224,4 +235,4 @@ FLinearColor FWorldPartitionDebugHelper::GetHeatMapColor(float ValueNormalized)
 	// Interpolate between the 2 colors
 	FLinearColor Result = Colors[idx1] + Fraction * (Colors[idx2] - Colors[idx1]);
 	return Result;
-}
+}
\ No newline at end of file
diff --git a/Engine/Source/Runtime/Engine/Private/WorldPartition/WorldPartitionRuntimeCell.cpp b/Engine/Source/Runtime/Engine/Private/WorldPartition/WorldPartitionRuntimeCell.cpp
--- a/Engine/Source/Runtime/Engine/Private/WorldPartition/WorldPartitionRuntimeCell.cpp
+++ b/Engine/Source/Runtime/Engine/Private/WorldPartition/WorldPartitionRuntimeCell.cpp
@@ -5,7 +5,7 @@
 #include "Engine/Level.h"
 #include "Misc/HierarchicalLogArchive.h"
 #include "WorldPartition/WorldPartitionDebugHelper.h"
-#include "WorldPartition/DataLayer/DataLayerSubsystem.h"
+#include "WorldPartition/DataLayer/DataLayerManager.h"
 #include "WorldPartition/DataLayer/DataLayersID.h"
 #include "WorldPartition/ContentBundle/ContentBundleDescriptor.h"
 #include "WorldPartition/DataLayer/WorldDataLayers.h"
@@ -95,27 +95,23 @@ FLinearColor UWorldPartitionRuntimeCell::GetDebugStreamingPriorityColor() const
 	return FLinearColor::Transparent;
 }
 
-TArray<const UDataLayerInstance*> UWorldPartitionRuntimeCell::GetDataLayerInstances() const
+bool UWorldPartitionRuntimeCell::HasAnyDataLayerInEffectiveRuntimeState(EDataLayerRuntimeState InState) const
 {
-	if (const UDataLayerSubsystem* DataLayerSubsystem = UWorld::GetSubsystem<UDataLayerSubsystem>(GetWorld()))
-	{
-		return DataLayerSubsystem->GetDataLayerInstances(GetDataLayers());
-	}
+	const UDataLayerManager* DataLayerManager = HasDataLayers() ? UDataLayerManager::GetDataLayerManager(this) : nullptr;
+	return DataLayerManager ? DataLayerManager->IsAnyDataLayerInEffectiveRuntimeState(GetDataLayers(), InState) : false;
+}
 
-	return TArray<const UDataLayerInstance*>();
+TArray<const UDataLayerInstance*> UWorldPartitionRuntimeCell::GetDataLayerInstances() const
+{
+	const UDataLayerManager* DataLayerManager = HasDataLayers() ? UDataLayerManager::GetDataLayerManager(this) : nullptr;
+	return DataLayerManager ? DataLayerManager->GetDataLayerInstances(GetDataLayers()) : TArray<const UDataLayerInstance*>();
 }
 
 bool UWorldPartitionRuntimeCell::ContainsDataLayer(const UDataLayerAsset* DataLayerAsset) const
 {
-	if (const UDataLayerSubsystem* DataLayerSubsystem = UWorld::GetSubsystem<UDataLayerSubsystem>(GetWorld()))
-	{
-		if (const UDataLayerInstance* DataLayerInstance = DataLayerSubsystem->GetDataLayerInstance(DataLayerAsset))
-		{
-			return ContainsDataLayer(DataLayerInstance);
-		}
-	}
-
-	return false;
+	const UDataLayerManager* DataLayerManager = HasDataLayers() ? UDataLayerManager::GetDataLayerManager(this) : nullptr;
+	const UDataLayerInstance* DataLayerInstance = DataLayerManager ? DataLayerManager->GetDataLayerInstance(DataLayerAsset) : nullptr;
+	return DataLayerInstance ? ContainsDataLayer(DataLayerInstance) : false;
 }
 
 bool UWorldPartitionRuntimeCell::ContainsDataLayer(const UDataLayerInstance* DataLayerInstance) const
