- Added UDataLayerManager which is local to each WorldPartition.

diff --git a/Engine/Source/Runtime/Engine/Private/WorldPartition/WorldPartitionRuntimeSpatialHash.cpp b/Engine/Source/Runtime/Engine/Private/WorldPartition/WorldPartitionRuntimeSpatialHash.cpp
--- a/Engine/Source/Runtime/Engine/Private/WorldPartition/WorldPartitionRuntimeSpatialHash.cpp
+++ b/Engine/Source/Runtime/Engine/Private/WorldPartition/WorldPartitionRuntimeSpatialHash.cpp
@@ -4,33 +4,34 @@
 	WorldPartitionRuntimeSpatialHash.cpp: UWorldPartitionRuntimeSpatialHash implementation
 =============================================================================*/
 #include "WorldPartition/WorldPartitionRuntimeSpatialHash.h"
-
-#include "Engine/Engine.h"
-#include "Math/TranslationMatrix.h"
-#include "Misc/ArchiveMD5.h"
-#include "Misc/HierarchicalLogArchive.h"
+#include "WorldPartition/WorldPartitionRuntimeCellDataSpatialHash.h"
+#include "WorldPartition/WorldPartitionRuntimeSpatialHashGridPreviewer.h"
+#include "WorldPartition/WorldPartitionRuntimeHash.h"
 #include "WorldPartition/WorldPartitionStreamingPolicy.h"
-#include "UObject/UnrealType.h"
-#include "Engine/Level.h"
-#include "Engine/Canvas.h"
-#include "DrawDebugHelpers.h"
-#include "GlobalRenderResources.h"
-#include "UObject/ObjectSaveContext.h"
-#include "Components/LineBatchComponent.h"
-#include "ProfilingDebugging/ScopedTimers.h"
 #include "WorldPartition/WorldPartitionDebugHelper.h"
-#include "WorldPartition/HLOD/HLODLayer.h"
 #include "WorldPartition/WorldPartitionLevelStreamingDynamic.h"
+#include "WorldPartition/WorldPartitionLog.h"
+#include "WorldPartition/WorldPartitionStreamingGeneration.h"
+#include "WorldPartition/WorldPartitionSubsystem.h"
 #include "WorldPartition/RuntimeSpatialHash/RuntimeSpatialHashGridHelper.h"
-#include "WorldPartition/DataLayer/DataLayerSubsystem.h"
 #include "WorldPartition/ContentBundle/ContentBundleWorldSubsystem.h"
 #include "WorldPartition/ContentBundle/ContentBundleDescriptor.h"
 #include "WorldPartition/ContentBundle/ContentBundleBase.h"
-#include "WorldPartition/WorldPartitionLog.h"
-#include "WorldPartition/WorldPartitionRuntimeHash.h"
-#include "WorldPartition/WorldPartitionRuntimeSpatialHashGridPreviewer.h"
-#include "WorldPartition/WorldPartitionStreamingGeneration.h"
-#include "WorldPartition/WorldPartitionRuntimeCellDataSpatialHash.h"
+#include "WorldPartition/DataLayer/DataLayerManager.h"
+#include "WorldPartition/DataLayer/DataLayerInstance.h"
+#include "WorldPartition/HLOD/HLODLayer.h"
+#include "Components/LineBatchComponent.h"
+#include "Engine/Engine.h"
+#include "Engine/Level.h"
+#include "Engine/Canvas.h"
+#include "DrawDebugHelpers.h"
+#include "GlobalRenderResources.h"
+#include "Math/TranslationMatrix.h"
+#include "Misc/ArchiveMD5.h"
+#include "Misc/HierarchicalLogArchive.h"
+#include "ProfilingDebugging/ScopedTimers.h"
+#include "UObject/UnrealType.h"
+#include "UObject/ObjectSaveContext.h"
 
 #include UE_INLINE_GENERATED_CPP_BY_NAME(WorldPartitionRuntimeSpatialHash)
 
@@ -206,7 +207,7 @@ void FSpatialHashStreamingGrid::GetCells(const FWorldPartitionStreamingQuerySour
 	}
 }
 
-void FSpatialHashStreamingGrid::GetCells(const TArray<FWorldPartitionStreamingSource>& Sources, const UDataLayerSubsystem* DataLayerSubsystem, UWorldPartitionRuntimeHash::FStreamingSourceCells& OutActivateCells, UWorldPartitionRuntimeHash::FStreamingSourceCells& OutLoadCells, bool bEnableZCulling) const
+void FSpatialHashStreamingGrid::GetCells(const TArray<FWorldPartitionStreamingSource>& Sources, UWorldPartitionRuntimeHash::FStreamingSourceCells& OutActivateCells, UWorldPartitionRuntimeHash::FStreamingSourceCells& OutLoadCells, bool bEnableZCulling) const
 {
 	TRACE_CPUPROFILER_EVENT_SCOPE(FSpatialHashStreamingGrid::GetCells);
 
@@ -244,7 +245,7 @@ void FSpatialHashStreamingGrid::GetCells(const TArray<FWorldPartitionStreamingSo
 						const FVector2D CellMinMaxZ(Cell->GetContentBounds().Min.Z, Cell->GetContentBounds().Max.Z);
 						if (!bEnableZCulling || TRange<double>(CellMinMaxZ.X, CellMinMaxZ.Y).Overlaps(TRange<double>(Shape.GetCenter().Z - Shape.GetRadius(), Shape.GetCenter().Z + Shape.GetRadius())))
 						{
-							if (!Cell->HasDataLayers() || (DataLayerSubsystem && DataLayerSubsystem->IsAnyDataLayerInEffectiveRuntimeState(Cell->GetDataLayers(), EDataLayerRuntimeState::Activated)))
+							if (!Cell->HasDataLayers() || Cell->HasAnyDataLayerInEffectiveRuntimeState(EDataLayerRuntimeState::Activated))
 							{
 								if (Source.TargetState == EStreamingSourceTargetState::Loaded)
 								{
@@ -257,7 +258,7 @@ void FSpatialHashStreamingGrid::GetCells(const TArray<FWorldPartitionStreamingSo
 									bAddedActivatedCell = !GRuntimeSpatialHashUseAlignedGridLevels && GRuntimeSpatialHashSnapNonAlignedGridLevelsToLowerLevels;
 								}
 							}
-							else if (DataLayerSubsystem && DataLayerSubsystem->IsAnyDataLayerInEffectiveRuntimeState(Cell->GetDataLayers(), EDataLayerRuntimeState::Loaded))
+							else if (Cell->HasAnyDataLayerInEffectiveRuntimeState(EDataLayerRuntimeState::Loaded))
 							{
 								OutLoadCells.AddCell(Cell, Source, Shape);
 							}
@@ -272,7 +273,7 @@ void FSpatialHashStreamingGrid::GetCells(const TArray<FWorldPartitionStreamingSo
 		});
 	}
 
-	GetNonSpatiallyLoadedCells(DataLayerSubsystem, OutActivateCells.GetCells(), OutLoadCells.GetCells());
+	GetNonSpatiallyLoadedCells(OutActivateCells.GetCells(), OutLoadCells.GetCells());
 
 	if (!GRuntimeSpatialHashUseAlignedGridLevels && GRuntimeSpatialHashSnapNonAlignedGridLevelsToLowerLevels)
 	{
@@ -329,7 +330,7 @@ void FSpatialHashStreamingGrid::GetCells(const TArray<FWorldPartitionStreamingSo
 		{
 			ForEachRuntimeCell(ParentCell.Key, [&](const UWorldPartitionRuntimeCell* Cell)
 			{
-				if (!Cell->HasDataLayers() || (DataLayerSubsystem && DataLayerSubsystem->IsAnyDataLayerInEffectiveRuntimeState(Cell->GetDataLayers(), EDataLayerRuntimeState::Activated)))
+				if (!Cell->HasDataLayers() || Cell->HasAnyDataLayerInEffectiveRuntimeState(EDataLayerRuntimeState::Activated))
 				{
 					for (const auto& Info : ParentCell.Value)
 					{
@@ -518,22 +519,22 @@ void FSpatialHashStreamingGrid::ForEachRuntimeCell(TFunctionRef<bool(const UWorl
 	ForEachGridLevel(InjectedGridLevels);
 }
 
-void FSpatialHashStreamingGrid::GetNonSpatiallyLoadedCells(const UDataLayerSubsystem* DataLayerSubsystem, TSet<const UWorldPartitionRuntimeCell*>& OutActivateCells, TSet<const UWorldPartitionRuntimeCell*>& OutLoadCells) const
+void FSpatialHashStreamingGrid::GetNonSpatiallyLoadedCells(TSet<const UWorldPartitionRuntimeCell*>& OutActivateCells, TSet<const UWorldPartitionRuntimeCell*>& OutLoadCells) const
 {
 	if (GridLevels.Num() > 0)
 	{
-		auto ForEachGridLevelCells = [DataLayerSubsystem, &OutActivateCells, &OutLoadCells](const FSpatialHashStreamingGridLevel& InGridLevel)
+		auto ForEachGridLevelCells = [&OutActivateCells, &OutLoadCells](const FSpatialHashStreamingGridLevel& InGridLevel)
 		{
 			for (const FSpatialHashStreamingGridLayerCell& LayerCell : InGridLevel.LayerCells)
 			{
 				for (const UWorldPartitionRuntimeCell* Cell : LayerCell.GridCells)
 				{
-					if (!Cell->HasDataLayers() || (DataLayerSubsystem && DataLayerSubsystem->IsAnyDataLayerInEffectiveRuntimeState(Cell->GetDataLayers(), EDataLayerRuntimeState::Activated)))
+					if (!Cell->HasDataLayers() || Cell->HasAnyDataLayerInEffectiveRuntimeState(EDataLayerRuntimeState::Activated))
 					{
 						check(Cell->IsAlwaysLoaded() || Cell->HasDataLayers() || Cell->GetContentBundleID().IsValid());
 						OutActivateCells.Add(Cell);
 					}
-					else if (DataLayerSubsystem && DataLayerSubsystem->IsAnyDataLayerInEffectiveRuntimeState(Cell->GetDataLayers(), EDataLayerRuntimeState::Loaded))
+					else if (Cell->HasAnyDataLayerInEffectiveRuntimeState(EDataLayerRuntimeState::Loaded))
 					{
 						check(Cell->HasDataLayers());
 						OutLoadCells.Add(Cell);
@@ -551,19 +552,19 @@ void FSpatialHashStreamingGrid::GetNonSpatiallyLoadedCells(const UDataLayerSubsy
 	}
 }
 
-void FSpatialHashStreamingGrid::GetFilteredCellsForDebugDraw(const FSpatialHashStreamingGridLayerCell* LayerCell, const UDataLayerSubsystem* DataLayerSubsystem, TArray<const UWorldPartitionRuntimeCell*>& FilteredCells) const
+void FSpatialHashStreamingGrid::GetFilteredCellsForDebugDraw(const FSpatialHashStreamingGridLayerCell* LayerCell, TArray<const UWorldPartitionRuntimeCell*>& FilteredCells) const
 {
 	if (LayerCell)
 	{
-		Algo::CopyIf(LayerCell->GridCells, FilteredCells, [DataLayerSubsystem](const UWorldPartitionRuntimeCell* GridCell)
+		Algo::CopyIf(LayerCell->GridCells, FilteredCells, [](const UWorldPartitionRuntimeCell* Cell)
 		{
-			if (GridCell->IsDebugShown())
+			if (Cell->IsDebugShown())
 			{
-				EStreamingStatus StreamingStatus = GridCell->GetStreamingStatus();
-				const TArray<FName>& DataLayers = GridCell->GetDataLayers();
-				return (!DataLayers.Num() ||
-					DataLayerSubsystem->IsAnyDataLayerInEffectiveRuntimeState(DataLayers, EDataLayerRuntimeState::Loaded) ||
-					DataLayerSubsystem->IsAnyDataLayerInEffectiveRuntimeState(DataLayers, EDataLayerRuntimeState::Activated) ||
+				EStreamingStatus StreamingStatus = Cell->GetStreamingStatus();
+				const TArray<FName>& DataLayers = Cell->GetDataLayers();
+				return (!Cell->HasDataLayers() ||
+					Cell->HasAnyDataLayerInEffectiveRuntimeState(EDataLayerRuntimeState::Loaded) ||
+					Cell->HasAnyDataLayerInEffectiveRuntimeState(EDataLayerRuntimeState::Activated) ||
 					((StreamingStatus != LEVEL_Unloaded) && (StreamingStatus != LEVEL_UnloadedButStillAround)));
 			}
 			return false;
@@ -577,16 +578,27 @@ EWorldPartitionRuntimeCellVisualizeMode FSpatialHashStreamingGrid::GetStreamingC
 	return VisualizeMode;
 }
 
-void FSpatialHashStreamingGrid::Draw3D(UWorld* World, const TArray<FWorldPartitionStreamingSource>& Sources, const FTransform& Transform) const
+static TMap<FName, FColor> GetDataLayerDebugColors(const UWorldPartition* InWorldPartition)
 {
-	const EWorldPartitionRuntimeCellVisualizeMode VisualizeMode = GetStreamingCellVisualizeMode();
-	const UContentBundleManager* ContentBundleManager = World->ContentBundleManager;
-	const UDataLayerSubsystem* DataLayerSubsystem = World->GetSubsystem<UDataLayerSubsystem>();
-	TMap<FName, FColor> DataLayerDebugColors;
-	if (DataLayerSubsystem)
+	TMap<FName, FColor> DebugColors;
+	if (const UDataLayerManager* DataLayerManager = InWorldPartition->GetDataLayerManager())
 	{
-		DataLayerSubsystem->GetDataLayerDebugColors(DataLayerDebugColors);
+		DataLayerManager->ForEachDataLayerInstance([&DebugColors](UDataLayerInstance* DataLayerInstance)
+		{
+			DebugColors.Add(DataLayerInstance->GetDataLayerFName(), DataLayerInstance->GetDebugColor());
+			return true;
+		});
 	}
+	return DebugColors;
+}
+
+void FSpatialHashStreamingGrid::Draw3D(const UWorldPartitionRuntimeSpatialHash* Owner, const TArray<FWorldPartitionStreamingSource>& Sources, const FTransform& Transform) const
+{
+	const UWorldPartition* WorldPartition = Owner->GetOuterUWorldPartition();
+	UWorld* OwningWorld = WorldPartition->GetWorld();
+	const EWorldPartitionRuntimeCellVisualizeMode VisualizeMode = GetStreamingCellVisualizeMode();
+	const UContentBundleManager* ContentBundleManager = OwningWorld->ContentBundleManager;
+	TMap<FName, FColor> DataLayerDebugColors = GetDataLayerDebugColors(WorldPartition);
 
 	const FSquare2DGridHelper& Helper = GetGridHelper();
 	int32 MinGridLevel = FMath::Clamp<int32>(GShowRuntimeSpatialHashGridLevel, 0, GridLevels.Num() - 1);
@@ -611,7 +623,7 @@ void FSpatialHashStreamingGrid::Draw3D(UWorld* World, const TArray<FWorldPartiti
 				FilteredCells.Reset();
 				ForEachLayerCell(Coords, [&](const FSpatialHashStreamingGridLayerCell* LayerCell)
 				{
-					GetFilteredCellsForDebugDraw(LayerCell, DataLayerSubsystem, FilteredCells);
+					GetFilteredCellsForDebugDraw(LayerCell, FilteredCells);
 				});
 				for (const UWorldPartitionRuntimeCell* Cell : FilteredCells)
 				{
@@ -627,9 +639,9 @@ void FSpatialHashStreamingGrid::Draw3D(UWorld* World, const TArray<FWorldPartiti
 
 					// Draw Cell using its debug color
 					FColor CellColor = Cell->GetDebugColor(VisualizeMode).ToFColor(false).WithAlpha(16);
-					DrawDebugSolidBox(World, CellBox, CellColor, Transform, false, -1.f, 255);
+					DrawDebugSolidBox(OwningWorld, CellBox, CellColor, Transform, false, -1.f, 255);
 					FVector CellPos = Transform.TransformPosition(CellBox.GetCenter());
-					DrawDebugBox(World, CellPos, CellBox.GetExtent(), Transform.GetRotation(), CellColor.WithAlpha(255), false, -1.f, 255, 10.f);
+					DrawDebugBox(OwningWorld, CellPos, CellBox.GetExtent(), Transform.GetRotation(), CellColor.WithAlpha(255), false, -1.f, 255, 10.f);
 
 					// Draw Cell's DataLayers colored boxes
 					if (DataLayerDebugColors.Num() && Cell->GetDataLayers().Num() > 0)
@@ -641,7 +653,7 @@ void FSpatialHashStreamingGrid::Draw3D(UWorld* World, const TArray<FWorldPartiti
 						for (const FName& DataLayer : Cell->GetDataLayers())
 						{
 							const FColor& DataLayerColor = DataLayerDebugColors[DataLayer];
-							DrawDebugSolidBox(World, DataLayerColoredBox, DataLayerColor, Transform, false, -1.f, 255);
+							DrawDebugSolidBox(OwningWorld, DataLayerColoredBox, DataLayerColor, Transform, false, -1.f, 255);
 							DataLayerColoredBox = DataLayerColoredBox.TransformBy(DataLayerOffsetMatrix);
 						}
 					}
@@ -649,7 +661,7 @@ void FSpatialHashStreamingGrid::Draw3D(UWorld* World, const TArray<FWorldPartiti
 					// Draw Content Bundle colored boxes
 					if (ContentBundleManager && FWorldPartitionDebugHelper::CanDrawContentBundles() && Cell->GetContentBundleID().IsValid())
 					{
-						if (const FContentBundleBase* ContentBundle = ContentBundleManager->GetContentBundle(World, Cell->GetContentBundleID()))
+						if (const FContentBundleBase* ContentBundle = ContentBundleManager->GetContentBundle(OwningWorld, Cell->GetContentBundleID()))
 						{
 							check(ContentBundle->GetDescriptor());
 							FBox ContentBundleColoredBox = CellBox;
@@ -657,7 +669,7 @@ void FSpatialHashStreamingGrid::Draw3D(UWorld* World, const TArray<FWorldPartiti
 							ContentBundleColoredBox.Max.X = ContentBundleColoredBox.Min.X + ContentBundleSizeX;
 							FTranslationMatrix ContentBundleOffsetMatrix(FVector(ContentBundleColoredBox.GetSize().X - ContentBundleSizeX, 0, 0));
 							ContentBundleColoredBox = ContentBundleColoredBox.TransformBy(ContentBundleOffsetMatrix);
-							DrawDebugSolidBox(World, ContentBundleColoredBox, ContentBundle->GetDescriptor()->GetDebugColor(), Transform, false, -1.f, 255);
+							DrawDebugSolidBox(OwningWorld, ContentBundleColoredBox, ContentBundle->GetDescriptor()->GetDebugColor(), Transform, false, -1.f, 255);
 						}
 					}
 				}
@@ -672,18 +684,18 @@ void FSpatialHashStreamingGrid::Draw3D(UWorld* World, const TArray<FWorldPartiti
 			FVector StartTrace = Source.Location + FVector(0.f, 0.f, 100.f);
 			FVector EndTrace = StartTrace - FVector(0.f, 0.f, 1000000.f);
 			FHitResult Hit;
-			if (World->LineTraceSingleByObjectType(Hit, StartTrace, EndTrace, FCollisionObjectQueryParams(ECC_WorldStatic), FCollisionQueryParams(SCENE_QUERY_STAT(DebugWorldPartitionTrace), true)))
+			if (OwningWorld->LineTraceSingleByObjectType(Hit, StartTrace, EndTrace, FCollisionObjectQueryParams(ECC_WorldStatic), FCollisionQueryParams(SCENE_QUERY_STAT(DebugWorldPartitionTrace), true)))
 			{
 				SourceLocationZ = Hit.ImpactPoint.Z;
 			}
 		}
 
 		const FColor Color = Source.GetDebugColor();
-		Source.ForEachShape(GetLoadingRange(), GridName, HLODLayer, /*bProjectIn2D*/ true, [&Color, &SourceLocationZ, &Transform, &World, this](const FSphericalSector& Shape)
+		Source.ForEachShape(GetLoadingRange(), GridName, HLODLayer, /*bProjectIn2D*/ true, [&Color, &SourceLocationZ, &Transform, &OwningWorld, this](const FSphericalSector& Shape)
 		{
 			FSphericalSector ZOffsettedShape = Shape;
 			ZOffsettedShape.SetCenter(FVector(FVector2D(ZOffsettedShape.GetCenter()), SourceLocationZ));
-			DrawStreamingSource3D(World, ZOffsettedShape, Transform, Color);
+			DrawStreamingSource3D(OwningWorld, ZOffsettedShape, Transform, Color);
 		});
 	}
 }
@@ -716,18 +728,15 @@ void FSpatialHashStreamingGrid::DrawStreamingSource3D(UWorld* World, const FSphe
 	}
 }
 
-void FSpatialHashStreamingGrid::Draw2D(UCanvas* Canvas, UWorld* World, const TArray<FWorldPartitionStreamingSource>& Sources, const FBox& Region, const FBox2D& GridScreenBounds, TFunctionRef<FVector2D(const FVector2D&)> WorldToScreen) const
+void FSpatialHashStreamingGrid::Draw2D(UCanvas* Canvas, const UWorldPartitionRuntimeSpatialHash* Owner, const TArray<FWorldPartitionStreamingSource>& Sources, const FBox& Region, const FBox2D& GridScreenBounds, TFunctionRef<FVector2D(const FVector2D&)> WorldToScreen) const
 {
 	TRACE_CPUPROFILER_EVENT_SCOPE(FSpatialHashStreamingGrid::Draw2D);
 
+	const UWorldPartition* WorldPartition = Owner->GetOuterUWorldPartition();
+	UWorld* OwningWorld = WorldPartition->GetWorld();
 	const EWorldPartitionRuntimeCellVisualizeMode VisualizeMode = GetStreamingCellVisualizeMode();
-	const UContentBundleManager* ContentBundleManager = World->ContentBundleManager;
-	const UDataLayerSubsystem* DataLayerSubsystem = World->GetSubsystem<UDataLayerSubsystem>();
-	TMap<FName, FColor> DataLayerDebugColors;
-	if (DataLayerSubsystem)
-	{
-		DataLayerSubsystem->GetDataLayerDebugColors(DataLayerDebugColors);
-	}
+	const UContentBundleManager* ContentBundleManager = OwningWorld->ContentBundleManager;
+	TMap<FName, FColor> DataLayerDebugColors = GetDataLayerDebugColors(WorldPartition);
 
 	FCanvas* CanvasObject = Canvas->Canvas;
 	int32 MinGridLevel = FMath::Clamp<int32>(GShowRuntimeSpatialHashGridLevel, 0, GridLevels.Num() - 1);
@@ -777,7 +786,7 @@ void FSpatialHashStreamingGrid::Draw2D(UCanvas* Canvas, UWorld* World, const TAr
 			FilteredCells.Reset();
 			ForEachLayerCell(Coords, [&](const FSpatialHashStreamingGridLayerCell* LayerCell)
 			{
-				GetFilteredCellsForDebugDraw(LayerCell, DataLayerSubsystem, FilteredCells);
+				GetFilteredCellsForDebugDraw(LayerCell, FilteredCells);
 			});
 
 			if (FilteredCells.Num())
@@ -812,7 +821,7 @@ void FSpatialHashStreamingGrid::Draw2D(UCanvas* Canvas, UWorld* World, const TAr
 					// Draw Content Bundle colored boxes
 					if (ContentBundleManager && FWorldPartitionDebugHelper::CanDrawContentBundles() && Cell->GetContentBundleID().IsValid())
 					{
-						if (const FContentBundleBase* ContentBundle = ContentBundleManager->GetContentBundle(World, Cell->GetContentBundleID()))
+						if (const FContentBundleBase* ContentBundle = ContentBundleManager->GetContentBundle(OwningWorld, Cell->GetContentBundleID()))
 						{
 							check(ContentBundle->GetDescriptor());
 							FVector2D BoxSize = CellBoundsSize;
@@ -1326,11 +1335,30 @@ bool UWorldPartitionRuntimeSpatialHash::CreateStreamingGrid(const FSpatialHashRu
 					{
 						for (const FGuid& ActorGuid : ActorSetInstance->ActorSet->Actors)
 						{
-							FilteredActors.Emplace(ActorGuid, ActorSetInstance);
+							// Instanced world partition, ContainerID is the main container, but it's not the main world partition,
+							// so the always loaded actors don't be part of the process of ForceExternalActorLevelReference/AlwaysLoadedActorsForPIE.
+							// In PIE, always loaded actors of an instanced world partition will go in the always loaded cell.
+							// 
+							// WorldDataLayers::ShouldLevelKeepRefIfExternal return true, so it will always be part of the partitioned 
+							// persistent level of a world partition (ideally we would need this information to be part of the ActorDesc).
+							// 
+							// In PIE, for an instanced world partition, we don't want this WorldDataLayers actor to be both in the persistent level 
+							// and also part of the always loaded cell level.
+							//
+							// @todo_ow: We need to implement PIE always loaded actors of instanced world partitions to be part
+							//			 of the persistent level and get rid of the always loaded cell (to have the same behavior
+							//			 as non-instanced world partition and as cooked world partition).
+							//
+							// For now, we force skip AWorldDataLayers actors.
+							IStreamingGenerationContext::FActorInstance ActorInstance(ActorGuid, ActorSetInstance);
+							if (!ActorInstance.GetActorDescView().GetActorNativeClass()->IsChildOf<AWorldDataLayers>())
+							{
+								FilteredActors.Emplace(ActorGuid, ActorSetInstance);
+							}
 						}
 					}
 				}
-
+				
 				if (!FilteredActors.Num())
 				{
 					continue;
@@ -1375,13 +1403,13 @@ bool UWorldPartitionRuntimeSpatialHash::CreateStreamingGrid(const FSpatialHashRu
 
 					if (DataLayerCount > 0)
 					{
-						if (const UDataLayerSubsystem* DataLayerSubsystem = UWorld::GetSubsystem<UDataLayerSubsystem>(World))
+						if (const UDataLayerManager* DataLayerManager = UDataLayerManager::GetDataLayerManager(StreamingCell))
 						{
 							TArray<const UDataLayerInstance*> DataLayerObjects;
 							Builder += TEXT(" DL[");
 							for (int i = 0; i < DataLayerCount; ++i)
 							{
-								const UDataLayerInstance* DataLayer = DataLayerSubsystem->GetDataLayerInstance(DataLayers[i]);
+								const UDataLayerInstance* DataLayer = DataLayerManager->GetDataLayerInstance(DataLayers[i]);
 								DataLayerObjects.Add(DataLayer);
 								Builder += DataLayer->GetDataLayerShortName();
 								Builder += TEXT(",");
@@ -1465,15 +1493,19 @@ FAutoConsoleCommand UWorldPartitionRuntimeSpatialHash::OverrideLoadingRangeComma
 			UWorld* World = Context.World();
 			if (World && World->IsGameWorld())
 			{
-				if (UWorldPartition* WorldPartition = World->GetWorldPartition())
+				if (UWorldPartitionSubsystem* WorldPartitionSubsystem = World->GetSubsystem<UWorldPartitionSubsystem>())
 				{
-					if (UWorldPartitionRuntimeSpatialHash* RuntimeSpatialHash = Cast<UWorldPartitionRuntimeSpatialHash>(WorldPartition->RuntimeHash))
+					WorldPartitionSubsystem->ForEachWorldPartition([GridIndex, OverrideLoadingRange](UWorldPartition* WorldPartition)
 					{
-						if (RuntimeSpatialHash->StreamingGrids.IsValidIndex(GridIndex))
+						if (UWorldPartitionRuntimeSpatialHash* RuntimeSpatialHash = Cast<UWorldPartitionRuntimeSpatialHash>(WorldPartition->RuntimeHash))
 						{
-							RuntimeSpatialHash->StreamingGrids[GridIndex].OverrideLoadingRange = OverrideLoadingRange;
+							if (RuntimeSpatialHash->StreamingGrids.IsValidIndex(GridIndex))
+							{
+								RuntimeSpatialHash->StreamingGrids[GridIndex].OverrideLoadingRange = OverrideLoadingRange;
+							}
 						}
-					}
+						return true;
+					});
 				}
 			}
 		}
@@ -1518,8 +1550,6 @@ void UWorldPartitionRuntimeSpatialHash::ForEachStreamingCellsQuery(const FWorldP
 
 void UWorldPartitionRuntimeSpatialHash::ForEachStreamingCellsSources(const TArray<FWorldPartitionStreamingSource>& Sources, TFunctionRef<bool(const UWorldPartitionRuntimeCell*, EStreamingSourceTargetState)> Func) const
 {
-	const UDataLayerSubsystem* DataLayerSubsystem = GetWorld()->GetSubsystem<UDataLayerSubsystem>();
-
 	FStreamingSourceCells ActivateStreamingSourceCells;
 	FStreamingSourceCells LoadStreamingSourceCells;
 
@@ -1530,7 +1560,7 @@ void UWorldPartitionRuntimeSpatialHash::ForEachStreamingCellsSources(const TArra
 		{
 			if (IsCellRelevantFor(StreamingGrid.bClientOnlyVisible))
 			{
-				StreamingGrid.GetNonSpatiallyLoadedCells(DataLayerSubsystem, ActivateStreamingSourceCells.GetCells(), LoadStreamingSourceCells.GetCells());
+				StreamingGrid.GetNonSpatiallyLoadedCells(ActivateStreamingSourceCells.GetCells(), LoadStreamingSourceCells.GetCells());
 			}
 		});
 	}
@@ -1541,7 +1571,7 @@ void UWorldPartitionRuntimeSpatialHash::ForEachStreamingCellsSources(const TArra
 		{
 			if (IsCellRelevantFor(StreamingGrid.bClientOnlyVisible))
 			{
-				StreamingGrid.GetCells(Sources, DataLayerSubsystem, ActivateStreamingSourceCells, LoadStreamingSourceCells, GetEffectiveEnableZCulling(bEnableZCulling));
+				StreamingGrid.GetCells(Sources, ActivateStreamingSourceCells, LoadStreamingSourceCells, GetEffectiveEnableZCulling(bEnableZCulling));
 			}
 		});
 	}
@@ -1715,8 +1745,6 @@ bool UWorldPartitionRuntimeSpatialHash::Draw2D(UCanvas* Canvas, const TArray<FWo
 	}
 
 	UWorldPartition* WorldPartition = GetOuterUWorldPartition();
-	UWorld* World = WorldPartition->GetWorld();
-
 	const float CanvasMaxScreenSize = PartitionCanvasSize.X;
 	const float GridMaxScreenSize = CanvasMaxScreenSize / FilteredStreamingGrids.Num();
 	const float GridEffectiveScreenRatio = 1.f;
@@ -1753,7 +1781,7 @@ bool UWorldPartitionRuntimeSpatialHash::Draw2D(UCanvas* Canvas, const TArray<FWo
 		const float WorldToScreenScale = (0.5f * GridEffectiveScreenSize) / RegionExtent;
 		auto WorldToScreen = [&](const FVector2D& WorldPos) { return (WorldToScreenScale * (WorldPos - GridReferenceWorldPos)) + GridScreenOffset; };
 
-		StreamingGrid->Draw2D(Canvas, World, Sources, Region, GridScreenBounds, WorldToScreen);
+		StreamingGrid->Draw2D(Canvas, this, Sources, Region, GridScreenBounds, WorldToScreen);
 		GridsBounds += GridScreenBounds;
 
 		// Draw WorldPartition name
@@ -1797,14 +1825,12 @@ bool UWorldPartitionRuntimeSpatialHash::Draw2D(UCanvas* Canvas, const TArray<FWo
 
 void UWorldPartitionRuntimeSpatialHash::Draw3D(const TArray<FWorldPartitionStreamingSource>& Sources) const
 {
-	UWorld* World = GetWorld();
 	UWorldPartition* WorldPartition = GetOuterUWorldPartition();
-
 	FTransform Transform = WorldPartition->GetInstanceTransform();
 	TArray<const FSpatialHashStreamingGrid*> FilteredStreamingGrids = GetFilteredStreamingGrids();
 	for (const FSpatialHashStreamingGrid* StreamingGrid : FilteredStreamingGrids)
 	{
-		StreamingGrid->Draw3D(World, Sources, Transform);
+		StreamingGrid->Draw3D(this, Sources, Transform);
 	}
 }
 
diff --git a/Engine/Source/Runtime/Engine/Private/WorldPartition/WorldPartitionStreamingGeneration.cpp b/Engine/Source/Runtime/Engine/Private/WorldPartition/WorldPartitionStreamingGeneration.cpp
--- a/Engine/Source/Runtime/Engine/Private/WorldPartition/WorldPartitionStreamingGeneration.cpp
+++ b/Engine/Source/Runtime/Engine/Private/WorldPartition/WorldPartitionStreamingGeneration.cpp
@@ -13,7 +13,7 @@
 #include "Misc/Paths.h"
 #include "WorldPartition/WorldPartitionRuntimeHash.h"
 #include "WorldPartition/WorldPartitionStreamingPolicy.h"
-#include "WorldPartition/DataLayer/DataLayerSubsystem.h"
+#include "WorldPartition/DataLayer/DataLayerManager.h"
 #include "WorldPartition/DataLayer/DataLayerUtils.h"
 #include "WorldPartition/ErrorHandling/WorldPartitionStreamingGenerationNullErrorHandler.h"
 #include "WorldPartition/ErrorHandling/WorldPartitionStreamingGenerationLogErrorHandler.h"
@@ -140,7 +140,7 @@ class FWorldPartitionStreamingGenerator
 						// Main container will get inherited properties from the actor descriptors
 						ActorSetInstance.RuntimeGrid = ReferenceActorDescView.GetRuntimeGrid();
 						ActorSetInstance.bIsSpatiallyLoaded = ReferenceActorDescView.GetIsSpatiallyLoaded();
-						ActorSetInstance.DataLayers = UDataLayerSubsystem::GetRuntimeDataLayerInstances(ContainerDescriptor.Container->GetWorld(), ReferenceActorDescView.GetRuntimeDataLayers());
+						ActorSetInstance.DataLayers = StreamingGenerator->GetRuntimeDataLayerInstances(ReferenceActorDescView.GetRuntimeDataLayerInstanceNames());
 					}
 					else
 					{
@@ -149,9 +149,8 @@ class FWorldPartitionStreamingGenerator
 						ActorSetInstance.bIsSpatiallyLoaded = ContainerInstanceDescriptor.bIsSpatiallyLoaded;
 
 						TSet<FName> CombinedDataLayers = ContainerInstanceDescriptor.RuntimeDataLayers;
-						CombinedDataLayers.Append(ReferenceActorDescView.GetRuntimeDataLayers());
-
-						ActorSetInstance.DataLayers = UDataLayerSubsystem::GetRuntimeDataLayerInstances(ContainerDescriptor.Container->GetWorld(), CombinedDataLayers.Array());						
+						CombinedDataLayers.Append(ReferenceActorDescView.GetRuntimeDataLayerInstanceNames());
+						ActorSetInstance.DataLayers = StreamingGenerator->GetRuntimeDataLayerInstances(CombinedDataLayers.Array());
 					}
 
 					ActorSetInstance.Bounds.Init();
@@ -237,12 +236,22 @@ class FWorldPartitionStreamingGenerator
 
 	void ResolveRuntimeDataLayers(FWorldPartitionActorDescView& ActorDescView, const FActorDescViewMap& ActorDescViewMap)
 	{
-		TArray<FName> RuntimeDataLayerInstanceNames;
-		RuntimeDataLayerInstanceNames.Reserve(ActorDescView.GetDataLayers().Num());
+		const UDataLayerManager* DataLayerManager = UDataLayerManager::GetDataLayerManager(WorldPartitionContext);
 
-		if (FDataLayerUtils::ResolveRuntimeDataLayerInstanceNames(ActorDescView, ActorDescViewMap, RuntimeDataLayerInstanceNames))
+		// Resolve DataLayerInstanceNames of ActorDescView only when necessary (i.e. when container is a template)
+		if (!ActorDescView.GetActorDesc()->HasResolvedDataLayerInstanceNames())
 		{
-			ActorDescView.SetRuntimeDataLayers(RuntimeDataLayerInstanceNames);
+			// Build a WorldDataLayerActorDescs if DataLayerManager can't resolve Data Layers (i.e. when validating changelists and World is not loaded)
+			const bool bDataLayerManagerCanResolve = DataLayerManager && DataLayerManager->CanResolveDataLayers();
+			const TArray<const FWorldDataLayersActorDesc*> WorldDataLayerActorDescs = !bDataLayerManagerCanResolve ? FDataLayerUtils::FindWorldDataLayerActorDescs(ActorDescViewMap) : TArray<const FWorldDataLayersActorDesc*>();
+			const TArray<FName> DataLayerInstanceNames = FDataLayerUtils::ResolvedDataLayerInstanceNames(DataLayerManager, ActorDescView.GetActorDesc(), WorldDataLayerActorDescs);
+			ActorDescView.SetDataLayerInstanceNames(DataLayerInstanceNames);
+		}
+
+		TArray<FName> RuntimeDataLayerInstanceNames;
+		if (FDataLayerUtils::ResolveRuntimeDataLayerInstanceNames(DataLayerManager, ActorDescView, ActorDescViewMap, RuntimeDataLayerInstanceNames))
+		{
+			ActorDescView.SetRuntimeDataLayerInstanceNames(RuntimeDataLayerInstanceNames);
 		}
 	}
 
@@ -290,7 +299,7 @@ class FWorldPartitionStreamingGenerator
 
 			// Pretend that this actor descriptor belongs to the original container, even if it's not present. It's essentially a proxy
 			// descriptor on top an existing one and at this point no code should require to access the container to resolve it anyways.
-			ModifiedActorDesc->SetContainer(const_cast<UActorDescContainer*>(InContainer));
+			ModifiedActorDesc->SetContainer(const_cast<UActorDescContainer*>(InContainer), InActor->GetWorld());
 
 			return ModifiedActorDesc;
 		};
@@ -346,7 +355,7 @@ class FWorldPartitionStreamingGenerator
 		// Append new unsaved actors for the persistent level
 		if (bHandleUnsavedActors)
 		{
-			for (AActor* Actor : InContainer->GetWorld()->PersistentLevel->Actors)
+			for (AActor* Actor : InContainer->GetWorldPartition()->GetWorld()->PersistentLevel->Actors)
 			{
 				if (IsValid(Actor) && Actor->IsPackageExternal() && Actor->IsMainPackageActor() && !Actor->IsEditorOnly() && (Actor->GetContentBundleGuid() == InContainer->GetContentBundleGuid()) && !InContainer->GetActorDesc(Actor->GetActorGuid()))
 				{
@@ -390,7 +399,7 @@ class FWorldPartitionStreamingGenerator
 			}
 
 			// Data layers are accumulated down the hierarchy chain, since level instances supports data layers assignation on actors
-			InOutRuntimeDataLayers.Append(InParentActorDescView.GetRuntimeDataLayers());
+			InOutRuntimeDataLayers.Append(InParentActorDescView.GetRuntimeDataLayerInstanceNames());
 		};
 
 		// Parse actor containers
@@ -510,19 +519,19 @@ class FWorldPartitionStreamingGenerator
 		if (bIsMainContainer)
 		{
 			TArray<FGuid> LevelScriptReferences;
-			if (UWorld* World = ContainerDescriptor.Container->GetWorld())
+			if (WorldPartitionContext)
 			{
 				// Gather all references to external actors from the level script and make them always loaded
-				if (ULevelScriptBlueprint* LevelScriptBlueprint = World->PersistentLevel->GetLevelScriptBlueprint(true))
+				if (ULevelScriptBlueprint* LevelScriptBlueprint = WorldPartitionContext->GetTypedOuter<UWorld>()->PersistentLevel->GetLevelScriptBlueprint(true))
 				{
 					TArray<AActor*> LevelScriptExternalActorReferences = ActorsReferencesUtils::GetExternalActorReferences(LevelScriptBlueprint);
 					Algo::Transform(LevelScriptExternalActorReferences, LevelScriptReferences, [](const AActor* Actor) { return Actor->GetActorGuid(); });
 				}
 
 				// Validate data layers
-				if (UDataLayerSubsystem* DataLayerSubsystem = UWorld::GetSubsystem<UDataLayerSubsystem>(World))
+				if (UDataLayerManager* DataLayerManager = UDataLayerManager::GetDataLayerManager(WorldPartitionContext))
 				{
-					DataLayerSubsystem->ForEachDataLayer([this](const UDataLayerInstance* DataLayerInstance)
+					DataLayerManager->ForEachDataLayerInstance([this](const UDataLayerInstance* DataLayerInstance)
 					{
 						DataLayerInstance->Validate(ErrorHandler);
 						return true;
@@ -544,7 +553,7 @@ class FWorldPartitionStreamingGenerator
 						ActorDescView->SetForcedNonSpatiallyLoaded();
 					}
 
-					if (ActorDescView->GetRuntimeDataLayers().Num())
+					if (ActorDescView->GetRuntimeDataLayerInstanceNames().Num())
 					{
 						ErrorHandler->OnInvalidReferenceLevelScriptDataLayers(*ActorDescView);
 						ActorDescView->SetInvalidDataLayers();
@@ -583,10 +592,10 @@ class FWorldPartitionStreamingGenerator
 				// Validate grid placement
 				auto IsReferenceDataLayersValid = [](const FWorldPartitionActorDescView& RefererActorDescView, const FWorldPartitionActorDescView& ReferenceActorDescView)
 				{
-					if (RefererActorDescView.GetRuntimeDataLayers().Num() == ReferenceActorDescView.GetRuntimeDataLayers().Num())
+					if (RefererActorDescView.GetRuntimeDataLayerInstanceNames().Num() == ReferenceActorDescView.GetRuntimeDataLayerInstanceNames().Num())
 					{
-						const TSet<FName> RefererActorDescDataLayers(RefererActorDescView.GetRuntimeDataLayers());
-						const TSet<FName> ReferenceActorDescDataLayers(ReferenceActorDescView.GetRuntimeDataLayers());
+						const TSet<FName> RefererActorDescDataLayers(RefererActorDescView.GetRuntimeDataLayerInstanceNames());
+						const TSet<FName> ReferenceActorDescDataLayers(ReferenceActorDescView.GetRuntimeDataLayerInstanceNames());
 
 						return RefererActorDescDataLayers.Includes(ReferenceActorDescDataLayers);
 					}
@@ -796,11 +805,13 @@ public:
 	struct FWorldPartitionStreamingGeneratorParams
 	{
 		FWorldPartitionStreamingGeneratorParams()
-			: ModifiedActorsDescList(nullptr)
+			: WorldPartitionContext(nullptr)
+			, ModifiedActorsDescList(nullptr)
 			, ErrorHandler(&NullErrorHandler)
 			, bEnableStreaming(false)
 		{}
 
+		const UWorldPartition* WorldPartitionContext;
 		FActorDescList* ModifiedActorsDescList;
 		IStreamingGenerationErrorHandler* ErrorHandler;
 		bool bEnableStreaming;
@@ -812,7 +823,8 @@ public:
 	};
 
 	FWorldPartitionStreamingGenerator(const FWorldPartitionStreamingGeneratorParams& Params)
-		: bEnableStreaming(Params.bEnableStreaming)
+		: WorldPartitionContext(Params.WorldPartitionContext)
+		, bEnableStreaming(Params.bEnableStreaming)
 		, ModifiedActorsDescList(Params.ModifiedActorsDescList)
 		, FilteredClasses(Params.FilteredClasses)
 		, IsValidGrid(Params.IsValidGrid)
@@ -929,7 +941,15 @@ public:
 		return StreamingGenerationContext.Get();
 	}
 
+	TArray<const UDataLayerInstance*> GetRuntimeDataLayerInstances(const TArray<FName>& RuntimeDataLayers) const
+	{
+		static TArray<const UDataLayerInstance*> EmptyArray;
+		UDataLayerManager* DataLayerManager = UDataLayerManager::GetDataLayerManager(WorldPartitionContext);
+		return DataLayerManager ? DataLayerManager->GetRuntimeDataLayerInstances(RuntimeDataLayers) : EmptyArray;
+	}
+
 private:
+	const UWorldPartition* WorldPartitionContext;
 	bool bEnableStreaming;
 	FActorDescList* ModifiedActorsDescList;
 	TArray<TSubclassOf<AActor>> FilteredClasses;
@@ -983,6 +1003,7 @@ bool UWorldPartition::GenerateContainerStreaming(const UActorDescContainer* InAc
 	FHierarchicalLogArchive HierarchicalLogAr(*LogFileAr);
 
 	FWorldPartitionStreamingGenerator::FWorldPartitionStreamingGeneratorParams StreamingGeneratorParams;
+	StreamingGeneratorParams.WorldPartitionContext = this;
 	StreamingGeneratorParams.ModifiedActorsDescList = ModifiedActorsDescList;
 	StreamingGeneratorParams.IsValidGrid = [this](FName GridName) { return RuntimeHash->IsValidGrid(GridName); };
 	StreamingGeneratorParams.ErrorHandler = ErrorHandler;
@@ -1032,6 +1053,7 @@ void UWorldPartition::GenerateHLOD(ISourceControlHelper* SourceControlHelper, bo
 
 		FStreamingGenerationLogErrorHandler LogErrorHandler;
 		FWorldPartitionStreamingGenerator::FWorldPartitionStreamingGeneratorParams StreamingGeneratorParams;
+		StreamingGeneratorParams.WorldPartitionContext = this;
 		StreamingGeneratorParams.ErrorHandler = &LogErrorHandler;
 		StreamingGeneratorParams.bEnableStreaming = IsStreamingEnabled();
 		StreamingGeneratorParams.FilteredClasses.Add(AWorldPartitionHLOD::StaticClass());
@@ -1053,6 +1075,7 @@ void UWorldPartition::CheckForErrors(IStreamingGenerationErrorHandler* ErrorHand
 	FActorDescList ModifiedActorDescList;
 
 	FWorldPartitionStreamingGenerator::FWorldPartitionStreamingGeneratorParams StreamingGeneratorParams;
+	StreamingGeneratorParams.WorldPartitionContext = this;
 	StreamingGeneratorParams.ModifiedActorsDescList = &ModifiedActorDescList;
 	StreamingGeneratorParams.ErrorHandler = ErrorHandler;
 	StreamingGeneratorParams.IsValidGrid = [this](FName GridName) { return RuntimeHash->IsValidGrid(GridName); };
@@ -1069,6 +1092,7 @@ void UWorldPartition::CheckForErrors(IStreamingGenerationErrorHandler* ErrorHand
 
 	ForEachActorDescContainer([this, &StreamingGeneratorParams](const UActorDescContainer* InActorDescContainer)
 	{
+		check(StreamingGeneratorParams.WorldPartitionContext == InActorDescContainer->GetWorldPartition());
 		FWorldPartitionStreamingGenerator StreamingGenerator(StreamingGeneratorParams);
 		StreamingGenerator.PreparationPhase(InActorDescContainer);
 	});
@@ -1090,7 +1114,8 @@ void UWorldPartition::CheckForErrors(const FCheckForErrorsParams& Params)
 	FActorDescList ModifiedActorDescList;
 
 	FWorldPartitionStreamingGenerator::FWorldPartitionStreamingGeneratorParams StreamingGeneratorParams;
-	StreamingGeneratorParams.ModifiedActorsDescList = Params.ActorDescContainer->GetWorld() ? &ModifiedActorDescList : nullptr;
+	StreamingGeneratorParams.WorldPartitionContext = Params.ActorDescContainer->GetWorldPartition();
+	StreamingGeneratorParams.ModifiedActorsDescList = !Params.ActorDescContainer->IsTemplateContainer() ? &ModifiedActorDescList : nullptr;
 	StreamingGeneratorParams.ErrorHandler = Params.ErrorHandler;
 	StreamingGeneratorParams.IsValidGrid = [](FName GridName) { return true; };
 	StreamingGeneratorParams.bEnableStreaming = Params.bEnableStreaming;
diff --git a/Engine/Source/Runtime/Engine/Private/WorldPartition/WorldPartitionStreamingPolicy.cpp b/Engine/Source/Runtime/Engine/Private/WorldPartition/WorldPartitionStreamingPolicy.cpp
--- a/Engine/Source/Runtime/Engine/Private/WorldPartition/WorldPartitionStreamingPolicy.cpp
+++ b/Engine/Source/Runtime/Engine/Private/WorldPartition/WorldPartitionStreamingPolicy.cpp
@@ -11,7 +11,7 @@
 #include "WorldPartition/WorldPartitionSubsystem.h"
 #include "WorldPartition/WorldPartitionLog.h"
 #include "WorldPartition/WorldPartitionReplay.h"
-#include "WorldPartition/DataLayer/DataLayerSubsystem.h"
+#include "WorldPartition/DataLayer/DataLayerManager.h"
 #include "WorldPartition/DataLayer/WorldDataLayers.h"
 #include "WorldPartition/HLOD/HLODSubsystem.h"
 #include "WorldPartition/ContentBundle/ContentBundle.h"
@@ -374,11 +374,11 @@ const TSet<FName>& UWorldPartitionStreamingPolicy::GetServerDisallowedStreamingO
 			TArray<FString> AllDLAssetsStrings;
 			GServerDisallowStreamingOutDataLayersString.ParseIntoArray(AllDLAssetsStrings, TEXT(","));
 
-			if (const UDataLayerSubsystem* DataLayerSubsystem = UWorld::GetSubsystem<UDataLayerSubsystem>(GetWorld()))
+			if (const UDataLayerManager* DataLayerManager = WorldPartition->GetDataLayerManager())
 			{
 				for (const FString& DataLayerAssetName : AllDLAssetsStrings)
 				{
-					if (const UDataLayerInstance* DataLayerInstance = DataLayerSubsystem->GetDataLayerInstanceFromAssetName(FName(DataLayerAssetName)))
+					if (const UDataLayerInstance* DataLayerInstance = DataLayerManager->GetDataLayerInstanceFromAssetName(FName(DataLayerAssetName)))
 					{
 						ServerDisallowStreamingOutDataLayers.Add(DataLayerInstance->GetDataLayerFName());
 					}
@@ -503,9 +503,9 @@ void UWorldPartitionStreamingPolicy::UpdateStreamingState()
 				}
 			}
 
-			const UDataLayerSubsystem* DataLayerSubsystem = GetWorld()->GetSubsystem<UDataLayerSubsystem>();
-			TSet<FName> EffectiveActiveDataLayerNames = DataLayerSubsystem->GetEffectiveActiveDataLayerNames();
-			TSet<FName> EffectiveLoadedDataLayerNames = DataLayerSubsystem->GetEffectiveLoadedDataLayerNames();
+			const UDataLayerManager* DataLayerManager = WorldPartition->GetDataLayerManager();
+			TSet<FName> EffectiveActiveDataLayerNames = DataLayerManager->GetEffectiveActiveDataLayerNames();
+			TSet<FName> EffectiveLoadedDataLayerNames = DataLayerManager->GetEffectiveLoadedDataLayerNames();
 
 			auto CanServerDeactivateOrUnloadDataLayerCell = [&ServerDisallowStreamingOutDataLayers](const UWorldPartitionRuntimeCell* Cell)
 			{
@@ -974,7 +974,7 @@ bool UWorldPartitionStreamingPolicy::IsStreamingCompleted(const TArray<FWorldPar
 	const UWorld* World = GetWorld();
 	check(World);
 	check(World->IsGameWorld());
-	const UDataLayerSubsystem* DataLayerSubsystem = World->GetSubsystem<UDataLayerSubsystem>();
+	const UDataLayerManager* DataLayerManager = WorldPartition->GetDataLayerManager();
 	const bool bTestProvidedStreamingSource = !!InStreamingSources;
 
 	// Always test non-spatial cells
@@ -984,17 +984,17 @@ bool UWorldPartitionStreamingPolicy::IsStreamingCompleted(const TArray<FWorldPar
 		FWorldPartitionStreamingQuerySource& QuerySource = QuerySources.Emplace_GetRef();
 		QuerySource.bSpatialQuery = false;
 		QuerySource.bDataLayersOnly = false;
-		QuerySource.DataLayers = DataLayerSubsystem->GetEffectiveActiveDataLayerNames().Array();
+		QuerySource.DataLayers = DataLayerManager->GetEffectiveActiveDataLayerNames().Array();
 		if (!IsStreamingCompleted(EWorldPartitionRuntimeCellState::Activated, QuerySources, true))
 		{
 			return false;
 		}
 
 		// Test only loaded data layers
-		if (!DataLayerSubsystem->GetEffectiveLoadedDataLayerNames().IsEmpty())
+		if (!DataLayerManager->GetEffectiveLoadedDataLayerNames().IsEmpty())
 		{
 			QuerySource.bDataLayersOnly = true;
-			QuerySource.DataLayers = DataLayerSubsystem->GetEffectiveLoadedDataLayerNames().Array();
+			QuerySource.DataLayers = DataLayerManager->GetEffectiveLoadedDataLayerNames().Array();
 			if (!IsStreamingCompleted(EWorldPartitionRuntimeCellState::Loaded, QuerySources, true))
 			{
 				return false;
@@ -1019,7 +1019,7 @@ bool UWorldPartitionStreamingPolicy::IsStreamingCompleted(const TArray<FWorldPar
 		QuerySource.bUseGridLoadingRange = true;
 		QuerySource.Radius = 0.f;
 		QuerySource.bDataLayersOnly = false;
-		QuerySource.DataLayers = (StreamingSource.TargetState == EStreamingSourceTargetState::Loaded) ? DataLayerSubsystem->GetEffectiveLoadedDataLayerNames().Array() : DataLayerSubsystem->GetEffectiveActiveDataLayerNames().Array();
+		QuerySource.DataLayers = (StreamingSource.TargetState == EStreamingSourceTargetState::Loaded) ? DataLayerManager->GetEffectiveLoadedDataLayerNames().Array() : DataLayerManager->GetEffectiveActiveDataLayerNames().Array();
 
 		// Execute query
 		const EWorldPartitionRuntimeCellState QueryState = (StreamingSource.TargetState == EStreamingSourceTargetState::Loaded) ? EWorldPartitionRuntimeCellState::Loaded : EWorldPartitionRuntimeCellState::Activated;
@@ -1034,13 +1034,13 @@ bool UWorldPartitionStreamingPolicy::IsStreamingCompleted(const TArray<FWorldPar
 
 bool UWorldPartitionStreamingPolicy::IsStreamingCompleted(EWorldPartitionRuntimeCellState QueryState, const TArray<FWorldPartitionStreamingQuerySource>& QuerySources, bool bExactState) const
 {
-	const UDataLayerSubsystem* DataLayerSubsystem = GetWorld()->GetSubsystem<UDataLayerSubsystem>();
+	const UDataLayerManager* DataLayerManager = WorldPartition->GetDataLayerManager();
 	const bool bIsHLODEnabled = UHLODSubsystem::IsHLODEnabled();
 
 	bool bResult = true;
 	for (const FWorldPartitionStreamingQuerySource& QuerySource : QuerySources)
 	{
-		WorldPartition->RuntimeHash->ForEachStreamingCellsQuery(QuerySource, [QuerySource, QueryState, bExactState, bIsHLODEnabled, DataLayerSubsystem, &bResult](const UWorldPartitionRuntimeCell* Cell)
+		WorldPartition->RuntimeHash->ForEachStreamingCellsQuery(QuerySource, [QuerySource, QueryState, bExactState, bIsHLODEnabled, DataLayerManager, &bResult](const UWorldPartitionRuntimeCell* Cell)
 		{
 			EWorldPartitionRuntimeCellState CellState = Cell->GetCurrentState();
 			if (CellState != QueryState)
@@ -1058,7 +1058,7 @@ bool UWorldPartitionStreamingPolicy::IsStreamingCompleted(EWorldPartitionRuntime
 				{
 					for (const FName& CellDataLayer : Cell->GetDataLayers())
 					{
-						if (!QuerySource.DataLayers.Contains(CellDataLayer) && DataLayerSubsystem->GetDataLayerEffectiveRuntimeStateByName(CellDataLayer) > EDataLayerRuntimeState::Unloaded)
+						if (!QuerySource.DataLayers.Contains(CellDataLayer) && DataLayerManager->GetDataLayerInstanceEffectiveRuntimeState(DataLayerManager->GetDataLayerInstanceFromName(CellDataLayer)) > EDataLayerRuntimeState::Unloaded)
 						{
 							bSkipCell = true;
 							break;
diff --git a/Engine/Source/Runtime/Engine/Private/WorldPartition/WorldPartitionSubsystem.cpp b/Engine/Source/Runtime/Engine/Private/WorldPartition/WorldPartitionSubsystem.cpp
--- a/Engine/Source/Runtime/Engine/Private/WorldPartition/WorldPartitionSubsystem.cpp
+++ b/Engine/Source/Runtime/Engine/Private/WorldPartition/WorldPartitionSubsystem.cpp
@@ -5,7 +5,7 @@
 #include "UnrealEngine.h"
 #include "WorldPartition/WorldPartition.h"
 #include "WorldPartition/WorldPartitionDebugHelper.h"
-#include "WorldPartition/DataLayer/DataLayerSubsystem.h"
+#include "WorldPartition/DataLayer/DataLayerManager.h"
 #include "WorldPartition/ContentBundle/ContentBundleWorldSubsystem.h"
 #include "Engine/Canvas.h"
 #include "Engine/CoreSettings.h"
@@ -163,12 +163,13 @@ void UWorldPartitionSubsystem::Deinitialize()
 	Super::Deinitialize();
 }
 
+// We allow creating UWorldPartitionSubsystem for inactive worlds as WorldPartition initialization is necessary 
+// because DataLayerManager is required to be initialized when duplicating a partitioned world.
 bool UWorldPartitionSubsystem::DoesSupportWorldType(const EWorldType::Type WorldType) const
 {
 	return Super::DoesSupportWorldType(WorldType) || WorldType == EWorldType::Inactive;
 }
 
-#if WITH_EDITOR
 void UWorldPartitionSubsystem::ForEachWorldPartition(TFunctionRef<bool(UWorldPartition*)> Func)
 {
 	for (UWorldPartition* WorldPartition : RegisteredWorldPartitions)
@@ -179,7 +180,6 @@ void UWorldPartitionSubsystem::ForEachWorldPartition(TFunctionRef<bool(UWorldPar
 		}
 	}
 }
-#endif
 
 void UWorldPartitionSubsystem::OnWorldPartitionInitialized(UWorldPartition* InWorldPartition)
 {
@@ -474,10 +474,12 @@ void UWorldPartitionSubsystem::Draw(UCanvas* Canvas, class APlayerController* PC
 		WorldPartition->DrawStreamingStatusLegend(Canvas, CurrentOffset);
 	}
 
-	UDataLayerSubsystem* DataLayerSubsystem = WorldPartition->GetWorld()->GetSubsystem<UDataLayerSubsystem>();
-	if (DataLayerSubsystem && (GDrawDataLayers || GDrawDataLayersLoadTime || GDrawRuntimeHash2D))
+	if (GDrawDataLayers || GDrawDataLayersLoadTime || GDrawRuntimeHash2D)
 	{
-		DataLayerSubsystem->DrawDataLayersStatus(Canvas, CurrentOffset);
+		if (UDataLayerManager* DataLayerManager = WorldPartition->GetDataLayerManager())
+		{
+			DataLayerManager->DrawDataLayersStatus(Canvas, CurrentOffset);
+		}
 	}
 
 	UContentBundleManager* ContentBundleManager = GetWorld()->ContentBundleManager;
diff --git a/Engine/Source/Runtime/Engine/Private/WorldPartitionBlueprintLibrary.cpp b/Engine/Source/Runtime/Engine/Private/WorldPartitionBlueprintLibrary.cpp
--- a/Engine/Source/Runtime/Engine/Private/WorldPartitionBlueprintLibrary.cpp
+++ b/Engine/Source/Runtime/Engine/Private/WorldPartitionBlueprintLibrary.cpp
@@ -1,6 +1,7 @@
 // Copyright Epic Games, Inc. All Rights Reserved.
 
 #include "WorldPartition/WorldPartitionBlueprintLibrary.h"
+#include "WorldPartition/DataLayer/DataLayerManager.h"
 
 #include UE_INLINE_GENERATED_CPP_BY_NAME(WorldPartitionBlueprintLibrary)
 
@@ -261,3 +262,8 @@ bool UWorldPartitionBlueprintLibrary::GetIntersectingActorDescs(const FBox& InBo
 #endif
 	return false;
 }
+
+UDataLayerManager* UWorldPartitionBlueprintLibrary::GetDataLayerManager(class UObject* WorldContextObject)
+{
+	return UDataLayerManager::GetDataLayerManager(WorldContextObject);
+}
diff --git a/Engine/Source/Runtime/Engine/Private/WorldSettings.cpp b/Engine/Source/Runtime/Engine/Private/WorldSettings.cpp
--- a/Engine/Source/Runtime/Engine/Private/WorldSettings.cpp
+++ b/Engine/Source/Runtime/Engine/Private/WorldSettings.cpp
@@ -33,7 +33,7 @@
 #include "Editor.h"
 #include "Misc/TransactionObjectEvent.h"
 #include "HierarchicalLOD.h"
-#include "WorldPartition/DataLayer/DataLayerSubsystem.h"
+#include "WorldPartition/DataLayer/DataLayerManager.h"
 #endif 
 
 #define LOCTEXT_NAMESPACE "ErrorChecking"
@@ -488,9 +488,9 @@ void AWorldSettings::SaveDefaultWorldPartitionSettings()
 	{
 		DefaultWorldPartitionSettings.LoadedEditorRegions = WorldPartition->GetUserLoadedEditorRegions();
 
-		if (const UDataLayerSubsystem* DataLayerSubsystem = UWorld::GetSubsystem<UDataLayerSubsystem>(GetWorld()))
+		if (const UDataLayerManager* DataLayerManager = UDataLayerManager::GetDataLayerManager(GetWorld()))
 		{
-			DataLayerSubsystem->GetUserLoadedInEditorStates(DefaultWorldPartitionSettings.LoadedDataLayers, DefaultWorldPartitionSettings.NotLoadedDataLayers);
+			DataLayerManager->GetUserLoadedInEditorStates(DefaultWorldPartitionSettings.LoadedDataLayers, DefaultWorldPartitionSettings.NotLoadedDataLayers);
 		}
 	}
 }
diff --git a/Engine/Source/Runtime/Engine/Public/WorldPartition/ActorDescContainer.h b/Engine/Source/Runtime/Engine/Public/WorldPartition/ActorDescContainer.h
--- a/Engine/Source/Runtime/Engine/Public/WorldPartition/ActorDescContainer.h
+++ b/Engine/Source/Runtime/Engine/Public/WorldPartition/ActorDescContainer.h
@@ -42,12 +42,10 @@ public:
 	};
 
 	UE_DEPRECATED(5.1, "UActorDescContainer::Initialize is deprecated, UActorDescContainer::Initialize with UActorDescContainer::FInitializeParams should be used instead.")
-	void Initialize(UWorld* World, FName InPackageName);
+	void Initialize(UWorld* InWorld, FName InPackageName);
 	void Initialize(const FInitializeParams& InitParams);
 	void Update();
 	void Uninitialize();
-	
-	virtual UWorld* GetWorld() const override;
 
 #if WITH_EDITOR
 	bool IsInitialized() const { return bContainerInitialized; }
@@ -62,7 +60,9 @@ public:
 	FGuid GetContentBundleGuid() const { return ContentBundleGuid; }
 	void SetContentBundleGuid(const FGuid& InGetContentBundleGuid) { ContentBundleGuid = InGetContentBundleGuid; }
 
+	bool IsTemplateContainer() const;
 	bool IsMainPartitionContainer() const;
+	UWorldPartition* GetWorldPartition() const;
 
 	FString GetExternalActorPath() const;
 
@@ -72,10 +72,7 @@ public:
 	void LoadAllActors(TArray<FWorldPartitionReference>& OutReferences);
 
 	bool IsActorDescHandled(const AActor* Actor) const;
-#endif
 
-#if WITH_EDITOR
-public:
 	DECLARE_EVENT_OneParam(UWorldPartition, FActorDescAddedEvent, FWorldPartitionActorDesc*);
 	FActorDescAddedEvent OnActorDescAddedEvent;
 	
@@ -91,19 +88,9 @@ public:
 	bool HasInvalidActors() const { return InvalidActors.Num() > 0; }
 	const TArray<TUniquePtr<FWorldPartitionActorDesc>>& GetInvalidActors() const { return InvalidActors; }
 	void ClearInvalidActors() { InvalidActors.Empty(); }
-#endif
-
-	UPROPERTY(Transient)
-	TObjectPtr<UWorld> World;
-
-protected:
-	//~ Begin UObject Interface
-	virtual void BeginDestroy() override;
-	//~ End UObject Interface
 
-#if WITH_EDITOR
 	//~ Begin FActorDescList Interface
-	virtual void AddActorDescriptor(FWorldPartitionActorDesc* ActorDesc) override;
+	virtual void AddActorDescriptor(FWorldPartitionActorDesc* ActorDesc, UWorld* InWorldContext) override;
 	virtual void RemoveActorDescriptor(FWorldPartitionActorDesc* ActorDesc) override;
 	//~ End FActorDescList Interface
 
@@ -122,7 +109,15 @@ protected:
 	TArray<TUniquePtr<FWorldPartitionActorDesc>> InvalidActors;
 #endif
 
+protected:
+	//~ Begin UObject Interface
+	virtual void BeginDestroy() override;
+	//~ End UObject Interface
+
 private:
+	// GetWorld() should never be called on an ActorDescContainer to avoid any confusion as it can be used as a template
+	virtual UWorld* GetWorld() const override { check(false); return nullptr; }
+
 #if WITH_EDITOR
 	bool ShouldRegisterDelegates();
 	void RegisterEditorDelegates();
diff --git a/Engine/Source/Runtime/Engine/Public/WorldPartition/ActorDescList.h b/Engine/Source/Runtime/Engine/Public/WorldPartition/ActorDescList.h
--- a/Engine/Source/Runtime/Engine/Public/WorldPartition/ActorDescList.h
+++ b/Engine/Source/Runtime/Engine/Public/WorldPartition/ActorDescList.h
@@ -151,7 +151,7 @@ public:
 	};
 
 protected:
-	virtual void AddActorDescriptor(FWorldPartitionActorDesc* ActorDesc);
+	virtual void AddActorDescriptor(FWorldPartitionActorDesc* ActorDesc, UWorld* InWorldContext);
 	virtual void RemoveActorDescriptor(FWorldPartitionActorDesc* ActorDesc);
 	TUniquePtr<FWorldPartitionActorDesc>* GetActorDescriptor(const FGuid& ActorGuid);
 
diff --git a/Engine/Source/Runtime/Engine/Public/WorldPartition/DataLayer/DataLayerInstance.h b/Engine/Source/Runtime/Engine/Public/WorldPartition/DataLayer/DataLayerInstance.h
--- a/Engine/Source/Runtime/Engine/Public/WorldPartition/DataLayer/DataLayerInstance.h
+++ b/Engine/Source/Runtime/Engine/Public/WorldPartition/DataLayer/DataLayerInstance.h
@@ -73,7 +73,7 @@ public:
 	static FText GetDataLayerText(const UDataLayerInstance* InDataLayer);
 
 	virtual bool IsLocked() const;
-	virtual bool IsReadOnly() const { return false; }
+	virtual bool IsReadOnly() const;
 	virtual bool CanEditChange(const FProperty* InProperty) const;
 	virtual bool AddActor(AActor* Actor) const { return false; }
 	virtual bool RemoveActor(AActor* Actor) const { return false; }
@@ -118,6 +118,9 @@ public:
 	const UDataLayerInstance* GetParent() const { return Parent; }
 	UDataLayerInstance* GetParent() { return Parent; }
 
+	EDataLayerRuntimeState GetRuntimeState() const;
+	EDataLayerRuntimeState GetEffectiveRuntimeState() const;
+
 	const TArray<TObjectPtr<UDataLayerInstance>>& GetChildren() const { return Children; }
 	void ForEachChild(TFunctionRef<bool(const UDataLayerInstance*)> Operation) const;
 
@@ -132,6 +135,9 @@ public:
 
 private:
 	void AddChild(UDataLayerInstance* DataLayer);
+	bool SetRuntimeState(EDataLayerRuntimeState InState, bool bInIsRecursive = false) const;
+	friend class UDataLayerManager;
+
 #if WITH_EDITOR
 	void RemoveChild(UDataLayerInstance* DataLayer);
 #endif
diff --git a/Engine/Source/Runtime/Engine/Public/WorldPartition/DataLayer/DataLayerManager.h b/Engine/Source/Runtime/Engine/Public/WorldPartition/DataLayer/DataLayerManager.h
new file mode 100644
--- /dev/null
+++ b/Engine/Source/Runtime/Engine/Public/WorldPartition/DataLayer/DataLayerManager.h
@@ -0,0 +1,218 @@
+// Copyright Epic Games, Inc. All Rights Reserved.
+
+#pragma once
+
+#include "WorldPartition/WorldPartition.h"
+#include "WorldPartition/DataLayer/ActorDataLayer.h"
+#include "WorldPartition/DataLayer/DataLayerInstance.h"
+#include "WorldPartition/DataLayer/DataLayerEditorContext.h"
+#include "WorldPartition/DataLayer/WorldDataLayers.h"
+#include "Templates/SubclassOf.h"
+#include "DataLayerManager.generated.h"
+
+class UDataLayerAsset;
+class UDataLayerInstanceWithAsset;
+class UDataLayerInstance;
+class UDataLayerLoadingPolicy;
+class UActorDescContainer;
+class UCanvas;
+class UWorld;
+
+DECLARE_DYNAMIC_MULTICAST_DELEGATE_TwoParams(FOnDataLayerInstanceRuntimeStateChanged, const UDataLayerInstance*, DataLayer, EDataLayerRuntimeState, State);
+
+#if WITH_EDITOR
+class ENGINE_API FDataLayersEditorBroadcast
+{
+public:
+	static FDataLayersEditorBroadcast& Get();
+	static void StaticOnActorDataLayersEditorLoadingStateChanged(bool bIsFromUserChange);
+	/** Broadcasts whenever one or more DataLayers editor loading state changed */
+	DECLARE_EVENT_OneParam(FDataLayersEditorBroadcast, FOnActorDataLayersEditorLoadingStateChanged, bool /*bIsFromUserChange*/);
+	FOnActorDataLayersEditorLoadingStateChanged& OnActorDataLayersEditorLoadingStateChanged() { return DataLayerEditorLoadingStateChanged; }
+
+private:
+	FOnActorDataLayersEditorLoadingStateChanged DataLayerEditorLoadingStateChanged;
+
+	friend class UDataLayerManager;
+};
+#endif
+
+UCLASS(Config = Engine, Within = WorldPartition)
+class ENGINE_API UDataLayerManager : public UObject
+{
+	GENERATED_BODY()
+
+public:
+
+	static UDataLayerManager* GetDataLayerManager(const UObject* InObject);
+
+	//~ Begin Blueprint interface
+
+	UFUNCTION(BlueprintCallable, Category = DataLayers)
+	const UDataLayerInstance* GetDataLayerInstanceFromAsset(const UDataLayerAsset* InDataLayerAsset) const;
+
+	UFUNCTION(BlueprintCallable, Category = DataLayers)
+	const UDataLayerInstance* GetDataLayerInstanceFromName(const FName& InDataLayerInstanceName) const;
+
+	UFUNCTION(BlueprintCallable, Category = DataLayers)
+	bool SetDataLayerInstanceRuntimeState(const UDataLayerInstance* InDataLayerInstance, EDataLayerRuntimeState InState, bool bInIsRecursive = false);
+
+	UFUNCTION(BlueprintCallable, Category = DataLayers)
+	EDataLayerRuntimeState GetDataLayerInstanceRuntimeState(const UDataLayerInstance* InDataLayerInstance) const;
+
+	UFUNCTION(BlueprintCallable, Category = DataLayers)
+	EDataLayerRuntimeState GetDataLayerInstanceEffectiveRuntimeState(const UDataLayerInstance* InDataLayerInstance) const;
+
+	UPROPERTY(BlueprintAssignable)
+	FOnDataLayerInstanceRuntimeStateChanged OnDataLayerInstanceRuntimeStateChanged;
+
+	//~ End Blueprint interface
+
+	template<class T>
+	const UDataLayerInstance* GetDataLayerInstance(const T& InDataLayerIdentifier) const;
+
+	template<class T>
+	TArray<const UDataLayerInstance*> GetDataLayerInstances(const TArray<T>& InDataLayerIdentifiers) const;
+
+	template<class T>
+	TArray<FName> GetDataLayerInstanceNames(const TArray<T>& InDataLayerIdentifiers) const;
+
+	void ForEachDataLayerInstance(TFunctionRef<bool(UDataLayerInstance*)> Func);
+	void ForEachDataLayerInstance(TFunctionRef<bool(UDataLayerInstance*)> Func) const;
+
+	//~ Begin Runtime State
+	const TSet<FName>& GetEffectiveActiveDataLayerNames() const;
+	const TSet<FName>& GetEffectiveLoadedDataLayerNames() const;
+	bool IsAnyDataLayerInEffectiveRuntimeState(const TArray<FName>& InDataLayerNames, EDataLayerRuntimeState InState) const;
+	//~ End Runtime State
+
+private:
+
+	//~ Begin Initialization/Deinitialization
+	UDataLayerManager();
+	void Initialize();
+	void DeInitialize();
+	//~ End Initialization/Deinitialization
+
+	//~ Begin Debugging
+	void DrawDataLayersStatus(UCanvas* Canvas, FVector2D& Offset) const;
+	void DumpDataLayers(FOutputDevice& OutputDevice) const;
+	TArray<UDataLayerInstance*> ConvertArgsToDataLayers(const TArray<FString>& InArgs);
+	//~ End Debugging
+
+	AWorldDataLayers* GetWorldDataLayers() const;
+	const UDataLayerInstance* GetDataLayerInstanceFromAssetName(const FName& InDataLayerAssetFullName) const;
+
+	void BroadcastOnDataLayerInstanceRuntimeStateChanged(const UDataLayerInstance* InDataLayer, EDataLayerRuntimeState InState);
+
+	/** Console command used to toggle activation of a Data Layer */
+	static class FAutoConsoleCommand ToggleDataLayerActivation;
+
+	/** Console command used to set Runtime Data Layer state */
+	static class FAutoConsoleCommand SetDataLayerRuntimeStateCommand;
+
+	/** Console command used to list Data Layers */
+	static class FAutoConsoleCommandWithOutputDevice DumpDataLayersCommand;
+
+	/** Data layers load time */
+	mutable TMap<const UDataLayerInstance*, double> ActiveDataLayersLoadTime;
+
+	friend class UWorldPartition;
+	friend class AWorldDataLayers;
+	friend class UDataLayerSubsystem;
+	friend class UDataLayerInstanceWithAsset;
+	friend class UWorldPartitionStreamingPolicy;
+	friend class UWorldPartitionRuntimeCell;
+	friend class UWorldPartitionSubsystem;
+	friend class FWorldPartitionDebugHelper;
+
+#if WITH_EDITOR
+private:
+
+	//~ Begin Editor Context
+	void PushActorEditorContext() const;
+	void PopActorEditorContext() const;
+	TArray<UDataLayerInstance*> GetActorEditorContextDataLayers() const;
+	uint32 GetDataLayerEditorContextHash() const;
+	//~ End Editor Context
+
+	//~ Begin ActorDesc
+	void OnActorDescContainerInitialized(UActorDescContainer* InActorDescContainer) const;
+	bool CanResolveDataLayers() const;
+	void ResolveActorDescContainersDataLayers() const;
+	void ResolveActorDescContainerDataLayers(UActorDescContainer* InActorDescContainer) const;
+	void ResolveActorDescDataLayers(FWorldPartitionActorDesc* InActorDesc) const;
+	void ResolveActorDescContainerDataLayersInternal(UActorDescContainer* InActorDescContainer, FWorldPartitionActorDesc* InActorDesc) const;
+	//~ End
+
+	//~ Begin Editor loading
+	TSubclassOf<UDataLayerLoadingPolicy> GetDataLayerLoadingPolicyClass() const;
+	bool ResolveIsLoadedInEditor(const TArray<FName>& InDataLayerInstanceNames) const;
+	//~ End Editor loading
+
+	// Helper
+	TArray<const UDataLayerInstance*> GetRuntimeDataLayerInstances(const TArray<FName>& InDataLayerInstanceNames) const;
+
+	//~ Begin User settings
+	void UpdateDataLayerEditorPerProjectUserSettings() const;
+	void GetUserLoadedInEditorStates(TArray<FName>& OutDataLayersLoadedInEditor, TArray<FName>& OutDataLayersNotLoadedInEditor) const;
+	//~ End User settings
+
+	/** Used by Editor Context */
+	mutable int32 DataLayerActorEditorContextID;
+
+	/** Whether manager can resolve data layers */
+	bool bCanResolveDataLayers;
+
+	friend struct SDataLayerTreeLabel;
+	friend class IWorldPartitionActorLoaderInterface;
+	friend class FWorldPartitionStreamingGenerator;
+	friend class FWorldPartitionActorDesc;
+	friend class UWorldPartitionRuntimeSpatialHash;
+	friend class UWorldPartitionConvertCommandlet;
+	friend class FDataLayersBroadcast;
+	friend class FDataLayerUtils;
+	friend class FDataLayerMode;
+	friend class FDataLayerHierarchy;
+	friend class FActorBrowsingMode;
+	friend class UDataLayerEditorSubsystem;
+	friend class UActorPartitionSubsystem;
+	friend class UActorDescContainer;
+	friend class AWorldSettings;
+	friend class AActor;
+#endif
+
+#if WITH_EDITORONLY_DATA
+private:
+	UPROPERTY(Config)
+	TSoftClassPtr<UDataLayerLoadingPolicy> DataLayerLoadingPolicyClass;
+
+	UPROPERTY(Transient)
+	TObjectPtr<UDataLayerLoadingPolicy> DataLayerLoadingPolicy;
+#endif
+};
+
+PRAGMA_DISABLE_DEPRECATION_WARNINGS
+template<class T>
+const UDataLayerInstance* UDataLayerManager::GetDataLayerInstance(const T& InDataLayerIdentifier) const
+{
+	AWorldDataLayers* WorldDataLayers = GetWorldDataLayers();
+	return WorldDataLayers ? WorldDataLayers->GetDataLayerInstance(InDataLayerIdentifier) : nullptr;
+}
+
+template<class T>
+TArray<const UDataLayerInstance*> UDataLayerManager::GetDataLayerInstances(const TArray<T>& InDataLayerIdentifiers) const
+{
+	static TArray<const UDataLayerInstance*> EmptyArray;
+	AWorldDataLayers* WorldDataLayers = GetWorldDataLayers();
+	return WorldDataLayers ? WorldDataLayers->GetDataLayerInstances(InDataLayerIdentifiers) : EmptyArray;
+}
+
+template<class T>
+TArray<FName> UDataLayerManager::GetDataLayerInstanceNames(const TArray<T>& InDataLayerIdentifiers) const
+{
+	static TArray<FName> EmptyArray;
+	AWorldDataLayers* WorldDataLayers = GetWorldDataLayers();
+	return WorldDataLayers ? WorldDataLayers->GetDataLayerInstanceNames(InDataLayerIdentifiers) : EmptyArray;
+}
+PRAGMA_ENABLE_DEPRECATION_WARNINGS
\ No newline at end of file
diff --git a/Engine/Source/Runtime/Engine/Public/WorldPartition/DataLayer/DataLayerSubsystem.h b/Engine/Source/Runtime/Engine/Public/WorldPartition/DataLayer/DataLayerSubsystem.h
--- a/Engine/Source/Runtime/Engine/Public/WorldPartition/DataLayer/DataLayerSubsystem.h
+++ b/Engine/Source/Runtime/Engine/Public/WorldPartition/DataLayer/DataLayerSubsystem.h
@@ -5,15 +5,13 @@
 #include "Subsystems/WorldSubsystem.h"
 #include "WorldPartition/DataLayer/ActorDataLayer.h"
 #include "WorldPartition/DataLayer/DataLayerInstance.h"
-#include "WorldPartition/DataLayer/DataLayerEditorContext.h"
-#include "WorldPartition/DataLayer/WorldDataLayersCollection.h"
+#include "WorldPartition/DataLayer/DataLayerManager.h"
 #include "DataLayerSubsystem.generated.h"
 
 /**
  * UDataLayerSubsystem
  */
 
-class UActorDescContainer;
 class UDataLayerLoadingPolicy;
 class UCanvas;
 
@@ -24,38 +22,12 @@ DECLARE_MULTICAST_DELEGATE_OneParam(FOnWorldDataLayersPostRegister, AWorldDataLa
 DECLARE_MULTICAST_DELEGATE_OneParam(FOnWorldDataLayersPreUnregister, AWorldDataLayers*);
 #endif
 
-#if WITH_EDITOR
-class ENGINE_API FDataLayersEditorBroadcast
-{
-	friend class UDataLayerEditorSubsystem;
-
-public:
-	static FDataLayersEditorBroadcast& Get();
-
-	static void StaticOnActorDataLayersEditorLoadingStateChanged(bool bIsFromUserChange);
-
-	/** Broadcasts whenever one or more DataLayers editor loading state changed */
-	DECLARE_EVENT_OneParam(FDataLayersEditorBroadcast, FOnActorDataLayersEditorLoadingStateChanged, bool /*bIsFromUserChange*/);
-	FOnActorDataLayersEditorLoadingStateChanged& OnActorDataLayersEditorLoadingStateChanged() { return DataLayerEditorLoadingStateChanged; }
-
-private:
-	FOnActorDataLayersEditorLoadingStateChanged DataLayerEditorLoadingStateChanged;
-};
-#endif
-
+/** This class is deprecated, it has been replaced by DataLayerManager. */
 UCLASS(Config = Engine)
 class ENGINE_API UDataLayerSubsystem : public UWorldSubsystem
 {
 	GENERATED_BODY()
 
-public:
-	UDataLayerSubsystem();
-
-	//~ Begin USubsystem Interface.
-	virtual void Initialize(FSubsystemCollectionBase& Collection) override;
-	virtual void Deinitialize() override;
-	//~ End USubsystem Interface.
-
 protected:
 	virtual bool DoesSupportWorldType(const EWorldType::Type WorldType) const override;
 
@@ -63,113 +35,108 @@ public:
 	//~ Begin Blueprint callable functions
 
 	/** Find a Data Layer by its asset. */
-	UFUNCTION(BlueprintCallable, Category = DataLayers)
+	UE_DEPRECATED(5.3, "Use UDataLayerManager::GetDataLayerInstanceFromAsset instead.")
+	UFUNCTION(BlueprintCallable, Category = DataLayers, meta = (DeprecatedFunction, DeprecationMessage = "DataLayerSubsystem is deprecated, use the equivalent function in DataLayerManager"))
 	UDataLayerInstance* GetDataLayerInstanceFromAsset(const UDataLayerAsset* InDataLayerAsset) const;
 
-	UFUNCTION(BlueprintCallable, Category = DataLayers)
+	UE_DEPRECATED(5.3, "Use UDataLayerManager::GetDataLayerInstanceRuntimeState instead.")
+	UFUNCTION(BlueprintCallable, Category = DataLayers, meta = (DeprecatedFunction, DeprecationMessage = "DataLayerSubsystem is deprecated, use the equivalent function in DataLayerManager"))
 	EDataLayerRuntimeState GetDataLayerInstanceRuntimeState(const UDataLayerAsset* InDataLayerAsset) const;
 
-	UFUNCTION(BlueprintCallable, Category = DataLayers)
+	UE_DEPRECATED(5.3, "Use UDataLayerManager::GetDataLayerInstanceEffectiveRuntimeState instead.")
+	UFUNCTION(BlueprintCallable, Category = DataLayers, meta = (DeprecatedFunction, DeprecationMessage = "DataLayerSubsystem is deprecated, use the equivalent function in DataLayerManager"))
 	EDataLayerRuntimeState GetDataLayerInstanceEffectiveRuntimeState(const UDataLayerAsset* InDataLayerAsset) const;
 
 	/** Set the Data Layer state using its name. */
-	UFUNCTION(BlueprintCallable, Category = DataLayers, BlueprintAuthorityOnly)
+	UE_DEPRECATED(5.3, "Use UDataLayerManager::SetDataLayerInstanceRuntimeState instead.")
+	UFUNCTION(BlueprintCallable, Category = DataLayers, BlueprintAuthorityOnly, meta = (DeprecatedFunction, DeprecationMessage = "DataLayerSubsystem is deprecated, use the equivalent function in DataLayerManager"))
 	void SetDataLayerInstanceRuntimeState(const UDataLayerAsset* InDataLayerAsset, EDataLayerRuntimeState InState, bool bInIsRecursive = false);
 
 	/** Called when a Data Layer changes state. */
+	UE_DEPRECATED(5.3, "Use UDataLayerManager::OnDataLayerInstanceRuntimeStateChanged instead.")
 	UPROPERTY(BlueprintAssignable)
 	FOnDataLayerRuntimeStateChanged OnDataLayerRuntimeStateChanged;
 
 	//~ End Blueprint callable functions
 
-#if WITH_EDITOR
-	bool ResolveIsLoadedInEditor(const TArray<FName>& InDataLayerInstanceNames) const;
-	bool CanResolveDataLayers() const;
-
-	bool RemoveDataLayer(const UDataLayerInstance* InDataLayer);
-	int32 RemoveDataLayers(const TArray<UDataLayerInstance*>& InDataLayerInstances);
-
-	void UpdateDataLayerEditorPerProjectUserSettings() const;
-	void GetUserLoadedInEditorStates(TArray<FName>& OutDataLayersLoadedInEditor, TArray<FName>& OutDataLayersNotLoadedInEditor) const;
-	
-	static TArray<const UDataLayerInstance*> GetRuntimeDataLayerInstances(UWorld* InWorld, const TArray<FName>& DataLayerInstanceNames);
-
-	void PushActorEditorContext() const;
-	void PopActorEditorContext() const;
-	TArray<UDataLayerInstance*> GetActorEditorContextDataLayers() const;
-	uint32 GetDataLayerEditorContextHash() const;
-#endif
-
 	template<class T>
+	UE_DEPRECATED(5.3, "Use UDataLayerManager::GetDataLayerInstance instead.")
 	UDataLayerInstance* GetDataLayerInstance(const T& InDataLayerIdentifier, const ULevel* InLevelContext = nullptr) const;
 
 	template<class T>
+	UE_DEPRECATED(5.3, "Use UDataLayerManager::GetDataLayerInstances instead.")
 	TArray<const UDataLayerInstance*> GetDataLayerInstances(const TArray<T>& InDataLayerIdentifiers, const ULevel* InLevelContext = nullptr) const;
 
 	template<class T>
+	UE_DEPRECATED(5.3, "Use UDataLayerManager::GetDataLayerInstanceNames instead.")
 	TArray<FName> GetDataLayerInstanceNames(const TArray<T>& InDataLayerIdentifiers, const ULevel* InLevelContext = nullptr) const;
 
+	UE_DEPRECATED(5.3, "This function has been deprecated.")
 	const UDataLayerInstance* GetDataLayerInstanceFromAssetName(const FName& InDataLayerAssetFullName) const;
 
-	void ForEachDataLayer(TFunctionRef<bool(UDataLayerInstance*)> Func, const ULevel* InLevelContext = nullptr);
-	void ForEachDataLayer(TFunctionRef<bool(UDataLayerInstance*)> Func, const ULevel* InLevelContext = nullptr) const;
-
+	UE_DEPRECATED(5.3, "Use UDataLayerManager::SetDataLayerRuntimeState instead.")
 	void SetDataLayerRuntimeState(const UDataLayerInstance* InDataLayerInstance, EDataLayerRuntimeState InState, bool bInIsRecursive = false);
+
+	UE_DEPRECATED(5.3, "Use UDataLayerManager::GetDataLayerRuntimeState instead.")
 	EDataLayerRuntimeState GetDataLayerRuntimeState(const UDataLayerInstance* InDataLayer) const;
+
+	UE_DEPRECATED(5.3, "Use UDataLayerManager::GetDataLayerRuntimeStateByName instead.")
 	EDataLayerRuntimeState GetDataLayerRuntimeStateByName(const FName& InDataLayerName) const;
+
+	UE_DEPRECATED(5.3, "Use UDataLayerManager::GetDataLayerEffectiveRuntimeState instead.")
 	EDataLayerRuntimeState GetDataLayerEffectiveRuntimeState(const UDataLayerInstance* InDataLayer) const;
+
+	UE_DEPRECATED(5.3, "Use UDataLayerManager::GetDataLayerEffectiveRuntimeStateByName instead.")
 	EDataLayerRuntimeState GetDataLayerEffectiveRuntimeStateByName(const FName& InDataLayerName) const;
-	bool IsAnyDataLayerInEffectiveRuntimeState(const TArray<FName>& InDataLayerNames, EDataLayerRuntimeState InState) const;
-	const TSet<FName>& GetEffectiveActiveDataLayerNames() const;
-	const TSet<FName>& GetEffectiveLoadedDataLayerNames() const;
 
-#if WITH_EDITOR
-	/** Called once a AWorldDataLayers has been registered. */
-	static FOnWorldDataLayersPostRegister OnWorldDataLayerPostRegister;
+	UE_DEPRECATED(5.3, "This function has been deprecated.")
+	bool IsAnyDataLayerInEffectiveRuntimeState(const TArray<FName>& InDataLayerNames, EDataLayerRuntimeState InState) const;
 
-	/** Called before a AWorldDataLayers has been unregistered. */
-	static FOnWorldDataLayersPreUnregister OnWorldDataLayerPreUnregister;
-#endif
+	UE_DEPRECATED(5.3, "This function has been deprecated.")
+	const TSet<FName>& GetEffectiveActiveDataLayerNames() const;
 
-	void RegisterWorldDataLayer(AWorldDataLayers* WorldDataLayers);
-	void UnregisterWorldDataLayer(AWorldDataLayers* WorldDataLayers);
+	UE_DEPRECATED(5.3, "This function has been deprecated.")
+	const TSet<FName>& GetEffectiveLoadedDataLayerNames() const;
 
-	void ForEachWorldDataLayer(TFunctionRef<bool(AWorldDataLayers*)> Func);
-	void ForEachWorldDataLayer(TFunctionRef<bool(AWorldDataLayers*)> Func) const;
+	UE_DEPRECATED(5.3, "This function has been deprecated.")
+	void GetDataLayerDebugColors(TMap<FName, FColor>& OutMapping) const {}
 
-	void GetDataLayerDebugColors(TMap<FName, FColor>& OutMapping) const;
+	UE_DEPRECATED(5.3, "Use function has been deprecated.")
 	void DrawDataLayersStatus(UCanvas* Canvas, FVector2D& Offset) const;
-	static TArray<UDataLayerInstance*> ConvertArgsToDataLayers(UWorld* World, const TArray<FString>& InArgs);
 
+	UE_DEPRECATED(5.3, "This function has been deprecated.")
 	void DumpDataLayers(FOutputDevice& OutputDevice) const;
 
-	//~ Begin Deprecated
-	
+#if WITH_EDITOR
+	UE_DEPRECATED(5.3, "This function has been deprecated.")
+	static TArray<const UDataLayerInstance*> GetRuntimeDataLayerInstances(UWorld* InWorld, const TArray<FName>& InDataLayerInstanceNames);
+#endif
+
 	PRAGMA_DISABLE_DEPRECATION_WARNINGS // Suppress compiler warning on override of deprecated function
 
-	UE_DEPRECATED(5.0, "Use SetDataLayerRuntimeState() instead.")
-	UFUNCTION(BlueprintCallable, Category = DataLayers, BlueprintAuthorityOnly, meta = (DeprecatedFunction, DeprecationMessage = "Use SetDataLayerRuntimeState instead"))
+	UE_DEPRECATED(5.0, "Use SetDataLayerRuntimeState instead.")
+	UFUNCTION(BlueprintCallable, Category = DataLayers, BlueprintAuthorityOnly, meta = (DeprecatedFunction, DeprecationMessage = "DataLayerSubsystem is deprecated, use the equivalent function in DataLayerManager"))
 	void SetDataLayerState(const FActorDataLayer& InDataLayer, EDataLayerState InState) { SetDataLayerRuntimeState(InDataLayer, (EDataLayerRuntimeState)InState); }
 
-	UE_DEPRECATED(5.0, "Use SetDataLayerRuntimeStateByLabel() instead.")
-	UFUNCTION(BlueprintCallable, Category = DataLayers, BlueprintAuthorityOnly, meta = (DeprecatedFunction, DeprecationMessage = "Use SetDataLayerRuntimeStateByLabel instead"))
+	UE_DEPRECATED(5.0, "Use SetDataLayerRuntimeStateByLabel instead.")
+	UFUNCTION(BlueprintCallable, Category = DataLayers, BlueprintAuthorityOnly, meta = (DeprecatedFunction, DeprecationMessage = "DataLayerSubsystem is deprecated, use the equivalent function in DataLayerManager"))
 	void SetDataLayerStateByLabel(const FName& InDataLayerLabel, EDataLayerState InState) { SetDataLayerRuntimeStateByLabel(InDataLayerLabel, (EDataLayerRuntimeState)InState); }
 
-	UE_DEPRECATED(5.0, "Use GetDataLayerRuntimeState() instead.")
-	UFUNCTION(BlueprintCallable, Category = DataLayers, meta = (DeprecatedFunction, DeprecationMessage = "Use GetDataLayerRuntimeState instead"))
+	UE_DEPRECATED(5.0, "Use GetDataLayerRuntimeState instead.")
+	UFUNCTION(BlueprintCallable, Category = DataLayers, meta = (DeprecatedFunction, DeprecationMessage = "DataLayerSubsystem is deprecated, use the equivalent function in DataLayerManager"))
 	EDataLayerState GetDataLayerState(const FActorDataLayer& InDataLayer) const { return (EDataLayerState)GetDataLayerRuntimeState(InDataLayer); }
 
-	UE_DEPRECATED(5.0, "Use GetDataLayerRuntimeStateByLabel() instead.")
-	UFUNCTION(BlueprintCallable, Category = DataLayers, meta = (DeprecatedFunction, DeprecationMessage = "Use GetDataLayerRuntimeStateByLabel instead"))
+	UE_DEPRECATED(5.0, "Use GetDataLayerRuntimeStateByLabel instead.")
+	UFUNCTION(BlueprintCallable, Category = DataLayers, meta = (DeprecatedFunction, DeprecationMessage = "DataLayerSubsystem is deprecated, use the equivalent function in DataLayerManager"))
 	EDataLayerState GetDataLayerStateByLabel(const FName& InDataLayerLabel) const { return (EDataLayerState)GetDataLayerRuntimeStateByLabel(InDataLayerLabel); }
 
-	UE_DEPRECATED(5.0, "Use GetDataLayerRuntimeState() instead.")
+	UE_DEPRECATED(5.0, "Use GetDataLayerRuntimeState instead.")
 	EDataLayerState GetDataLayerState(const UDataLayerInstance* InDataLayer) const { return (EDataLayerState)GetDataLayerRuntimeState(InDataLayer); }
 
-	UE_DEPRECATED(5.0, "Use GetDataLayerRuntimeStateByName() instead.")
+	UE_DEPRECATED(5.0, "Use UDataLayerManager::GetDataLayerInstanceRuntimeState instead.")
 	EDataLayerState GetDataLayerStateByName(const FName& InDataLayerName) const { return (EDataLayerState)GetDataLayerRuntimeStateByName(InDataLayerName); }
 
-	UE_DEPRECATED(5.0, "Use IsAnyDataLayerInEffectiveRuntimeState() instead.")
+	UE_DEPRECATED(5.0, "Use UDataLayerManager::IsAnyDataLayerInEffectiveRuntimeState instead.")
 	bool IsAnyDataLayerInState(const TArray<FName>& InDataLayerNames, EDataLayerState InState) const { return IsAnyDataLayerInEffectiveRuntimeState(InDataLayerNames, (EDataLayerRuntimeState)InState); }
 
 	PRAGMA_ENABLE_DEPRECATION_WARNINGS
@@ -182,95 +149,67 @@ public:
 	UFUNCTION(BlueprintCallable, Category = DataLayers, meta = (DeprecatedFunction, DeprecationMessage = "GetLoadedDataLayerNames will be removed."))
 	const TSet<FName>& GetLoadedDataLayerNames() const { static TSet<FName> EmptySet; return EmptySet; }
 
-	UE_DEPRECATED(5.1, "Use GetDataLayerFromAsset() instead.")
-	UFUNCTION(BlueprintCallable, Category = DataLayers)
+	UE_DEPRECATED(5.1, "Use GetDataLayerInstanceFromAsset instead.")
+	UFUNCTION(BlueprintCallable, Category = DataLayers, meta = (DeprecatedFunction, DeprecationMessage = "DataLayerSubsystem is deprecated, use GetDataLayerInstanceFromAsset in DataLayerManager"))
 	UDataLayerInstance* GetDataLayer(const FActorDataLayer& InDataLayer) const;
 
-	UE_DEPRECATED(5.1, "Use GetDataLayerFromAsset() instead.")
-	UFUNCTION(BlueprintCallable, Category = DataLayers)
+	UE_DEPRECATED(5.1, "Use GetDataLayerInstanceFromAsset instead.")
+	UFUNCTION(BlueprintCallable, Category = DataLayers, meta = (DeprecatedFunction, DeprecationMessage = "DataLayerSubsystem is deprecated, use GetDataLayerInstanceFromAsset in DataLayerManager"))
 	UDataLayerInstance* GetDataLayerFromName(FName InDataLayerName) const;
 
-	UE_DEPRECATED(5.1, "Use GetDataLayerFromAsset() instead.")
-	UFUNCTION(BlueprintCallable, Category = DataLayers)
+	UE_DEPRECATED(5.1, "Use GetDataLayerInstanceFromAsset instead.")
+	UFUNCTION(BlueprintCallable, Category = DataLayers, meta = (DeprecatedFunction, DeprecationMessage = "DataLayerSubsystem is deprecated, use GetDataLayerInstanceFromAsset in DataLayerManager"))
 	UDataLayerInstance* GetDataLayerFromLabel(FName InDataLayerLabel) const;
 
-	UE_DEPRECATED(5.1, "Use GetDataLayerInstanceRuntimeState() instead.")
-	UFUNCTION(BlueprintCallable, Category = DataLayers)
+	UE_DEPRECATED(5.1, "Use GetDataLayerInstanceRuntimeState instead.")
+	UFUNCTION(BlueprintCallable, Category = DataLayers, meta = (DeprecatedFunction, DeprecationMessage = "DataLayerSubsystem is deprecated, use GetDataLayerInstanceRuntimeState in DataLayerManager"))
 	EDataLayerRuntimeState GetDataLayerRuntimeState(const FActorDataLayer& InDataLayer) const;
 
-	UE_DEPRECATED(5.1, "Use GetDataLayerInstanceRuntimeState() instead.")
-	UFUNCTION(BlueprintCallable, Category = DataLayers)
+	UE_DEPRECATED(5.1, "Use GetDataLayerInstanceRuntimeState instead.")
+	UFUNCTION(BlueprintCallable, Category = DataLayers, meta = (DeprecatedFunction, DeprecationMessage = "DataLayerSubsystem is deprecated, use GetDataLayerInstanceRuntimeState in DataLayerManager"))
 	EDataLayerRuntimeState GetDataLayerRuntimeStateByLabel(const FName& InDataLayerLabel) const;
 
-	UE_DEPRECATED(5.1, "Use GetDataLayerInstanceEffectiveRuntimeState() instead.")
-	UFUNCTION(BlueprintCallable, Category = DataLayers)
+	UE_DEPRECATED(5.1, "Use GetDataLayerInstanceEffectiveRuntimeState instead.")
+	UFUNCTION(BlueprintCallable, Category = DataLayers, meta = (DeprecatedFunction, DeprecationMessage = "DataLayerSubsystem is deprecated, use GetDataLayerInstanceEffectiveRuntimeState in DataLayerManager"))
 	EDataLayerRuntimeState GetDataLayerEffectiveRuntimeState(const FActorDataLayer& InDataLayer) const;
 
-	UE_DEPRECATED(5.1, "Use GetDataLayerInstanceEffectiveRuntimeState() instead.")
-	UFUNCTION(BlueprintCallable, Category = DataLayers)
+	UE_DEPRECATED(5.1, "Use GetDataLayerInstanceEffectiveRuntimeState instead.")
+	UFUNCTION(BlueprintCallable, Category = DataLayers, meta = (DeprecatedFunction, DeprecationMessage = "DataLayerSubsystem is deprecated, use GetDataLayerInstanceEffectiveRuntimeState in DataLayerManager"))
 	EDataLayerRuntimeState GetDataLayerEffectiveRuntimeStateByLabel(const FName& InDataLayerLabel) const;
 
 	UE_DEPRECATED(5.1, "Use SetDataLayerRuntimeState() with UDataLayerAsset* overload instead.")
-	UFUNCTION(BlueprintCallable, Category = DataLayers, BlueprintAuthorityOnly)
+	UFUNCTION(BlueprintCallable, Category = DataLayers, BlueprintAuthorityOnly, meta = (DeprecatedFunction, DeprecationMessage = "DataLayerSubsystem is deprecated, use SetDataLayerInstanceRuntimeState in DataLayerManager"))
 	void SetDataLayerRuntimeState(const FActorDataLayer& InDataLayer, EDataLayerRuntimeState InState, bool bInIsRecursive = false);
 
-	UE_DEPRECATED(5.1, "Use SetDataLayerInstanceRuntimeState() instead.")
-	UFUNCTION(BlueprintCallable, Category = DataLayers, BlueprintAuthorityOnly)
+	UE_DEPRECATED(5.1, "Use SetDataLayerInstanceRuntimeState instead.")
+	UFUNCTION(BlueprintCallable, Category = DataLayers, BlueprintAuthorityOnly, meta = (DeprecatedFunction, DeprecationMessage = "DataLayerSubsystem is deprecated, use SetDataLayerInstanceRuntimeState in DataLayerManager"))
 	void SetDataLayerRuntimeStateByLabel(const FName& InDataLayerLabel, EDataLayerRuntimeState InState, bool bInIsRecursive = false);
 
-	//~ End Deprecated
-
 private:
+	UE_DEPRECATED(5.3, "Use UDataLayerManager::DataLayerLoadingPolicyClass instead.")
 	UPROPERTY(Config)
 	TSoftClassPtr<UDataLayerLoadingPolicy> DataLayerLoadingPolicyClass;
 
-	UPROPERTY(Transient)
-	TObjectPtr<UDataLayerLoadingPolicy> DataLayerLoadingPolicy;
-
-#if WITH_EDITOR
-	void RegisterEarlyLoadedWorldDataLayers();
-	void OnActorDescContainerInitialized(UActorDescContainer* InActorDescContainer);
-	void ResolveActorDescContainer(UActorDescContainer* InActorDescContainer);
-	void ResolveActorDescContainers();
-
-	mutable int32 DataLayerActorEditorContextID;
-#endif
-
-	/** Console command used to toggle activation of a DataLayer */
-	static class FAutoConsoleCommand ToggleDataLayerActivation;
-
-	/** Data layers load time */
-	mutable TMap<const UDataLayerInstance*, double> ActiveDataLayersLoadTime;
-
-	/** Console command used to set Runtime DataLayer state*/
-	static class FAutoConsoleCommand SetDataLayerRuntimeStateCommand;
-
-	void OnEffectiveRuntimeDataLayerStatesChanged(AWorldDataLayers* InWorldDataLayers);
-	void UpdateCachedEffectiveRuntimeStates() const;
-
-	mutable bool bIsCachedEffectiveStateDirty;
-	mutable TSet<FName> CachedEffectiveActiveDataLayerNames;
-	mutable TSet<FName> CachedEffectiveLoadedDataLayerNames;
-
-	friend class AWorldDataLayers;
-
-	FWorldDataLayersCollection WorldDataLayerCollection;
+	friend class UDataLayerManager;
 };
 
 template<class T>
 UDataLayerInstance* UDataLayerSubsystem::GetDataLayerInstance(const T& InDataLayerIdentifier, const ULevel* InLevelContext /* = nullptr */) const
 {
-	return WorldDataLayerCollection.GetDataLayerInstance(InDataLayerIdentifier, InLevelContext);
+	const UDataLayerManager* DataLayerManager = UDataLayerManager::GetDataLayerManager(GetWorld());
+	return DataLayerManager ? const_cast<UDataLayerInstance*>(DataLayerManager->GetDataLayerInstance(InDataLayerIdentifier)) : nullptr;
 }
 
 template<class T>
 TArray<FName> UDataLayerSubsystem::GetDataLayerInstanceNames(const TArray<T>& InDataLayerIdentifiers, const ULevel* InLevelContext /* = nullptr */) const
 {
-	return WorldDataLayerCollection.GetDataLayerInstanceNames(InDataLayerIdentifiers, InLevelContext);
+	const UDataLayerManager* DataLayerManager = UDataLayerManager::GetDataLayerManager(GetWorld());
+	return DataLayerManager ? DataLayerManager->GetDataLayerInstanceNames(InDataLayerIdentifiers) : TArray<FName>();
 }
 
 template<class T>
 TArray<const UDataLayerInstance*> UDataLayerSubsystem::GetDataLayerInstances(const TArray<T>& InDataLayerIdentifiers, const ULevel* InLevelContext /* = nullptr */) const
 {
-	return WorldDataLayerCollection.GetDataLayerInstances(InDataLayerIdentifiers, InLevelContext);
+	const UDataLayerManager* DataLayerManager = UDataLayerManager::GetDataLayerManager(GetWorld());
+	return DataLayerManager ? DataLayerManager->GetDataLayerInstances(InDataLayerIdentifiers, InLevelContext) : TArray<const UDataLayerInstance*>();
 }
\ No newline at end of file
diff --git a/Engine/Source/Runtime/Engine/Public/WorldPartition/DataLayer/DataLayerUtils.h b/Engine/Source/Runtime/Engine/Public/WorldPartition/DataLayer/DataLayerUtils.h
--- a/Engine/Source/Runtime/Engine/Public/WorldPartition/DataLayer/DataLayerUtils.h
+++ b/Engine/Source/Runtime/Engine/Public/WorldPartition/DataLayer/DataLayerUtils.h
@@ -10,6 +10,7 @@
 #include "WorldPartition/WorldPartitionActorDescView.h"
 #include "WorldPartition/WorldPartitionStreamingGeneration.h"
 
+class UDataLayerManager;
 class FDataLayerInstanceDesc;
 class FWorldDataLayersActorDesc;
 class FWorldPartitionActorDesc;
@@ -26,15 +27,15 @@ public:
 		return IconNameByType[static_cast<uint32>(DataLayerType)];
 	}
 
-	static TArray<FName> ResolvedDataLayerInstanceNames(const FWorldPartitionActorDesc* InActorDesc, const TArray<const FWorldDataLayersActorDesc*>& InWorldDataLayersActorDescs = TArray<const FWorldDataLayersActorDesc*>(), UWorld* InWorld = nullptr, bool* bOutIsResultValid = nullptr);
+	static TArray<FName> ResolvedDataLayerInstanceNames(const UDataLayerManager* InDataLayerManager, const FWorldPartitionActorDesc* InActorDesc, const TArray<const FWorldDataLayersActorDesc*>& InWorldDataLayersActorDescs = TArray<const FWorldDataLayersActorDesc*>());
 	
-	static bool ResolveRuntimeDataLayerInstanceNames(const FWorldPartitionActorDescView& InActorDescView, const FActorDescViewMap& ActorDescViewMap, TArray<FName>& OutRuntimeDataLayerInstanceNames);
+	static bool ResolveRuntimeDataLayerInstanceNames(const UDataLayerManager* InDataLayerManager, const FWorldPartitionActorDescView& InActorDescView, const FActorDescViewMap& ActorDescViewMap, TArray<FName>& OutRuntimeDataLayerInstanceNames);
 
 	static const FDataLayerInstanceDesc* GetDataLayerInstanceDescFromInstanceName(const TArray<const FWorldDataLayersActorDesc*>& InWorldDataLayersActorDescs, const FName& DataLayerInstanceName);
 
 	static const FDataLayerInstanceDesc* GetDataLayerInstanceDescFromAssetPath(const TArray<const FWorldDataLayersActorDesc*>& InWorldDataLayersActorDescs, const FName& DataLayerAssetPath);
 
-	static bool FindWorldDataLayerActorDescs(const FActorDescViewMap& ActorDescViewMap, TArray<const FWorldDataLayersActorDesc*>& OutWorldDataLayersActorDescs);
+	static TArray<const FWorldDataLayersActorDesc*> FindWorldDataLayerActorDescs(const FActorDescViewMap& ActorDescViewMap);
 
 	static bool AreWorldDataLayersActorDescsSane(const TArray<const FWorldDataLayersActorDesc*>& InWorldDataLayersActorDescs);
 #endif
diff --git a/Engine/Source/Runtime/Engine/Public/WorldPartition/DataLayer/WorldDataLayers.h b/Engine/Source/Runtime/Engine/Public/WorldPartition/DataLayer/WorldDataLayers.h
--- a/Engine/Source/Runtime/Engine/Public/WorldPartition/DataLayer/WorldDataLayers.h
+++ b/Engine/Source/Runtime/Engine/Public/WorldPartition/DataLayer/WorldDataLayers.h
@@ -45,8 +45,6 @@ class ENGINE_API AWorldDataLayers : public AInfo
 
 public:
 	virtual void PostLoad() override;
-	virtual void PreRegisterAllComponents() override;
-	virtual void PostUnregisterAllComponents() override;
 	virtual void RewindForReplay() override;
 	virtual void BeginPlay() override;
 
@@ -89,11 +87,11 @@ public:
 	// Allows overriding of DataLayers with PlayFromHere
 	template<class T>
 	void OverwriteDataLayerRuntimeStates(const TArray<T>* InActiveDataLayers, const TArray<T>* InLoadedDataLayers );
-#endif
 
 	bool IsSubWorldDataLayers() const;
-	bool IsTheMainWorldDataLayers() const;
 	bool IsRuntimeRelevant() const;
+#endif
+
 	static FName GetWorldPartionWorldDataLayersName() { return FName(TEXT("WorldDataLayers")); } // reserved for ULevel::WorldDataLayers
 	
 	void DumpDataLayers(FOutputDevice& OutputDevice) const;
@@ -159,7 +157,7 @@ public:
 	UE_DEPRECATED(5.1, "Convert UDataLayers to UDataLayerAsset and UDataLayerInstance using DataLayerToAssetCommandLet")
 	const UDataLayerInstance* GetDataLayerFromLabel(const FName& InDataLayerLabel) const;
 
-	UE_DEPRECATED(5.1, "Use UDataLayerSubsystem::GetUserLoadedInEditorStates insteaed")
+	UE_DEPRECATED(5.3, "This function has been deprecated.")
 	void GetUserLoadedInEditorStates(TArray<FName>& OutDataLayersLoadedInEditor, TArray<FName>& OutDataLayersNotLoadedInEditor) const {}
 
 	//~ End Deprecated
diff --git a/Engine/Source/Runtime/Engine/Public/WorldPartition/LevelInstance/LevelInstanceActorDesc.h b/Engine/Source/Runtime/Engine/Public/WorldPartition/LevelInstance/LevelInstanceActorDesc.h
--- a/Engine/Source/Runtime/Engine/Public/WorldPartition/LevelInstance/LevelInstanceActorDesc.h
+++ b/Engine/Source/Runtime/Engine/Public/WorldPartition/LevelInstance/LevelInstanceActorDesc.h
@@ -35,16 +35,17 @@ protected:
 	virtual bool Equals(const FWorldPartitionActorDesc* Other) const override;
 	virtual void TransferFrom(const FWorldPartitionActorDesc* From) override;
 	virtual void Serialize(FArchive& Ar) override;
-	virtual void SetContainer(UActorDescContainer* InContainer) override;
+	virtual void SetContainer(UActorDescContainer* InContainer, UWorld* InWorldContext) override;
 
 	FName LevelPackage;
 	FTransform LevelInstanceTransform;
 	ELevelInstanceRuntimeBehavior DesiredRuntimeBehavior;
 
 	TWeakObjectPtr<UActorDescContainer> LevelInstanceContainer;
+	TWeakObjectPtr<UWorld> LevelInstanceContainerWorldContext;
 
 private:
-	void RegisterContainerInstance(UWorld* InWorld);
+	void RegisterContainerInstance(UWorld* InWorldContext);
 	void UnregisterContainerInstance();
 	void UpdateBounds();
 };
diff --git a/Engine/Source/Runtime/Engine/Public/WorldPartition/WorldPartition.h b/Engine/Source/Runtime/Engine/Public/WorldPartition/WorldPartition.h
--- a/Engine/Source/Runtime/Engine/Public/WorldPartition/WorldPartition.h
+++ b/Engine/Source/Runtime/Engine/Public/WorldPartition/WorldPartition.h
@@ -32,6 +32,7 @@ class UWorldPartitionRuntimeCell;
 class UWorldPartitionRuntimeHash;
 class UWorldPartitionStreamingPolicy;
 class IWorldPartitionCell;
+class UDataLayerManager;
 class IStreamingGenerationErrorHandler;
 class FLoaderAdapterAlwaysLoadedActors;
 class FLoaderAdapterActorList;
@@ -285,6 +286,8 @@ public:
 	void DisableStreamingIn();
 	void EnableStreamingIn();
 
+	UDataLayerManager* GetDataLayerManager() const;
+
 #if WITH_EDITORONLY_DATA
 	UPROPERTY(DuplicateTransient)
 	TObjectPtr<UWorldPartitionEditorHash> EditorHash;
@@ -360,6 +363,9 @@ private:
 	mutable TOptional<bool> bCachedIsServerStreamingEnabled;
 	mutable TOptional<bool> bCachedIsServerStreamingOutEnabled;
 
+	UPROPERTY(Transient)
+	TObjectPtr<UDataLayerManager> DataLayerManager;
+
 	UPROPERTY(Transient)
 	mutable TObjectPtr<UWorldPartitionStreamingPolicy> StreamingPolicy;
 
diff --git a/Engine/Source/Runtime/Engine/Public/WorldPartition/WorldPartitionActorDesc.h b/Engine/Source/Runtime/Engine/Public/WorldPartition/WorldPartitionActorDesc.h
--- a/Engine/Source/Runtime/Engine/Public/WorldPartition/WorldPartitionActorDesc.h
+++ b/Engine/Source/Runtime/Engine/Public/WorldPartition/WorldPartitionActorDesc.h
@@ -108,9 +108,10 @@ PRAGMA_ENABLE_DEPRECATION_WARNINGS
 	inline bool GetActorIsHLODRelevant() const { return bActorIsHLODRelevant; }
 	inline FName GetHLODLayer() const { return HLODLayer; }
 	inline const TArray<FName>& GetDataLayers() const { return DataLayers; }
-	inline const TArray<FName>& GetDataLayerInstanceNames() const { return DataLayerInstanceNames; }
+	inline bool HasResolvedDataLayerInstanceNames() const { return ResolvedDataLayerInstanceNames.IsSet(); }
+	const TArray<FName>& GetDataLayerInstanceNames() const;
 	inline const TArray<FName>& GetTags() const { return Tags; }
-	inline void SetDataLayerInstanceNames(const TArray<FName>& InDataLayerInstanceNames) { DataLayerInstanceNames = InDataLayerInstanceNames; }
+	inline void SetDataLayerInstanceNames(const TArray<FName>& InDataLayerInstanceNames) { ResolvedDataLayerInstanceNames = InDataLayerInstanceNames; }
 	inline FName GetActorPackage() const { return ActorPackage; }
 	inline FSoftObjectPath GetActorSoftPath() const { return ActorPath; }
 	inline FName GetActorLabel() const { return ActorLabel; }
@@ -202,7 +203,7 @@ public:
 		return Container;
 	}
 
-	virtual void SetContainer(UActorDescContainer* InContainer)
+	virtual void SetContainer(UActorDescContainer* InContainer, UWorld* InWorldContext)
 	{
 		check(!Container || !InContainer);
 		Container = InContainer;
@@ -240,13 +241,7 @@ public:
 protected:
 	FWorldPartitionActorDesc();
 
-	virtual void TransferFrom(const FWorldPartitionActorDesc* From)
-	{
-		Container = From->Container;
-		SoftRefCount = From->SoftRefCount;
-		HardRefCount = From->HardRefCount;
-		bIsForcedNonSpatiallyLoaded = From->bIsForcedNonSpatiallyLoaded;
-	}
+	virtual void TransferFrom(const FWorldPartitionActorDesc* From);
 
 	virtual void TransferWorldData(const FWorldPartitionActorDesc* From)
 	{
@@ -287,7 +282,7 @@ protected:
 	UClass*							ActorNativeClass;
 	mutable TWeakObjectPtr<AActor>	ActorPtr;
 	UActorDescContainer*			Container;
-	TArray<FName>					DataLayerInstanceNames;
+	TOptional<TArray<FName>>		ResolvedDataLayerInstanceNames; // Can only resolve in ActorDesc if Container is not used as a template
 	bool							bIsForcedNonSpatiallyLoaded;
 
 	static TMap<TSubclassOf<AActor>, FActorDescDeprecator> Deprecators;
diff --git a/Engine/Source/Runtime/Engine/Public/WorldPartition/WorldPartitionActorDescView.h b/Engine/Source/Runtime/Engine/Public/WorldPartition/WorldPartitionActorDescView.h
--- a/Engine/Source/Runtime/Engine/Public/WorldPartition/WorldPartitionActorDescView.h
+++ b/Engine/Source/Runtime/Engine/Public/WorldPartition/WorldPartitionActorDescView.h
@@ -38,9 +38,8 @@ public:
 	bool GetActorIsRuntimeOnly() const;
 	bool GetActorIsHLODRelevant() const;
 	FName GetHLODLayer() const;
-	const TArray<FName>& GetDataLayers() const;
 	const TArray<FName>& GetDataLayerInstanceNames() const;
-	const TArray<FName>& GetRuntimeDataLayers() const;
+	const TArray<FName>& GetRuntimeDataLayerInstanceNames() const;
 	const TArray<FName>& GetTags() const;
 	FName GetActorPackage() const;
 	
@@ -79,8 +78,9 @@ PRAGMA_ENABLE_DEPRECATION_WARNINGS
 	void SetForcedNonSpatiallyLoaded();
 	void SetInvalidRuntimeGrid();
 	void SetInvalidDataLayers();
-	void SetRuntimeDataLayers(TArray<FName>& InRuntimeDataLayers);
+	void SetRuntimeDataLayerInstanceNames(TArray<FName>& InRuntimeDataLayerInstanceNames);
 	void SetRuntimeReferences(TArray<FGuid>& InRuntimeReferences);
+	void SetDataLayerInstanceNames(const TArray<FName>& InDataLayerInstanceNames);
 
 	AActor* GetActor() const;
 
@@ -101,7 +101,8 @@ protected:
 	bool bIsForcedNonSpatiallyLoaded;
 	bool bInvalidDataLayers;
 	bool bInvalidRuntimeGrid;
-	TOptional<TArray<FName>> RuntimeDataLayers;
+	TOptional<TArray<FName>> ResolvedDataLayerInstanceNames;
+	TOptional<TArray<FName>> RuntimeDataLayerInstanceNames;
 	TOptional<TArray<FGuid>> RuntimeReferences;
 };
 #endif
\ No newline at end of file
diff --git a/Engine/Source/Runtime/Engine/Public/WorldPartition/WorldPartitionActorLoaderInterface.h b/Engine/Source/Runtime/Engine/Public/WorldPartition/WorldPartitionActorLoaderInterface.h
--- a/Engine/Source/Runtime/Engine/Public/WorldPartition/WorldPartitionActorLoaderInterface.h
+++ b/Engine/Source/Runtime/Engine/Public/WorldPartition/WorldPartitionActorLoaderInterface.h
@@ -73,5 +73,9 @@ public:
 	};
 
 	virtual ILoaderAdapter* GetLoaderAdapter() =0;
+
+private:
+	static bool PassDataLayersFilter(class UWorld* InWorld, const TArray<FName>& InDataLayerInstanceNames);
+	friend class ILoaderAdapter;
 #endif
 };
\ No newline at end of file
diff --git a/Engine/Source/Runtime/Engine/Public/WorldPartition/WorldPartitionBlueprintLibrary.h b/Engine/Source/Runtime/Engine/Public/WorldPartition/WorldPartitionBlueprintLibrary.h
--- a/Engine/Source/Runtime/Engine/Public/WorldPartition/WorldPartitionBlueprintLibrary.h
+++ b/Engine/Source/Runtime/Engine/Public/WorldPartition/WorldPartitionBlueprintLibrary.h
@@ -10,6 +10,7 @@
 
 class FWorldPartitionActorDesc;
 class UWorldPartition;
+class UDataLayerManager;
 
 /**
  * Snapshot of an actor descriptor, which represents the state of an actor on disk.
@@ -139,4 +140,10 @@ public:
 	 */
 	UFUNCTION(BlueprintCallable, Category="World Partition")
 	static bool GetIntersectingActorDescs(const FBox& InBox, TArray<FActorDesc>& OutActorDescs);
+
+	/**
+	 * Returns the Data Layer Manager for this object.
+	 */
+	UFUNCTION(BlueprintCallable, Category = "World Partition", meta = (WorldContext = "WorldContextObject"))
+	static UDataLayerManager* GetDataLayerManager(class UObject* WorldContextObject);
 };
\ No newline at end of file
diff --git a/Engine/Source/Runtime/Engine/Public/WorldPartition/WorldPartitionRuntimeCell.h b/Engine/Source/Runtime/Engine/Public/WorldPartition/WorldPartitionRuntimeCell.h
--- a/Engine/Source/Runtime/Engine/Public/WorldPartition/WorldPartitionRuntimeCell.h
+++ b/Engine/Source/Runtime/Engine/Public/WorldPartition/WorldPartitionRuntimeCell.h
@@ -190,6 +190,8 @@ class ENGINE_API UWorldPartitionRuntimeCell : public UObject, public IWorldParti
 	virtual FName GetLevelPackageName() const override;
 	//~End IWorldPartitionCell Interface
 
+	bool HasAnyDataLayerInEffectiveRuntimeState(EDataLayerRuntimeState InState) const;
+
 	bool GetBlockOnSlowLoading() const { return bBlockOnSlowLoading; }
 #if WITH_EDITOR
 	bool NeedsActorToCellRemapping() const;
