Add two new parameter interfaces for audio sources

diff --git a/Engine/Source/Runtime/Engine/Classes/Sound/SoundAttenuation.h b/Engine/Source/Runtime/Engine/Classes/Sound/SoundAttenuation.h
--- a/Engine/Source/Runtime/Engine/Classes/Sound/SoundAttenuation.h
+++ b/Engine/Source/Runtime/Engine/Classes/Sound/SoundAttenuation.h
@@ -542,4 +542,30 @@ namespace Audio
 
 		ENGINE_API Audio::FParameterInterfacePtr GetInterface();
 	} // namespace SpatializationInterface
+
+	namespace SourceOrientationInterface
+	{
+		ENGINE_API const extern FName Name;
+
+		namespace Inputs
+		{
+			ENGINE_API const extern FName Azimuth;
+			ENGINE_API const extern FName Elevation;
+		} // namespace Inputs
+
+		ENGINE_API Audio::FParameterInterfacePtr GetInterface();
+	} // namespace EmitterInterface
+
+	namespace ListenerOrientationInterface
+	{
+		ENGINE_API const extern FName Name;
+
+		namespace Inputs
+		{
+			ENGINE_API const extern FName Azimuth;
+			ENGINE_API const extern FName Elevation;
+		} // namespace Inputs
+
+		ENGINE_API Audio::FParameterInterfacePtr GetInterface();
+	} // namespace EmitterInterface
 } // namespace Audio
diff --git a/Engine/Source/Runtime/Engine/Private/ActiveSound.cpp b/Engine/Source/Runtime/Engine/Private/ActiveSound.cpp
--- a/Engine/Source/Runtime/Engine/Private/ActiveSound.cpp
+++ b/Engine/Source/Runtime/Engine/Private/ActiveSound.cpp
@@ -599,11 +599,15 @@ void FActiveSound::UpdateInterfaceParameters(const TArray<FListener>& InListener
 
 	FParameterInterfacePtr AttenuationInterface = AttenuationInterface::GetInterface();
 	FParameterInterfacePtr SpatializationInterface = SpatializationInterface::GetInterface();
+	FParameterInterfacePtr SourceOrientationInterface = SourceOrientationInterface::GetInterface();
+	FParameterInterfacePtr ListenerOrientationInterface = ListenerOrientationInterface::GetInterface();
 
 	const bool bImplementsAttenuation = Sound->ImplementsParameterInterface(AttenuationInterface);
 	const bool bImplementsSpatialization = Sound->ImplementsParameterInterface(SpatializationInterface);
+	const bool bImplementsSourceOrientation = Sound->ImplementsParameterInterface(SourceOrientationInterface);
+	const bool bImplementsListenerOrientation = Sound->ImplementsParameterInterface(ListenerOrientationInterface);
 
-	if (!bImplementsAttenuation && !bImplementsSpatialization)
+	if (false == (bImplementsAttenuation || bImplementsSpatialization || bImplementsSourceOrientation || bImplementsListenerOrientation))
 	{
 		return;
 	}
@@ -612,6 +616,7 @@ void FActiveSound::UpdateInterfaceParameters(const TArray<FListener>& InListener
 
 	const FListener& Listener = InListeners[ClosestListenerIndex];
 	const FVector SourceDirection = Transform.GetLocation() - Listener.Transform.GetLocation();
+	
 	if (bImplementsAttenuation)
 	{
 		const float Distance = SourceDirection.Size();
@@ -631,6 +636,33 @@ void FActiveSound::UpdateInterfaceParameters(const TArray<FListener>& InListener
 		});
 	}
 
+	if (bImplementsSourceOrientation)
+	{
+		const FVector ListenerDirectionNormal = Transform.InverseTransformVectorNoScale(-SourceDirection).GetSafeNormal();
+		const FVector2D ListenerAzimuthAndElevation = FMath::GetAzimuthAndElevation(ListenerDirectionNormal, FVector::ForwardVector, FVector::RightVector, FVector::UpVector);
+		const float Azimuth = FMath::RadiansToDegrees(ListenerAzimuthAndElevation.X);
+		const float Elevation = FMath::RadiansToDegrees(ListenerAzimuthAndElevation.Y);
+		ParamsToUpdate.Append(
+		{
+			{ SourceOrientationInterface::Inputs::Azimuth, Azimuth },
+			{ SourceOrientationInterface::Inputs::Elevation, Elevation }
+		});
+	}
+
+	if (bImplementsListenerOrientation)
+	{
+		const FVector ListenerDirectionNormal = Listener.Transform.GetRotation().GetForwardVector();
+		const FVector2D ListenerAzimuthAndElevation = FMath::GetAzimuthAndElevation(ListenerDirectionNormal, FVector::ForwardVector, FVector::RightVector, FVector::UpVector);
+		const float Azimuth = FMath::RadiansToDegrees(ListenerAzimuthAndElevation.X);
+		const float Elevation = FMath::RadiansToDegrees(ListenerAzimuthAndElevation.Y);
+		
+		ParamsToUpdate.Append(
+		{
+			{ ListenerOrientationInterface::Inputs::Azimuth, Azimuth },
+			{ ListenerOrientationInterface::Inputs::Elevation, Elevation }
+		});
+	}
+
 	InstanceTransmitter->SetParameters(MoveTemp(ParamsToUpdate));
 }
 
diff --git a/Engine/Source/Runtime/Engine/Private/AudioSettings.cpp b/Engine/Source/Runtime/Engine/Private/AudioSettings.cpp
--- a/Engine/Source/Runtime/Engine/Private/AudioSettings.cpp
+++ b/Engine/Source/Runtime/Engine/Private/AudioSettings.cpp
@@ -282,6 +282,8 @@ void UAudioSettings::RegisterParameterInterfaces()
 		Audio::IAudioParameterInterfaceRegistry& InterfaceRegistry = Audio::IAudioParameterInterfaceRegistry::Get();
 		InterfaceRegistry.RegisterInterface(Audio::AttenuationInterface::GetInterface());
 		InterfaceRegistry.RegisterInterface(Audio::SpatializationInterface::GetInterface());
+		InterfaceRegistry.RegisterInterface(Audio::SourceOrientationInterface::GetInterface());
+		InterfaceRegistry.RegisterInterface(Audio::ListenerOrientationInterface::GetInterface());
 	}
 }
 
diff --git a/Engine/Source/Runtime/Engine/Private/SoundAttenuation.cpp b/Engine/Source/Runtime/Engine/Private/SoundAttenuation.cpp
--- a/Engine/Source/Runtime/Engine/Private/SoundAttenuation.cpp
+++ b/Engine/Source/Runtime/Engine/Private/SoundAttenuation.cpp
@@ -233,5 +233,99 @@ namespace Audio
 		}
 	} // namespace SpatializationInterface
 #undef AUDIO_PARAMETER_INTERFACE_NAMESPACE
+
+#define AUDIO_PARAMETER_INTERFACE_NAMESPACE "UE.Source.Orientation"
+	namespace SourceOrientationInterface
+	{
+		const FName Name = AUDIO_PARAMETER_INTERFACE_NAMESPACE;
+
+		namespace Inputs
+		{
+			const FName Azimuth = AUDIO_PARAMETER_INTERFACE_MEMBER_DEFINE("Azimuth");
+			const FName Elevation = AUDIO_PARAMETER_INTERFACE_MEMBER_DEFINE("Elevation");
+		} // namespace Inputs
+
+		Audio::FParameterInterfacePtr GetInterface()
+		{
+			struct FInterface : public Audio::FParameterInterface
+			{
+				FInterface()
+					: FParameterInterface(SourceOrientationInterface::Name, { 1, 0 }, *USoundBase::StaticClass())
+				{
+					Inputs =
+					{
+						{
+							FText(),
+							NSLOCTEXT("SourceOrientation", "AzimuthDescription", "Horizontal angle between emitter forward and listener location in degrees."),
+							FName(),
+							{ Inputs::Azimuth, 0.0f }
+						},
+						{
+							FText(),
+							NSLOCTEXT("SourceOrientation", "ElevationDescription", "Vertical angle between emitter forward and listener location in degrees."),
+							FName(),
+							{ Inputs::Elevation, 0.0f }
+						}
+					};
+				}
+			};
+
+			static FParameterInterfacePtr InterfacePtr;
+			if (!InterfacePtr.IsValid())
+			{
+				InterfacePtr = MakeShared<FInterface>();
+			}
+
+			return InterfacePtr;
+		}
+	} // namespace SourceOrientationInterface
+#undef AUDIO_PARAMETER_INTERFACE_NAMESPACE
+
+#define AUDIO_PARAMETER_INTERFACE_NAMESPACE "UE.Listener.Orientation"
+	namespace ListenerOrientationInterface
+	{
+		const FName Name = AUDIO_PARAMETER_INTERFACE_NAMESPACE;
+
+		namespace Inputs
+		{
+			const FName Azimuth = AUDIO_PARAMETER_INTERFACE_MEMBER_DEFINE("Azimuth");
+			const FName Elevation = AUDIO_PARAMETER_INTERFACE_MEMBER_DEFINE("Elevation");
+		} // namespace Inputs
+
+		Audio::FParameterInterfacePtr GetInterface()
+		{
+			struct FInterface : public Audio::FParameterInterface
+			{
+				FInterface()
+					: FParameterInterface(ListenerOrientationInterface::Name, { 1, 0 }, *USoundBase::StaticClass())
+				{
+					Inputs =
+					{
+						{
+							FText(),
+							NSLOCTEXT("ListenerOrientation", "AzimuthDescription", "Horizontal viewing angle of the current listener in world."),
+							FName(),
+							{ Inputs::Azimuth, 0.0f }
+						},
+						{
+							FText(),
+							NSLOCTEXT("ListenerOrientation", "ElevationDescription", "Vertical viewing angle of the current listener in world."),
+							FName(),
+							{ Inputs::Elevation, 0.0f }
+						}
+					};
+				}
+			};
+
+			static FParameterInterfacePtr InterfacePtr;
+			if (!InterfacePtr.IsValid())
+			{
+				InterfacePtr = MakeShared<FInterface>();
+			}
+
+			return InterfacePtr;
+		}
+	} // namespace ListenerOrientationInterface
+#undef AUDIO_PARAMETER_INTERFACE_NAMESPACE
 } // namespace Audio
 #undef LOCTEXT_NAMESPACE
