Adding missing include.

diff --git a/Engine/Source/Runtime/Engine/Private/UnrealEngine.cpp b/Engine/Source/Runtime/Engine/Private/UnrealEngine.cpp
--- a/Engine/Source/Runtime/Engine/Private/UnrealEngine.cpp
+++ b/Engine/Source/Runtime/Engine/Private/UnrealEngine.cpp
@@ -231,6 +231,7 @@ UnrealEngine.cpp: Implements the UEngine class and helpers.
 #include "HDRHelper.h"
 
 #if WITH_DUMPGPU
+#include "RenderGraphBuilder.h"
 #endif
 
 #if WITH_DUMPGPU
