Fixed NonUnity build.

diff --git a/Engine/Source/Runtime/Engine/Private/ProfilingDebugging/LevelStreamingProfilingSubsystem.cpp b/Engine/Source/Runtime/Engine/Private/ProfilingDebugging/LevelStreamingProfilingSubsystem.cpp
--- a/Engine/Source/Runtime/Engine/Private/ProfilingDebugging/LevelStreamingProfilingSubsystem.cpp
+++ b/Engine/Source/Runtime/Engine/Private/ProfilingDebugging/LevelStreamingProfilingSubsystem.cpp
@@ -4,6 +4,7 @@
 
 #include "Engine/Level.h"
 #include "Engine/LevelStreaming.h"
+#include "Engine/World.h"
 #include "HAL/FileManager.h"
 #include "Misc/Paths.h"
 #include "Streaming/LevelStreamingDelegates.h"
diff --git a/Engine/Source/Runtime/Engine/Private/WorldPartition/DataLayer/DataLayerManager.cpp b/Engine/Source/Runtime/Engine/Private/WorldPartition/DataLayer/DataLayerManager.cpp
--- a/Engine/Source/Runtime/Engine/Private/WorldPartition/DataLayer/DataLayerManager.cpp
+++ b/Engine/Source/Runtime/Engine/Private/WorldPartition/DataLayer/DataLayerManager.cpp
@@ -11,7 +11,9 @@
 #include "WorldPartition/WorldPartitionRuntimeCell.h"
 #include "WorldPartition/WorldPartitionHelpers.h"
 #include "Engine/Canvas.h"
+#include "Engine/Engine.h"
 #include "Debug/DebugDrawService.h"
+#include "Misc/Paths.h"
 
 #if WITH_EDITOR
 #include "WorldPartition/DataLayer/WorldDataLayersActorDesc.h"
diff --git a/Engine/Source/Runtime/Engine/Private/WorldPartition/HLOD/HLODDestruction.cpp b/Engine/Source/Runtime/Engine/Private/WorldPartition/HLOD/HLODDestruction.cpp
--- a/Engine/Source/Runtime/Engine/Private/WorldPartition/HLOD/HLODDestruction.cpp
+++ b/Engine/Source/Runtime/Engine/Private/WorldPartition/HLOD/HLODDestruction.cpp
@@ -3,6 +3,7 @@
 #include "WorldPartition/HLOD/HLODDestruction.h"
 
 #include "Engine/LevelStreaming.h"
+#include "Engine/World.h"
 #include "LevelUtils.h"
 #include "UObject/ScriptInterface.h"
 #include "WorldPartition/HLOD/DestructibleHLODComponent.h"
diff --git a/Engine/Source/Runtime/Engine/Public/WorldPartition/ActorDescContainer.h b/Engine/Source/Runtime/Engine/Public/WorldPartition/ActorDescContainer.h
--- a/Engine/Source/Runtime/Engine/Public/WorldPartition/ActorDescContainer.h
+++ b/Engine/Source/Runtime/Engine/Public/WorldPartition/ActorDescContainer.h
@@ -11,6 +11,7 @@
 #include "ActorDescContainer.generated.h"
 
 class FLinkerInstancingContext;
+class UWorldPartition;
 
 UCLASS()
 class ENGINE_API UActorDescContainer : public UObject, public FActorDescList
