Do not sync load compressed data during PostLoad when SoundWave loading behavior set to RetainOnLoad

diff --git a/Engine/Source/Runtime/Engine/Private/SoundWave.cpp b/Engine/Source/Runtime/Engine/Private/SoundWave.cpp
--- a/Engine/Source/Runtime/Engine/Private/SoundWave.cpp
+++ b/Engine/Source/Runtime/Engine/Private/SoundWave.cpp
@@ -1671,7 +1671,7 @@ void USoundWave::PostLoad()
 		{
 			if (!GIsEditor && !DisableRetainingCVar && ActualLoadingBehavior == ESoundWaveLoadingBehavior::RetainOnLoad)
 			{
-				RetainCompressedAudio(true);
+				RetainCompressedAudio(false);
 			}
 			else
 			{
