Switched SparseVolumeTextures to new DDC API. Also introduced UStreamableSparseVolumeTexture as new base class for static and animated SVT classes. This is in preparation for streaming and also to clean up and deduplicate some code that could be shared. Fixed missing call to ExpandVolumeBounds() when importing an animated SVT.

#rnx

diff --git a/Engine/Source/Editor/SparseVolumeTexture/Private/SparseVolumeTextureFactory.cpp b/Engine/Source/Editor/SparseVolumeTexture/Private/SparseVolumeTextureFactory.cpp
--- a/Engine/Source/Editor/SparseVolumeTexture/Private/SparseVolumeTextureFactory.cpp
+++ b/Engine/Source/Editor/SparseVolumeTexture/Private/SparseVolumeTextureFactory.cpp
@@ -482,6 +482,7 @@ UObject* USparseVolumeTextureFactory::FactoryCreateFile(UClass* InClass, UObject
 		FBox VolumeBounds(FVector(FLT_MAX), FVector(-FLT_MAX));
 		ExpandVolumeBounds(ImportOptions, PreviewData.GridInfo, VolumeBounds);
 		StaticSVTexture->VolumeBounds = VolumeBounds;
+		StaticSVTexture->Frames.SetNum(1);
 
 		FSparseVolumeRawSource SparseVolumeRawSource{};
 
@@ -505,7 +506,7 @@ UObject* USparseVolumeTextureFactory::FactoryCreateFile(UClass* InClass, UObject
 
 		// Serialize the raw source data into the asset object.
 		{
-			UE::Serialization::FEditorBulkDataWriter RawDataArchiveWriter(StaticSVTexture->StaticFrame.RawData);
+			UE::Serialization::FEditorBulkDataWriter RawDataArchiveWriter(StaticSVTexture->Frames[0].RawData);
 			SparseVolumeRawSource.Serialize(RawDataArchiveWriter);
 		}
 
@@ -535,7 +536,7 @@ UObject* USparseVolumeTextureFactory::FactoryCreateFile(UClass* InClass, UObject
 
 		// Allocate space for each frame
 		AnimatedSVTexture->FrameCount = NumFrames;
-		AnimatedSVTexture->AnimationFrames.SetNum(NumFrames);
+		AnimatedSVTexture->Frames.SetNum(NumFrames);
 
 		FBox VolumeBounds(FVector(FLT_MAX), FVector(-FLT_MAX));
 
@@ -585,6 +586,8 @@ UObject* USparseVolumeTextureFactory::FactoryCreateFile(UClass* InClass, UObject
 				}
 			}
 
+			ExpandVolumeBounds(ImportOptions, FrameGridInfo, VolumeBounds);
+
 			FSparseVolumeRawSource SparseVolumeRawSource{};
 
 			FOpenVDBToSVTConversionResult SVTResult;
@@ -607,7 +610,7 @@ UObject* USparseVolumeTextureFactory::FactoryCreateFile(UClass* InClass, UObject
 
 			// Serialize the raw source data from this frame into the asset object.
 			{
-				UE::Serialization::FEditorBulkDataWriter RawDataArchiveWriter(AnimatedSVTexture->AnimationFrames[FrameIdx].RawData);
+				UE::Serialization::FEditorBulkDataWriter RawDataArchiveWriter(AnimatedSVTexture->Frames[FrameIdx].RawData);
 				SparseVolumeRawSource.Serialize(RawDataArchiveWriter);
 			}
 
diff --git a/Engine/Source/Runtime/Engine/Classes/SparseVolumeTexture/SparseVolumeTexture.h b/Engine/Source/Runtime/Engine/Classes/SparseVolumeTexture/SparseVolumeTexture.h
--- a/Engine/Source/Runtime/Engine/Classes/SparseVolumeTexture/SparseVolumeTexture.h
+++ b/Engine/Source/Runtime/Engine/Classes/SparseVolumeTexture/SparseVolumeTexture.h
@@ -14,6 +14,7 @@
 #include "SparseVolumeTexture.generated.h"
 
 namespace UE { namespace Shader	{ enum class EValueType : uint8; } }
+namespace UE { namespace DerivedData { class FRequestOwner; } }
 
 #define SPARSE_VOLUME_TILE_RES	16
 
@@ -100,7 +101,7 @@ public:
 	FSparseVolumeTextureSceneProxy();
 	virtual ~FSparseVolumeTextureSceneProxy() override;
 
-	void InitialiseRuntimeData(FSparseVolumeTextureRuntime& SparseVolumeTextureRuntime);
+	void InitializeRuntimeData(FSparseVolumeTextureRuntime& SparseVolumeTextureRuntime);
 
 	const FSparseVolumeAssetHeader& GetHeader() const 
 	{
@@ -136,10 +137,6 @@ private:
 
 struct ENGINE_API FSparseVolumeTextureFrame
 {
-	FSparseVolumeTextureFrame();
-	virtual ~FSparseVolumeTextureFrame();
-	bool BuildRuntimeData();
-
 	// The frame data that can be streamed in when in game.
 	FByteBulkData						RuntimeStreamedInData;
 
@@ -153,6 +150,11 @@ struct ENGINE_API FSparseVolumeTextureFrame
 	/** The raw data that can be loaded when we want to update cook the data with different settings or updated code without re importing. */
 	UE::Serialization::FEditorBulkData	RawData;
 #endif
+
+	FSparseVolumeTextureFrame();
+	virtual ~FSparseVolumeTextureFrame();
+	bool BuildRuntimeData();
+	void Serialize(FArchive& Ar, UStreamableSparseVolumeTexture* Owner, int32 FrameIndex);
 };
 
 
@@ -178,7 +180,6 @@ public:
 
 	virtual int32 GetFrameCount() const { return 0; }
 	virtual const FSparseVolumeAssetHeader* GetSparseVolumeTextureHeader() const { return nullptr; }
-	virtual FSparseVolumeTextureSceneProxy* GetSparseVolumeTextureSceneProxy() { return nullptr; }
 	virtual const FSparseVolumeTextureSceneProxy* GetSparseVolumeTextureSceneProxy() const { return nullptr; }
 	virtual FBox GetVolumeBounds() const { return FBox(); }
 
@@ -196,36 +197,55 @@ public:
 private:
 };
 
-
 UCLASS(ClassGroup = Rendering, BlueprintType)//, hidecategories = (Object))
-class ENGINE_API UStaticSparseVolumeTexture : public USparseVolumeTexture
+class UStreamableSparseVolumeTexture : public USparseVolumeTexture
 {
 	GENERATED_UCLASS_BODY()
 
 public:
 
-	UPROPERTY(VisibleAnywhere, Category = "Rendering", meta = (DisplayName = "Volume Bounds"))
-	FBox VolumeBounds;
-
-	UStaticSparseVolumeTexture();
-	virtual ~UStaticSparseVolumeTexture() = default;
+	UStreamableSparseVolumeTexture();
+	virtual ~UStreamableSparseVolumeTexture() = default;
 
 	//~ Begin UObject Interface.
-	virtual void GetResourceSizeEx(FResourceSizeEx& CumulativeResourceSize) override;
 	virtual void PostLoad() override;
 	virtual void BeginDestroy() override;
 	virtual void Serialize(FArchive& Ar) override;
-	virtual void PreSave(FObjectPreSaveContext ObjectSaveContext) override;
 #if WITH_EDITOR
 	virtual void PostEditChangeProperty(FPropertyChangedEvent& PropertyChangedEvent) override;
 #endif
 	//~ End UObject Interface.
 
+	const FSparseVolumeTextureSceneProxy* GetStreamedFrameProxyOrFallback(int32 FrameIndex) const;
+
+protected:
+	TArray<FSparseVolumeTextureFrame, TInlineAllocator<1u>> Frames;
+
+	void GenerateOrLoadDDCRuntimeDataAndCreateSceneProxy();
+	void GenerateOrLoadDDCRuntimeDataForFrame(FSparseVolumeTextureFrame& Frame, UE::DerivedData::FRequestOwner& DDCRequestOwner);
+	void ConvertRawSourceDataToSparseVolumeTextureRuntime(FSparseVolumeTextureFrame& Frame);
+};
+
+UCLASS(ClassGroup = Rendering, BlueprintType)//, hidecategories = (Object))
+class ENGINE_API UStaticSparseVolumeTexture : public UStreamableSparseVolumeTexture
+{
+	GENERATED_UCLASS_BODY()
+
+public:
+
+	UPROPERTY(VisibleAnywhere, Category = "Rendering", meta = (DisplayName = "Volume Bounds"))
+	FBox VolumeBounds;
+
+	UStaticSparseVolumeTexture();
+	virtual ~UStaticSparseVolumeTexture() = default;
+
+	//~ Begin UObject Interface.
+	virtual void GetResourceSizeEx(FResourceSizeEx& CumulativeResourceSize) override;
+	//~ End UObject Interface.
 
 	//~ Begin USparseVolumeTexture Interface.
 	int32 GetFrameCount() const override { return 1; }
 	const FSparseVolumeAssetHeader* GetSparseVolumeTextureHeader() const override;
-	FSparseVolumeTextureSceneProxy* GetSparseVolumeTextureSceneProxy() override;
 	const FSparseVolumeTextureSceneProxy* GetSparseVolumeTextureSceneProxy() const override;
 	FBox GetVolumeBounds() const override;
 	//~ End USparseVolumeTexture Interface.
@@ -235,17 +255,11 @@ private:
 #if WITH_EDITOR
 	friend class USparseVolumeTextureFactory; // Importer
 #endif
-
-	FSparseVolumeTextureFrame StaticFrame;
-
-	void ConvertRawSourceDataToSparseVolumeTextureRuntime();
-	void GenerateOrLoadDDCRuntimeData();
-	void GenerateOrLoadDDCRuntimeDataAndCreateSceneProxy();
 };
 
 // UAnimatedSparseVolumeTexture inherit from USparseVolumeTexture to be viewed using the first frame by default.
 UCLASS(ClassGroup = Rendering, BlueprintType)//, hidecategories = (Object))
-class ENGINE_API UAnimatedSparseVolumeTexture : public USparseVolumeTexture
+class ENGINE_API UAnimatedSparseVolumeTexture : public UStreamableSparseVolumeTexture
 {
 	GENERATED_UCLASS_BODY()
 
@@ -263,26 +277,17 @@ public:
 
 	//~ Begin UObject Interface.
 	virtual void GetResourceSizeEx(FResourceSizeEx& CumulativeResourceSize) override;
-	virtual void PostLoad() override;
-	virtual void BeginDestroy() override;
-	virtual void Serialize(FArchive& Ar) override;
-	virtual void PreSave(FObjectPreSaveContext ObjectSaveContext) override;
-#if WITH_EDITOR
-	virtual void PostEditChangeProperty(FPropertyChangedEvent& PropertyChangedEvent) override;
-#endif
 	//~ End UObject Interface.
 
-
 	//~ Begin USparseVolumeTexture Interface.
 	int32 GetFrameCount() const override { return FrameCount; }
 	const FSparseVolumeAssetHeader* GetSparseVolumeTextureHeader() const override;
-	FSparseVolumeTextureSceneProxy* GetSparseVolumeTextureSceneProxy() override;
 	const FSparseVolumeTextureSceneProxy* GetSparseVolumeTextureSceneProxy() const override;
 	FBox GetVolumeBounds() const override;
 	//~ End USparseVolumeTexture Interface.
 
 	// Used for debugging a specific frame of an animated sequence.
-	FSparseVolumeTextureSceneProxy* GetSparseVolumeTextureFrameSceneProxy(int32 FrameIndex);
+	const FSparseVolumeTextureSceneProxy* GetSparseVolumeTextureFrameSceneProxy(int32 FrameIndex) const;
 	const FSparseVolumeAssetHeader* GetSparseVolumeTextureFrameHeader(int32 FrameIndex) const;
 
 private:
@@ -291,14 +296,7 @@ private:
 	friend class USparseVolumeTextureFactory; // Importer
 #endif
 	
-	const bool bLoadAllFramesToProxies = true;	// SVT_TODO remove that once streaming is working
-	TArray<FSparseVolumeTextureFrame> AnimationFrames;
 	int32 PreviewFrameIndex;
-
-	int32 GetFrameCountToLoad() const;
-	void ConvertRawSourceDataToSparseVolumeTextureRuntime(int32 FrameIndex);
-	void GenerateOrLoadDDCRuntimeData(int32 FrameIndex);
-	void GenerateOrLoadDDCRuntimeDataAndCreateSceneProxy(int32 FrameIndex);
 };
 
 // USparseVolumeTextureFrame inherits from USparseVolumeTexture to be viewed using any given frame of a UAnimatedSparseVolumeTexture (or UStaticSparseVolumeTexture)
@@ -313,18 +311,17 @@ public:
 
 	static USparseVolumeTextureFrame* CreateFrame(USparseVolumeTexture* Texture, int32 FrameIndex);
 
-	void Init(FSparseVolumeTextureSceneProxy* InSceneProxy, const FSparseVolumeAssetHeader* InAssetHeader, const FBox& InVolumeBounds);
+	void Initialize(const FSparseVolumeTextureSceneProxy* InSceneProxy, const FSparseVolumeAssetHeader* InAssetHeader, const FBox& InVolumeBounds);
 
 	//~ Begin USparseVolumeTexture Interface.
 	int32 GetFrameCount() const override { return 1; }
 	const FSparseVolumeAssetHeader* GetSparseVolumeTextureHeader() const override;
-	FSparseVolumeTextureSceneProxy* GetSparseVolumeTextureSceneProxy() override;
 	const FSparseVolumeTextureSceneProxy* GetSparseVolumeTextureSceneProxy() const override;
 	FBox GetVolumeBounds() const override;
 	//~ End USparseVolumeTexture Interface.
 
 private:
-	FSparseVolumeTextureSceneProxy* SceneProxy;
+	const FSparseVolumeTextureSceneProxy* SceneProxy;
 	const FSparseVolumeAssetHeader* AssetHeader;
 	FBox VolumeBounds;
 };
diff --git a/Engine/Source/Runtime/Engine/Private/SparseVolumeTexture/SparseVolumeTexture.cpp b/Engine/Source/Runtime/Engine/Private/SparseVolumeTexture/SparseVolumeTexture.cpp
--- a/Engine/Source/Runtime/Engine/Private/SparseVolumeTexture/SparseVolumeTexture.cpp
+++ b/Engine/Source/Runtime/Engine/Private/SparseVolumeTexture/SparseVolumeTexture.cpp
@@ -11,7 +11,8 @@
 #include "Shader/ShaderTypes.h"
 
 #if WITH_EDITORONLY_DATA
-#include "DerivedDataCacheInterface.h"
+#include "DerivedDataCache.h"
+#include "DerivedDataRequestOwner.h"
 #endif
 
 #include "Serialization/LargeMemoryReader.h"
@@ -134,7 +135,7 @@ FSparseVolumeTextureSceneProxy::~FSparseVolumeTextureSceneProxy()
 {
 }
 
-void FSparseVolumeTextureSceneProxy::InitialiseRuntimeData(FSparseVolumeTextureRuntime& SparseVolumeTextureRuntimeIn)
+void FSparseVolumeTextureSceneProxy::InitializeRuntimeData(FSparseVolumeTextureRuntime& SparseVolumeTextureRuntimeIn)
 {
 	SparseVolumeTextureRuntime = &SparseVolumeTextureRuntimeIn;
 }
@@ -240,6 +241,61 @@ bool FSparseVolumeTextureFrame::BuildRuntimeData()
 	return false;
 }
 
+void FSparseVolumeTextureFrame::Serialize(FArchive& Ar, UStreamableSparseVolumeTexture* Owner, int32 FrameIndex)
+{
+	FStripDataFlags StripFlags(Ar);
+
+	const bool bInlinePayload = (FrameIndex == 0) || true; // SVT_TODO: remove once streaming is working
+	RuntimeStreamedInData.SetBulkDataFlags(bInlinePayload ? BULKDATA_ForceInlinePayload : BULKDATA_Force_NOT_InlinePayload);
+
+	if (StripFlags.IsEditorDataStripped() && Ar.IsLoadingFromCookedPackage())
+	{
+		// In this case we are loading in game with a cooked build so we only need to load the runtime data.
+
+		// Read cooked bulk data from archive
+		RuntimeStreamedInData.Serialize(Ar, Owner);
+
+		// Create runtime data from cooked bulk data
+		{
+			FBulkDataReader BulkDataReader(RuntimeStreamedInData);
+			SparseVolumeTextureRuntime.Serialize(BulkDataReader);
+		}
+
+		// The bulk data is no longer needed
+		RuntimeStreamedInData.RemoveBulkData();
+
+		// Runtime data is now valid, create the render thread proxy
+		SparseVolumeTextureSceneProxy.InitializeRuntimeData(SparseVolumeTextureRuntime);
+		BeginInitResource(&SparseVolumeTextureSceneProxy);
+	}
+	else if (Ar.IsCooking())
+	{
+		// We are cooking the game, serialize the asset out.
+
+		const bool bBuiltRuntimeData = BuildRuntimeData();
+		check(bBuiltRuntimeData); // SVT_TODO: actual error handling
+
+		// Write runtime data into RuntimeStreamedInData
+		{
+			FBulkDataWriter BulkDataWriter(RuntimeStreamedInData);
+			SparseVolumeTextureRuntime.Serialize(BulkDataWriter);
+		}
+
+		// And now write the cooked bulk data to the archive
+		RuntimeStreamedInData.Serialize(Ar, Owner);
+	}
+	else if (!Ar.IsObjectReferenceCollector())
+	{
+#if WITH_EDITORONLY_DATA
+		// When in EDITOR:
+		//  - We only serialize raw data 
+		//  - The runtime data is fetched/put from/to DDC
+		//  - This EditorBulk data do not load the full and huge OpenVDB data. That is only done explicitly later.
+		RawData.Serialize(Ar, Owner);
+#endif
+	}
+}
+
 ////////////////////////////////////////////////////////////////////////////////////////////////
 
 USparseVolumeTexture::USparseVolumeTexture(const FObjectInitializer& ObjectInitializer)
@@ -325,96 +381,47 @@ UE::Shader::EValueType USparseVolumeTexture::GetUniformParameterType(int32 Index
 
 ////////////////////////////////////////////////////////////////////////////////////////////////
 
-UStaticSparseVolumeTexture::UStaticSparseVolumeTexture(const FObjectInitializer& ObjectInitializer)
+UStreamableSparseVolumeTexture::UStreamableSparseVolumeTexture(const FObjectInitializer& ObjectInitializer)
 	: Super(ObjectInitializer)
-	, StaticFrame()
-{
-}
-
-void UStaticSparseVolumeTexture::GetResourceSizeEx(FResourceSizeEx& CumulativeResourceSize)
 {
-	Super::GetResourceSizeEx(CumulativeResourceSize);
-	// CumulativeResourceSize.AddDedicatedSystemMemoryBytes but not the RawData size
 }
 
-void UStaticSparseVolumeTexture::PostLoad()
+void UStreamableSparseVolumeTexture::PostLoad()
 {
 	GenerateOrLoadDDCRuntimeDataAndCreateSceneProxy();
 
 	Super::PostLoad();
 }
 
-void UStaticSparseVolumeTexture::BeginDestroy()
+void UStreamableSparseVolumeTexture::BeginDestroy()
 {
 	Super::BeginDestroy();
-
-	BeginReleaseResource(&StaticFrame.SparseVolumeTextureSceneProxy);
+	for (FSparseVolumeTextureFrame& Frame : Frames)
+	{
+		BeginReleaseResource(&Frame.SparseVolumeTextureSceneProxy);
+	}
 }
 
-void UStaticSparseVolumeTexture::Serialize(FArchive& Ar)
+void UStreamableSparseVolumeTexture::Serialize(FArchive& Ar)
 {
 	Super::Serialize(Ar);
 
-	FStripDataFlags StripFlags(Ar);
+	int32 NumFrames = Frames.Num();
+	Ar << NumFrames;
 
-	if (StripFlags.IsEditorDataStripped() && Ar.IsLoadingFromCookedPackage())
+	if (Ar.IsLoading())
 	{
-		// In this case we are loading in game with a cooked build so we only need to load the runtime data.
-
-		// Read cooked bulk data from archive
-		StaticFrame.RuntimeStreamedInData.Serialize(Ar, this);
-
-		// Create runtime data from cooked bulk data
-		{
-			FBulkDataReader BulkDataReader(StaticFrame.RuntimeStreamedInData);
-			StaticFrame.SparseVolumeTextureRuntime.Serialize(BulkDataReader);
-		}
-
-		// The bulk data is no longer needed
-		StaticFrame.RuntimeStreamedInData.RemoveBulkData();
-
-		// Runtime data is now valid, create the render thread proxy
-		StaticFrame.SparseVolumeTextureSceneProxy.InitialiseRuntimeData(StaticFrame.SparseVolumeTextureRuntime);
-		BeginInitResource(&StaticFrame.SparseVolumeTextureSceneProxy);
+		Frames.Reset(NumFrames);
+		Frames.AddDefaulted(NumFrames);
 	}
-	else if (Ar.IsCooking())
+	for (int32 FrameIndex = 0; FrameIndex < NumFrames; ++FrameIndex)
 	{
-		// We are cooking the game, serialize the asset out.
-		
-		// The runtime bulk data for static sparse volume texture is always loaded, not streamed in.
-		StaticFrame.RuntimeStreamedInData.SetBulkDataFlags(BULKDATA_ForceInlinePayload);
-
-		const bool bBuiltRuntimeData = StaticFrame.BuildRuntimeData();
-		check(bBuiltRuntimeData); // SVT_TODO: actual error handling
-		
-		// Write runtime data into RuntimeStreamedInData
-		{
-			FBulkDataWriter BulkDataWriter(StaticFrame.RuntimeStreamedInData);
-			StaticFrame.SparseVolumeTextureRuntime.Serialize(BulkDataWriter);
-		}
-
-		// And now write the cooked bulk data to the archive
-		StaticFrame.RuntimeStreamedInData.Serialize(Ar, this);
-	}
-	else if (!Ar.IsObjectReferenceCollector())
-	{
-#if WITH_EDITORONLY_DATA
-		// When in EDITOR:
-		//  - We only serialize raw data 
-		//  - The runtime data is fetched/put from/to DDC
-		//  - This EditorBulk data do not load the full and huge OpenVDB data. That is only done explicitly later.
-		StaticFrame.RawData.Serialize(Ar, this);
-#endif
+		Frames[FrameIndex].Serialize(Ar, this, FrameIndex);
 	}
 }
 
-void UStaticSparseVolumeTexture::PreSave(FObjectPreSaveContext ObjectSaveContext)
-{
-	Super::PreSave(ObjectSaveContext);
-}
-
 #if WITH_EDITOR
-void UStaticSparseVolumeTexture::PostEditChangeProperty(FPropertyChangedEvent& PropertyChangedEvent)
+void UStreamableSparseVolumeTexture::PostEditChangeProperty(FPropertyChangedEvent& PropertyChangedEvent)
 {
 	Super::PostEditChangeProperty(PropertyChangedEvent);
 
@@ -422,272 +429,180 @@ void UStaticSparseVolumeTexture::PostEditChangeProperty(FPropertyChangedEvent& P
 }
 #endif // WITH_EDITOR
 
-const FSparseVolumeAssetHeader* UStaticSparseVolumeTexture::GetSparseVolumeTextureHeader() const
+const FSparseVolumeTextureSceneProxy* UStreamableSparseVolumeTexture::GetStreamedFrameProxyOrFallback(int32 FrameIndex) const
 {
-	return &StaticFrame.SparseVolumeTextureRuntime.Header;
-}
-FSparseVolumeTextureSceneProxy* UStaticSparseVolumeTexture::GetSparseVolumeTextureSceneProxy()
-{ 
-	return &StaticFrame.SparseVolumeTextureSceneProxy; 
-}
-const FSparseVolumeTextureSceneProxy* UStaticSparseVolumeTexture::GetSparseVolumeTextureSceneProxy() const
-{ 
-	return &StaticFrame.SparseVolumeTextureSceneProxy; 
-}
-
-FBox UStaticSparseVolumeTexture::GetVolumeBounds() const
-{
-	return VolumeBounds;
-}
-
-void UStaticSparseVolumeTexture::ConvertRawSourceDataToSparseVolumeTextureRuntime()
-{
-#if WITH_EDITOR
-	// Check if the virtualized bulk data payload is available now
-	if (StaticFrame.RawData.HasPayloadData())
-	{
-		const bool bSuccess = StaticFrame.BuildRuntimeData();
-		ensure(bSuccess);
-	}
-	else
-	{
-		UE_LOG(LogSparseVolumeTexture, Error, TEXT("SparseVolumeTexture - Raw source data is not available for %s. Using default data."), *GetName());
-		StaticFrame.SparseVolumeTextureRuntime.SetAsDefaultTexture();
-	}
-#endif // WITH_EDITOR
+	check(!Frames.IsEmpty());
+	FrameIndex = FrameIndex % Frames.Num();
+	return &Frames[FrameIndex].SparseVolumeTextureSceneProxy;
 }
 
-void UStaticSparseVolumeTexture::GenerateOrLoadDDCRuntimeData()
+void UStreamableSparseVolumeTexture::GenerateOrLoadDDCRuntimeDataAndCreateSceneProxy()
 {
 #if WITH_EDITORONLY_DATA
-	static const FString SparseVolumeTextureDDCVersion = TEXT("381AE2A9-A903-4C8F-8486-891E24D6EC70");	// Bump this if you want to ignore all cached data so far.
-	const FString DerivedDataKey = StaticFrame.RawData.GetIdentifier().ToString() + SparseVolumeTextureDDCVersion;
-
-	bool bSuccess = true;
-	TArray<uint8> DerivedData;
-	if (GetDerivedDataCacheRef().GetSynchronous(*DerivedDataKey, DerivedData, GetPathName()))
+	UE::DerivedData::FRequestOwner DDCRequestOwner(UE::DerivedData::EPriority::Normal);
 	{
-		UE_LOG(LogSparseVolumeTexture, Display, TEXT("SparseVolumeTexture - Caching %s"), *GetName());
-
-		FMemoryReader Ar(DerivedData, /*bIsPersistent=*/ true);
-
-		int64 UncompressedSize = 0;
-		Ar << UncompressedSize;
+		UE::DerivedData::FRequestBarrier DDCRequestBarrier(DDCRequestOwner);
+		for (FSparseVolumeTextureFrame& Frame : Frames)
+		{
+			// Release any previously allocated render thread proxy
+			BeginReleaseResource(&Frame.SparseVolumeTextureSceneProxy);
 
-		uint8* DecompressionBuffer = reinterpret_cast<uint8*>(FMemory::Malloc(UncompressedSize));
-		Ar.SerializeCompressedNew(DecompressionBuffer, UncompressedSize);
+			GenerateOrLoadDDCRuntimeDataForFrame(Frame, DDCRequestOwner);
+		}
+	}
 
-		FLargeMemoryReader LargeMemReader(DecompressionBuffer, UncompressedSize, ELargeMemoryReaderFlags::Persistent | ELargeMemoryReaderFlags::TakeOwnership);
+	// Wait for all DDC requests to complete before creating the proxies
+	DDCRequestOwner.Wait();
 
-		StaticFrame.SparseVolumeTextureRuntime.Serialize(LargeMemReader);
+	for (FSparseVolumeTextureFrame& Frame : Frames)
+	{
+		// Runtime data is now valid, create the render thread proxy
+		Frame.SparseVolumeTextureSceneProxy.InitializeRuntimeData(Frame.SparseVolumeTextureRuntime);
+		BeginInitResource(&Frame.SparseVolumeTextureSceneProxy);
 	}
-	else
+	
+#else
+	// SVT_TODO do we really need to do this in the cooked case?
+	for (FSparseVolumeTextureFrame& Frame : Frames)
 	{
-		// Using a LargeMemoryWriter for serialization since the data can be bigger than 2 GB
-		FLargeMemoryWriter LargeMemWriter(0, /*bIsPersistent=*/ true);
+		// Release any previously allocated render thread proxy
+		BeginReleaseResource(&Frame.SparseVolumeTextureSceneProxy);
 
-		ConvertRawSourceDataToSparseVolumeTextureRuntime();
-		StaticFrame.SparseVolumeTextureRuntime.Serialize(LargeMemWriter);
+		// Create the render thread proxy
+		Frame.SparseVolumeTextureSceneProxy.InitializeRuntimeData(Frame.SparseVolumeTextureRuntime);
+		BeginInitResource(&Frame.SparseVolumeTextureSceneProxy);
+	}
+#endif
+}
 
-		int64 UncompressedSize = LargeMemWriter.TotalSize();
+void UStreamableSparseVolumeTexture::GenerateOrLoadDDCRuntimeDataForFrame(FSparseVolumeTextureFrame& Frame, UE::DerivedData::FRequestOwner& DDCRequestOwner)
+{
+#if WITH_EDITORONLY_DATA
+	using namespace UE::DerivedData;
 
-		// Since the DDC doesn't support data bigger than 2 GB, we only cache for such uncompressed size.
-		constexpr int64 SizeThreshold = 2147483648;	// 2GB
-		const bool bIsCacheable = UncompressedSize < SizeThreshold;
-		if (bIsCacheable)
-		{
-			FMemoryWriter CompressedArchive(DerivedData, true);
+	static const FString SparseVolumeTextureDDCVersion = TEXT("381AE2A9-A903-4C8F-8486-891E24D6EC70");	// Bump this if you want to ignore all cached data so far.
+	const FString DerivedDataKey = Frame.RawData.GetIdentifier().ToString() + SparseVolumeTextureDDCVersion;
 
-			CompressedArchive << UncompressedSize; // needed for allocating decompression buffer
-			CompressedArchive.SerializeCompressedNew(LargeMemWriter.GetData(), UncompressedSize);
+	const FCacheKey Key = ConvertLegacyCacheKey(DerivedDataKey);
+	const FSharedString Name = MakeStringView(GetPathName());
 
-			GetDerivedDataCacheRef().Put(*DerivedDataKey, DerivedData, GetPathName());
-		}
-		else
+	UE::DerivedData::GetCache().GetValue({ {Name, Key} }, DDCRequestOwner,
+		[this, &Frame, &DDCRequestOwner](FCacheGetValueResponse&& Response)
 		{
-			UE_LOG(LogSparseVolumeTexture, Error, TEXT("SparseVolumeTexture - the asset is too large to fit in Derived Data Cache %s"), *GetName());
-		}
-	}
+			if (Response.Status == EStatus::Ok)
+			{
+				DDCRequestOwner.LaunchTask(TEXT("UStreamableSparseVolumeTexture_DerivedDataLoad"),
+					[this, &Frame, Value = MoveTemp(Response.Value)]()
+					{
+						FSharedBuffer Data = Value.GetData().Decompress();
+						FMemoryReaderView Ar(Data, true /*bIsPersistent*/);
+						Frame.SparseVolumeTextureRuntime.Serialize(Ar);
+					});
+			}
+			else if (Response.Status == EStatus::Error)
+			{
+				DDCRequestOwner.LaunchTask(TEXT("UStreamableSparseVolumeTexture_DerivedDataBuild"),
+					[this, &Frame, &DDCRequestOwner, Name = Response.Name, Key = Response.Key]()
+					{
+						ConvertRawSourceDataToSparseVolumeTextureRuntime(Frame);
+
+						// Using a LargeMemoryWriter for serialization since the data can be bigger than 2 GB
+						FLargeMemoryWriter LargeMemWriter(0, /*bIsPersistent=*/ true);
+						Frame.SparseVolumeTextureRuntime.Serialize(LargeMemWriter);
+
+						int64 UncompressedSize = LargeMemWriter.TotalSize();
+
+						// Since the DDC doesn't support data bigger than 2 GB, we only cache for such uncompressed size.
+						constexpr int64 SizeThreshold = 2147483648;	// 2GB
+						const bool bIsCacheable = UncompressedSize < SizeThreshold;
+						if (bIsCacheable)
+						{
+							FValue Value = FValue::Compress(FSharedBuffer::MakeView(LargeMemWriter.GetView()));
+							UE::DerivedData::GetCache().PutValue({ {Name, Key, Value} }, DDCRequestOwner);
+						}
+						else
+						{
+							UE_LOG(LogSparseVolumeTexture, Error, TEXT("SparseVolumeTexture - the asset is too large to fit in Derived Data Cache %s"), *GetName());
+						}
+					});
+			}
+		});
 #endif // WITH_EDITORONLY_DATA
 }
 
-void UStaticSparseVolumeTexture::GenerateOrLoadDDCRuntimeDataAndCreateSceneProxy()
+void UStreamableSparseVolumeTexture::ConvertRawSourceDataToSparseVolumeTextureRuntime(FSparseVolumeTextureFrame& Frame)
 {
-	// Release any previously allocated render thread proxy
-	BeginReleaseResource(&StaticFrame.SparseVolumeTextureSceneProxy);
-
 #if WITH_EDITORONLY_DATA
-	// We only fetch/put DDC when in editor. Otherwise, StaticFrame.SparseVolumeTextureRuntime is serialize in.
-	GenerateOrLoadDDCRuntimeData();
-#endif
-
-	// Runtime data is now valid, create the render thread proxy
-	StaticFrame.SparseVolumeTextureSceneProxy.InitialiseRuntimeData(StaticFrame.SparseVolumeTextureRuntime);
-	BeginInitResource(&StaticFrame.SparseVolumeTextureSceneProxy);
+	// Check if the virtualized bulk data payload is available now
+	if (Frame.RawData.HasPayloadData())
+	{
+		const bool bSuccess = Frame.BuildRuntimeData();
+		ensure(bSuccess);
+	}
+	else
+	{
+		UE_LOG(LogSparseVolumeTexture, Error, TEXT("SparseVolumeTexture - Raw source data is not available for %s. Using default data."), *GetName());
+		Frame.SparseVolumeTextureRuntime.SetAsDefaultTexture();
+	}
+#endif // WITH_EDITORONLY_DATA
 }
 
 ////////////////////////////////////////////////////////////////////////////////////////////////
 
-UAnimatedSparseVolumeTexture::UAnimatedSparseVolumeTexture(const FObjectInitializer& ObjectInitializer)
+UStaticSparseVolumeTexture::UStaticSparseVolumeTexture(const FObjectInitializer& ObjectInitializer)
 	: Super(ObjectInitializer)
-	, FrameCount(0)
-	, AnimationFrames()
 {
 }
 
-void UAnimatedSparseVolumeTexture::GetResourceSizeEx(FResourceSizeEx& CumulativeResourceSize)
+void UStaticSparseVolumeTexture::GetResourceSizeEx(FResourceSizeEx& CumulativeResourceSize)
 {
 	Super::GetResourceSizeEx(CumulativeResourceSize);
 	// CumulativeResourceSize.AddDedicatedSystemMemoryBytes but not the RawData size
 }
 
-void UAnimatedSparseVolumeTexture::PostLoad()
+const FSparseVolumeAssetHeader* UStaticSparseVolumeTexture::GetSparseVolumeTextureHeader() const
 {
-	const int32 FrameCountToLoad = GetFrameCountToLoad();
-	for (int32 FrameIndex = 0; FrameIndex < FrameCountToLoad; FrameIndex++)
-	{
-		GenerateOrLoadDDCRuntimeDataAndCreateSceneProxy(FrameIndex);
-	}
-
-	Super::PostLoad();
+	return &GetStreamedFrameProxyOrFallback(0)->GetHeader();
 }
 
-void UAnimatedSparseVolumeTexture::BeginDestroy()
-{
-	Super::BeginDestroy();
-
-	const int32 FrameCountToLoad = GetFrameCountToLoad();
-	for(int32 FrameIndex = 0; FrameIndex < FrameCountToLoad; FrameIndex++)
-	{
-		FSparseVolumeTextureFrame& Frame = AnimationFrames[FrameIndex];
-		BeginReleaseResource(&Frame.SparseVolumeTextureSceneProxy);
-	}
+const FSparseVolumeTextureSceneProxy* UStaticSparseVolumeTexture::GetSparseVolumeTextureSceneProxy() const
+{ 
+	return GetStreamedFrameProxyOrFallback(0);
 }
 
-void UAnimatedSparseVolumeTexture::Serialize(FArchive& Ar)
+FBox UStaticSparseVolumeTexture::GetVolumeBounds() const
 {
-	Super::Serialize(Ar);
-
-	FStripDataFlags StripFlags(Ar);
-
-	if (StripFlags.IsEditorDataStripped() && Ar.IsLoadingFromCookedPackage())
-	{
-		// In this case we are loading in game with a cooked build so we only need to load the runtime data.
-
-		AnimationFrames.SetNum(FrameCount);
-		for (int32 i = 0; i < FrameCount; ++i)
-		{
-			FSparseVolumeTextureFrame& Frame = AnimationFrames[i];
-			
-			// Read cooked bulk data from archive
-			Frame.RuntimeStreamedInData.Serialize(Ar, this);
-
-			// Create runtime data from cooked bulk data
-			{
-				FBulkDataReader RuntimeStreamedInData(Frame.RuntimeStreamedInData);
-				Frame.SparseVolumeTextureRuntime.Serialize(RuntimeStreamedInData);
-			}
-
-			// The bulk data is no longer needed
-			Frame.RuntimeStreamedInData.RemoveBulkData();
-
-			// Runtime data is now valid, create the render thread proxy
-			Frame.SparseVolumeTextureSceneProxy.InitialiseRuntimeData(Frame.SparseVolumeTextureRuntime);
-			BeginInitResource(&Frame.SparseVolumeTextureSceneProxy);
-		}
-	}
-	else if (Ar.IsCooking())
-	{
-		// We are cooking the game, serialize the asset out.
-
-		check(AnimationFrames.Num() == FrameCount);
-		for (int32 i = 0; i < FrameCount; ++i)
-		{
-			FSparseVolumeTextureFrame& Frame = AnimationFrames[i];
-
-			// SVT_TODO: Until we have streaming, the runtime bulk data for all sparse volume texture frames is always loaded.
-			AnimationFrames[i].RuntimeStreamedInData.SetBulkDataFlags(BULKDATA_ForceInlinePayload);
-
-			const bool bBuiltRuntimeData = Frame.BuildRuntimeData();
-			check(bBuiltRuntimeData); // SVT_TODO
-
-			// Write runtime data into RuntimeStreamedInData
-			{
-				FBulkDataWriter RuntimeStreamedInData(Frame.RuntimeStreamedInData);
-				Frame.SparseVolumeTextureRuntime.Serialize(RuntimeStreamedInData);
-			}
-
-			// And now write the cooked bulk data to the archive
-			Frame.RuntimeStreamedInData.Serialize(Ar, this);
-		}
-	}
-	else if (!Ar.IsObjectReferenceCollector())
-	{
-#if WITH_EDITORONLY_DATA
-		// When in EDITOR:
-		//  - We only serialize raw data 
-		//  - The runtime data is fetched/put from/to DDC
-		//  - This EditorBulk data do not load the full and huge OpenVDB data. That is only done explicitly later.
-		if (Ar.IsSaving())
-		{
-			check(AnimationFrames.Num() == FrameCount);
-		}
-		else if (Ar.IsLoading())
-		{
-			AnimationFrames.SetNum(FrameCount);
-		}
-
-		for (int32 i = 0; i < FrameCount; ++i)
-		{
-			AnimationFrames[i].RawData.Serialize(Ar, this);
-		}
-#endif
-	}
+	return VolumeBounds;
 }
 
-void UAnimatedSparseVolumeTexture::PreSave(FObjectPreSaveContext ObjectSaveContext)
-{
-	Super::PreSave(ObjectSaveContext);
-}
+////////////////////////////////////////////////////////////////////////////////////////////////
 
-#if WITH_EDITOR
-void UAnimatedSparseVolumeTexture::PostEditChangeProperty(FPropertyChangedEvent& PropertyChangedEvent)
+UAnimatedSparseVolumeTexture::UAnimatedSparseVolumeTexture(const FObjectInitializer& ObjectInitializer)
+	: Super(ObjectInitializer)
+	, FrameCount(0)
 {
-	Super::PostEditChangeProperty(PropertyChangedEvent);
-
-	const int32 FrameCountToLoad = GetFrameCountToLoad();
-	for (int32 FrameIndex = 0; FrameIndex < FrameCountToLoad; FrameIndex++)
-	{
-		GenerateOrLoadDDCRuntimeDataAndCreateSceneProxy(FrameIndex);
-	}
 }
-#endif // WITH_EDITOR
 
-const FSparseVolumeAssetHeader* UAnimatedSparseVolumeTexture::GetSparseVolumeTextureHeader() const
+void UAnimatedSparseVolumeTexture::GetResourceSizeEx(FResourceSizeEx& CumulativeResourceSize)
 {
-	// When an AnimatedSparseVolumeTexture is used as SparseVolumeTexture, it can only be previewed using its first frame.
-	check(AnimationFrames.Num() >= 1);
-	const int32 FrameIndex = PreviewFrameIndex % GetFrameCountToLoad();
-	const FSparseVolumeTextureFrame& Frame = AnimationFrames[FrameIndex];
-	return &Frame.SparseVolumeTextureRuntime.Header;
+	Super::GetResourceSizeEx(CumulativeResourceSize);
+	// CumulativeResourceSize.AddDedicatedSystemMemoryBytes but not the RawData size
 }
 
-FSparseVolumeTextureSceneProxy* UAnimatedSparseVolumeTexture::GetSparseVolumeTextureSceneProxy()
+const FSparseVolumeAssetHeader* UAnimatedSparseVolumeTexture::GetSparseVolumeTextureHeader() const
 {
 	// When an AnimatedSparseVolumeTexture is used as SparseVolumeTexture, it can only be previewed using its first frame.
-	check(AnimationFrames.Num() >= 1);
-	const int32 FrameIndex = PreviewFrameIndex % GetFrameCountToLoad();
-	FSparseVolumeTextureFrame& Frame = AnimationFrames[FrameIndex];
-	return &Frame.SparseVolumeTextureSceneProxy;
+	check(FrameCount >= 1);
+	const int32 FrameIndex = PreviewFrameIndex % FrameCount;
+	return &GetStreamedFrameProxyOrFallback(FrameIndex)->GetHeader();
 }
 
 const FSparseVolumeTextureSceneProxy* UAnimatedSparseVolumeTexture::GetSparseVolumeTextureSceneProxy() const
 {
 	// When an AnimatedSparseVolumeTexture is used as SparseVolumeTexture, it can only be previewed using its first frame.
-	check(AnimationFrames.Num() >= 1);
-	const int32 FrameIndex = PreviewFrameIndex % GetFrameCountToLoad();
-	const FSparseVolumeTextureFrame& Frame = AnimationFrames[FrameIndex];
-	return &Frame.SparseVolumeTextureSceneProxy;
+	check(FrameCount >= 1);
+	const int32 FrameIndex = PreviewFrameIndex % FrameCount;
+	return GetStreamedFrameProxyOrFallback(FrameIndex);
 }
 
 FBox UAnimatedSparseVolumeTexture::GetVolumeBounds() const
@@ -695,122 +610,18 @@ FBox UAnimatedSparseVolumeTexture::GetVolumeBounds() const
 	return VolumeBounds;
 }
 
-FSparseVolumeTextureSceneProxy* UAnimatedSparseVolumeTexture::GetSparseVolumeTextureFrameSceneProxy(int32 FrameIndex)
+const FSparseVolumeTextureSceneProxy* UAnimatedSparseVolumeTexture::GetSparseVolumeTextureFrameSceneProxy(int32 FrameIndex) const
 {
-	check(AnimationFrames.Num() >= 1);
-	FrameIndex = FrameIndex % GetFrameCountToLoad();
-	// SVT_TODO when streaming is enabled, this will likely change.
-	return &AnimationFrames[FrameIndex].SparseVolumeTextureSceneProxy;
+	check(FrameCount >= 1);
+	FrameIndex = FrameIndex % FrameCount;
+	return GetStreamedFrameProxyOrFallback(FrameIndex);
 }
 
 const FSparseVolumeAssetHeader* UAnimatedSparseVolumeTexture::GetSparseVolumeTextureFrameHeader(int32 FrameIndex) const
 {
-	check(AnimationFrames.Num() >= 1);
-	FrameIndex = FrameIndex % GetFrameCountToLoad();
-	// SVT_TODO when streaming is enabled, this will likely change.
-	return &AnimationFrames[FrameIndex].SparseVolumeTextureRuntime.Header;
-}
-
-int32 UAnimatedSparseVolumeTexture::GetFrameCountToLoad() const
-{
-	if (FrameCount > 0)
-	{
-		return bLoadAllFramesToProxies ? FrameCount : 1;
-	}
-	return 0;
-}
-
-void UAnimatedSparseVolumeTexture::ConvertRawSourceDataToSparseVolumeTextureRuntime(int32 FrameIndex)
-{
-#if WITH_EDITOR
-	FSparseVolumeTextureFrame& Frame = AnimationFrames[FrameIndex];
-
-	// Check if the virtualized bulk data payload is available now
-	if (Frame.RawData.HasPayloadData())
-	{
-		const bool bSuccess = Frame.BuildRuntimeData();
-		ensure(bSuccess);
-	}
-	else
-	{
-		UE_LOG(LogSparseVolumeTexture, Error, TEXT("AnimatedSparseVolumeTexture - Raw source data is not available for %s - Frame %i. Using default data."), *GetName(), FrameIndex);
-		Frame.SparseVolumeTextureRuntime.SetAsDefaultTexture();
-	}
-#endif // WITH_EDITOR
-}
-
-void UAnimatedSparseVolumeTexture::GenerateOrLoadDDCRuntimeData(int32 FrameIndex)
-{
-#if WITH_EDITORONLY_DATA
-	static const FString SparseVolumeTextureDDCVersion = TEXT("381AE2A9-A903-4C8F-8486-891E24D6EC71");	// Bump this if you want to ignore all cached data so far.
-
-	FSparseVolumeTextureFrame& Frame = AnimationFrames[FrameIndex];
-
-	const FString DerivedDataKey = Frame.RawData.GetIdentifier().ToString() + SparseVolumeTextureDDCVersion;
-
-	bool bSuccess = true;
-	TArray<uint8> DerivedData;
-	if (GetDerivedDataCacheRef().GetSynchronous(*DerivedDataKey, DerivedData, GetPathName()))
-	{
-		UE_LOG(LogSparseVolumeTexture, Display, TEXT("SparseVolumeTexture - Caching %s"), *GetName());
-
-		FMemoryReader Ar(DerivedData, /*bIsPersistent=*/ true);
-
-		int64 UncompressedSize = 0;
-		Ar << UncompressedSize;
-
-		uint8* DecompressionBuffer = reinterpret_cast<uint8*>(FMemory::Malloc(UncompressedSize));
-		Ar.SerializeCompressedNew(DecompressionBuffer, UncompressedSize);
-
-		FLargeMemoryReader LargeMemReader(DecompressionBuffer, UncompressedSize, ELargeMemoryReaderFlags::Persistent | ELargeMemoryReaderFlags::TakeOwnership);
-
-		Frame.SparseVolumeTextureRuntime.Serialize(LargeMemReader);
-	}
-	else
-	{
-		// Using a LargeMemoryWriter for serialization since the data can be bigger than 2 GB
-		FLargeMemoryWriter LargeMemWriter(0, /*bIsPersistent=*/ true);
-
-		ConvertRawSourceDataToSparseVolumeTextureRuntime(FrameIndex);
-		Frame.SparseVolumeTextureRuntime.Serialize(LargeMemWriter);
-
-		int64 UncompressedSize = LargeMemWriter.TotalSize();
-
-		// Since the DDC doesn't support data bigger than 2 GB, we only cache for such uncompressed size.
-		constexpr int64 SizeThreshold = 2147483648;	// 2GB
-		const bool bIsCacheable = UncompressedSize < SizeThreshold;
-		if (bIsCacheable)
-		{
-			FMemoryWriter CompressedArchive(DerivedData, true);
-
-			CompressedArchive << UncompressedSize; // needed for allocating decompression buffer
-			CompressedArchive.SerializeCompressedNew(LargeMemWriter.GetData(), UncompressedSize);
-
-			GetDerivedDataCacheRef().Put(*DerivedDataKey, DerivedData, GetPathName());
-		}
-		else
-		{
-			UE_LOG(LogSparseVolumeTexture, Error, TEXT("SparseVolumeTexture - the asset is too large to fit in Derived Data Cache %s - Frame %i"), *GetName(), FrameIndex);
-		}
-	}
-#endif // WITH_EDITORONLY_DATA
-}
-
-void UAnimatedSparseVolumeTexture::GenerateOrLoadDDCRuntimeDataAndCreateSceneProxy(int32 FrameIndex)
-{
-	FSparseVolumeTextureFrame& Frame = AnimationFrames[FrameIndex];
-
-	// Release any previously allocated render thread proxy
-	BeginReleaseResource(&Frame.SparseVolumeTextureSceneProxy);
-
-#if WITH_EDITORONLY_DATA
-	// We only fetch/put DDC when in editor. Otherwise, StaticFrame.SparseVolumeTextureRuntime is serialize in.
-	GenerateOrLoadDDCRuntimeData(FrameIndex);
-#endif
-
-	// Runtime data is now valid, create the render thread proxy
-	Frame.SparseVolumeTextureSceneProxy.InitialiseRuntimeData(Frame.SparseVolumeTextureRuntime);
-	BeginInitResource(&Frame.SparseVolumeTextureSceneProxy);
+	check(FrameCount >= 1);
+	FrameIndex = FrameIndex % FrameCount;
+	return &GetStreamedFrameProxyOrFallback(FrameIndex)->GetHeader();
 }
 
 ////////////////////////////////////////////////////////////////////////////////////////////////
@@ -823,24 +634,25 @@ USparseVolumeTextureFrame::USparseVolumeTextureFrame(const FObjectInitializer& O
 USparseVolumeTextureFrame* USparseVolumeTextureFrame::CreateFrame(USparseVolumeTexture* Texture, int32 FrameIndex)
 {
 	USparseVolumeTextureFrame* Frame = NewObject<USparseVolumeTextureFrame>();
-	FSparseVolumeTextureSceneProxy* Proxy = nullptr;
-	if (Texture->IsA<UAnimatedSparseVolumeTexture>())
+	const FSparseVolumeTextureSceneProxy* Proxy = nullptr;
+	if (Texture->IsA<UStreamableSparseVolumeTexture>())
 	{
-		UAnimatedSparseVolumeTexture* AnimatedSVT = CastChecked<UAnimatedSparseVolumeTexture>(Texture);
-		check(AnimatedSVT);
-		Proxy = AnimatedSVT->GetSparseVolumeTextureFrameSceneProxy(FrameIndex);
+		UStreamableSparseVolumeTexture* StreamableSVT = CastChecked<UStreamableSparseVolumeTexture>(Texture);
+		check(StreamableSVT);
+		Proxy = StreamableSVT->GetStreamedFrameProxyOrFallback(FrameIndex);
 	}
 	else
 	{
 		Proxy = Texture->GetSparseVolumeTextureSceneProxy();
 	}
+	check(Proxy);
 
-	Frame->Init(Proxy, Texture->GetSparseVolumeTextureHeader(), Texture->GetVolumeBounds());
+	Frame->Initialize(Proxy, &Proxy->GetHeader(), Texture->GetVolumeBounds());
 
 	return Frame;
 }
 
-void USparseVolumeTextureFrame::Init(FSparseVolumeTextureSceneProxy* InSceneProxy, const FSparseVolumeAssetHeader* InAssetHeader, const FBox& InVolumeBounds)
+void USparseVolumeTextureFrame::Initialize(const FSparseVolumeTextureSceneProxy* InSceneProxy, const FSparseVolumeAssetHeader* InAssetHeader, const FBox& InVolumeBounds)
 {
 	SceneProxy = InSceneProxy;
 	AssetHeader = InAssetHeader;
@@ -852,11 +664,6 @@ const FSparseVolumeAssetHeader* USparseVolumeTextureFrame::GetSparseVolumeTextur
 	return AssetHeader;
 }
 
-FSparseVolumeTextureSceneProxy* USparseVolumeTextureFrame::GetSparseVolumeTextureSceneProxy()
-{
-	return SceneProxy;
-}
-
 const FSparseVolumeTextureSceneProxy* USparseVolumeTextureFrame::GetSparseVolumeTextureSceneProxy() const
 {
 	return SceneProxy;
