- Fixed deprecated UDataLayerSubsystem::GetDataLayerInstances to properly map to UDataLayerManager::GetDataLayerInstances.
- Added other missing deprecated UDataLayerSubsystem methods.

diff --git a/Engine/Source/Runtime/Engine/Public/WorldPartition/DataLayer/DataLayerSubsystem.h b/Engine/Source/Runtime/Engine/Public/WorldPartition/DataLayer/DataLayerSubsystem.h
--- a/Engine/Source/Runtime/Engine/Public/WorldPartition/DataLayer/DataLayerSubsystem.h
+++ b/Engine/Source/Runtime/Engine/Public/WorldPartition/DataLayer/DataLayerSubsystem.h
@@ -89,27 +89,75 @@ public:
 	UE_DEPRECATED(5.3, "Use UDataLayerManager::GetDataLayerEffectiveRuntimeStateByName instead.")
 	EDataLayerRuntimeState GetDataLayerEffectiveRuntimeStateByName(const FName& InDataLayerName) const;
 
-	UE_DEPRECATED(5.3, "This function has been deprecated.")
+	UE_DEPRECATED(5.3, "Use UDataLayerManager::IsAnyDataLayerInEffectiveRuntimeState instead.")
 	bool IsAnyDataLayerInEffectiveRuntimeState(const TArray<FName>& InDataLayerNames, EDataLayerRuntimeState InState) const;
 
-	UE_DEPRECATED(5.3, "This function has been deprecated.")
+	UE_DEPRECATED(5.3, "Use UDataLayerManager::GetEffectiveActiveDataLayerNames instead.")
 	const TSet<FName>& GetEffectiveActiveDataLayerNames() const;
 
-	UE_DEPRECATED(5.3, "This function has been deprecated.")
+	UE_DEPRECATED(5.3, "Use UDataLayerManager::GetEffectiveLoadedDataLayerNames instead.")
 	const TSet<FName>& GetEffectiveLoadedDataLayerNames() const;
 
 	UE_DEPRECATED(5.3, "This function has been deprecated.")
 	void GetDataLayerDebugColors(TMap<FName, FColor>& OutMapping) const {}
 
-	UE_DEPRECATED(5.3, "Use function has been deprecated.")
+	UE_DEPRECATED(5.3, "This function has been deprecated.")
 	void DrawDataLayersStatus(UCanvas* Canvas, FVector2D& Offset) const;
 
 	UE_DEPRECATED(5.3, "This function has been deprecated.")
 	void DumpDataLayers(FOutputDevice& OutputDevice) const;
 
+	UE_DEPRECATED(5.3, "This function has been deprecated.")
+	void RegisterWorldDataLayer(AWorldDataLayers* WorldDataLayers) {}
+
+	UE_DEPRECATED(5.3, "This function has been deprecated.")
+	void UnregisterWorldDataLayer(AWorldDataLayers* WorldDataLayers) {}
+
+	UE_DEPRECATED(5.3, "This function has been deprecated.")
+	void ForEachWorldDataLayer(TFunctionRef<bool(AWorldDataLayers*)> Func) {}
+
+	UE_DEPRECATED(5.3, "This function has been deprecated.")
+	void ForEachWorldDataLayer(TFunctionRef<bool(AWorldDataLayers*)> Func) const {}
+
+	UE_DEPRECATED(5.3, "This function has been deprecated.")
+	void ForEachDataLayer(TFunctionRef<bool(UDataLayerInstance*)> Func, const ULevel* InLevelContext = nullptr) {}
+
+	UE_DEPRECATED(5.3, "This function has been deprecated.")
+	void ForEachDataLayer(TFunctionRef<bool(UDataLayerInstance*)> Func, const ULevel* InLevelContext = nullptr) const {}
+
 #if WITH_EDITOR
 	UE_DEPRECATED(5.3, "This function has been deprecated.")
 	static TArray<const UDataLayerInstance*> GetRuntimeDataLayerInstances(UWorld* InWorld, const TArray<FName>& InDataLayerInstanceNames);
+
+	UE_DEPRECATED(5.3, "This function has been deprecated.")
+	bool ResolveIsLoadedInEditor(const TArray<FName>& InDataLayerInstanceNames) const { return false; }
+
+	UE_DEPRECATED(5.3, "This function has been deprecated.")
+	bool CanResolveDataLayers() const { return false; }
+
+	UE_DEPRECATED(5.3, "This function has been deprecated.")
+	bool RemoveDataLayer(const UDataLayerInstance* InDataLayer) { return false; }
+
+	UE_DEPRECATED(5.3, "This function has been deprecated.")
+	int32 RemoveDataLayers(const TArray<UDataLayerInstance*>& InDataLayerInstances) { return 0; }
+
+	UE_DEPRECATED(5.3, "This function has been deprecated.")
+	void UpdateDataLayerEditorPerProjectUserSettings() const {}
+
+	UE_DEPRECATED(5.3, "This function has been deprecated.")
+	void GetUserLoadedInEditorStates(TArray<FName>& OutDataLayersLoadedInEditor, TArray<FName>& OutDataLayersNotLoadedInEditor) const {}
+
+	UE_DEPRECATED(5.3, "This function has been deprecated.")
+	void PushActorEditorContext() const {}
+
+	UE_DEPRECATED(5.3, "This function has been deprecated.")
+	void PopActorEditorContext() const {}
+
+	UE_DEPRECATED(5.3, "This function has been deprecated.")
+	TArray<UDataLayerInstance*> GetActorEditorContextDataLayers() const { return TArray<UDataLayerInstance*>(); }
+
+	UE_DEPRECATED(5.3, "This function has been deprecated.")
+	uint32 GetDataLayerEditorContextHash() const { return 0; }
 #endif
 
 	PRAGMA_DISABLE_DEPRECATION_WARNINGS // Suppress compiler warning on override of deprecated function
