PrimitiveComponent StreamableTextures: Cook Determinism: Turn off the optimization that reads Texture->DidSerializeStreamingMipsForPlatform to more aggressively turn off streaming for PrimitiveComponents that do not need it. The implementation of the optimization relied on the texture being cooked previous in the same cook session, and caused indeterminism in the cook. We should reimplement this optimization using calls to calculate the information if not already calculated.
#rnx

diff --git a/Engine/Source/Runtime/Engine/Private/Components/PrimitiveComponent.cpp b/Engine/Source/Runtime/Engine/Private/Components/PrimitiveComponent.cpp
--- a/Engine/Source/Runtime/Engine/Private/Components/PrimitiveComponent.cpp
+++ b/Engine/Source/Runtime/Engine/Private/Components/PrimitiveComponent.cpp
@@ -1425,18 +1425,8 @@ void UPrimitiveComponent::PreSave(FObjectPreSaveContext ObjectSaveContext)
 		{
 			if (Texture && Texture->IsCandidateForTextureStreamingOnPlatformDuringCook(ObjectSaveContext.GetTargetPlatform()))
 			{
-				// furthermore refine based on what actually happened with that texture in *this* cook session
-				//	if that texture was serialized out before us
-				// can be false if texture IsCandidate but did not actually make streaming textures
-				// if texture has not serialized yet and set DidSerializeStreamingMipsForPlatform, we must assume it would be true
-				
-				const FString PlatformName = ObjectSaveContext.GetTargetPlatform()->PlatformName();
-				const bool * pBoolDidStream = Texture->DidSerializeStreamingMipsForPlatform.Find(PlatformName);
-				if ( pBoolDidStream == nullptr || *pBoolDidStream == true )
-				{
-					bHasStreamableTextures = true;
-					break;
-				}
+				bHasStreamableTextures = true;
+				break;
 			}
 		}
 
