Fix missing Hlods grids when generating with a container that only has LevelInstances
Fix dangling pointer when a high enough amount of grid hlod actors were generated
Re-enable plugin Hlod generation

diff --git a/Engine/Source/Runtime/Engine/Private/WorldPartition/RuntimeSpatialHash/RuntimeSpatialHashHLOD.cpp b/Engine/Source/Runtime/Engine/Private/WorldPartition/RuntimeSpatialHash/RuntimeSpatialHashHLOD.cpp
--- a/Engine/Source/Runtime/Engine/Private/WorldPartition/RuntimeSpatialHash/RuntimeSpatialHashHLOD.cpp
+++ b/Engine/Source/Runtime/Engine/Private/WorldPartition/RuntimeSpatialHash/RuntimeSpatialHashHLOD.cpp
@@ -272,37 +272,40 @@ static TArray<FGuid> GenerateHLODsForGrid(UWorldPartition* WorldPartition, const
 }
 
 // Find all referenced HLODLayer assets
-static TMap<UHLODLayer*, int32> GatherHLODLayers(const UActorDescContainer* ActorDescContainer)
+static TMap<UHLODLayer*, int32> GatherHLODLayers(const IStreamingGenerationContext* StreamingGenerationContext, const UWorldPartition* WorldPartition)
 {
 	// Gather up all HLODLayers referenced by the actors, along with the HLOD level at which it was used
 	TMap<UHLODLayer*, int32> HLODLayersLevel;
 
-	for (FActorDescList::TConstIterator<> ActorDescIterator(ActorDescContainer); ActorDescIterator; ++ActorDescIterator)
+	StreamingGenerationContext->ForEachActorSetContainer([&HLODLayersLevel, WorldPartition](const IStreamingGenerationContext::FActorSetContainer& ActorSetContainer)
 	{
-		const FWorldPartitionActorDesc& ActorDesc = **ActorDescIterator;
-		if (!ActorDesc.GetActorNativeClass()->IsChildOf<AWorldPartitionHLOD>())
+		const FActorDescViewMap* ActorDescViewMap = ActorSetContainer.ActorDescViewMap;
+		ActorDescViewMap->ForEachActorDescView([&HLODLayersLevel, WorldPartition](const FWorldPartitionActorDescView& ActorDescView)
 		{
-			if (ActorDesc.GetActorIsHLODRelevant())
+			if (!ActorDescView.GetActorNativeClass()->IsChildOf<AWorldPartitionHLOD>())
 			{
-				UHLODLayer* HLODLayer = UHLODLayer::GetHLODLayer(ActorDesc, ActorDescContainer->GetWorldPartition());
-
-				// If layer was already encountered, no need to process it again
-				if (!HLODLayersLevel.Contains(HLODLayer))
+				if (ActorDescView.GetActorIsHLODRelevant())
 				{
-					// Walk up the parent HLOD layers, keep track of HLOD level
-					int32 CurrentHLODLevel = 0;
-					while (HLODLayer != nullptr)
+					UHLODLayer* HLODLayer = UHLODLayer::GetHLODLayer(ActorDescView, WorldPartition);
+		
+					// If layer was already encountered, no need to process it again
+					if (!HLODLayersLevel.Contains(HLODLayer))
 					{
-						int32& HLODLevel = HLODLayersLevel.FindOrAdd(HLODLayer);
-						HLODLevel = FMath::Max(HLODLevel, CurrentHLODLevel);
-
-						HLODLayer = HLODLayer->GetParentLayer().LoadSynchronous();
-						CurrentHLODLevel++;
+						// Walk up the parent HLOD layers, keep track of HLOD level
+						int32 CurrentHLODLevel = 0;
+						while (HLODLayer != nullptr)
+						{
+							int32& HLODLevel = HLODLayersLevel.FindOrAdd(HLODLayer);
+							HLODLevel = FMath::Max(HLODLevel, CurrentHLODLevel);
+		
+							HLODLayer = HLODLayer->GetParentLayer().LoadSynchronous();
+							CurrentHLODLevel++;
+						}
 					}
 				}
 			}
-		}
-	}
+		});
+	});
 
 	return HLODLayersLevel;
 }
@@ -443,7 +446,7 @@ bool UWorldPartitionRuntimeSpatialHash::GenerateHLOD(ISourceControlHelper* Sourc
 	const UActorDescContainer* ActorDescContainer = MainActorSetContainer->ActorDescContainer;
 
 	// Find all used HLOD layers
-	TMap<UHLODLayer*, int32> HLODLayersLevels = GatherHLODLayers(ActorDescContainer);
+	TMap<UHLODLayer*, int32> HLODLayersLevels = GatherHLODLayers(StreamingGenerationContext, WorldPartition);
 	TArray<UHLODLayer*> HLODLayers;
 	HLODLayersLevels.GetKeys(HLODLayers);
 
@@ -539,6 +542,7 @@ bool UWorldPartitionRuntimeSpatialHash::GenerateHLOD(ISourceControlHelper* Sourc
 		TArray<const IStreamingGenerationContext::FActorSetInstance*> HLODActorSetInstancePtrs;
 		HLODActorSetInstances.Reserve(GridsHLODActors[HLODGridName].Num());
 		HLODActorSetInstancePtrs.Reserve(GridsHLODActors[HLODGridName].Num());
+		MainActorSetContainer->ActorSets.Reserve(MainActorSetContainer->ActorSets.Num() + GridsHLODActors[HLODGridName].Num());
 
 		for (const FGuid& HLODActorGuid : GridsHLODActors[HLODGridName])
 		{
diff --git a/Engine/Source/Runtime/Engine/Private/WorldPartition/WorldPartitionStreamingGeneration.cpp b/Engine/Source/Runtime/Engine/Private/WorldPartition/WorldPartitionStreamingGeneration.cpp
--- a/Engine/Source/Runtime/Engine/Private/WorldPartition/WorldPartitionStreamingGeneration.cpp
+++ b/Engine/Source/Runtime/Engine/Private/WorldPartition/WorldPartitionStreamingGeneration.cpp
@@ -179,13 +179,21 @@ class FWorldPartitionStreamingGenerator
 			return ActorSetContainers.IsValidIndex(MainWorldActorSetContainerIndex) ? &ActorSetContainers[MainWorldActorSetContainerIndex] : nullptr;
 		}
 
-		virtual void ForEachActorSetInstance(TFunctionRef<void(const FActorSetInstance&)> Func) const
+		virtual void ForEachActorSetInstance(TFunctionRef<void(const FActorSetInstance&)> Func) const override
 		{
 			for (const FActorSetInstance& ActorSetInstance : ActorSetInstances)
 			{
 				Func(ActorSetInstance);
 			}
 		}
+
+		virtual void ForEachActorSetContainer(TFunctionRef<void(const FActorSetContainer&)> Func) const override
+		{
+			for (const FActorSetContainer& ActorSetContainer : ActorSetContainers)
+			{
+				Func(ActorSetContainer);
+			}
+		}
 		//~End IStreamingGenerationContext interface};
 
 	private:
@@ -1046,11 +1054,6 @@ void UWorldPartition::GenerateHLOD(ISourceControlHelper* SourceControlHelper, bo
 {
 	ForEachActorDescContainer([this, &SourceControlHelper, bCreateActorsOnly](UActorDescContainer* InActorDescContainer)
 	{
-		if (!InActorDescContainer->IsMainPartitionContainer())
-		{
-			return;
-		}
-
 		FStreamingGenerationLogErrorHandler LogErrorHandler;
 		FWorldPartitionStreamingGenerator::FWorldPartitionStreamingGeneratorParams StreamingGeneratorParams;
 		StreamingGeneratorParams.WorldPartitionContext = this;
diff --git a/Engine/Source/Runtime/Engine/Public/WorldPartition/WorldPartitionStreamingGenerationContext.h b/Engine/Source/Runtime/Engine/Public/WorldPartition/WorldPartitionStreamingGenerationContext.h
--- a/Engine/Source/Runtime/Engine/Public/WorldPartition/WorldPartitionStreamingGenerationContext.h
+++ b/Engine/Source/Runtime/Engine/Public/WorldPartition/WorldPartitionStreamingGenerationContext.h
@@ -55,5 +55,6 @@ public:
 	virtual FBox GetWorldBounds() const = 0;
 	virtual const FActorSetContainer* GetMainWorldContainer() const = 0;
 	virtual void ForEachActorSetInstance(TFunctionRef<void(const FActorSetInstance&)> Func) const = 0;
+	virtual void ForEachActorSetContainer(TFunctionRef<void(const FActorSetContainer&)> Func) const = 0;
 };
 #endif
\ No newline at end of file
