RayTracing: Add support for SplineMesh components

This allows spline meshes to work in path tracing notably, as well as RT shadows and other raytraced effects. A new cvar "r.RayTracing.Geometry.SplineMeshes" can be used to quickly enable or disable support.

diff --git a/Engine/Shaders/Private/LocalVertexFactory.ush b/Engine/Shaders/Private/LocalVertexFactory.ush
--- a/Engine/Shaders/Private/LocalVertexFactory.ush
+++ b/Engine/Shaders/Private/LocalVertexFactory.ush
@@ -211,7 +211,7 @@ FVertexFactoryInput LoadVertexFactoryInputForHGS(uint TriangleIndex, int VertexI
 #endif
 
 
-#if COMPUTESHADER
+#if COMPUTESHADER && MANUAL_VERTEX_FETCH
 
 FVertexFactoryInput LoadVertexFactoryInputForDynamicUpdate(uint TriangleIndex, int VertexIndex, uint PrimitiveId)
 {
@@ -242,6 +242,10 @@ FVertexFactoryInput LoadVertexFactoryInputForDynamicUpdate(uint TriangleIndex, i
 
 	VF_GPUSCENE_SET_INPUT_FOR_RT(Input, PrimitiveId, 0U);
 
+#if USE_SPLINEDEFORM
+	Input.Position.w = 1.0;
+#endif
+
 	return Input;
 }
 
@@ -771,7 +775,9 @@ float4 CalcWorldPosition(float4 Position, FLWCMatrix LocalToWorld)
 {
 #if USE_INSTANCING
 	return TransformLocalToTranslatedWorld(mul(Position, InstanceTransform).xyz, LocalToWorld);
-#elif USE_SPLINEDEFORM
+#elif USE_SPLINEDEFORM && !RAYHITGROUPSHADER
+	// NOTE: The #if above does not include raytracing shaders because the cached position already
+	// has the spline deformation baked in.
 /*
 	// Make transform for this point along spline
 	float4x3 SliceTransform = CalcSliceTransform(dot(Position.xyz, SplineMeshDir));
@@ -809,7 +815,18 @@ half3x3 CalcTangentToLocal(FVertexFactoryInput Input, inout float TangentSign)
 
 #if USE_SPLINEDEFORM
 	// Make slice rotation matrix, and use that to transform tangents
-	half3x3 SliceRot = CalcSliceRot(dot(Input.Position.xyz, SplineMeshDir));
+#if RAYHITGROUPSHADER
+	// When raytracing, Input.Position already has the spline deformation baked in, so we need to
+	// go back to the original vertex buffer to compute the SliceRot matrix.
+	uint VertexOffset = LocalVF.VertexFetch_Parameters[VF_VertexOffset] + Input.VertexId * 3;
+	float3 SliceInputPos = 0;
+	SliceInputPos.x = LocalVF.VertexFetch_PositionBuffer[VertexOffset + 0];
+	SliceInputPos.y = LocalVF.VertexFetch_PositionBuffer[VertexOffset + 1];
+	SliceInputPos.z = LocalVF.VertexFetch_PositionBuffer[VertexOffset + 2];
+#else
+	float3 SliceInputPos = Input.Position.xyz;
+#endif
+	half3x3 SliceRot = CalcSliceRot(dot(SliceInputPos.xyz, SplineMeshDir));
 
 	TangentX = mul(TangentX, SliceRot);
 	TangentZ.xyz = mul(TangentZ.xyz, SliceRot);
@@ -1174,7 +1191,7 @@ float4 VertexFactoryGetPreviousWorldPosition(FVertexFactoryInput Input, FVertexF
 #if USE_INSTANCING
 		float4x4 InstanceTransform = GetInstancePrevTransform(Intermediates);
 		PrevLocalPosition = mul(Input.Position, InstanceTransform);
-#elif USE_SPLINEDEFORM
+#elif USE_SPLINEDEFORM && !RAYHITGROUPSHADER
 		// Just like CalcWorldPosition...
 		float4x3 SliceTransform = CalcSliceTransform(dot(Input.Position.xyz, SplineMeshDir));
 		// Transform into mesh space
diff --git a/Engine/Source/Runtime/Engine/Private/Components/SplineMeshComponent.cpp b/Engine/Source/Runtime/Engine/Private/Components/SplineMeshComponent.cpp
--- a/Engine/Source/Runtime/Engine/Private/Components/SplineMeshComponent.cpp
+++ b/Engine/Source/Runtime/Engine/Private/Components/SplineMeshComponent.cpp
@@ -49,7 +49,7 @@ FSplineMeshInstanceData::FSplineMeshInstanceData(const USplineMeshComponent* Sou
 
 void FSplineMeshVertexFactoryShaderParameters::Bind(const FShaderParameterMap& ParameterMap)
 {
-	SplineMeshParams.Bind(ParameterMap, TEXT("SplineParams"), SPF_Mandatory);
+	SplineMeshParams.Bind(ParameterMap, TEXT("SplineParams"), SPF_Optional);
 }
 
 void FSplineMeshVertexFactoryShaderParameters::GetElementShaderBindings(
@@ -111,7 +111,11 @@ void FSplineMeshVertexFactoryShaderParameters::GetElementShaderBindings(
 //////////////////////////////////////////////////////////////////////////
 // SplineMeshVertexFactory
 
-IMPLEMENT_VERTEX_FACTORY_PARAMETER_TYPE(FSplineMeshVertexFactory, SF_Vertex, FSplineMeshVertexFactoryShaderParameters);
+IMPLEMENT_VERTEX_FACTORY_PARAMETER_TYPE(FSplineMeshVertexFactory, SF_Vertex     , FSplineMeshVertexFactoryShaderParameters);
+#if RHI_RAYTRACING
+IMPLEMENT_VERTEX_FACTORY_PARAMETER_TYPE(FSplineMeshVertexFactory, SF_RayHitGroup, FSplineMeshVertexFactoryShaderParameters);
+IMPLEMENT_VERTEX_FACTORY_PARAMETER_TYPE(FSplineMeshVertexFactory, SF_Compute    , FSplineMeshVertexFactoryShaderParameters);
+#endif
 
 IMPLEMENT_VERTEX_FACTORY_TYPE(FSplineMeshVertexFactory, "/Engine/Private/LocalVertexFactory.ush",
 	  EVertexFactoryFlags::UsedWithMaterials
@@ -121,6 +125,9 @@ IMPLEMENT_VERTEX_FACTORY_TYPE(FSplineMeshVertexFactory, "/Engine/Private/LocalVe
 	| EVertexFactoryFlags::SupportsPositionOnly
 	| EVertexFactoryFlags::SupportsPrimitiveIdStream
 	| EVertexFactoryFlags::SupportsPSOPrecaching
+	| EVertexFactoryFlags::SupportsRayTracing
+	| EVertexFactoryFlags::SupportsRayTracingDynamicGeometry
+	| EVertexFactoryFlags::SupportsManualVertexFetch
 );
 
 void InitSplineMeshVertexFactoryComponents(
diff --git a/Engine/Source/Runtime/Engine/Private/SplineMeshSceneProxy.cpp b/Engine/Source/Runtime/Engine/Private/SplineMeshSceneProxy.cpp
--- a/Engine/Source/Runtime/Engine/Private/SplineMeshSceneProxy.cpp
+++ b/Engine/Source/Runtime/Engine/Private/SplineMeshSceneProxy.cpp
@@ -7,9 +7,16 @@
 #include "Misc/DelayedAutoRegister.h"
 #include "StaticMeshComponentLODInfo.h"
 #include "StaticMeshResources.h"
+#include "RayTracingInstance.h"
 
 IMPLEMENT_TYPE_LAYOUT(FSplineMeshVertexFactoryShaderParameters);
 
+static TAutoConsoleVariable<int32> CVarRayTracingSplineMeshes(
+	TEXT("r.RayTracing.Geometry.SplineMeshes"),
+	1,
+	TEXT("Include splines meshes in ray tracing effects (default = 1 (spline meshes enabled in ray tracing))"));
+
+
 bool FSplineMeshVertexFactory::ShouldCompilePermutation(const FVertexFactoryShaderPermutationParameters& Parameters)
 {
 	return (Parameters.MaterialParameters.bIsUsedWithSplineMeshes || Parameters.MaterialParameters.bIsSpecialEngineMaterial)
@@ -19,12 +26,6 @@ bool FSplineMeshVertexFactory::ShouldCompilePermutation(const FVertexFactoryShad
 /** Modify compile environment to enable spline deformation */
 void FSplineMeshVertexFactory::ModifyCompilationEnvironment(const FVertexFactoryShaderPermutationParameters& Parameters, FShaderCompilerEnvironment& OutEnvironment)
 {
-	const bool ContainsManualVertexFetch = OutEnvironment.GetDefinitions().Contains("MANUAL_VERTEX_FETCH");
-	if (!ContainsManualVertexFetch)
-	{
-		OutEnvironment.SetDefine(TEXT("MANUAL_VERTEX_FETCH"), TEXT("0"));
-	}
-
 	OutEnvironment.SetDefine(TEXT("VF_SUPPORTS_SPEEDTREE_WIND"), TEXT("0"));
 	FLocalVertexFactory::ModifyCompilationEnvironment(Parameters, OutEnvironment);
 
@@ -71,6 +72,9 @@ FSplineMeshSceneProxy::FSplineMeshSceneProxy(USplineMeshComponent* InComponent)
 {
 	bSupportsDistanceFieldRepresentation = false;
 	bSupportsMeshCardRepresentation = false;
+#if RHI_RAYTRACING
+	bDynamicRayTracingGeometry = true;
+#endif
 
 	// make sure all the materials are okay to be rendered as a spline mesh
 	for (FStaticMeshSceneProxy::FLODInfo& LODInfo : LODs)
@@ -169,3 +173,81 @@ bool FSplineMeshSceneProxy::GetCollisionMeshElement(int32 LODIndex, int32 BatchI
 	}
 	return false;
 }
+
+#if RHI_RAYTRACING
+void FSplineMeshSceneProxy::GetDynamicRayTracingInstances(struct FRayTracingMaterialGatheringContext& Context, TArray<FRayTracingInstance>& OutRayTracingInstances)
+{
+	if (CVarRayTracingSplineMeshes.GetValueOnRenderThread() == 0)
+	{
+		return;
+	}
+
+	ESceneDepthPriorityGroup PrimitiveDPG = GetStaticDepthPriorityGroup();
+	const int32 LODIndex = FMath::Max(GetLOD(Context.ReferenceView), (int32)GetCurrentFirstLODIdx_RenderThread());
+	const FStaticMeshLODResources& LODModel = RenderData->LODResources[LODIndex];
+
+	FRayTracingGeometry& Geometry = DynamicRayTracingGeometries[LODIndex];
+
+	FRayTracingInstance& RayTracingInstance = OutRayTracingInstances.AddDefaulted_GetRef();
+
+	const int32 NumBatches = GetNumMeshBatches();
+	const int32 NumRayTracingMaterialEntries = LODModel.Sections.Num() * NumBatches;
+
+	if (NumRayTracingMaterialEntries != CachedRayTracingMaterials.Num() || CachedRayTracingMaterialsLODIndex != LODIndex)
+	{
+		CachedRayTracingMaterials.Reset();
+		CachedRayTracingMaterials.Reserve(NumRayTracingMaterialEntries);
+
+		for (int32 BatchIndex = 0; BatchIndex < NumBatches; BatchIndex++)
+		{
+			for (int32 SectionIndex = 0; SectionIndex < LODModel.Sections.Num(); SectionIndex++)
+			{
+				FMeshBatch& MeshBatch = CachedRayTracingMaterials.AddDefaulted_GetRef();
+
+				bool bResult = GetMeshElement(LODIndex, BatchIndex, SectionIndex, PrimitiveDPG, false, false, MeshBatch);
+				if (!bResult)
+				{
+					// Hidden material
+					MeshBatch.MaterialRenderProxy = UMaterial::GetDefaultMaterial(MD_Surface)->GetRenderProxy();
+					MeshBatch.VertexFactory = &RenderData->LODVertexFactories[LODIndex].VertexFactory;
+				}
+
+				MeshBatch.SegmentIndex = SectionIndex;
+				MeshBatch.MeshIdInPrimitive = SectionIndex;
+			}
+		}
+
+		CachedRayTracingMaterialsLODIndex = LODIndex;
+	}
+
+	RayTracingInstance.Geometry = &Geometry;
+	// scene proxies live for the duration of Render(), making array views below safe
+	const FMatrix& ThisLocalToWorld = GetLocalToWorld();
+	RayTracingInstance.InstanceTransformsView = MakeArrayView(&ThisLocalToWorld, 1);
+	RayTracingInstance.MaterialsView = MakeArrayView(CachedRayTracingMaterials);
+	CachedRayTracingInstanceMaskAndFlags = Context.BuildInstanceMaskAndFlags(RayTracingInstance, *this);
+
+	if (RenderData->LODVertexFactories[LODIndex].VertexFactory.GetType()->SupportsRayTracingDynamicGeometry())
+	{
+		Context.DynamicRayTracingGeometriesToUpdate.Add(
+			FRayTracingDynamicGeometryUpdateParams
+			{
+				CachedRayTracingMaterials, // TODO: this copy can be avoided if FRayTracingDynamicGeometryUpdateParams supported array views
+				false,
+				(uint32)LODModel.GetNumVertices(),
+				uint32((SIZE_T)LODModel.GetNumVertices() * sizeof(FVector3f)),
+				Geometry.Initializer.TotalPrimitiveCount,
+				&Geometry,
+				nullptr /* VertexBuffer */,
+				true
+			}
+		);
+	}
+
+	check(CachedRayTracingMaterials.Num() == RayTracingInstance.GetMaterials().Num());
+	checkf(RayTracingInstance.Geometry->Initializer.Segments.Num() == CachedRayTracingMaterials.Num(), TEXT("Segments/Materials mismatch. Number of segments: %d. Number of Materials: %d. LOD Index: %d"),
+		RayTracingInstance.Geometry->Initializer.Segments.Num(),
+		CachedRayTracingMaterials.Num(),
+		LODIndex);
+}
+#endif // RHI_RAYTRACING
diff --git a/Engine/Source/Runtime/Engine/Private/SplineMeshSceneProxy.h b/Engine/Source/Runtime/Engine/Private/SplineMeshSceneProxy.h
--- a/Engine/Source/Runtime/Engine/Private/SplineMeshSceneProxy.h
+++ b/Engine/Source/Runtime/Engine/Private/SplineMeshSceneProxy.h
@@ -111,9 +111,9 @@ public:
 	virtual bool GetCollisionMeshElement(int32 LODIndex, int32 BatchIndex, int32 ElementIndex, uint8 InDepthPriorityGroup, const FMaterialRenderProxy* RenderProxy, FMeshBatch& OutMeshBatch) const override;
 
 #if RHI_RAYTRACING
-	virtual bool HasRayTracingRepresentation() const override { return false; }
-	virtual bool IsRayTracingRelevant() const override final { return false; }
-	virtual bool IsRayTracingStaticRelevant() const override final { return false; }
+	virtual bool HasRayTracingRepresentation() const override { return true; }
+	virtual bool IsRayTracingRelevant() const override { return true; }
+	virtual void GetDynamicRayTracingInstances(struct FRayTracingMaterialGatheringContext& Context, TArray<FRayTracingInstance>& OutRayTracingInstances) override;
 #endif // RHI_RAYTRACING
 
 	virtual FPrimitiveViewRelevance GetViewRelevance(const FSceneView* View) const override
