Strata - fixed reroute node to not crash when previewing Strata data.

diff --git a/Engine/Source/Runtime/Engine/Classes/Materials/MaterialExpressionNamedReroute.h b/Engine/Source/Runtime/Engine/Classes/Materials/MaterialExpressionNamedReroute.h
--- a/Engine/Source/Runtime/Engine/Classes/Materials/MaterialExpressionNamedReroute.h
+++ b/Engine/Source/Runtime/Engine/Classes/Materials/MaterialExpressionNamedReroute.h
@@ -63,6 +63,7 @@ public:
 	//~ Begin UMaterialExpression Interface
 #if WITH_EDITOR
 	virtual int32 Compile(class FMaterialCompiler* Compiler, int32 OutputIndex) override;
+	virtual int32 CompilePreview(FMaterialCompiler* Compiler, int32 OutputIndex) override;
 	virtual void GetCaption(TArray<FString>& OutCaptions) const override;
 
 	virtual FText GetCreationDescription() const override;
@@ -109,6 +110,8 @@ public:
 	//~ Begin UMaterialExpression Interface
 #if WITH_EDITOR
 	virtual int32 Compile(class FMaterialCompiler* Compiler, int32 OutputIndex) override;
+	virtual int32 CompilePreview(FMaterialCompiler* Compiler, int32 OutputIndex) override;
+
 	virtual void GetCaption(TArray<FString>& OutCaptions) const override;
 	virtual uint32 GetOutputType(int32 OutputIndex) override;
 	
diff --git a/Engine/Source/Runtime/Engine/Classes/Materials/MaterialExpressionReroute.h b/Engine/Source/Runtime/Engine/Classes/Materials/MaterialExpressionReroute.h
--- a/Engine/Source/Runtime/Engine/Classes/Materials/MaterialExpressionReroute.h
+++ b/Engine/Source/Runtime/Engine/Classes/Materials/MaterialExpressionReroute.h
@@ -21,6 +21,7 @@ class ENGINE_API UMaterialExpressionReroute : public UMaterialExpressionRerouteB
 	//~ Begin UMaterialExpression Interface
 #if WITH_EDITOR
 	virtual int32 Compile(class FMaterialCompiler* Compiler, int32 OutputIndex) override;
+	virtual int32 CompilePreview(FMaterialCompiler* Compiler, int32 OutputIndex) override;
 	virtual void GetCaption(TArray<FString>& OutCaptions) const override;
 
 	virtual FText GetCreationDescription() const override;
diff --git a/Engine/Source/Runtime/Engine/Private/Materials/MaterialExpressions.cpp b/Engine/Source/Runtime/Engine/Private/Materials/MaterialExpressions.cpp
--- a/Engine/Source/Runtime/Engine/Private/Materials/MaterialExpressions.cpp
+++ b/Engine/Source/Runtime/Engine/Private/Materials/MaterialExpressions.cpp
@@ -18161,6 +18161,17 @@ int32 UMaterialExpressionReroute::Compile(FMaterialCompiler* Compiler, int32 Out
 	return Result;
 }
 
+int32 UMaterialExpressionReroute::CompilePreview(FMaterialCompiler* Compiler, int32 OutputIndex)
+{
+	int32 ResultCodeChunk = Compile(Compiler, OutputIndex);
+
+	if (Input.Expression && Input.Expression->IsResultStrataMaterial(Input.OutputIndex))
+	{
+		ResultCodeChunk = Compiler->StrataCompilePreview(ResultCodeChunk);
+	}
+	return ResultCodeChunk;
+}
+
 void UMaterialExpressionReroute::GetCaption(TArray<FString>& OutCaptions) const
 {
 	OutCaptions.Add(TEXT("Reroute Node (reroutes wires)"));
@@ -18267,6 +18278,17 @@ int32 UMaterialExpressionNamedRerouteDeclaration::Compile(FMaterialCompiler* Com
 	return Input.Compile(Compiler);
 }
 
+int32 UMaterialExpressionNamedRerouteDeclaration::CompilePreview(FMaterialCompiler* Compiler, int32 OutputIndex)
+{
+	int32 ResultCodeChunk = Compile(Compiler, OutputIndex);
+
+	if (Input.Expression && Input.Expression->IsResultStrataMaterial(Input.OutputIndex))
+	{
+		ResultCodeChunk = Compiler->StrataCompilePreview(ResultCodeChunk);
+	}
+	return ResultCodeChunk;
+}
+
 void UMaterialExpressionNamedRerouteDeclaration::GetCaption(TArray<FString>& OutCaptions) const
 {
 	OutCaptions.Add(Name.ToString());
@@ -18431,6 +18453,17 @@ int32 UMaterialExpressionNamedRerouteUsage::Compile(FMaterialCompiler* Compiler,
 	}
 	return Declaration->Compile(Compiler, OutputIndex);
 }
+int32 UMaterialExpressionNamedRerouteUsage::CompilePreview(FMaterialCompiler* Compiler, int32 OutputIndex)
+{
+	int32 ResultCodeChunk = Compile(Compiler, OutputIndex);
+
+	FExpressionInput Input = Declaration->TraceInputsToRealInput();
+	if (Input.Expression && Input.Expression->IsResultStrataMaterial(Input.OutputIndex))
+	{
+		ResultCodeChunk = Compiler->StrataCompilePreview(ResultCodeChunk);
+	}
+	return ResultCodeChunk;
+}
 
 void UMaterialExpressionNamedRerouteUsage::GetCaption(TArray<FString>& OutCaptions) const
 {
