Removing deprecated ExactPingV2

diff --git a/Engine/Source/Runtime/Engine/Classes/GameFramework/PlayerState.h b/Engine/Source/Runtime/Engine/Classes/GameFramework/PlayerState.h
--- a/Engine/Source/Runtime/Engine/Classes/GameFramework/PlayerState.h
+++ b/Engine/Source/Runtime/Engine/Classes/GameFramework/PlayerState.h
@@ -31,33 +31,6 @@ struct PingAvgData
 	}
 };
 
-/**
- * Struct keeping track of the lowest ping values over a given second.
- */
-struct UE_DEPRECATED(4.27, "ExactPingV2 is no longer used. Please use ExactPing instead.") PingAvgDataV2
-{
-	PRAGMA_DISABLE_DEPRECATION_WARNINGS
-	/** List of ping values */
-	TArray<uint16> PingValues;
-
-	/** The maximum number of ping values we will keep track of */
-	static const uint8 MAX_PING_VALUES_SIZE = 7;
-
-	/** Default constructor */
-	PingAvgDataV2()
-		: AvgPingV2(0.0f)
-	{
-		for (int32 i = 0; i < MAX_PING_VALUES_SIZE; i++)
-		{
-			PingValues.Add(MAX_uint16);
-		}
-	}
-	
-	/** The average of the values in PingValues, calculated after 1s. */
-	float AvgPingV2;
-	PRAGMA_ENABLE_DEPRECATION_WARNINGS
-};
-
 DECLARE_DYNAMIC_MULTICAST_DELEGATE_ThreeParams(FOnPlayerStatePawnSet, APlayerState*, Player, APawn*, NewPawn, APawn*, OldPawn);
 
 /**
@@ -146,9 +119,6 @@ public:
 	/** Exact ping in milliseconds as float (rounded and compressed in replicated CompressedPing) */
 	float ExactPing;
 
-	UE_DEPRECATED(4.27, "Please use ExactPing instead.")
-	float ExactPingV2;
-
 	/** Used to match up InactivePlayerState with rejoining playercontroller. */
 	UPROPERTY()
 	FString SavedNetworkAddress;
@@ -188,10 +158,6 @@ private:
 	 */
 	PingAvgData		PingBucket[4];
 
-	PRAGMA_DISABLE_DEPRECATION_WARNINGS
-	PingAvgDataV2	PingBucketV2[4];
-	PRAGMA_ENABLE_DEPRECATION_WARNINGS
-
 	/** The timestamp for when the current PingBucket began filling */
 	float			CurPingBucketTimestamp;
 
diff --git a/Engine/Source/Runtime/Engine/Private/PlayerState.cpp b/Engine/Source/Runtime/Engine/Private/PlayerState.cpp
--- a/Engine/Source/Runtime/Engine/Private/PlayerState.cpp
+++ b/Engine/Source/Runtime/Engine/Private/PlayerState.cpp
@@ -69,10 +69,6 @@ void APlayerState::UpdatePing(float InPing)
 
 		PingBucket[CurPingBucket].PingSum = FMath::FloorToInt(InPingInMs);
 		PingBucket[CurPingBucket].PingCount = 1;
-
-		PRAGMA_DISABLE_DEPRECATION_WARNINGS
-		PingBucketV2[CurPingBucket] = PingAvgDataV2();
-		PRAGMA_ENABLE_DEPRECATION_WARNINGS
 	}
 	// Limit the number of pings we accept per-bucket, to avoid overflowing PingBucket values
 	else if (PingBucket[CurPingBucket].PingCount < 7)
@@ -80,27 +76,6 @@ void APlayerState::UpdatePing(float InPing)
 		PingBucket[CurPingBucket].PingSum += FMath::FloorToInt(InPingInMs);
 		PingBucket[CurPingBucket].PingCount++;
 	}
-
-	PRAGMA_DISABLE_DEPRECATION_WARNINGS
-	TArray<uint16>& CurrentBucketPingValues = PingBucketV2[CurPingBucket].PingValues;
-
-	// This makes sure we will actually add the ping value to the list, since much of the time the new ping value will be higher than
-	// what is already there.
-	if (InPingInMs < CurrentBucketPingValues[PingAvgDataV2::MAX_PING_VALUES_SIZE - 1])
-	{
-		for (int32 i = 0; i < PingAvgDataV2::MAX_PING_VALUES_SIZE; i++)
-		{
-			// This will keep the list of ping values we're currently using sorted - this will insert
-			// a ping value if it is less than any in current list and remove the new max value.
-			if (InPingInMs < CurrentBucketPingValues[i])
-			{
-				CurrentBucketPingValues.Insert(FMath::FloorToInt(InPingInMs), i);
-				CurrentBucketPingValues.RemoveAt(PingAvgDataV2::MAX_PING_VALUES_SIZE);
-				break;
-			}
-		}
-	}
-	PRAGMA_ENABLE_DEPRECATION_WARNINGS
 }
 
 void APlayerState::RecalculateAvgPing()
@@ -114,29 +89,6 @@ void APlayerState::RecalculateAvgPing()
 		Count += PingBucket[i].PingCount;
 	}
 
-	PRAGMA_DISABLE_DEPRECATION_WARNINGS
-	int32 SumV2 = 0;
-	int32 NumValidValues = 0;
-	for (; NumValidValues < PingAvgDataV2::MAX_PING_VALUES_SIZE; NumValidValues++)
-	{
-		if (PingBucketV2[CurPingBucket].PingValues[NumValidValues] != MAX_uint16)
-		{
-			SumV2 += PingBucketV2[CurPingBucket].PingValues[NumValidValues];
-		}
-	}
-
-	// Use NumValidValues instead of MAX_PING_VALUES_SIZE in case there are fewer valid values.
-	PingBucketV2[CurPingBucket].AvgPingV2 = NumValidValues > 0 ? SumV2 / NumValidValues : MAX_flt;
-
-	float AvgSumV2 = 0.0f;
-	for (int32 i = 0; i < UE_ARRAY_COUNT(PingBucketV2); i++)
-	{
-		AvgSumV2 += PingBucketV2[i].AvgPingV2;
-	}
-
-	ExactPingV2 = AvgSumV2 / UE_ARRAY_COUNT(PingBucketV2);
-	PRAGMA_ENABLE_DEPRECATION_WARNINGS
-
 	// Calculate the average, and divide it by 4 to optimize replication
 	ExactPing = (Count > 0 ? ((float)Sum / (float)Count) : 0.f);
 
