WorldPartition:
- Fix PIE SoftObjectPaths on Sub WPs (non-level instance actors)
- Fix Cooked SoftObjectPaths on AlwaysLoaded cell

diff --git a/Engine/Source/Runtime/Engine/Private/WorldPartition/WorldPartitionLevelStreamingPolicy.cpp b/Engine/Source/Runtime/Engine/Private/WorldPartition/WorldPartitionLevelStreamingPolicy.cpp
--- a/Engine/Source/Runtime/Engine/Private/WorldPartition/WorldPartitionLevelStreamingPolicy.cpp
+++ b/Engine/Source/Runtime/Engine/Private/WorldPartition/WorldPartitionLevelStreamingPolicy.cpp
@@ -18,6 +18,11 @@
 
 #if WITH_EDITOR
 #include "Misc/PackageName.h"
+
+namespace WorldPartitionLevelStreamingPolicy
+{
+	static FName PersistentCellName(TEXT("PersistentCell"));
+}
 #endif
 
 int32 UWorldPartitionLevelStreamingPolicy::GetCellLoadingCount() const
@@ -98,8 +103,11 @@ TSubclassOf<UWorldPartitionRuntimeCell> UWorldPartitionLevelStreamingPolicy::Get
 
 void UWorldPartitionLevelStreamingPolicy::PrepareActorToCellRemapping()
 {
+	FString SourceWorldPath, RemappedWorldPath;
+	const bool bInstancedWorld = WorldPartition->GetTypedOuter<UWorld>()->GetSoftObjectPathMapping(SourceWorldPath, RemappedWorldPath);
+
 	// Build Actor-to-Cell remapping
-	WorldPartition->RuntimeHash->ForEachStreamingCells([this](const UWorldPartitionRuntimeCell* Cell)
+	WorldPartition->RuntimeHash->ForEachStreamingCells([this, bInstancedWorld, &SourceWorldPath, &RemappedWorldPath](const UWorldPartitionRuntimeCell* Cell)
 	{
 		const UWorldPartitionRuntimeLevelStreamingCell* StreamingCell = Cast<const UWorldPartitionRuntimeLevelStreamingCell>(Cell);
 		check(StreamingCell);
@@ -107,20 +115,58 @@ void UWorldPartitionLevelStreamingPolicy::PrepareActorToCellRemapping()
 		{
 			FString RemappedActorPath;
 			FString CellActorPath = CellObjectMap.Path.ToString();
+			bool bActorPathNeedsRemapping = false;
 
-			// Add actor container id to actor path so that we can distinguish between actors of different Level Instances
-			const bool ActorPathNeedsRemapping = FWorldPartitionLevelHelper::RemapActorPath(CellObjectMap.ContainerID, CellActorPath, RemappedActorPath);
+			if (bInstancedWorld && CellActorPath.StartsWith(RemappedWorldPath))
+			{
+				check(!IsRunningCookCommandlet());
+				// When calling PrepareActorToCellRemapping on an instanced world (PIE) we want to use the source world path to build the actor cell mapping
+				// 
+				// We will get actor paths in this format (UEDPIE prefix + level instance suffix):
+				// 
+				// '/Game/SomePath/UEDPIE_0_WorldName_LevelInstance1.WorldName:PersistentLevel.ActorA'
+				// 
+				// but we need to build an instancing/pie agnostic source mapping with the following key:
+				//
+				// '/Game/SomePath/WorldName.WorldName:PersistentLevel.ActorA'
+				check(CellObjectMap.ContainerID.IsMainContainer());
+				bActorPathNeedsRemapping = true;
+
+				FSoftObjectPath TmpSoftObjectPath(CellActorPath);
+				RemappedActorPath = FSoftObjectPath(FTopLevelAssetPath(SourceWorldPath), TmpSoftObjectPath.GetSubPathString()).ToString();
+			}
+			else if(!CellObjectMap.ContainerID.IsMainContainer())
+			{
+				// Add actor container id to actor path so that we can distinguish between actors of different Level Instances
+				//
+				// '/Game/SomePath/LevelInstance.LevelInstance:PersistentLevel.ActorA' will be remapped to
+				// 
+				// '/Game/SomePath/LevelInstance.LevelInstance:PersistentLevel.ActorA_{ContainerID}'
+				bActorPathNeedsRemapping = FWorldPartitionLevelHelper::RemapActorPath(CellObjectMap.ContainerID, CellActorPath, RemappedActorPath);
+			}
 
-			if (ActorPathNeedsRemapping || Cell->NeedsActorToCellRemapping())
+			// The use cases for remapping are the following:
+			//
+			// - Spatially loaded or Datalayer Actors from the main World Partition map that get moved into a Streaming Cell. In thise case an actor path like:
+			//		- '/Game/SomePath/WorldName.WorldName:PersistentLevel.ActorA' would be mapped to a cell name ex: 'WorldName_MainGrid_L0_X5_Y-4'
+			// - Always loaded Actors from the main World:
+			//		- In PIE they get remapped to the top level Cell 'WorldName_MainGrid_L{MAX}_X0_Y0'
+			//		- In Cook they don't need remapping as the top level Cell is the PersistentLevel (Cell->NeedsActorToCellRemapping() returns false)
+			// - Embedded Level Instance actors always need remapping because they originate from a different map and will be moved to a main world cell
+			//		- '/Game/SomePath/LevelInstance.LevelInstance:PersistentLevel.ActorA_{ContainerID}' will be mapped to a cell name ex: 'WorldName_MainGrid_L0_X5_Y-4'
+			//		-  AlwaysLoaded Embedded actors in Cook differ from the main World actors as they need to be remapped to the PersistentLevel and this is why we have the 
+			//		   special cell name: WorldPartitionLevelStreamingPolicy::PersistentCellName
+			if (bActorPathNeedsRemapping || Cell->NeedsActorToCellRemapping())
 			{
-				const FString& ActorPath = ActorPathNeedsRemapping ? RemappedActorPath : CellActorPath;
+				const FString& ActorPath = bActorPathNeedsRemapping ? RemappedActorPath : CellActorPath;
+				FName CellName = Cell->NeedsActorToCellRemapping() ? StreamingCell->GetFName() : WorldPartitionLevelStreamingPolicy::PersistentCellName;
 
-				ActorToCellRemapping.Add(FName(*ActorPath), StreamingCell->GetFName());
+				ActorToCellRemapping.Add(FName(*ActorPath), CellName);
 
 				const int32 LastDotPos = ActorPath.Find(TEXT("."), ESearchCase::IgnoreCase, ESearchDir::FromEnd);
 				check(LastDotPos != INDEX_NONE);
 
-				SubObjectsToCellRemapping.Add(FName(*ActorPath.Mid(LastDotPos + 1)), StreamingCell->GetFName());
+				SubObjectsToCellRemapping.Add(FName(*ActorPath.Mid(LastDotPos + 1)), CellName);
 			}
 		}
 		return true;
@@ -158,22 +204,31 @@ void UWorldPartitionLevelStreamingPolicy::RemapSoftObjectPath(FSoftObjectPath& O
 		if (!SrcObjectPath.GetSubPathString().IsEmpty())
 		{
 			UWorld* OuterWorld = WorldPartition->GetTypedOuter<UWorld>();
-			const FString PackagePath = UWorldPartitionLevelStreamingPolicy::GetCellPackagePath(*CellName, OuterWorld);
-			FString PrefixPath;
-			if (IsRunningCookCommandlet())
+			if (*CellName == WorldPartitionLevelStreamingPolicy::PersistentCellName)
 			{
-				//@todo_ow: Temporary workaround. This information should be provided by the COTFS
-				const UPackage* Package = OuterWorld->GetPackage();
-				PrefixPath = FString::Printf(TEXT("%s/%s/_Generated_"), *FPackageName::GetLongPackagePath(Package->GetPathName()), *FPackageName::GetShortName(Package->GetName()));
+				check(IsRunningCookCommandlet());
+				ObjectPath = FSoftObjectPath(FTopLevelAssetPath(OuterWorld), SrcObjectPath.GetSubPathString());
 			}
-
-			// Use the WorldPartition world name here instead of using the world name from the path to support converting level instance paths to main world paths.
-			ObjectPath = FSoftObjectPath(FTopLevelAssetPath(FString::Printf(TEXT("%s%s.%s"), *PrefixPath, *PackagePath, *OuterWorld->GetName())), SrcObjectPath.GetSubPathString());
-			// Put back PIE prefix
-			if (OuterWorld->IsPlayInEditor() && (PIEInstanceID != INDEX_NONE))
+			else
 			{
-				ObjectPath.FixupForPIE(PIEInstanceID);
-			}
+				const FString PackagePath = UWorldPartitionLevelStreamingPolicy::GetCellPackagePath(*CellName, OuterWorld);
+				FString PrefixPath;
+
+				if (IsRunningCookCommandlet())
+				{
+					//@todo_ow: Temporary workaround. This information should be provided by the COTFS
+					const UPackage* Package = OuterWorld->GetPackage();
+					PrefixPath = FString::Printf(TEXT("%s/%s/_Generated_"), *FPackageName::GetLongPackagePath(Package->GetPathName()), *FPackageName::GetShortName(Package->GetName()));
+				}
+
+				// Use the WorldPartition world name here instead of using the world name from the path to support converting level instance paths to main world paths.
+				ObjectPath = FSoftObjectPath(FTopLevelAssetPath(FString::Printf(TEXT("%s%s.%s"), *PrefixPath, *PackagePath, *OuterWorld->GetName())), SrcObjectPath.GetSubPathString());
+				// Put back PIE prefix
+				if (OuterWorld->IsPlayInEditor() && (PIEInstanceID != INDEX_NONE))
+				{
+					ObjectPath.FixupForPIE(PIEInstanceID);
+				}
+			}	
 		}
 	}
 }
