Minor clean-up: make commonly used AudioDevice functions const correct
#rnx

diff --git a/Engine/Source/Runtime/Engine/Private/AudioDeviceManager.cpp b/Engine/Source/Runtime/Engine/Private/AudioDeviceManager.cpp
--- a/Engine/Source/Runtime/Engine/Private/AudioDeviceManager.cpp
+++ b/Engine/Source/Runtime/Engine/Private/AudioDeviceManager.cpp
@@ -1,6 +1,7 @@
 // Copyright Epic Games, Inc. All Rights Reserved.
 #include "AudioDeviceManager.h"
 
+#include "Algo/Transform.h"
 #include "Audio/AudioDebug.h"
 #include "AudioAnalytics.h"
 #include "AudioMixerDevice.h"
@@ -1243,29 +1244,32 @@ uint8 FAudioDeviceManager::GetNumMainAudioDeviceWorlds() const
 	}
 }
 
-TArray<FAudioDevice*> FAudioDeviceManager::GetAudioDevices()
+TArray<FAudioDevice*> FAudioDeviceManager::GetAudioDevices() const
 {
 	TArray<FAudioDevice*> DeviceList;
-	FScopeLock ScopeLock(&DeviceMapCriticalSection);
-	for (auto& Device : Devices)
+
 	{
-		DeviceList.Add(Device.Value.Device);
+		FScopeLock ScopeLock(&DeviceMapCriticalSection);
+		Algo::Transform(Devices, DeviceList, [](const TPair<Audio::FDeviceId, FAudioDeviceContainer>& Pair)
+		{
+			return Pair.Value.Device;
+		});
 	}
 
 	return DeviceList;
 }
 
-TArray<UWorld*> FAudioDeviceManager::GetWorldsUsingAudioDevice(const Audio::FDeviceId& InID)
+TArray<UWorld*> FAudioDeviceManager::GetWorldsUsingAudioDevice(const Audio::FDeviceId& InID) const
 {
-	FScopeLock ScopeLock(&DeviceMapCriticalSection);
-	if (Devices.Contains(InID))
-	{
-		return Devices[InID].WorldsUsingThisDevice;
-	}
-	else
 	{
-		return TArray<UWorld*>();
+		FScopeLock ScopeLock(&DeviceMapCriticalSection);
+		if (Devices.Contains(InID))
+		{
+			return Devices[InID].WorldsUsingThisDevice;
+		}
 	}
+
+	return { };
 }
 
 #if INSTRUMENT_AUDIODEVICE_HANDLES
diff --git a/Engine/Source/Runtime/Engine/Public/AudioDeviceManager.h b/Engine/Source/Runtime/Engine/Public/AudioDeviceManager.h
--- a/Engine/Source/Runtime/Engine/Public/AudioDeviceManager.h
+++ b/Engine/Source/Runtime/Engine/Public/AudioDeviceManager.h
@@ -281,9 +281,9 @@ public:
 	TMap<int32, FSoundBuffer*>	WaveBufferMap;
 
 	/** Returns all the audio devices managed by device manager. */
-	TArray<FAudioDevice*> GetAudioDevices();
+	TArray<FAudioDevice*> GetAudioDevices() const;
 
-	TArray<UWorld*> GetWorldsUsingAudioDevice(const Audio::FDeviceId& InID);
+	TArray<UWorld*> GetWorldsUsingAudioDevice(const Audio::FDeviceId& InID) const;
 
 #if INSTRUMENT_AUDIODEVICE_HANDLES
 	void AddStackWalkForContainer(Audio::FDeviceId InId, uint32 StackWalkID, FString&& InStackWalk);
