Temporarily removed DeprecatedFunction tag on BP functions in DataLayerSubsystem to remove warnings coming from BPs (until data is properly fixed).

diff --git a/Engine/Source/Runtime/Engine/Public/WorldPartition/DataLayer/DataLayerSubsystem.h b/Engine/Source/Runtime/Engine/Public/WorldPartition/DataLayer/DataLayerSubsystem.h
--- a/Engine/Source/Runtime/Engine/Public/WorldPartition/DataLayer/DataLayerSubsystem.h
+++ b/Engine/Source/Runtime/Engine/Public/WorldPartition/DataLayer/DataLayerSubsystem.h
@@ -36,20 +36,20 @@ public:
 
 	/** Find a Data Layer by its asset. */
 	UE_DEPRECATED(5.3, "Use UDataLayerManager::GetDataLayerInstanceFromAsset instead.")
-	UFUNCTION(BlueprintCallable, Category = DataLayers, meta = (DeprecatedFunction, DeprecationMessage = "DataLayerSubsystem is deprecated, use the equivalent function in DataLayerManager"))
+	UFUNCTION(BlueprintCallable, Category = DataLayers)
 	UDataLayerInstance* GetDataLayerInstanceFromAsset(const UDataLayerAsset* InDataLayerAsset) const;
 
 	UE_DEPRECATED(5.3, "Use UDataLayerManager::GetDataLayerInstanceRuntimeState instead.")
-	UFUNCTION(BlueprintCallable, Category = DataLayers, meta = (DeprecatedFunction, DeprecationMessage = "DataLayerSubsystem is deprecated, use the equivalent function in DataLayerManager"))
+	UFUNCTION(BlueprintCallable, Category = DataLayers)
 	EDataLayerRuntimeState GetDataLayerInstanceRuntimeState(const UDataLayerAsset* InDataLayerAsset) const;
 
 	UE_DEPRECATED(5.3, "Use UDataLayerManager::GetDataLayerInstanceEffectiveRuntimeState instead.")
-	UFUNCTION(BlueprintCallable, Category = DataLayers, meta = (DeprecatedFunction, DeprecationMessage = "DataLayerSubsystem is deprecated, use the equivalent function in DataLayerManager"))
+	UFUNCTION(BlueprintCallable, Category = DataLayers)
 	EDataLayerRuntimeState GetDataLayerInstanceEffectiveRuntimeState(const UDataLayerAsset* InDataLayerAsset) const;
 
 	/** Set the Data Layer state using its name. */
 	UE_DEPRECATED(5.3, "Use UDataLayerManager::SetDataLayerInstanceRuntimeState instead.")
-	UFUNCTION(BlueprintCallable, Category = DataLayers, BlueprintAuthorityOnly, meta = (DeprecatedFunction, DeprecationMessage = "DataLayerSubsystem is deprecated, use the equivalent function in DataLayerManager"))
+	UFUNCTION(BlueprintCallable, Category = DataLayers)
 	void SetDataLayerInstanceRuntimeState(const UDataLayerAsset* InDataLayerAsset, EDataLayerRuntimeState InState, bool bInIsRecursive = false);
 
 	/** Called when a Data Layer changes state. */
@@ -198,39 +198,39 @@ public:
 	const TSet<FName>& GetLoadedDataLayerNames() const { static TSet<FName> EmptySet; return EmptySet; }
 
 	UE_DEPRECATED(5.1, "Use GetDataLayerInstanceFromAsset instead.")
-	UFUNCTION(BlueprintCallable, Category = DataLayers, meta = (DeprecatedFunction, DeprecationMessage = "DataLayerSubsystem is deprecated, use GetDataLayerInstanceFromAsset in DataLayerManager"))
+	UFUNCTION(BlueprintCallable, Category = DataLayers)
 	UDataLayerInstance* GetDataLayer(const FActorDataLayer& InDataLayer) const;
 
 	UE_DEPRECATED(5.1, "Use GetDataLayerInstanceFromAsset instead.")
-	UFUNCTION(BlueprintCallable, Category = DataLayers, meta = (DeprecatedFunction, DeprecationMessage = "DataLayerSubsystem is deprecated, use GetDataLayerInstanceFromAsset in DataLayerManager"))
+	UFUNCTION(BlueprintCallable, Category = DataLayers)
 	UDataLayerInstance* GetDataLayerFromName(FName InDataLayerName) const;
 
 	UE_DEPRECATED(5.1, "Use GetDataLayerInstanceFromAsset instead.")
-	UFUNCTION(BlueprintCallable, Category = DataLayers, meta = (DeprecatedFunction, DeprecationMessage = "DataLayerSubsystem is deprecated, use GetDataLayerInstanceFromAsset in DataLayerManager"))
+	UFUNCTION(BlueprintCallable, Category = DataLayers)
 	UDataLayerInstance* GetDataLayerFromLabel(FName InDataLayerLabel) const;
 
 	UE_DEPRECATED(5.1, "Use GetDataLayerInstanceRuntimeState instead.")
-	UFUNCTION(BlueprintCallable, Category = DataLayers, meta = (DeprecatedFunction, DeprecationMessage = "DataLayerSubsystem is deprecated, use GetDataLayerInstanceRuntimeState in DataLayerManager"))
+	UFUNCTION(BlueprintCallable, Category = DataLayers)
 	EDataLayerRuntimeState GetDataLayerRuntimeState(const FActorDataLayer& InDataLayer) const;
 
 	UE_DEPRECATED(5.1, "Use GetDataLayerInstanceRuntimeState instead.")
-	UFUNCTION(BlueprintCallable, Category = DataLayers, meta = (DeprecatedFunction, DeprecationMessage = "DataLayerSubsystem is deprecated, use GetDataLayerInstanceRuntimeState in DataLayerManager"))
+	UFUNCTION(BlueprintCallable, Category = DataLayers)
 	EDataLayerRuntimeState GetDataLayerRuntimeStateByLabel(const FName& InDataLayerLabel) const;
 
 	UE_DEPRECATED(5.1, "Use GetDataLayerInstanceEffectiveRuntimeState instead.")
-	UFUNCTION(BlueprintCallable, Category = DataLayers, meta = (DeprecatedFunction, DeprecationMessage = "DataLayerSubsystem is deprecated, use GetDataLayerInstanceEffectiveRuntimeState in DataLayerManager"))
+	UFUNCTION(BlueprintCallable, Category = DataLayers)
 	EDataLayerRuntimeState GetDataLayerEffectiveRuntimeState(const FActorDataLayer& InDataLayer) const;
 
 	UE_DEPRECATED(5.1, "Use GetDataLayerInstanceEffectiveRuntimeState instead.")
-	UFUNCTION(BlueprintCallable, Category = DataLayers, meta = (DeprecatedFunction, DeprecationMessage = "DataLayerSubsystem is deprecated, use GetDataLayerInstanceEffectiveRuntimeState in DataLayerManager"))
+	UFUNCTION(BlueprintCallable, Category = DataLayers)
 	EDataLayerRuntimeState GetDataLayerEffectiveRuntimeStateByLabel(const FName& InDataLayerLabel) const;
 
 	UE_DEPRECATED(5.1, "Use SetDataLayerRuntimeState() with UDataLayerAsset* overload instead.")
-	UFUNCTION(BlueprintCallable, Category = DataLayers, BlueprintAuthorityOnly, meta = (DeprecatedFunction, DeprecationMessage = "DataLayerSubsystem is deprecated, use SetDataLayerInstanceRuntimeState in DataLayerManager"))
+	UFUNCTION(BlueprintCallable, Category = DataLayers, BlueprintAuthorityOnly)
 	void SetDataLayerRuntimeState(const FActorDataLayer& InDataLayer, EDataLayerRuntimeState InState, bool bInIsRecursive = false);
 
 	UE_DEPRECATED(5.1, "Use SetDataLayerInstanceRuntimeState instead.")
-	UFUNCTION(BlueprintCallable, Category = DataLayers, BlueprintAuthorityOnly, meta = (DeprecatedFunction, DeprecationMessage = "DataLayerSubsystem is deprecated, use SetDataLayerInstanceRuntimeState in DataLayerManager"))
+	UFUNCTION(BlueprintCallable, Category = DataLayers, BlueprintAuthorityOnly)
 	void SetDataLayerRuntimeStateByLabel(const FName& InDataLayerLabel, EDataLayerRuntimeState InState, bool bInIsRecursive = false);
 
 private:
