- Set UDataLayerManager::SetDataLayerInstanceRuntimeState BlueprintAuthorityOnly.
- Added UDataLayerManager::SetDataLayerRuntimeState helper to avoid having to resolve a DataLayerInstance for a Data Layer asset.

diff --git a/Engine/Source/Runtime/Engine/Private/WorldPartition/DataLayer/DataLayerManager.cpp b/Engine/Source/Runtime/Engine/Private/WorldPartition/DataLayer/DataLayerManager.cpp
--- a/Engine/Source/Runtime/Engine/Private/WorldPartition/DataLayer/DataLayerManager.cpp
+++ b/Engine/Source/Runtime/Engine/Private/WorldPartition/DataLayer/DataLayerManager.cpp
@@ -206,6 +206,15 @@ bool UDataLayerManager::SetDataLayerInstanceRuntimeState(const UDataLayerInstanc
 	return InDataLayerInstance ? InDataLayerInstance->SetRuntimeState(InState, bInIsRecursive) : false;
 }
 
+bool UDataLayerManager::SetDataLayerRuntimeState(const UDataLayerAsset* InDataLayerAsset, EDataLayerRuntimeState InState, bool bInIsRecursive)
+{
+	if (const UDataLayerInstance* DataLayerInstance = InDataLayerAsset ? GetDataLayerInstanceFromAsset(InDataLayerAsset) : nullptr)
+	{
+		return SetDataLayerInstanceRuntimeState(DataLayerInstance, InState, bInIsRecursive);
+	}
+	return false;
+}
+
 void UDataLayerManager::BroadcastOnDataLayerInstanceRuntimeStateChanged(const UDataLayerInstance* InDataLayer, EDataLayerRuntimeState InState)
 {
 	// For backward compatibility
diff --git a/Engine/Source/Runtime/Engine/Public/WorldPartition/DataLayer/DataLayerManager.h b/Engine/Source/Runtime/Engine/Public/WorldPartition/DataLayer/DataLayerManager.h
--- a/Engine/Source/Runtime/Engine/Public/WorldPartition/DataLayer/DataLayerManager.h
+++ b/Engine/Source/Runtime/Engine/Public/WorldPartition/DataLayer/DataLayerManager.h
@@ -54,9 +54,12 @@ public:
 	UFUNCTION(BlueprintCallable, Category = DataLayers)
 	const UDataLayerInstance* GetDataLayerInstanceFromName(const FName& InDataLayerInstanceName) const;
 
-	UFUNCTION(BlueprintCallable, Category = DataLayers)
+	UFUNCTION(BlueprintCallable, Category = DataLayers, BlueprintAuthorityOnly)
 	bool SetDataLayerInstanceRuntimeState(const UDataLayerInstance* InDataLayerInstance, EDataLayerRuntimeState InState, bool bInIsRecursive = false);
 
+	UFUNCTION(BlueprintCallable, Category = DataLayers, BlueprintAuthorityOnly)
+	bool SetDataLayerRuntimeState(const UDataLayerAsset* InDataLayerAsset, EDataLayerRuntimeState InState, bool bInIsRecursive = false);
+
 	UFUNCTION(BlueprintCallable, Category = DataLayers)
 	EDataLayerRuntimeState GetDataLayerInstanceRuntimeState(const UDataLayerInstance* InDataLayerInstance) const;
 
