Merge DataLayer & ContentBundle information when creating the cell names.
Temp fix to reduce the cells path lenghts until we hash the whole file name

diff --git a/Engine/Source/Runtime/Engine/Private/WorldPartition/WorldPartitionRuntimeSpatialHash.cpp b/Engine/Source/Runtime/Engine/Private/WorldPartition/WorldPartitionRuntimeSpatialHash.cpp
--- a/Engine/Source/Runtime/Engine/Private/WorldPartition/WorldPartitionRuntimeSpatialHash.cpp
+++ b/Engine/Source/Runtime/Engine/Private/WorldPartition/WorldPartitionRuntimeSpatialHash.cpp
@@ -32,6 +32,7 @@
 #include "ProfilingDebugging/ScopedTimers.h"
 #include "UObject/UnrealType.h"
 #include "UObject/ObjectSaveContext.h"
+#include "Templates/TypeHash.h"
 
 #include UE_INLINE_GENERATED_CPP_BY_NAME(WorldPartitionRuntimeSpatialHash)
 
@@ -1234,14 +1235,24 @@ FString UWorldPartitionRuntimeSpatialHash::GetCellNameString(FName InWorldPackag
 	const FString ShortPackageName = UWorld::RemovePIEPrefix(FPackageName::GetShortName(InWorldPackageName));
 	FString CellName = FString::Printf(TEXT("%s_%s_%s"), *ShortPackageName, *InGridName.ToString(), *GetCellCoordString(InCellGlobalCoord));
 
-	if (InDataLayerID.GetHash())
+	// todo_ow : Hash the whole CellName & save a debug string inside the cell instead of hashing only the content bundle & data layer information
+	if (InDataLayerID.GetHash() && InContentBundleID.IsValid())
 	{
-		CellName += FString::Printf(TEXT("_DL%X"), InDataLayerID.GetHash());
+		uint32 Hash = InDataLayerID.GetHash();
+		Hash = HashCombine(Hash, GetTypeHash(InContentBundleID));
+		CellName += FString::Printf(TEXT("_M%X"), Hash);
 	}
-
-	if (InContentBundleID.IsValid())
+	else
 	{
-		CellName += FString::Printf(TEXT("_CB%s"), *UContentBundleDescriptor::GetContentBundleCompactString(InContentBundleID));
+		if (InDataLayerID.GetHash())
+		{
+			CellName += FString::Printf(TEXT("_DL%X"), InDataLayerID.GetHash());
+		}
+
+		if (InContentBundleID.IsValid())
+		{
+			CellName += FString::Printf(TEXT("_CB%s"), *UContentBundleDescriptor::GetContentBundleCompactString(InContentBundleID));
+		}
 	}
 
 	return CellName;
