Framework - Adding an OnPlayerStateChanged to APawn, to allow folks to detect now when the player state is set on a pawn.  There was a few pull requests to make SetPlayerState virtual but going to do this as an alternative.

diff --git a/Engine/Source/Runtime/Engine/Classes/GameFramework/Pawn.h b/Engine/Source/Runtime/Engine/Classes/GameFramework/Pawn.h
--- a/Engine/Source/Runtime/Engine/Classes/GameFramework/Pawn.h
+++ b/Engine/Source/Runtime/Engine/Classes/GameFramework/Pawn.h
@@ -10,6 +10,7 @@
 #include "Engine/EngineTypes.h"
 #include "GameFramework/Actor.h"
 #include "AI/Navigation/NavAgentInterface.h"
+
 #include "Pawn.generated.h"
 
 class AController;
@@ -174,6 +175,12 @@ public:
 	template<class T>
 	T* GetPlayerStateChecked() const { return CastChecked<T>(PlayerState); }
 
+protected:
+	/** Called on both the clientand server when ever SetPlayerState is called on this pawn. */
+	virtual void OnPlayerStateChanged(APlayerState* NewPlayerState, APlayerState* OldPlayerState) { }
+
+public:
+
 	/** Playback of replays writes blended pitch to this, rather than the RemoteViewPitch. This is to avoid having to compress and interpolated value. */
 	float BlendedReplayViewPitch;
 
diff --git a/Engine/Source/Runtime/Engine/Private/Pawn.cpp b/Engine/Source/Runtime/Engine/Private/Pawn.cpp
--- a/Engine/Source/Runtime/Engine/Private/Pawn.cpp
+++ b/Engine/Source/Runtime/Engine/Private/Pawn.cpp
@@ -610,15 +610,21 @@ void APawn::OnRep_PlayerState()
 
 void APawn::SetPlayerState(APlayerState* NewPlayerState)
 {
+	APlayerState* OldPlayerState = PlayerState;
+
 	if (PlayerState && PlayerState->GetPawn() == this)
 	{
 		FSetPlayerStatePawn(PlayerState, nullptr);
 	}
+
 	PlayerState = NewPlayerState;
+
 	if (PlayerState)
 	{
 		FSetPlayerStatePawn(PlayerState, this);
 	}
+
+	OnPlayerStateChanged(NewPlayerState, OldPlayerState);
 }
 
 void APawn::PossessedBy(AController* NewController)
