-Fixed a case where we would leak avoidance objects in UCharacterMovementComponent::SetAvoidanceEnabled if the avoidance manager wasn't set to auto purge.
-Fixed a typo with the avoidance manager's bAutoPurgeOutdatedObjects flag

diff --git a/Engine/Source/Runtime/Engine/Classes/AI/Navigation/AvoidanceManager.h b/Engine/Source/Runtime/Engine/Classes/AI/Navigation/AvoidanceManager.h
--- a/Engine/Source/Runtime/Engine/Classes/AI/Navigation/AvoidanceManager.h
+++ b/Engine/Source/Runtime/Engine/Classes/AI/Navigation/AvoidanceManager.h
@@ -178,6 +178,9 @@ class ENGINE_API UAvoidanceManager : public UObject, public FSelfRegisteringExec
 	/** Used to manually mark data associated with given ID as 'dead' (and reusable) */
 	void RemoveAvoidanceObject(const int32 AvoidanceUID);
 
+	/** Returns true if the avoidance objects are auto purged in RemoveOutdatedObjects */
+	bool IsAutoPurgeEnabled() const { return bAutoPurgeOutdatedObjects; }
+
 protected:
 	//~ Begin FExec Interface
 	virtual bool Exec_Dev(UWorld* Inworld,const TCHAR* Cmd,FOutputDevice& Ar) override;
@@ -217,6 +220,9 @@ protected:
 	/** If set then UAvoidanceManager::RemoveOutdatedObjects is responsisble
 	 *	for removing dead entries in AvoidanceObjects. If set to FALSE then the
 	 *	dead agents need to be removed manually with RemoveAvoidanceObject call */
+	uint32 bAutoPurgeOutdatedObjects : 1;
+
+	UE_DEPRECATED(5.3, "Please use bAutoPurgeOutdatedObjects instead.")
 	uint32 bAutoPurceOutdatedObjects : 1;
 
 #if !(UE_BUILD_SHIPPING || UE_BUILD_TEST)
diff --git a/Engine/Source/Runtime/Engine/Private/AI/Navigation/AvoidanceManager.cpp b/Engine/Source/Runtime/Engine/Private/AI/Navigation/AvoidanceManager.cpp
--- a/Engine/Source/Runtime/Engine/Private/AI/Navigation/AvoidanceManager.cpp
+++ b/Engine/Source/Runtime/Engine/Private/AI/Navigation/AvoidanceManager.cpp
@@ -60,7 +60,7 @@ UAvoidanceManager::UAvoidanceManager(const FObjectInitializer& ObjectInitializer
 	ArtificialRadiusExpansion = 1.5f;
 	TestHeightDifference_DEPRECATED = 500.0f;
 	bRequestedUpdateTimer = false;
-	bAutoPurceOutdatedObjects = true;
+	bAutoPurgeOutdatedObjects = true;
 	HeightCheckMargin = 10.0f;
 
 #if !(UE_BUILD_SHIPPING || UE_BUILD_TEST)
@@ -94,7 +94,7 @@ void UAvoidanceManager::RemoveAvoidanceObject(const int32 AvoidanceUID)
 
 void UAvoidanceManager::RemoveOutdatedObjects()
 {
-	if (bAutoPurceOutdatedObjects == false)
+	if (bAutoPurgeOutdatedObjects == false)
 	{
 		// will be handled manually
 		return;
diff --git a/Engine/Source/Runtime/Engine/Private/Components/CharacterMovementComponent.cpp b/Engine/Source/Runtime/Engine/Private/Components/CharacterMovementComponent.cpp
--- a/Engine/Source/Runtime/Engine/Private/Components/CharacterMovementComponent.cpp
+++ b/Engine/Source/Runtime/Engine/Private/Components/CharacterMovementComponent.cpp
@@ -10788,7 +10788,7 @@ void UCharacterMovementComponent::SetAvoidanceEnabled(bool bEnable)
 	{
 		bUseRVOAvoidance = bEnable;
 
-		// reset id, RegisterMovementComponent call is required to initialize update timers in avoidance manager
+		const int32 OldAvoidanceUID = AvoidanceUID;
 		AvoidanceUID = 0;
 
 		// this is a safety check - it's possible to not have CharacterOwner at this point if this function gets
@@ -10796,10 +10796,21 @@ void UCharacterMovementComponent::SetAvoidanceEnabled(bool bEnable)
 		ensure(GetCharacterOwner());
 		if (GetCharacterOwner() != nullptr)
 		{
-			UAvoidanceManager* AvoidanceManager = GetWorld()->GetAvoidanceManager();
-			if (AvoidanceManager && bEnable)
+			UWorld* World = GetWorld();
+			UAvoidanceManager* AvoidanceManager = World ? World->GetAvoidanceManager() : nullptr;
+
+			if (AvoidanceManager)
 			{
-				AvoidanceManager->RegisterMovementComponent(this, AvoidanceWeight);
+				if (bEnable)
+				{
+					AvoidanceManager->RegisterMovementComponent(this, AvoidanceWeight);
+				}
+				else if (!AvoidanceManager->IsAutoPurgeEnabled())
+				{
+					// When disabling avoidance the object needs to be removed immediately if the manager isn't set to auto purge. 
+					// Otherwise we would leak the object since there isn't anything left to remove it.
+					AvoidanceManager->RemoveAvoidanceObject(OldAvoidanceUID);
+				}
 			}
 		}
 	}
