Remove code preventing switching content bundle Hlods at runtime

diff --git a/Engine/Source/Runtime/Engine/Private/WorldPartition/HLOD/HLODSubsystem.cpp b/Engine/Source/Runtime/Engine/Private/WorldPartition/HLOD/HLODSubsystem.cpp
--- a/Engine/Source/Runtime/Engine/Private/WorldPartition/HLOD/HLODSubsystem.cpp
+++ b/Engine/Source/Runtime/Engine/Private/WorldPartition/HLOD/HLODSubsystem.cpp
@@ -285,12 +285,6 @@ const UHLODSubsystem::FCellData* UHLODSubsystem::GetCellData(const UWorldPartiti
 
 UHLODSubsystem::FCellData* UHLODSubsystem::GetCellData(const UWorldPartitionRuntimeCell* InCell)
 {
-	// @todo_ow ContentBundles do not support Hlods
-	if (InCell->GetContentBundleID().IsValid())
-	{
-		return nullptr;
-	}
-
 	const UWorldPartition* WorldPartition = InCell->GetTypedOuter<UWorldPartition>();
 	check(WorldPartition);
 	FWorldPartitionHLODRuntimeData* WorldPartitionHLODRuntimeData = WorldPartitionsHLODRuntimeData.Find(WorldPartition);
diff --git a/Engine/Source/Runtime/Engine/Private/WorldPartition/WorldPartitionLevelStreamingDynamic.cpp b/Engine/Source/Runtime/Engine/Private/WorldPartition/WorldPartitionLevelStreamingDynamic.cpp
--- a/Engine/Source/Runtime/Engine/Private/WorldPartition/WorldPartitionLevelStreamingDynamic.cpp
+++ b/Engine/Source/Runtime/Engine/Private/WorldPartition/WorldPartitionLevelStreamingDynamic.cpp
@@ -640,10 +640,6 @@ bool UWorldPartitionLevelStreamingDynamic::CanChangeVisibility(bool bMakeVisible
 		{
 			// Source cells that aren't HLOD can always be made visible
 			bool bAlwaysAllowVisibilityChange = bMakeVisible && !RuntimeLevelStreamingCell->GetIsHLOD();
-			
-			// @todo_ow ContentBundles do not support Hlods and events below are forwarding
-			// Events to the HLodSubsystem which assumes knowledge of all cells (not true with plugins)
-			bAlwaysAllowVisibilityChange |= RuntimeLevelStreamingCell->GetContentBundleID().IsValid();
 
 			if (!bAlwaysAllowVisibilityChange)
 			{
