* Fixed missing includes needed to compile non unity/pch on different configs

diff --git a/Engine/Plugins/Online/OnlineSubsystemTencent/Source/Private/OnlinePurchaseTencent.h b/Engine/Plugins/Online/OnlineSubsystemTencent/Source/Private/OnlinePurchaseTencent.h
--- a/Engine/Plugins/Online/OnlineSubsystemTencent/Source/Private/OnlinePurchaseTencent.h
+++ b/Engine/Plugins/Online/OnlineSubsystemTencent/Source/Private/OnlinePurchaseTencent.h
@@ -5,6 +5,7 @@
 #if WITH_TENCENT_RAIL_SDK
 
 #include "Interfaces/OnlinePurchaseInterface.h"
+#include "OnlineSubsystemTencentPackage.h"
 
 class FOnlineSubsystemTencent;
 struct FRequestAllAssetsTaskResult;
diff --git a/Engine/Source/Runtime/Engine/Private/ProfilingDebugging/LevelStreamingProfilingSubsystem.cpp b/Engine/Source/Runtime/Engine/Private/ProfilingDebugging/LevelStreamingProfilingSubsystem.cpp
--- a/Engine/Source/Runtime/Engine/Private/ProfilingDebugging/LevelStreamingProfilingSubsystem.cpp
+++ b/Engine/Source/Runtime/Engine/Private/ProfilingDebugging/LevelStreamingProfilingSubsystem.cpp
@@ -4,13 +4,16 @@
 
 #include "Engine/Level.h"
 #include "Engine/LevelStreaming.h"
-#include "Engine/World.h"
 #include "HAL/FileManager.h"
 #include "Misc/Paths.h"
 #include "Streaming/LevelStreamingDelegates.h"
 #include "WorldPartition/WorldPartition.h"
 #include "WorldPartition/WorldPartitionLevelStreamingDynamic.h"
 
+#if !WITH_EDITOR
+#include "Engine/World.h"
+#endif
+
 #include UE_INLINE_GENERATED_CPP_BY_NAME(LevelStreamingProfilingSubsystem)
 
 DEFINE_LOG_CATEGORY_STATIC(LogLevelStreamingProfiling, Log, All);
diff --git a/Engine/Source/Runtime/Engine/Private/WorldPartition/DataLayer/DataLayerManager.cpp b/Engine/Source/Runtime/Engine/Private/WorldPartition/DataLayer/DataLayerManager.cpp
--- a/Engine/Source/Runtime/Engine/Private/WorldPartition/DataLayer/DataLayerManager.cpp
+++ b/Engine/Source/Runtime/Engine/Private/WorldPartition/DataLayer/DataLayerManager.cpp
@@ -11,9 +11,7 @@
 #include "WorldPartition/WorldPartitionRuntimeCell.h"
 #include "WorldPartition/WorldPartitionHelpers.h"
 #include "Engine/Canvas.h"
-#include "Engine/Engine.h"
 #include "Debug/DebugDrawService.h"
-#include "Misc/Paths.h"
 
 #if WITH_EDITOR
 #include "WorldPartition/DataLayer/WorldDataLayersActorDesc.h"
@@ -28,6 +26,9 @@
 #include "EngineUtils.h"
 #include "LevelUtils.h"
 #include "Editor.h"
+#else
+#include "Engine/Engine.h"
+#include "Misc/Paths.h"
 #endif
 
 extern int32 GDrawDataLayersLoadTime;
diff --git a/Engine/Source/Runtime/Engine/Private/WorldPartition/HLOD/HLODDestruction.cpp b/Engine/Source/Runtime/Engine/Private/WorldPartition/HLOD/HLODDestruction.cpp
--- a/Engine/Source/Runtime/Engine/Private/WorldPartition/HLOD/HLODDestruction.cpp
+++ b/Engine/Source/Runtime/Engine/Private/WorldPartition/HLOD/HLODDestruction.cpp
@@ -3,7 +3,6 @@
 #include "WorldPartition/HLOD/HLODDestruction.h"
 
 #include "Engine/LevelStreaming.h"
-#include "Engine/World.h"
 #include "LevelUtils.h"
 #include "UObject/ScriptInterface.h"
 #include "WorldPartition/HLOD/DestructibleHLODComponent.h"
@@ -13,6 +12,9 @@
 
 #include UE_INLINE_GENERATED_CPP_BY_NAME(HLODDestruction)
 
+#if !WITH_EDITOR
+#include "Engine/World.h"
+#endif
 
 namespace
 {
