StreamableManager::OnPreGarbageCollection optimization, ~3x speedup on internal project

Changed FStreamable::ActiveHandles from weak to raw pointers that are synchronously removed as handles get released / canceled / deleted. Allows OnPreGarbageCollection to check if ActiveHandles.IsEmpty() instead of polling weak handles.

Might move all OnPreGarbageCollection work to AddReferencedObjects, which is running in parallel to the mark phase, in the future.

Also some driveby cleanup.

diff --git a/Engine/Source/Runtime/Engine/Private/StreamableManager.cpp b/Engine/Source/Runtime/Engine/Private/StreamableManager.cpp
--- a/Engine/Source/Runtime/Engine/Private/StreamableManager.cpp
+++ b/Engine/Source/Runtime/Engine/Private/StreamableManager.cpp
@@ -686,12 +686,25 @@ bool FStreamableHandle::HasUpdateDelegate() const
 	return UpdateDelegate.IsBound();
 }
 
+static void RemoveActiveHandle(FStreamable& Streamable, FStreamableHandle& Handle);		
+
 FStreamableHandle::~FStreamableHandle()
 {
-	check(IsInGameThread() || IsInGarbageCollectorThread());
+	checkf(IsInGameThread(), TEXT("Streamable handles aren't thread-safe and must be destroyed on the game thread"));
 
 	if (IsActive())
 	{
+		if (!bStalled)
+		{
+			for (const FSoftObjectPath&  RequestedAsset : RequestedAssets)
+			{
+				if (FStreamable* Streamable = OwningManager->FindStreamable(RequestedAsset))
+				{
+					RemoveActiveHandle(*Streamable, *this);
+				}
+			}
+		}
+
 		bReleased = true;
 		OwningManager = nullptr;
 		
@@ -973,70 +986,82 @@ TSharedPtr<FStreamableHandle> FStreamableHandle::GetOutermostHandle()
 /** Internal object, one of these per object paths managed by this system */
 struct FStreamable
 {
-	/** Hard Pointer to object */
-	UObject* Target;
-	/** If this object is currently being loaded */
-	bool	bAsyncLoadRequestOutstanding;
-	/** If this object failed to load, don't try again */
-	bool	bLoadFailed;
+	/** Hard GC pointer to object */
+	UObject* Target = nullptr;
+
+	/** Live handles keeping this alive. Handles may be duplicated. */
+	TArray<FStreamableHandle*> ActiveHandles;
 
-	/** List of handles that are waiting for this to load. The same handle may be here multiple times with redirectors */
+	/** Handles waiting for this to load. Handles may be duplicated with redirectors. */
 	TArray< TSharedRef< FStreamableHandle> > LoadingHandles;
 
-	/** List of handles that are keeping this streamable in memory. The same handle may be here multiple times */
-	TArray< TWeakPtr< FStreamableHandle> > ActiveHandles;
+	/** If this object is currently being loaded */
+	bool bAsyncLoadRequestOutstanding = false;
+
+	/** If this object failed to load, don't try again */
+	bool bLoadFailed = false;
 
-	FStreamable()
-		: Target(nullptr)
-		, bAsyncLoadRequestOutstanding(false)
-		, bLoadFailed(false)
-	{
-	}
 
-	~FStreamable()
+	void FreeHandles()
 	{
 		// Clear the loading handles 
-		for (int32 i = 0; i < LoadingHandles.Num(); i++)
+		for (TSharedRef<FStreamableHandle>& LoadingHandle : LoadingHandles)
 		{
-			LoadingHandles[i]->StreamablesLoading--;
+			LoadingHandle->StreamablesLoading--;
 		}
 		LoadingHandles.Empty();
 
 		// Cancel active handles, this list includes the loading handles
-		for (int32 i = 0; i < ActiveHandles.Num(); i++)
+		while (ActiveHandles.Num() > 0)
 		{
-			TSharedPtr<FStreamableHandle> ActiveHandle = ActiveHandles[i].Pin();
-
-			if (ActiveHandle.IsValid() && !ActiveHandle->bCanceled)
+			FStreamableHandle* ActiveHandle = ActiveHandles.Pop(/* shrink */ false);
+			if (!ActiveHandle->bCanceled)
 			{
 				// Full cancel isn't safe any more
 
 				ActiveHandle->bCanceled = true;
 				ActiveHandle->OwningManager = nullptr;
 
-				if (!ActiveHandle->bReleased)
+				if (ActiveHandle->bReleased)
 				{
-					FStreamableHandle::ExecuteDelegate(ActiveHandle->CancelDelegate, ActiveHandle);
-					ActiveHandle->UnbindDelegates();
+					ActiveHandle->NotifyParentsOfCancellation();
 				}
+				else
+				{
+					// Keep handle alive to stop the cancel callback from dropping the last reference
+					TSharedPtr<FStreamableHandle> SharedHandle = ActiveHandle->AsShared();
 
-				ActiveHandle->NotifyParentsOfCancellation();
-
+					FStreamableHandle::ExecuteDelegate(ActiveHandle->CancelDelegate, SharedHandle);
+					ActiveHandle->UnbindDelegates();
+					ActiveHandle->NotifyParentsOfCancellation();
+				}
 			}
 		}
-		ActiveHandles.Empty();
 	}
 
-	void AddLoadingRequest(TSharedRef<FStreamableHandle> NewRequest)
+	void AddLoadingRequest(TSharedRef<FStreamableHandle>& NewRequest)
 	{
 		// With redirectors we can end up adding the same handle multiple times, this is unusual but supported
-		ActiveHandles.Add(NewRequest);
+		ActiveHandles.Add(&NewRequest.Get());
 
 		LoadingHandles.Add(NewRequest);
 		NewRequest->StreamablesLoading++;
 	}
 };
 
+void RemoveActiveHandle(FStreamable& Streamable, FStreamableHandle& Handle)
+{
+	TArray<FStreamableHandle*>& ActiveHandles = Streamable.ActiveHandles;
+	for (int32 Idx = Streamable.ActiveHandles.Num() - 1; Idx >= 0; --Idx)
+	{
+		if (ActiveHandles[Idx] == &Handle)
+		{
+			ActiveHandles.RemoveAtSwap(Idx, 1, /* shrink */ false);
+		}
+	}
+}
+	
+
 FStreamableManager::FStreamableManager()
 	: FGCObject(EFlags::AddStableNativeReferencesOnly)
 {
@@ -1049,50 +1074,53 @@ FStreamableManager::~FStreamableManager()
 {
 	FCoreUObjectDelegates::GetPreGarbageCollectDelegate().RemoveAll(this);
 
-	for (TStreamableMap::TIterator It(StreamableItems); It; ++It)
+	for (const TPair<FSoftObjectPath, FStreamable*>& Pair : StreamableItems)
+	{
+		Pair.Value->FreeHandles();
+	}
+	
+	TStreamableMap TempStreamables;
+	Swap(TempStreamables, StreamableItems);
+	for (const TPair<FSoftObjectPath, FStreamable*>& TempPair : TempStreamables)
 	{
-		delete It.Value();
-		It.RemoveCurrent();
+		delete TempPair.Value;
 	}
 }
 
 void FStreamableManager::OnPreGarbageCollect()
 {
 	TRACE_CPUPROFILER_EVENT_SCOPE(FStreamableManager::OnPreGarbageCollect);
-	TSet<FSoftObjectPath> RedirectsToRemove;
-	TArray<FStreamable*> StreamablesToDelete;
 
-	// Remove any streamables with no active handles, as GC may have freed them
-	for (TStreamableMap::TIterator It(StreamableItems); It; ++It)
+	// Find streamables without active handles
+	TBitArray<> InactiveStreamables;
+	InactiveStreamables.SetNumUninitialized(StreamableItems.Num());
+	int32 Idx = 0;
+	for (const TPair<FSoftObjectPath, FStreamable*>& Pair : StreamableItems)
 	{
-		FStreamable* Existing = It.Value();
-
-		// Remove invalid handles, the weak pointers may be pointing to removed handles
-		for (int32 i = Existing->ActiveHandles.Num() - 1; i >= 0; i--)
-		{
-			TSharedPtr<FStreamableHandle> ActiveHandle = Existing->ActiveHandles[i].Pin();
-
-			if (!ActiveHandle.IsValid())
-			{
-				Existing->ActiveHandles.RemoveAtSwap(i, 1, false);
-			}
-		}
-
-		if (Existing->ActiveHandles.Num() == 0)
-		{
-			RedirectsToRemove.Add(It.Key());
-			StreamablesToDelete.Add(Existing);
-			It.RemoveCurrent();
-		}
+		InactiveStreamables[Idx++] = Pair.Value->ActiveHandles.IsEmpty();
 	}
 
-	for (FStreamable* StreamableToDelete : StreamablesToDelete)
+	int32 NumInactive = InactiveStreamables.CountSetBits();
+	if (NumInactive == 0)
 	{
-		delete StreamableToDelete;
+		return;
 	}
 
-	if (RedirectsToRemove.Num() > 0)
+	// Remove redirects
+	if (StreamableRedirects.Num() > 0)
 	{
+		TSet<FSoftObjectPath> RedirectsToRemove;
+		RedirectsToRemove.Reserve(NumInactive);
+
+		Idx = 0;
+		for (const TPair<FSoftObjectPath, FStreamable*>& Pair : StreamableItems)
+		{
+			if (InactiveStreamables[Idx++])
+			{
+				RedirectsToRemove.Add(Pair.Key);
+			}
+		}
+
 		for (TStreamableRedirects::TIterator It(StreamableRedirects); It; ++It)
 		{
 			if (RedirectsToRemove.Contains(It.Value().NewPath))
@@ -1101,6 +1129,18 @@ void FStreamableManager::OnPreGarbageCollect()
 			}
 		}
 	}
+
+	// Delete and remove inactive streamables
+	Idx = 0;
+	for (TStreamableMap::TIterator It(StreamableItems); It; ++It)
+	{
+		if (InactiveStreamables[Idx++])
+		{
+			It->Value->FreeHandles();
+			delete It->Value;
+			It.RemoveCurrent();
+		}
+	}
 }
 
 void FStreamableManager::AddReferencedObjects(FReferenceCollector& Collector)
@@ -1138,14 +1178,10 @@ bool FStreamableManager::GetReferencerPropertyName(UObject* Object, FString& Out
 			check(Existing);
 			if (Existing->Target == Object)
 			{
-				for (const TWeakPtr<FStreamableHandle>& WeakHandle : Existing->ActiveHandles)
+				if (Existing->ActiveHandles.Num() > 0)
 				{
-					TSharedPtr<FStreamableHandle> Handle = WeakHandle.Pin();
-					if (Handle.IsValid())
-					{
-						OutPropertyName = Handle->GetDebugName();
-						return true;
-					}
+					OutPropertyName = Existing->ActiveHandles[0]->GetDebugName();
+					return true;
 				}
 
 				if (Existing->LoadingHandles.Num() > 0)
@@ -1613,11 +1649,11 @@ void FStreamableManager::RemoveReferencedAsset(const FSoftObjectPath& Target, TS
 	// This should always be in the active handles list
 	if (ensureMsgf(Existing, TEXT("Failed to find existing streamable for %s"), *Target.ToString()))
 	{
-		ensureMsgf(Existing->ActiveHandles.Remove(Handle) > 0, TEXT("Failed to remove active handle for %s"), *Target.ToString());
+		ensureMsgf(Existing->ActiveHandles.RemoveSwap(&Handle.Get(), /* shrink */ false) > 0, TEXT("Failed to remove active handle for %s"), *Target.ToString());
 
 		// Try removing from loading list if it's still there, this won't call the callback as it's being called from cancel
 		// This may remove more than one copy if streamables were merged
-		int32 LoadingRemoved = Existing->LoadingHandles.Remove(Handle);
+		int32 LoadingRemoved = Existing->LoadingHandles.RemoveSwap(Handle);
 		if (LoadingRemoved > 0)
 		{
 			Handle->StreamablesLoading -= LoadingRemoved;
@@ -1796,22 +1832,18 @@ bool FStreamableManager::GetActiveHandles(const FSoftObjectPath& Target, TArray<
 	FStreamable* Existing = FindStreamable(Target);
 	if (Existing && Existing->ActiveHandles.Num() > 0)
 	{
-		for (TWeakPtr<FStreamableHandle> WeakHandle : Existing->ActiveHandles)
+		for (FStreamableHandle* ActiveHandle : Existing->ActiveHandles)
 		{
-			TSharedPtr<FStreamableHandle> Handle = WeakHandle.Pin();
+			ensure(ActiveHandle->OwningManager == this);
 
-			if (Handle.IsValid())
+			TSharedRef<FStreamableHandle> HandleRef = ActiveHandle->AsShared();
+			if (!bOnlyManagedHandles || ManagedActiveHandles.Contains(HandleRef))
 			{
-				ensure(Handle->OwningManager == this);
-
-				TSharedRef<FStreamableHandle> HandleRef = Handle.ToSharedRef();
-				if (!bOnlyManagedHandles || ManagedActiveHandles.Contains(HandleRef))
-				{
-					// Only add each handle once, we can have duplicates in the source list
-					HandleList.AddUnique(HandleRef);
-				}
+				// Only add each handle once, we can have duplicates in the source list
+				HandleList.AddUnique(MoveTemp(HandleRef));
 			}
 		}
+
 		return HandleList.Num() > 0;
 	}
 
@@ -1830,7 +1862,7 @@ FSoftObjectPath FStreamableManager::ResolveRedirects(const FSoftObjectPath& Targ
 	return Target;
 }
 
-FSoftObjectPath FStreamableManager::HandleLoadedRedirector(UObjectRedirector* LoadedRedirector, FSoftObjectPath RequestedPath, struct FStreamable* RequestedStreamable)
+FSoftObjectPath FStreamableManager::HandleLoadedRedirector(UObjectRedirector* LoadedRedirector, FSoftObjectPath RequestedPath, FStreamable* RequestedStreamable)
 {
 	UE_LOG(LogStreamableManager, Verbose, TEXT("     Found redirect %s"), *LoadedRedirector->GetPathName());
 
@@ -1863,12 +1895,10 @@ FSoftObjectPath FStreamableManager::HandleLoadedRedirector(UObjectRedirector* Lo
 		// This may result in the same handle being in the list twice! But the rest of the code is ready for that
 		// We let LoadFailed and InProgress stay as false on the new streamable because we've successfully loaded an object
 
-		RequestedStreamable->LoadingHandles.Append((*FoundStreamable)->LoadingHandles);
-		(*FoundStreamable)->LoadingHandles.Empty();
-
-		RequestedStreamable->ActiveHandles.Append((*FoundStreamable)->ActiveHandles);
-		(*FoundStreamable)->ActiveHandles.Empty();
+		RequestedStreamable->LoadingHandles.Append(MoveTemp((*FoundStreamable)->LoadingHandles));
+		RequestedStreamable->ActiveHandles.Append(MoveTemp((*FoundStreamable)->ActiveHandles));
 
+		// No handles remain so need to FreeHandles()
 		delete *FoundStreamable;
 	}
 
