Fix r.RayTracing.Geometry.StaticMeshes.WPO
- when cvar = 0, primitives were not being excluded from raytracing scene
- when cvar = 2, instead of disabling WPO evaluation, logic was ignoring bEvaluateWorldPositionOffsetInRayTracing from Static Mesh component and instead applying WPO for any meshes that used it.

diff --git a/Engine/Source/Runtime/Engine/Private/StaticMeshRender.cpp b/Engine/Source/Runtime/Engine/Private/StaticMeshRender.cpp
--- a/Engine/Source/Runtime/Engine/Private/StaticMeshRender.cpp
+++ b/Engine/Source/Runtime/Engine/Private/StaticMeshRender.cpp
@@ -303,24 +303,31 @@ FStaticMeshSceneProxy::FStaticMeshSceneProxy(UStaticMeshComponent* InComponent,
 #if RHI_RAYTRACING
 	bSupportRayTracing = InComponent->GetStaticMesh()->bSupportRayTracing;
 	bDynamicRayTracingGeometry = false;
+	bNeedsDynamicRayTracingGeometries = false;
 	
 	if (IsRayTracingAllowed() && bSupportRayTracing)
 	{
-		if (CVarRayTracingStaticMeshesWPO.GetValueOnAnyThread() > 0)
-		{
-			bDynamicRayTracingGeometry = MaterialRelevance.bUsesWorldPositionOffset;
-
-			if (CVarRayTracingStaticMeshesWPO.GetValueOnAnyThread() == 1)
-			{
-				bDynamicRayTracingGeometry &= InComponent->bEvaluateWorldPositionOffsetInRayTracing;
-			}
-		}
-
 		RayTracingGeometries.AddDefaulted(RenderData->LODResources.Num());
 		for (int32 LODIndex = 0; LODIndex < RenderData->LODResources.Num(); LODIndex++)
 		{
 			RayTracingGeometries[LODIndex] = &RenderData->LODResources[LODIndex].RayTracingGeometry;
 		}
+		
+		const bool bWantsRayTracingWPO = MaterialRelevance.bUsesWorldPositionOffset && InComponent->bEvaluateWorldPositionOffsetInRayTracing;
+
+		// r.RayTracing.Geometry.StaticMeshes.WPO is handled in the following way:
+		// 0 - mark ray tracing geometry as dynamic but don't create any dynamic geometries since it won't be included in ray tracing scene
+		// 1 - mark ray tracing geometry as dynamic and create dynamic geometries
+		// 2 - don't mark ray tracing geometry as dynamic nor create any dynamic geometries since WPO evaluation is disabled
+
+		// if r.RayTracing.Geometry.StaticMeshes.WPO == 2, WPO evaluation is disabled so don't need to mark geometry as dynamic
+		if (bWantsRayTracingWPO && CVarRayTracingStaticMeshesWPO.GetValueOnAnyThread() != 2)
+		{
+			bDynamicRayTracingGeometry = true;
+
+			// only need dynamic geometries when r.RayTracing.Geometry.StaticMeshes.WPO == 1
+			bNeedsDynamicRayTracingGeometries = CVarRayTracingStaticMeshesWPO.GetValueOnAnyThread() == 1;
+		}
 	}
 #endif
 
@@ -421,60 +428,42 @@ FStaticMeshSceneProxy::FStaticMeshSceneProxy(UStaticMeshComponent* InComponent,
 void FStaticMeshSceneProxy::SetEvaluateWorldPositionOffsetInRayTracing(bool NewValue)
 {
 #if RHI_RAYTRACING
-	// Skip this code if the cvar settings are forcing WPO evaluation on/off
-	if (CVarRayTracingStaticMeshesWPO.GetValueOnAnyThread() != 1)
+	if (!IsRayTracingAllowed() || !bSupportRayTracing)
+	{
 		return;
+	}
+
+	// if r.RayTracing.Geometry.StaticMeshes.WPO == 2, WPO evaluation is disabled so don't need to mark geometry as dynamic
+	NewValue &= MaterialRelevance.bUsesWorldPositionOffset && CVarRayTracingStaticMeshesWPO.GetValueOnAnyThread() != 2;
 
-	NewValue &= MaterialRelevance.bUsesWorldPositionOffset;
 	if (NewValue && !bDynamicRayTracingGeometry)
 	{
 		bDynamicRayTracingGeometry = true;
-		if (IsRayTracingAllowed())
-		{
-			DynamicRayTracingGeometries.AddDefaulted(RenderData->LODResources.Num());
 
-			for (int32 LODIndex = 0; LODIndex < RenderData->LODResources.Num(); LODIndex++)
-			{
-				auto& Initializer = DynamicRayTracingGeometries[LODIndex].Initializer;
-				Initializer = RenderData->LODResources[LODIndex].RayTracingGeometry.Initializer;
-				for (FRayTracingGeometrySegment& Segment : Initializer.Segments)
-				{
-					Segment.VertexBuffer = nullptr;
-				}
-				Initializer.bAllowUpdate = true;
-				Initializer.bFastBuild = true;
-				Initializer.Type = ERayTracingGeometryInitializerType::Rendering;
-			}
+		// only need dynamic geometries when r.RayTracing.Geometry.StaticMeshes.WPO == 1
+		bNeedsDynamicRayTracingGeometries = CVarRayTracingStaticMeshesWPO.GetValueOnAnyThread() == 1;
 
-			for (int32 i = 0; i < DynamicRayTracingGeometries.Num(); i++)
-			{
-				auto& Geometry = DynamicRayTracingGeometries[i];
-				Geometry.InitResource();
-			}
+		if (GetPrimitiveSceneInfo())
+		{
+			GetPrimitiveSceneInfo()->bIsRayTracingStaticRelevant = IsRayTracingStaticRelevant();
+		}
 
-			if (GetPrimitiveSceneInfo())
-			{
-				GetPrimitiveSceneInfo()->bIsRayTracingStaticRelevant = IsRayTracingStaticRelevant();
-			}
+		if (bNeedsDynamicRayTracingGeometries)
+		{
+			CreateDynamicRayTracingGeometries();
 		}
 	}
 	else if (!NewValue && bDynamicRayTracingGeometry)
 	{
 		bDynamicRayTracingGeometry = false;
-		if (IsRayTracingAllowed())
-		{
-			for (auto& Geometry : DynamicRayTracingGeometries)
-			{
-				Geometry.ReleaseResource();
-			}
-
-			DynamicRayTracingGeometries.Empty();
+		bNeedsDynamicRayTracingGeometries = false;
 
-			if (GetPrimitiveSceneInfo())
-			{
-				GetPrimitiveSceneInfo()->bIsRayTracingStaticRelevant = IsRayTracingStaticRelevant();
-			}
+		if (GetPrimitiveSceneInfo())
+		{
+			GetPrimitiveSceneInfo()->bIsRayTracingStaticRelevant = IsRayTracingStaticRelevant();
 		}
+
+		ReleaseDynamicRayTracingGeometries();
 	}
 #endif
 }
@@ -488,10 +477,7 @@ int32 FStaticMeshSceneProxy::GetLightMapCoordinateIndex() const
 FStaticMeshSceneProxy::~FStaticMeshSceneProxy()
 {
 #if RHI_RAYTRACING
-	for (auto& Geometry: DynamicRayTracingGeometries)
-	{
-		Geometry.ReleaseResource();
-	}
+	ReleaseDynamicRayTracingGeometries();
 #endif
 }
 
@@ -731,33 +717,55 @@ bool FStaticMeshSceneProxy::GetMeshElement(
 	}
 }
 
-void FStaticMeshSceneProxy::CreateRenderThreadResources()
-{
 #if RHI_RAYTRACING
-	if(IsRayTracingAllowed() && bSupportRayTracing)
+void FStaticMeshSceneProxy::CreateDynamicRayTracingGeometries()
+{
+	check(DynamicRayTracingGeometries.IsEmpty());
+
+	DynamicRayTracingGeometries.AddDefaulted(RenderData->LODResources.Num());
+
+	for (int32 LODIndex = 0; LODIndex < RenderData->LODResources.Num(); LODIndex++)
 	{
-		if (bDynamicRayTracingGeometry)
+		auto& Initializer = DynamicRayTracingGeometries[LODIndex].Initializer;
+		Initializer = RenderData->LODResources[LODIndex].RayTracingGeometry.Initializer;
+		for (FRayTracingGeometrySegment& Segment : Initializer.Segments)
 		{
-			DynamicRayTracingGeometries.AddDefaulted(RenderData->LODResources.Num());
-			for (int32 LODIndex = 0; LODIndex < RenderData->LODResources.Num(); LODIndex++)
-			{
-				auto& Initializer = DynamicRayTracingGeometries[LODIndex].Initializer;
-				Initializer = RenderData->LODResources[LODIndex].RayTracingGeometry.Initializer;
-				for (FRayTracingGeometrySegment& Segment : Initializer.Segments)
-				{
-					Segment.VertexBuffer = nullptr;
-				}
-				Initializer.bAllowUpdate = true;
-				Initializer.bFastBuild = true;
-				Initializer.Type = ERayTracingGeometryInitializerType::Rendering;
-			}
+			Segment.VertexBuffer = nullptr;
 		}
+		Initializer.bAllowUpdate = true;
+		Initializer.bFastBuild = true;
+		Initializer.Type = ERayTracingGeometryInitializerType::Rendering;
+	}
 
-		for(int32 i = 0; i < DynamicRayTracingGeometries.Num(); i++)
-		{
-			auto& Geometry = DynamicRayTracingGeometries[i];
-			Geometry.InitResource();
-		}
+	for (int32 i = 0; i < DynamicRayTracingGeometries.Num(); i++)
+	{
+		auto& Geometry = DynamicRayTracingGeometries[i];
+		Geometry.InitResource();
+	}
+}
+
+void FStaticMeshSceneProxy::ReleaseDynamicRayTracingGeometries()
+{
+	for (auto& Geometry : DynamicRayTracingGeometries)
+	{
+		Geometry.ReleaseResource();
+	}
+
+	DynamicRayTracingGeometries.Empty();
+}
+#endif
+
+void FStaticMeshSceneProxy::CreateRenderThreadResources()
+{
+#if RHI_RAYTRACING
+	if(IsRayTracingAllowed() && bNeedsDynamicRayTracingGeometries)
+	{
+		check(bDynamicRayTracingGeometry);
+		CreateDynamicRayTracingGeometries();
+	}
+	else
+	{
+		checkf(DynamicRayTracingGeometries.IsEmpty(), TEXT("Proxy shouldn't have entries in DynamicRayTracingGeometries."));
 	}
 #endif
 }
@@ -1788,11 +1796,14 @@ bool FStaticMeshSceneProxy::HasRayTracingRepresentation() const
 
 void FStaticMeshSceneProxy::GetDynamicRayTracingInstances(FRayTracingMaterialGatheringContext& Context, TArray<FRayTracingInstance>& OutRayTracingInstances )
 {
-	if (DynamicRayTracingGeometries.Num() <= 0 || CVarRayTracingStaticMeshes.GetValueOnRenderThread() == 0)
+	if (DynamicRayTracingGeometries.IsEmpty() || CVarRayTracingStaticMeshes.GetValueOnRenderThread() == 0)
 	{
 		return;
 	}
 
+	checkf(CVarRayTracingStaticMeshesWPO.GetValueOnRenderThread() == 1,
+		TEXT("Proxy should only have entries in DynamicRayTracingGeometries when WPO evaluation is enabled"));
+
 	if (!ensureMsgf(IsRayTracingRelevant(),
 		TEXT("GetDynamicRayTracingInstances() is only expected to be called for scene proxies that are compatible with ray tracing. ")
 		TEXT("RT-relevant primitive gathering code in FDeferredShadingSceneRenderer may be wrong.")))
diff --git a/Engine/Source/Runtime/Engine/Public/StaticMeshSceneProxy.h b/Engine/Source/Runtime/Engine/Public/StaticMeshSceneProxy.h
--- a/Engine/Source/Runtime/Engine/Public/StaticMeshSceneProxy.h
+++ b/Engine/Source/Runtime/Engine/Public/StaticMeshSceneProxy.h
@@ -224,8 +224,12 @@ protected:
 	float OverlayMaterialMaxDrawDistance;
 
 #if RHI_RAYTRACING
-	bool bSupportRayTracing;
-	bool bDynamicRayTracingGeometry;
+	void CreateDynamicRayTracingGeometries();
+	void ReleaseDynamicRayTracingGeometries();
+
+	bool bSupportRayTracing : 1;
+	bool bDynamicRayTracingGeometry : 1;
+	bool bNeedsDynamicRayTracingGeometries : 1;
 	TArray<FRayTracingGeometry, TInlineAllocator<MAX_MESH_LOD_COUNT>> DynamicRayTracingGeometries;
 	TArray<FMeshBatch> CachedRayTracingMaterials;
 	int16 CachedRayTracingMaterialsLODIndex = INDEX_NONE;
