[NO NON-PII COMMIT MESSAGE]

diff --git a/Engine/Source/Runtime/Engine/Classes/Components/SplineComponent.h b/Engine/Source/Runtime/Engine/Classes/Components/SplineComponent.h
--- a/Engine/Source/Runtime/Engine/Classes/Components/SplineComponent.h
+++ b/Engine/Source/Runtime/Engine/Classes/Components/SplineComponent.h
@@ -527,6 +527,14 @@ public:
 	UFUNCTION(BlueprintCallable, Category = Spline)
 	int32 GetNumberOfSplineSegments() const;
 
+	/** Get the input key (e.g. the time) of the control point of the spline at the specified index. */
+	UFUNCTION(BlueprintCallable, Category = Spline)
+	float GetInputKeyValueAtSplinePoint(int32 PointIndex) const;
+
+	/** Gets the spline point of the spline at the specified index */
+	UFUNCTION(BlueprintCallable, Category = Spline)
+	FSplinePoint GetSplinePointAt(int32 PointIndex, ESplineCoordinateSpace::Type CoordinateSpace) const;
+		
 	/** Get the location at spline point */
 	UFUNCTION(BlueprintCallable, Category = Spline)
 	FVector GetLocationAtSplinePoint(int32 PointIndex, ESplineCoordinateSpace::Type CoordinateSpace) const;
diff --git a/Engine/Source/Runtime/Engine/Private/Components/SplineComponent.cpp b/Engine/Source/Runtime/Engine/Private/Components/SplineComponent.cpp
--- a/Engine/Source/Runtime/Engine/Private/Components/SplineComponent.cpp
+++ b/Engine/Source/Runtime/Engine/Private/Components/SplineComponent.cpp
@@ -1009,6 +1009,29 @@ int32 USplineComponent::GetNumberOfSplineSegments() const
 	return (bClosedLoop ? NumPoints : NumPoints - 1);
 }
 
+float USplineComponent::GetInputKeyValueAtSplinePoint(int32 PointIndex) const
+{
+	const FInterpCurvePointVector& SplinePoint = GetPositionPointSafe(PointIndex);
+	return SplinePoint.InVal;
+}
+
+
+FSplinePoint USplineComponent::GetSplinePointAt(int32 PointIndex, ESplineCoordinateSpace::Type CoordinateSpace) const
+{
+	const FInterpCurvePointVector& SplinePoint = GetPositionPointSafe(PointIndex);
+
+	const FInterpCurvePointQuat& RotationPoint = GetRotationPointSafe(PointIndex);
+	const FRotator& Rotation = GetRotationAtSplineInputKey(RotationPoint.InVal, CoordinateSpace);
+
+	const FVector Scale = GetScaleAtSplinePoint(PointIndex);
+
+	return FSplinePoint(SplinePoint.InVal, 
+		SplinePoint.OutVal,
+		SplinePoint.ArriveTangent,
+		SplinePoint.LeaveTangent,
+		Rotation,
+		Scale);
+}
 
 FVector USplineComponent::GetLocationAtSplinePoint(int32 PointIndex, ESplineCoordinateSpace::Type CoordinateSpace) const
 {
