Change Canvas Translucent blend option to use AlphaBlend internally, to get a correct alpha channel output.

diff --git a/Engine/Source/Runtime/Engine/Private/UserInterface/Canvas.cpp b/Engine/Source/Runtime/Engine/Private/UserInterface/Canvas.cpp
--- a/Engine/Source/Runtime/Engine/Private/UserInterface/Canvas.cpp
+++ b/Engine/Source/Runtime/Engine/Private/UserInterface/Canvas.cpp
@@ -1550,7 +1550,7 @@ ESimpleElementBlendMode FCanvas::BlendToSimpleElementBlend(EBlendMode BlendMode)
 			return SE_BLEND_AlphaHoldout;
 		case BLEND_Translucent:
 		default:
-			return SE_BLEND_Translucent;
+			return SE_BLEND_AlphaBlend; // Same as translucent, but also write a proper alpha value
 	};
 }
 
