World Partition: Properly cluster parent actors & their children together if their spatially loaded flag differ.

diff --git a/Engine/Source/Runtime/Engine/Private/WorldPartition/WorldPartitionStreamingGeneration.cpp b/Engine/Source/Runtime/Engine/Private/WorldPartition/WorldPartitionStreamingGeneration.cpp
--- a/Engine/Source/Runtime/Engine/Private/WorldPartition/WorldPartitionStreamingGeneration.cpp
+++ b/Engine/Source/Runtime/Engine/Private/WorldPartition/WorldPartitionStreamingGeneration.cpp
@@ -585,11 +585,19 @@ class FWorldPartitionStreamingGenerator
 		int32 NbErrorsDetected = INDEX_NONE;
 		for(uint32 NbValidationPasses = 0; NbErrorsDetected; NbValidationPasses++)
 		{
+			// Type of work performed in this pass, for clarity
+			enum class EPassType
+			{
+				ErrorReporting,
+				Fixup
+			};
+			const EPassType PassType = NbValidationPasses == 0 ? EPassType::ErrorReporting : EPassType::Fixup;
+
 			NbErrorsDetected = 0;
 
-			ContainerDescriptor.ActorDescViewMap.ForEachActorDescView([this, &ContainerDescriptor, &NbErrorsDetected, &NbValidationPasses](FWorldPartitionActorDescView& ActorDescView)
+			ContainerDescriptor.ActorDescViewMap.ForEachActorDescView([this, &ContainerDescriptor, &NbErrorsDetected, PassType](FWorldPartitionActorDescView& ActorDescView)
 			{
-				// Validate data layers
+				// Validate grid placement
 				auto IsReferenceGridPlacementValid = [](const FWorldPartitionActorDescView& RefererActorDescView, const FWorldPartitionActorDescView& ReferenceActorDescView)
 				{
 					const bool bIsActorDescSpatiallyLoaded = RefererActorDescView.GetIsSpatiallyLoaded();
@@ -597,7 +605,7 @@ class FWorldPartitionStreamingGenerator
 					return bIsActorDescSpatiallyLoaded == bIsActorDescRefSpatiallyLoaded;
 				};
 
-				// Validate grid placement
+				// Validate data layers
 				auto IsReferenceDataLayersValid = [](const FWorldPartitionActorDescView& RefererActorDescView, const FWorldPartitionActorDescView& ReferenceActorDescView)
 				{
 					if (RefererActorDescView.GetRuntimeDataLayerInstanceNames().Num() == ReferenceActorDescView.GetRuntimeDataLayerInstanceNames().Num())
@@ -632,7 +640,7 @@ PRAGMA_ENABLE_DEPRECATION_WARNINGS
 
 				if (!IsRuntimeGridValid(ActorDescView))
 				{
-					if (!NbValidationPasses)
+					if (PassType == EPassType::ErrorReporting)
 					{
 						ErrorHandler->OnInvalidRuntimeGrid(ActorDescView, ActorDescView.GetRuntimeGrid());
 					}
@@ -669,7 +677,6 @@ PRAGMA_ENABLE_DEPRECATION_WARNINGS
 
 				// Add attach reference for the topmost parent, this reference is inverted since we consider the top most existing 
 				// parent to be refering to us, not the child to be referering the parent.
-				if (!NbValidationPasses)
 				{
 					FGuid ParentGuid = ActorDescView.GetParentActor();
 					FWorldPartitionActorDescView* TopParentDescView = nullptr;
@@ -685,8 +692,12 @@ PRAGMA_ENABLE_DEPRECATION_WARNINGS
 						}
 						else
 						{
-							// We had a guid but parent cannot be found, this will be a missing reference
-							References.Emplace(FActorReferenceInfo { ActorDescView.GetGuid(), &ActorDescView, ParentGuid, nullptr });
+							if (PassType == EPassType::ErrorReporting)
+							{
+								// We had a guid but parent cannot be found, this will report a missing reference error, but no error in the subsequent passes
+								References.Emplace(FActorReferenceInfo{ ActorDescView.GetGuid(), &ActorDescView, ParentGuid, nullptr });
+							}
+
 							break; 
 						}
 					}
@@ -698,7 +709,7 @@ PRAGMA_ENABLE_DEPRECATION_WARNINGS
 				}
 
 				TArray<FGuid> RuntimeReferences;
-				if (NbValidationPasses)
+				if (PassType == EPassType::Fixup)
 				{
 					RuntimeReferences.Reserve(ActorDescView.GetReferences().Num());
 				}
@@ -713,7 +724,7 @@ PRAGMA_ENABLE_DEPRECATION_WARNINGS
 						// Validate grid placement
 						if (!IsReferenceGridPlacementValid(*RefererActorDescView, *ReferenceActorDescView))
 						{
-							if (!NbValidationPasses)
+							if (PassType == EPassType::ErrorReporting)
 							{
 								ErrorHandler->OnInvalidReferenceGridPlacement(*RefererActorDescView, *ReferenceActorDescView);									
 							}
@@ -728,7 +739,7 @@ PRAGMA_ENABLE_DEPRECATION_WARNINGS
 
 						if (!IsReferenceDataLayersValid(*RefererActorDescView, *ReferenceActorDescView))
 						{
-							if (!NbValidationPasses)
+							if (PassType == EPassType::ErrorReporting)
 							{
 								ErrorHandler->OnInvalidReferenceDataLayers(*RefererActorDescView, *ReferenceActorDescView);									
 							}
@@ -743,7 +754,7 @@ PRAGMA_ENABLE_DEPRECATION_WARNINGS
 
 						if (!IsReferenceRuntimeGridValid(*RefererActorDescView, *ReferenceActorDescView))
 						{
-							if (!NbValidationPasses)
+							if (PassType == EPassType::ErrorReporting)
 							{
 								ErrorHandler->OnInvalidReferenceRuntimeGrid(*RefererActorDescView, *ReferenceActorDescView);
 							}
@@ -756,7 +767,7 @@ PRAGMA_ENABLE_DEPRECATION_WARNINGS
 							NbErrorsDetected++;
 						}
 
-						if (NbValidationPasses)
+						if (PassType == EPassType::Fixup)
 						{
 							RuntimeReferences.Add(Info.ReferenceGuid);
 						}
@@ -765,7 +776,7 @@ PRAGMA_ENABLE_DEPRECATION_WARNINGS
 					{
 						if (!ContainerDescriptor.FilteredActorDescViewMap.FindByGuid(Info.ReferenceGuid))
 						{
-							if (!NbValidationPasses)
+							if (PassType == EPassType::ErrorReporting)
 							{
 								FWorldPartitionActorDescView ReferendceActorDescView;
 								FWorldPartitionActorDescView* ReferendceActorDescViewPtr = nullptr;
@@ -787,7 +798,7 @@ PRAGMA_ENABLE_DEPRECATION_WARNINGS
 					}
 				}
 
-				if (NbValidationPasses)
+				if (PassType == EPassType::Fixup)
 				{
 					if (RuntimeReferences.Num() != ActorDescView.GetReferences().Num())
 					{
