World Partition: make sure actor set pointers aren't invalidated when adding new ones to the array during HLOD generation.

diff --git a/Engine/Source/Runtime/Engine/Private/WorldPartition/RuntimeSpatialHash/RuntimeSpatialHashHLOD.cpp b/Engine/Source/Runtime/Engine/Private/WorldPartition/RuntimeSpatialHash/RuntimeSpatialHashHLOD.cpp
--- a/Engine/Source/Runtime/Engine/Private/WorldPartition/RuntimeSpatialHash/RuntimeSpatialHashHLOD.cpp
+++ b/Engine/Source/Runtime/Engine/Private/WorldPartition/RuntimeSpatialHash/RuntimeSpatialHashHLOD.cpp
@@ -559,7 +559,7 @@ bool UWorldPartitionRuntimeSpatialHash::GenerateHLOD(ISourceControlHelper* Sourc
 			}
 			NewHLODActorSetInstance.ContentBundleID = MainActorSetContainer->ActorDescContainer->GetContentBundleGuid();
 			NewHLODActorSetInstance.ContainerInstance = MainActorSetContainer;
-			NewHLODActorSetInstance.ActorSet = &MainActorSetContainer->ActorSets.AddDefaulted_GetRef();
+			NewHLODActorSetInstance.ActorSet = MainActorSetContainer->ActorSets.Add_GetRef(MakeUnique<IStreamingGenerationContext::FActorSet>()).Get();
 			const_cast<IStreamingGenerationContext::FActorSet*>(NewHLODActorSetInstance.ActorSet)->Actors.Add(HLODActorDescView.GetGuid());
 
 			HLODActorSetInstancePtrs.Add(&NewHLODActorSetInstance);
diff --git a/Engine/Source/Runtime/Engine/Private/WorldPartition/WorldPartitionStreamingGeneration.cpp b/Engine/Source/Runtime/Engine/Private/WorldPartition/WorldPartitionStreamingGeneration.cpp
--- a/Engine/Source/Runtime/Engine/Private/WorldPartition/WorldPartitionStreamingGeneration.cpp
+++ b/Engine/Source/Runtime/Engine/Private/WorldPartition/WorldPartitionStreamingGeneration.cpp
@@ -94,7 +94,7 @@ class FWorldPartitionStreamingGenerator
 				ActorSetContainer.ActorSets.Empty(ContainerDescriptor.Clusters.Num());
 				for (const TArray<FGuid>& Cluster : ContainerDescriptor.Clusters)
 				{
-					FActorSet& ActorSet = ActorSetContainer.ActorSets.Emplace_GetRef();
+					FActorSet& ActorSet = *ActorSetContainer.ActorSets.Add_GetRef(MakeUnique<FActorSet>()).Get();
 					ActorSet.Actors = Cluster;
 				}
 
@@ -111,8 +111,9 @@ class FWorldPartitionStreamingGenerator
 				const FContainerDescriptor& ContainerDescriptor = StreamingGenerator->ContainerDescriptorsMap.FindChecked(ContainerInstanceDescriptor.Container);
 				const FActorSetContainer& ActorSetContainer = ActorSetContainers[ActorSetContainerMap.FindChecked(ContainerInstanceDescriptor.Container)];
 
-				for (const FActorSet& ActorSet : ActorSetContainer.ActorSets)
+				for (const TUniquePtr<FActorSet>& ActorSetPtr : ActorSetContainer.ActorSets)
 				{
+					const FActorSet& ActorSet = *ActorSetPtr;
 					const FWorldPartitionActorDescView& ReferenceActorDescView = ActorSetContainer.ActorDescViewMap->FindByGuidChecked(ActorSet.Actors[0]);
 
 					// Validate assumptions
diff --git a/Engine/Source/Runtime/Engine/Public/WorldPartition/WorldPartitionStreamingGenerationContext.h b/Engine/Source/Runtime/Engine/Public/WorldPartition/WorldPartitionStreamingGenerationContext.h
--- a/Engine/Source/Runtime/Engine/Public/WorldPartition/WorldPartitionStreamingGenerationContext.h
+++ b/Engine/Source/Runtime/Engine/Public/WorldPartition/WorldPartitionStreamingGenerationContext.h
@@ -18,9 +18,18 @@ public:
 
 	struct ENGINE_API FActorSetContainer
 	{
+		FActorSetContainer()
+			: ActorDescViewMap(nullptr)
+			, ActorDescContainer(nullptr)
+		{}
+
+		// Non-copyable
+		FActorSetContainer(const FActorSetContainer&) = delete;
+		FActorSetContainer& operator=(const FActorSetContainer&) = delete;
+
 		const FActorDescViewMap* ActorDescViewMap;
 		const UActorDescContainer* ActorDescContainer;
-		TArray<FActorSet> ActorSets;
+		TArray<TUniquePtr<FActorSet>> ActorSets;
 	};
 
 	struct ENGINE_API FActorSetInstance
