World Partition: cells package filenames are now the cell guid in base36 to reduce filename sizes.

diff --git a/Engine/Source/Runtime/Engine/Private/EngineUtils.cpp b/Engine/Source/Runtime/Engine/Private/EngineUtils.cpp
--- a/Engine/Source/Runtime/Engine/Private/EngineUtils.cpp
+++ b/Engine/Source/Runtime/Engine/Private/EngineUtils.cpp
@@ -25,7 +25,7 @@
 #include "Misc/UObjectToken.h"
 #include "Misc/MapErrors.h"
 #include "Misc/PathViews.h"
-
+#include "WorldPartition/WorldPartitionRuntimeCellInterface.h"
 #include "ProfilingDebugging/DiagnosticTable.h"
 #include "Interfaces/ITargetPlatform.h"
 
@@ -411,6 +411,12 @@ TArray<FSubLevelStatus> GetSubLevelsStatus( UWorld* World, bool SortByActorCount
 		LevelStatus.StreamingStatus = LEVEL_Visible;
 		LevelStatus.LODIndex = INDEX_NONE;
 		LevelStatus.ActorCount = World->GetActorCount();
+
+		if (const IWorldPartitionCell* WorldPartitionCell = World->PersistentLevel->GetWorldPartitionRuntimeCell())
+		{
+			LevelStatus.LevelLabel = WorldPartitionCell->GetDebugName();
+		}
+
 		Result.Add(LevelStatus);
 	}
 
@@ -449,6 +455,11 @@ TArray<FSubLevelStatus> GetSubLevelsStatus( UWorld* World, bool SortByActorCount
 
 			if (LevelStreaming->GetLoadedLevel())
 			{
+				if (const IWorldPartitionCell* WorldPartitionCell = LevelStreaming->GetLoadedLevel()->GetWorldPartitionRuntimeCell())
+				{
+					LevelStatus.LevelLabel = WorldPartitionCell->GetDebugName();
+				}
+
 				LevelStatus.ActorCount = LevelStreaming->GetLoadedLevel()->Actors.Num();
 
 				for (const AActor* Actor : LevelStreaming->GetLoadedLevel()->Actors)
diff --git a/Engine/Source/Runtime/Engine/Private/UnrealEngine.cpp b/Engine/Source/Runtime/Engine/Private/UnrealEngine.cpp
--- a/Engine/Source/Runtime/Engine/Private/UnrealEngine.cpp
+++ b/Engine/Source/Runtime/Engine/Private/UnrealEngine.cpp
@@ -17476,7 +17476,7 @@ int32 UEngine::RenderStatLevels(UWorld* World, FViewport* Viewport, FCanvas* Can
 
 		FColor	Color = ULevelStreaming::GetLevelStreamingStatusColor(LevelStatus.StreamingStatus);
 		const TCHAR* StatusDisplayName = ULevelStreaming::GetLevelStreamingStatusDisplayName(LevelStatus.StreamingStatus);
-		FString LevelPackageName = LevelStatus.PackageName.ToString();
+		FString LevelPackageName = LevelStatus.LevelLabel.IsEmpty() ? LevelStatus.PackageName.ToString() : LevelStatus.LevelLabel;
 		FString DisplayName = UWorld::RemovePIEPrefix(FPackageName::GetShortName(LevelPackageName));
 
 		if (LevelStatus.LODIndex != INDEX_NONE)
diff --git a/Engine/Source/Runtime/Engine/Private/WorldPartition/WorldPartitionRuntimeCell.cpp b/Engine/Source/Runtime/Engine/Private/WorldPartition/WorldPartitionRuntimeCell.cpp
--- a/Engine/Source/Runtime/Engine/Private/WorldPartition/WorldPartitionRuntimeCell.cpp
+++ b/Engine/Source/Runtime/Engine/Private/WorldPartition/WorldPartitionRuntimeCell.cpp
@@ -127,3 +127,8 @@ FName UWorldPartitionRuntimeCell::GetLevelPackageName() const
 	return NAME_None;
 #endif
 }
+
+FString UWorldPartitionRuntimeCell::GetDebugName() const
+{
+	return RuntimeCellData->GetDebugName();
+}
diff --git a/Engine/Source/Runtime/Engine/Private/WorldPartition/WorldPartitionRuntimeSpatialHash.cpp b/Engine/Source/Runtime/Engine/Private/WorldPartition/WorldPartitionRuntimeSpatialHash.cpp
--- a/Engine/Source/Runtime/Engine/Private/WorldPartition/WorldPartitionRuntimeSpatialHash.cpp
+++ b/Engine/Source/Runtime/Engine/Private/WorldPartition/WorldPartitionRuntimeSpatialHash.cpp
@@ -32,7 +32,6 @@
 #include "ProfilingDebugging/ScopedTimers.h"
 #include "UObject/UnrealType.h"
 #include "UObject/ObjectSaveContext.h"
-#include "Templates/TypeHash.h"
 
 #include UE_INLINE_GENERATED_CPP_BY_NAME(WorldPartitionRuntimeSpatialHash)
 
@@ -1235,24 +1234,14 @@ FString UWorldPartitionRuntimeSpatialHash::GetCellNameString(FName InWorldPackag
 	const FString ShortPackageName = UWorld::RemovePIEPrefix(FPackageName::GetShortName(InWorldPackageName));
 	FString CellName = FString::Printf(TEXT("%s_%s_%s"), *ShortPackageName, *InGridName.ToString(), *GetCellCoordString(InCellGlobalCoord));
 
-	// todo_ow : Hash the whole CellName & save a debug string inside the cell instead of hashing only the content bundle & data layer information
-	if (InDataLayerID.GetHash() && InContentBundleID.IsValid())
+	if (InDataLayerID.GetHash())
 	{
-		uint32 Hash = InDataLayerID.GetHash();
-		Hash = HashCombine(Hash, GetTypeHash(InContentBundleID));
-		CellName += FString::Printf(TEXT("_M%X"), Hash);
+		CellName += FString::Printf(TEXT("_DL%X"), InDataLayerID.GetHash());
 	}
-	else
-	{
-		if (InDataLayerID.GetHash())
-		{
-			CellName += FString::Printf(TEXT("_DL%X"), InDataLayerID.GetHash());
-		}
 
-		if (InContentBundleID.IsValid())
-		{
-			CellName += FString::Printf(TEXT("_CB%s"), *UContentBundleDescriptor::GetContentBundleCompactString(InContentBundleID));
-		}
+	if (InContentBundleID.IsValid())
+	{
+		CellName += FString::Printf(TEXT("_CB%s"), *UContentBundleDescriptor::GetContentBundleCompactString(InContentBundleID));
 	}
 
 	return CellName;
@@ -1380,7 +1369,7 @@ bool UWorldPartitionRuntimeSpatialHash::CreateStreamingGrid(const FSpatialHashRu
 				const FString CellName = GetCellNameString(OuterWorld->GetPackage()->GetFName(), CurrentStreamingGrid.GridName, CellGlobalCoords, GridCellDataChunk.GetDataLayersID(), GridCellDataChunk.GetContentBundleID());
 				const FGuid CellGuid = GetCellGuid(CurrentStreamingGrid.GridName, CellGlobalCoords, GridCellDataChunk.GetDataLayersID(), GridCellDataChunk.GetContentBundleID());
 
-				UWorldPartitionRuntimeCell* StreamingCell = CreateRuntimeCell(StreamingPolicy->GetRuntimeCellClass(), UWorldPartitionRuntimeCellDataSpatialHash::StaticClass(), FName(CellName));
+				UWorldPartitionRuntimeCell* StreamingCell = CreateRuntimeCell(StreamingPolicy->GetRuntimeCellClass(), UWorldPartitionRuntimeCellDataSpatialHash::StaticClass(), FName(CellGuid.ToString(EGuidFormats::Base36Encoded)));
 				UWorldPartitionRuntimeCellDataSpatialHash* CellDataSpatialHash = CastChecked<UWorldPartitionRuntimeCellDataSpatialHash>(StreamingCell->RuntimeCellData);
 
 				StreamingCell->SetIsAlwaysLoaded(bIsCellAlwaysLoaded);
@@ -1401,43 +1390,8 @@ bool UWorldPartitionRuntimeSpatialHash::CreateStreamingGrid(const FSpatialHashRu
 				CellDataSpatialHash->Position = FVector(Bounds.GetCenter(), 0.f);
 				CellDataSpatialHash->Extent = (float)CellExtent;
 				CellDataSpatialHash->GridName = RuntimeGrid.GridName;
-				
-				auto ComputeCellDebugName = [StreamingCell, CellDataSpatialHash, World, &CellGlobalCoords]()
-				{
-					TStringBuilder<512> Builder;
-					Builder += CellDataSpatialHash->GridName.ToString();
-					Builder += TEXT("_");
-					Builder += FString::Printf(TEXT("L%d_X%d_Y%d"), CellGlobalCoords.Z, CellGlobalCoords.X, CellGlobalCoords.Y);
-
-					const TArray<FName>& DataLayers = StreamingCell->GetDataLayers();
-					int32 DataLayerCount = DataLayers.Num();
-
-					if (DataLayerCount > 0)
-					{
-						if (const UDataLayerManager* DataLayerManager = UDataLayerManager::GetDataLayerManager(StreamingCell))
-						{
-							TArray<const UDataLayerInstance*> DataLayerObjects;
-							Builder += TEXT(" DL[");
-							for (int i = 0; i < DataLayerCount; ++i)
-							{
-								const UDataLayerInstance* DataLayer = DataLayerManager->GetDataLayerInstance(DataLayers[i]);
-								DataLayerObjects.Add(DataLayer);
-								Builder += DataLayer->GetDataLayerShortName();
-								Builder += TEXT(",");
-							}
-							Builder += FString::Printf(TEXT("ID:%X]"), FDataLayersID(DataLayerObjects).GetHash());
-						}
-					}
-
-					if (StreamingCell->GetContentBundleID().IsValid())
-					{
-						Builder += FString::Printf(TEXT(" CB[%s]"), *UContentBundleDescriptor::GetContentBundleCompactString(StreamingCell->GetContentBundleID()));
-					}
-
-					return Builder.ToString();
-				};
-
-				CellDataSpatialHash->DebugName = ComputeCellDebugName();
+			
+				CellDataSpatialHash->DebugName = CellName;
 
 				PopulateRuntimeCell(StreamingCell, FilteredActors, OutPackagesToGenerate);
 
diff --git a/Engine/Source/Runtime/Engine/Public/EngineUtils.h b/Engine/Source/Runtime/Engine/Public/EngineUtils.h
--- a/Engine/Source/Runtime/Engine/Public/EngineUtils.h
+++ b/Engine/Source/Runtime/Engine/Public/EngineUtils.h
@@ -712,6 +712,7 @@ struct FSubLevelActorDetails
 struct FSubLevelStatus
 {
 	FName				PackageName;
+	FString				LevelLabel;
 	EStreamingStatus	StreamingStatus;
 	int32				LODIndex;
 	bool				bInConsiderList;
diff --git a/Engine/Source/Runtime/Engine/Public/WorldPartition/WorldPartitionRuntimeCell.h b/Engine/Source/Runtime/Engine/Public/WorldPartition/WorldPartitionRuntimeCell.h
--- a/Engine/Source/Runtime/Engine/Public/WorldPartition/WorldPartitionRuntimeCell.h
+++ b/Engine/Source/Runtime/Engine/Public/WorldPartition/WorldPartitionRuntimeCell.h
@@ -188,6 +188,7 @@ class ENGINE_API UWorldPartitionRuntimeCell : public UObject, public IWorldParti
 		return Algo::AnyOf(DataLayers, [&InDataLayers](const FName& DataLayer) { return InDataLayers.Contains(DataLayer); });
 	}
 	virtual FName GetLevelPackageName() const override;
+	virtual FString GetDebugName() const override;
 	//~End IWorldPartitionCell Interface
 
 	bool HasAnyDataLayerInEffectiveRuntimeState(EDataLayerRuntimeState InState) const;
@@ -282,7 +283,6 @@ public:
 	const FBox& GetContentBounds() const { return RuntimeCellData->GetContentBounds(); }
 	FBox GetCellBounds() const { return RuntimeCellData->GetCellBounds(); }
 	virtual bool IsDebugShown() const;
-	virtual FString GetDebugName() const { return RuntimeCellData->GetDebugName(); }
 	//~End UWorldPartitionRuntimeCellData Proxy
 
 	UPROPERTY()
diff --git a/Engine/Source/Runtime/Engine/Public/WorldPartition/WorldPartitionRuntimeCellInterface.h b/Engine/Source/Runtime/Engine/Public/WorldPartition/WorldPartitionRuntimeCellInterface.h
--- a/Engine/Source/Runtime/Engine/Public/WorldPartition/WorldPartitionRuntimeCellInterface.h
+++ b/Engine/Source/Runtime/Engine/Public/WorldPartition/WorldPartitionRuntimeCellInterface.h
@@ -37,6 +37,8 @@ public:
 	virtual FBox GetCellBounds() const = 0;
 	/** Returns the associated level package name. */
 	virtual FName GetLevelPackageName() const = 0;
+	/** Returns the debug name associated with this cell. */
+	virtual FString GetDebugName() const = 0;
 
 #if WITH_EDITOR
 	/** Returns the referenced actor packages. */
