Speculative fix for a non unity compilation error

diff --git a/Engine/Source/Runtime/Engine/Private/Streaming/TextureStreamingHelpers.h b/Engine/Source/Runtime/Engine/Private/Streaming/TextureStreamingHelpers.h
--- a/Engine/Source/Runtime/Engine/Private/Streaming/TextureStreamingHelpers.h
+++ b/Engine/Source/Runtime/Engine/Private/Streaming/TextureStreamingHelpers.h
@@ -184,4 +184,4 @@ struct FRenderAssetStreamingStats
 };
 
 // Helper to access the level bStaticComponentsRegisteredInStreamingManager flag.
-extern bool OwnerLevelHasRegisteredStaticComponentsInStreamingManager(const AActor* Owner);
+extern bool OwnerLevelHasRegisteredStaticComponentsInStreamingManager(const class AActor* Owner);
