Added WorldPositionOffsetDisableDistance to FISMComponentDescriptor

diff --git a/Engine/Source/Runtime/Engine/Private/ISMPartition/ISMComponentDescriptor.cpp b/Engine/Source/Runtime/Engine/Private/ISMPartition/ISMComponentDescriptor.cpp
--- a/Engine/Source/Runtime/Engine/Private/ISMPartition/ISMComponentDescriptor.cpp
+++ b/Engine/Source/Runtime/Engine/Private/ISMPartition/ISMComponentDescriptor.cpp
@@ -62,7 +62,7 @@ void FISMComponentDescriptor::InitFrom(const UStaticMeshComponent* Template, boo
 	bVisible = Template->GetVisibleFlag();
 	bVisibleInRayTracing = Template->bVisibleInRayTracing;
 	bEvaluateWorldPositionOffset = Template->bEvaluateWorldPositionOffset;
-
+	WorldPositionOffsetDisableDistance = Template->WorldPositionOffsetDisableDistance;
 	// Determine if this instance must render with reversed culling based on both scale and the component property
 	const bool bIsLocalToWorldDeterminantNegative = Template->GetRenderMatrix().Determinant() < 0;
 	bReverseCulling = Template->bReverseCulling != bIsLocalToWorldDeterminantNegative;
@@ -135,6 +135,7 @@ bool FISMComponentDescriptor::operator==(const FISMComponentDescriptor& Other) c
 	bVisibleInRayTracing == Other.bVisibleInRayTracing &&
 	bEvaluateWorldPositionOffset == Other.bEvaluateWorldPositionOffset &&
 	bReverseCulling == Other.bReverseCulling &&
+	WorldPositionOffsetDisableDistance == Other.WorldPositionOffsetDisableDistance &&
 #if WITH_EDITORONLY_DATA
 	HLODBatchingPolicy == Other.HLODBatchingPolicy &&
 	bIncludeInHLOD == Other.bIncludeInHLOD &&
@@ -216,6 +217,7 @@ void FISMComponentDescriptor::InitComponent(UInstancedStaticMeshComponent* ISMCo
 	ISMComponent->bVisibleInRayTracing = bVisibleInRayTracing;
 	ISMComponent->bEvaluateWorldPositionOffset = bEvaluateWorldPositionOffset;
 	ISMComponent->bReverseCulling = bReverseCulling;
+	ISMComponent->WorldPositionOffsetDisableDistance = WorldPositionOffsetDisableDistance;
 	
 #if WITH_EDITORONLY_DATA
 	ISMComponent->HLODBatchingPolicy = HLODBatchingPolicy;
diff --git a/Engine/Source/Runtime/Engine/Public/ISMPartition/ISMComponentDescriptor.h b/Engine/Source/Runtime/Engine/Public/ISMPartition/ISMComponentDescriptor.h
--- a/Engine/Source/Runtime/Engine/Public/ISMPartition/ISMComponentDescriptor.h
+++ b/Engine/Source/Runtime/Engine/Public/ISMPartition/ISMComponentDescriptor.h
@@ -174,4 +174,7 @@ public:
 	UPROPERTY(EditAnywhere, Category = "Component Settings")
 	uint8 bConsiderForActorPlacementWhenHidden : 1;
 #endif
+
+	UPROPERTY(EditAnywhere, Category = "Component Settings")
+	int32 WorldPositionOffsetDisableDistance;
 };
