Improved logic that estimates what material parameters are actually connected by traversing the graph more accurately.

It now walks graph from material outputs backwards to material inputs following node inputs. When a static switch node is encountered, it tries to evaluate the value of the referenced static bool parameter and continues crawling the graph along the active input, ignoring the "false" subgraph. This change should cause the Material Editor to generate fewer shader permutations.

diff --git a/Engine/Source/Runtime/Engine/Private/Materials/Material.cpp b/Engine/Source/Runtime/Engine/Private/Materials/Material.cpp
--- a/Engine/Source/Runtime/Engine/Private/Materials/Material.cpp
+++ b/Engine/Source/Runtime/Engine/Private/Materials/Material.cpp
@@ -2019,21 +2019,9 @@ void UMaterial::UpdateCachedExpressionData()
 	}
 	else
 	{
-		if (!bUseMaterialAttributes)
-		{
-			for (int32 PropertyIndex = 0; PropertyIndex < MP_MAX; ++PropertyIndex)
-			{
-				const EMaterialProperty Property = (EMaterialProperty)PropertyIndex;
-				const FExpressionInput* Input = GetExpressionInputForProperty(Property);
-				if (Input && Input->IsConnected())
-				{
-					LocalCachedExpressionData->SetPropertyConnected(Property);
-				}
-			}
-		}
-
-		FMaterialCachedExpressionContext Context;
-		LocalCachedExpressionData->UpdateForExpressions(Context, GetExpressions(), EMaterialParameterAssociation::GlobalParameter, -1);
+		// Run graph analysis to harvest material information used to select which shader permutations are needed to be generated (e.g. detect what material properties are written to
+		// taking under consideration the presence of static switches to select among two input subgraphs).
+		LocalCachedExpressionData->AnalyzeMaterial(*this);
 	}
 
 	if (LocalCachedExpressionData->bHasMaterialLayers)
diff --git a/Engine/Source/Runtime/Engine/Private/Materials/MaterialCachedData.cpp b/Engine/Source/Runtime/Engine/Private/Materials/MaterialCachedData.cpp
--- a/Engine/Source/Runtime/Engine/Private/Materials/MaterialCachedData.cpp
+++ b/Engine/Source/Runtime/Engine/Private/Materials/MaterialCachedData.cpp
@@ -19,6 +19,12 @@
 #include "Materials/MaterialExpressionMakeMaterialAttributes.h"
 #include "Materials/MaterialExpressionMaterialAttributeLayers.h"
 #include "Materials/MaterialExpressionMaterialFunctionCall.h"
+#include "Materials/MaterialExpressionStaticSwitchParameter.h"
+#include "Materials/MaterialExpressionStaticSwitch.h"
+#include "Materials/MaterialExpressionFunctionOutput.h"
+#include "Materials/MaterialExpressionNamedReroute.h"
+#include "Materials/MaterialExpressionFunctionInput.h"
+#include "Materials/MaterialExpressionStaticBool.h"
 #include "Materials/MaterialFunctionInterface.h"
 #include "Materials/MaterialParameterCollection.h"
 #include "MaterialHLSLTree.h"
@@ -356,157 +362,405 @@ void FMaterialCachedExpressionData::UpdateForExpressions(const FMaterialCachedEx
 
 		Expression->GetLandscapeLayerNames(EditorOnlyData->LandscapeLayerNames);
 
-		if (UMaterialExpressionCollectionParameter* ExpressionCollectionParameter = Cast<UMaterialExpressionCollectionParameter>(Expression))
+		if (UMaterialExpressionMaterialAttributeLayers* LayersExpression = Cast<UMaterialExpressionMaterialAttributeLayers>(Expression))
 		{
-			UMaterialParameterCollection* Collection = ExpressionCollectionParameter->Collection;
-			if (Collection)
+			checkf(Association == GlobalParameter, TEXT("UMaterialExpressionMaterialAttributeLayers can't be nested"));
+			// Only a single layers expression is allowed/expected...creating additional layer expression will cause a compile error
+			if (!bHasMaterialLayers)
 			{
-				FMaterialParameterCollectionInfo NewInfo;
-				NewInfo.ParameterCollection = Collection;
-				NewInfo.StateId = Collection->StateId;
-				ParameterCollectionInfos.AddUnique(NewInfo);
+				const FMaterialLayersFunctions& Layers = Context.LayerOverrides ? *Context.LayerOverrides : LayersExpression->DefaultLayers;
+				UpdateForLayerFunctions(Context, Layers);
+
+				// TODO(?) - Layers for MIs are currently duplicated here and in FStaticParameterSet
+				bHasMaterialLayers = true;
+				MaterialLayers = Layers.GetRuntime();
+				EditorOnlyData->MaterialLayers = Layers.EditorOnly;
+				FMaterialLayersFunctions::Validate(MaterialLayers, EditorOnlyData->MaterialLayers);
+				LayersExpression->RebuildLayerGraph(false);
 			}
 		}
-		else if (UMaterialExpressionDynamicParameter* ExpressionDynamicParameter = Cast< UMaterialExpressionDynamicParameter>(Expression))
+		else if (UMaterialExpressionMaterialFunctionCall* FunctionCall = Cast<UMaterialExpressionMaterialFunctionCall>(Expression))
 		{
-			DynamicParameterNames.Empty(ExpressionDynamicParameter->ParamNames.Num());
-			for (const FString& Name : ExpressionDynamicParameter->ParamNames)
+			if (Context.bUpdateFunctionExpressions)
 			{
-				DynamicParameterNames.Add(*Name);
+				UpdateForFunction(Context, FunctionCall->MaterialFunction, GlobalParameter, -1);
+
+				// Update the function call node, so it can relink inputs and outputs as needed
+				// Update even if MaterialFunctionNode->MaterialFunction is NULL, because we need to remove the invalid inputs in that case
+				FunctionCall->UpdateFromFunctionResource();
 			}
+		} 
+	}
+}
+
+/**
+ * A helper struct used to crawl through the graph starting from material outputs in order to detect which material parameters are
+ * written to. This information is used to mark the connected properties in a MaterialCachedData instance.
+ * The graph is explored starting from the material outputs and walking backwards to its inputs. Whenever a static switch node is
+ * encountered, the algorithm finds out the value of the associated static parameter and continues exploration only along the
+ * activated subgraph. Because of this, unused subgraphs aren't visited so that any potential property writes in them don't interfere
+ * with the final result.
+ */
+struct FMaterialConnectedPropertiesAnalyzer
+{
+	/** The target FMaterialCachedExpressionData to set connected parameters to */
+	FMaterialCachedExpressionData&	CachedExpressionData;
+	
+	/** The current expression (if inside one, otherwise nullptr) */
+	UMaterialFunctionInterface*		CurrentFunction;
+	
+	/** The list of "fringe" material expressions that are yet to be explored */
+	TArray<UMaterialExpression*>	UnexploredExpressions;
+	
+	/** The set of visited expressions, used to avoid visiting the same expressions twice */
+	TSet<UMaterialExpression*>		VisitedExpressions;
+
+	/** Constructor */
+	FMaterialConnectedPropertiesAnalyzer(FMaterialCachedExpressionData& CachedExpressionData, UMaterialFunctionInterface* CurrentFunction = nullptr)
+		: CachedExpressionData(CachedExpressionData)
+		, CurrentFunction(CurrentFunction)
+	{
+	}
+
+	/** Pops an expression from the unexplored list, returning null if list is empty */
+	UMaterialExpression* PopUnvisitedExpression()
+	{
+		if (UnexploredExpressions.IsEmpty())
+		{
+			return nullptr;
 		}
-		else if (UMaterialExpressionLandscapeGrassOutput* ExpressionGrassOutput = Cast<UMaterialExpressionLandscapeGrassOutput>(Expression))
+		UMaterialExpression* Expression = UnexploredExpressions.Last();
+		UnexploredExpressions.Pop();
+		return Expression;
+	}
+
+	/** Pushes an expression to the unexplored list if has not yet been explored */
+	void PushUnexploredExpression(UMaterialExpression* Expression)
+	{
+		if (!Expression || VisitedExpressions.Contains(Expression))
 		{
-			for (const auto& Type : ExpressionGrassOutput->GrassTypes)
-			{
-				GrassTypes.AddUnique(Type.GrassType);
-			}
+			return;
 		}
-		else if (UMaterialExpressionQualitySwitch* QualitySwitchNode = Cast<UMaterialExpressionQualitySwitch>(Expression))
+		UnexploredExpressions.Add(Expression);
+		VisitedExpressions.Add(Expression);
+	}
+
+	/** Visits all expressions in the UnexploredExpressions list */
+	void VisitUnexploredExpressions(EMaterialParameterAssociation Association, int32 ParameterIndex)
+	{
+		for (;;)
 		{
-			const FExpressionInput DefaultInput = QualitySwitchNode->Default.GetTracedInput();
+			// Pop the current expression.
+			UMaterialExpression* Expression = PopUnvisitedExpression();
+			if (!Expression)
+			{
+				break;
+			}
 
-			for (int32 InputIndex = 0; InputIndex < EMaterialQualityLevel::Num; InputIndex++)
+			// Handle static switch nodes specially. Try to evaluate the condition static bool value and continue exploration along the active subgraph.
+			if (UMaterialExpressionStaticSwitchParameter* ExpressionStaticSwitchParameter = Cast<UMaterialExpressionStaticSwitchParameter>(Expression))
 			{
-				if (QualitySwitchNode->Inputs[InputIndex].IsConnected())
+				// Push the input expression that matches the state of the referenced parameter (ignoring the other branch).
+				bool ParamValue;
+				verify(EvaluateStaticBoolExpression(Expression, ParamValue, Association, ParameterIndex));
+				FExpressionInput* ExpressionInput = ParamValue ? &ExpressionStaticSwitchParameter->A : &ExpressionStaticSwitchParameter->B;
+				PushUnexploredExpression(ExpressionInput->Expression);
+				continue;
+			}
+			else if (UMaterialExpressionStaticSwitch* ExpressionStaticSwitch = Cast<UMaterialExpressionStaticSwitch>(Expression))
+			{
+				// Try to read back the input actual value.
+				bool ParamValue;
+				if (EvaluateStaticBoolExpression(ExpressionStaticSwitch->Value.Expression, ParamValue, Association, ParameterIndex))
+				{
+					// Push the active input expression.
+					FExpressionInput* ExpressionInput = ParamValue ? &ExpressionStaticSwitch->A : &ExpressionStaticSwitch->B;
+					PushUnexploredExpression(ExpressionInput->Expression);
+					continue;
+				}
+			}
+			
+			// Expression is either not a static switch node or it is but its conditional value could not be determined. We need to proceed exploring
+			// all input expressions to the current one (essentially walking backwards from outputs to inputs).
+			for (FExpressionInput* ExpressionInput : Expression->GetInputs())
+			{
+				PushUnexploredExpression(ExpressionInput->Expression);
+			}
+
+			// Explore the current expression
+			if (UMaterialExpressionCollectionParameter* ExpressionCollectionParameter = Cast<UMaterialExpressionCollectionParameter>(Expression))
+			{
+				UMaterialParameterCollection* Collection = ExpressionCollectionParameter->Collection;
+				if (Collection)
+				{
+					FMaterialParameterCollectionInfo NewInfo;
+					NewInfo.ParameterCollection = Collection;
+					NewInfo.StateId = Collection->StateId;
+					CachedExpressionData.ParameterCollectionInfos.AddUnique(NewInfo);
+				}
+			}
+			else if (UMaterialExpressionDynamicParameter* ExpressionDynamicParameter = Cast< UMaterialExpressionDynamicParameter>(Expression))
+			{
+				CachedExpressionData.DynamicParameterNames.Empty(ExpressionDynamicParameter->ParamNames.Num());
+				for (const FString& Name : ExpressionDynamicParameter->ParamNames)
+				{
+					CachedExpressionData.DynamicParameterNames.Add(*Name);
+				}
+			}
+			else if (UMaterialExpressionLandscapeGrassOutput* ExpressionGrassOutput = Cast<UMaterialExpressionLandscapeGrassOutput>(Expression))
+			{
+				for (const auto& Type : ExpressionGrassOutput->GrassTypes)
+				{
+					CachedExpressionData.GrassTypes.AddUnique(Type.GrassType);
+				}
+			}
+			else if (Expression->IsA(UMaterialExpressionRuntimeVirtualTextureOutput::StaticClass()))
+			{
+				CachedExpressionData.bHasRuntimeVirtualTextureOutput = true;
+			}
+			else if (Expression->IsA(UMaterialExpressionSceneColor::StaticClass()))
+			{
+				CachedExpressionData.bHasSceneColor = true;
+			}
+			else if (Expression->IsA(UMaterialExpressionPerInstanceRandom::StaticClass()))
+			{
+				CachedExpressionData.bHasPerInstanceRandom = true;
+			}
+			else if (Expression->IsA(UMaterialExpressionPerInstanceCustomData::StaticClass()))
+			{
+				CachedExpressionData.bHasPerInstanceCustomData = true;
+			}
+			else if (Expression->IsA(UMaterialExpressionPerInstanceCustomData3Vector::StaticClass()))
+			{
+				CachedExpressionData.bHasPerInstanceCustomData = true;
+			}
+			else if (Expression->IsA(UMaterialExpressionVertexInterpolator::StaticClass()))
+			{
+				CachedExpressionData.bHasVertexInterpolator = true;
+			}
+			else if (UMaterialExpressionQualitySwitch* ExpressionQualitySwitch = Cast<UMaterialExpressionQualitySwitch>(Expression))
+			{
+				const FExpressionInput DefaultInput = ExpressionQualitySwitch->Default.GetTracedInput();
+				for (int32 InputIndex = 0; InputIndex < EMaterialQualityLevel::Num; InputIndex++)
 				{
-					// We can ignore quality levels that are defined the same way as 'Default'
-					// This avoids compiling a separate explicit quality level resource, that will end up exactly the same as the default resource
-					const FExpressionInput Input = QualitySwitchNode->Inputs[InputIndex].GetTracedInput();
-					if (Input.Expression != DefaultInput.Expression ||
-						Input.OutputIndex != DefaultInput.OutputIndex)
+					if (ExpressionQualitySwitch->Inputs[InputIndex].IsConnected())
 					{
-						QualityLevelsUsed[InputIndex] = true;
+						// We can ignore quality levels that are defined the same way as 'Default'
+						// This avoids compiling a separate explicit quality level resource, that will end up exactly the same as the default resource
+						const FExpressionInput Input = ExpressionQualitySwitch->Inputs[InputIndex].GetTracedInput();
+						if (Input.Expression != DefaultInput.Expression || Input.OutputIndex != DefaultInput.OutputIndex)
+						{
+							CachedExpressionData.QualityLevelsUsed[InputIndex] = true;
+						}
 					}
 				}
 			}
+			else if (UMaterialExpressionSetMaterialAttributes* SetMatAttributes = Cast<UMaterialExpressionSetMaterialAttributes>(Expression))
+			{
+				for (int32 PinIndex = 0; PinIndex < SetMatAttributes->AttributeSetTypes.Num(); ++PinIndex)
+				{
+					// For this material attribute pin do we have something connected?
+					const FGuid& Guid = SetMatAttributes->AttributeSetTypes[PinIndex];
+					const FExpressionInput& AttributeInput = SetMatAttributes->Inputs[PinIndex + 1];
+					const EMaterialProperty MaterialProperty = FMaterialAttributeDefinitionMap::GetProperty(Guid);
+					if (AttributeInput.Expression)
+					{
+						CachedExpressionData.SetPropertyConnected(MaterialProperty);
+					}
+				}
+			}
+			else if (UMaterialExpressionMakeMaterialAttributes* MakeMatAttributes = Cast<UMaterialExpressionMakeMaterialAttributes>(Expression))
+			{
+				auto SetMatAttributeConditionally = [&](EMaterialProperty InMaterialProperty, bool InIsConnected)
+				{
+					if (InIsConnected)
+					{
+						CachedExpressionData.SetPropertyConnected(InMaterialProperty);
+					}
+				};
+
+				SetMatAttributeConditionally(EMaterialProperty::MP_BaseColor, MakeMatAttributes->BaseColor.IsConnected());
+				SetMatAttributeConditionally(EMaterialProperty::MP_Metallic, MakeMatAttributes->Metallic.IsConnected());
+				SetMatAttributeConditionally(EMaterialProperty::MP_Specular, MakeMatAttributes->Specular.IsConnected());
+				SetMatAttributeConditionally(EMaterialProperty::MP_Roughness, MakeMatAttributes->Roughness.IsConnected());
+				SetMatAttributeConditionally(EMaterialProperty::MP_Anisotropy, MakeMatAttributes->Anisotropy.IsConnected());
+				SetMatAttributeConditionally(EMaterialProperty::MP_EmissiveColor, MakeMatAttributes->EmissiveColor.IsConnected());
+				SetMatAttributeConditionally(EMaterialProperty::MP_Opacity, MakeMatAttributes->Opacity.IsConnected());
+				SetMatAttributeConditionally(EMaterialProperty::MP_OpacityMask, MakeMatAttributes->OpacityMask.IsConnected());
+				SetMatAttributeConditionally(EMaterialProperty::MP_Normal, MakeMatAttributes->Normal.IsConnected());
+				SetMatAttributeConditionally(EMaterialProperty::MP_Tangent, MakeMatAttributes->Tangent.IsConnected());
+				SetMatAttributeConditionally(EMaterialProperty::MP_WorldPositionOffset, MakeMatAttributes->WorldPositionOffset.IsConnected());
+				SetMatAttributeConditionally(EMaterialProperty::MP_SubsurfaceColor, MakeMatAttributes->SubsurfaceColor.IsConnected());
+				SetMatAttributeConditionally(EMaterialProperty::MP_CustomData0, MakeMatAttributes->ClearCoat.IsConnected());
+				SetMatAttributeConditionally(EMaterialProperty::MP_CustomData1, MakeMatAttributes->ClearCoatRoughness.IsConnected());
+				SetMatAttributeConditionally(EMaterialProperty::MP_AmbientOcclusion, MakeMatAttributes->AmbientOcclusion.IsConnected());
+				SetMatAttributeConditionally(EMaterialProperty::MP_Refraction, MakeMatAttributes->Refraction.IsConnected());
+				SetMatAttributeConditionally(EMaterialProperty::MP_CustomizedUVs0, MakeMatAttributes->CustomizedUVs[0].IsConnected());
+				SetMatAttributeConditionally(EMaterialProperty::MP_CustomizedUVs1, MakeMatAttributes->CustomizedUVs[1].IsConnected());
+				SetMatAttributeConditionally(EMaterialProperty::MP_CustomizedUVs2, MakeMatAttributes->CustomizedUVs[2].IsConnected());
+				SetMatAttributeConditionally(EMaterialProperty::MP_CustomizedUVs3, MakeMatAttributes->CustomizedUVs[3].IsConnected());
+				SetMatAttributeConditionally(EMaterialProperty::MP_CustomizedUVs4, MakeMatAttributes->CustomizedUVs[4].IsConnected());
+				SetMatAttributeConditionally(EMaterialProperty::MP_CustomizedUVs5, MakeMatAttributes->CustomizedUVs[5].IsConnected());
+				SetMatAttributeConditionally(EMaterialProperty::MP_CustomizedUVs6, MakeMatAttributes->CustomizedUVs[6].IsConnected());
+				SetMatAttributeConditionally(EMaterialProperty::MP_CustomizedUVs7, MakeMatAttributes->CustomizedUVs[7].IsConnected());
+				SetMatAttributeConditionally(EMaterialProperty::MP_PixelDepthOffset, MakeMatAttributes->PixelDepthOffset.IsConnected());
+				SetMatAttributeConditionally(EMaterialProperty::MP_ShadingModel, MakeMatAttributes->ShadingModel.IsConnected());
+			}
+			else if (UMaterialExpressionMaterialAttributeLayers* LayersExpression = Cast<UMaterialExpressionMaterialAttributeLayers>(Expression))
+			{
+				const FMaterialLayersFunctions& LayerFunctions = LayersExpression->DefaultLayers;
+
+				for (int32 LayerIndex = 0; LayerIndex < LayerFunctions.Layers.Num(); ++LayerIndex)
+				{
+					VisitFunction(LayerFunctions.Layers[LayerIndex], LayerParameter, LayerIndex);
+				}
+
+				for (int32 BlendIndex = 0; BlendIndex < LayerFunctions.Blends.Num(); ++BlendIndex)
+				{
+					VisitFunction(LayerFunctions.Blends[BlendIndex], BlendParameter, BlendIndex);
+				}
+			}
+			else if (UMaterialExpressionMaterialFunctionCall* FunctionCall = Cast<UMaterialExpressionMaterialFunctionCall>(Expression))
+			{
+				VisitFunction(FunctionCall->MaterialFunction, Association, ParameterIndex);
+			}
 		}
-		else if (Expression->IsA(UMaterialExpressionRuntimeVirtualTextureOutput::StaticClass()))
-		{
-			bHasRuntimeVirtualTextureOutput = true;
-		}
-		else if (Expression->IsA(UMaterialExpressionSceneColor::StaticClass()))
-		{
-			bHasSceneColor = true;
-		}
-		else if (Expression->IsA(UMaterialExpressionPerInstanceRandom::StaticClass()))
-		{
-			bHasPerInstanceRandom = true;
-		}
-		else if (Expression->IsA(UMaterialExpressionPerInstanceCustomData::StaticClass()))
+	}
+
+	/** Visits the inner expressions of a function call */
+	void VisitFunction(UMaterialFunctionInterface* Function, EMaterialParameterAssociation Association, int32 ParameterIndex)
+	{
+		if (!Function)
 		{
-			bHasPerInstanceCustomData = true;
+			return;
 		}
-		else if (Expression->IsA(UMaterialExpressionPerInstanceCustomData3Vector::StaticClass()))
+
+		// Get function outputs.
+		TArray<FFunctionExpressionInput> Inputs;
+		TArray<FFunctionExpressionOutput> Outputs;
+		Function->GetInputsAndOutputs(Inputs, Outputs);
+
+		// Push all function outputs as unexplored expressions to a new analyzer.
+		// Note: we use a different analyzer here so that multiple function calls to the same function are treated separately, since
+		// their actual state can be different due to static parameter overrides within a function.
+		FMaterialConnectedPropertiesAnalyzer FunctionAnalyzer{ CachedExpressionData, Function };
+		for (auto& Output : Outputs)
 		{
-			bHasPerInstanceCustomData = true;
+			FunctionAnalyzer.PushUnexploredExpression(Output.ExpressionOutput);
 		}
-		else if (Expression->IsA(UMaterialExpressionVertexInterpolator::StaticClass()))
+		
+		// Visit inner function expressions.
+		FunctionAnalyzer.VisitUnexploredExpressions(Association, ParameterIndex);
+	}
+
+	/** If expression is of type "static bool" this function walks the graph to find the actual value of the referenced static bool and writes the result to OutValue.
+	 * @return true if the boolean value could be determined.
+	 */
+	bool EvaluateStaticBoolExpression(UMaterialExpression* Expression, bool& OutValue, EMaterialParameterAssociation Association, int32 ParameterIndex)
+	{
+		if (!Expression)
 		{
-			bHasVertexInterpolator = true;
+			return false;
 		}
-		else if (UMaterialExpressionMaterialAttributeLayers* LayersExpression = Cast<UMaterialExpressionMaterialAttributeLayers>(Expression))
+		else if (Expression->IsA<UMaterialExpressionStaticSwitchParameter>() || Expression->IsA<UMaterialExpressionStaticBoolParameter>())
 		{
-			checkf(Association == GlobalParameter, TEXT("UMaterialExpressionMaterialAttributeLayers can't be nested"));
-			// Only a single layers expression is allowed/expected...creating additional layer expression will cause a compile error
-			if (!bHasMaterialLayers)
-			{
-				const FMaterialLayersFunctions& Layers = Context.LayerOverrides ? *Context.LayerOverrides : LayersExpression->DefaultLayers;
-				UpdateForLayerFunctions(Context, Layers);
+			// If this is a static switch expression, push the input expression that matches the state of the referenced parameter (ignoring the other branch).
+			// Read the parameter used by this expression if any.
+			FMaterialParameterMetadata ParameterMeta;
+			verify(Expression->GetParameterValue(ParameterMeta));
 
-				// TODO(?) - Layers for MIs are currently duplicated here and in FStaticParameterSet
-				bHasMaterialLayers = true;
-				MaterialLayers = Layers.GetRuntime();
-				EditorOnlyData->MaterialLayers = Layers.EditorOnly;
-				FMaterialLayersFunctions::Validate(MaterialLayers, EditorOnlyData->MaterialLayers);
-				LayersExpression->RebuildLayerGraph(false);
+			// If we're processing a function, give that a chance to override the parameter value
+			if (CurrentFunction)
+			{
+				FMaterialParameterMetadata OverrideParameterMeta;
+				if (CurrentFunction->GetParameterOverrideValue(ParameterMeta.Value.Type, Expression->GetParameterName(), OverrideParameterMeta))
+				{
+					ParameterMeta.Value = OverrideParameterMeta.Value;
+					ParameterMeta.ExpressionGuid = OverrideParameterMeta.ExpressionGuid;
+					ParameterMeta.bUsedAsAtlasPosition = OverrideParameterMeta.bUsedAsAtlasPosition;
+					ParameterMeta.ScalarAtlas = OverrideParameterMeta.ScalarAtlas;
+					ParameterMeta.ScalarCurve = OverrideParameterMeta.ScalarCurve;
+				}
 			}
+
+			// Read the parameter.
+			UObject* UnusedReferencedTexture = nullptr;
+			const FMaterialParameterInfo ParameterInfo(Expression->GetParameterName(), Association, ParameterIndex);
+			CachedExpressionData.AddParameter(ParameterInfo, ParameterMeta, UnusedReferencedTexture);
+
+			OutValue = ParameterMeta.Value.AsStaticSwitch();
+			return true;
 		}
-		else if (UMaterialExpressionMaterialFunctionCall* FunctionCall = Cast<UMaterialExpressionMaterialFunctionCall>(Expression))
+		else if (UMaterialExpressionStaticSwitch* ExpressionStaticSwitch = Cast<UMaterialExpressionStaticSwitch>(Expression))
 		{
-			if (Context.bUpdateFunctionExpressions)
+			if (!EvaluateStaticBoolExpression(ExpressionStaticSwitch->Value.Expression, OutValue, Association, ParameterIndex))
 			{
-				UpdateForFunction(Context, FunctionCall->MaterialFunction, GlobalParameter, -1);
-
-				// Update the function call node, so it can relink inputs and outputs as needed
-				// Update even if MaterialFunctionNode->MaterialFunction is NULL, because we need to remove the invalid inputs in that case
-				FunctionCall->UpdateFromFunctionResource();
+				return false;
 			}
+			FExpressionInput* OperandExpressionInput = OutValue ? &ExpressionStaticSwitch->A : &ExpressionStaticSwitch->B;
+			return EvaluateStaticBoolExpression(OperandExpressionInput->Expression, OutValue, Association, ParameterIndex);
 		}
-		else if (UMaterialExpressionSetMaterialAttributes* SetMatAttributes = Cast<UMaterialExpressionSetMaterialAttributes>(Expression))
+		else if (UMaterialExpressionNamedRerouteUsage* ExpressionRerouteUsage = Cast<UMaterialExpressionNamedRerouteUsage>(Expression))
 		{
-			for (int32 PinIndex = 0; PinIndex < SetMatAttributes->AttributeSetTypes.Num(); ++PinIndex)
+			if (!ExpressionRerouteUsage->Declaration)
 			{
-				// For this material attribute pin do we have something connected?
-				const FGuid& Guid = SetMatAttributes->AttributeSetTypes[PinIndex];
-				const FExpressionInput& AttributeInput = SetMatAttributes->Inputs[PinIndex + 1];
-				const EMaterialProperty MaterialProperty = FMaterialAttributeDefinitionMap::GetProperty(Guid);
-				if (AttributeInput.Expression)
-				{
-					SetPropertyConnected(MaterialProperty);
-				}
+				return false;
 			}
+			return EvaluateStaticBoolExpression(ExpressionRerouteUsage->Declaration->Input.Expression, OutValue, Association, ParameterIndex);
 		}
-		else if (UMaterialExpressionMakeMaterialAttributes* MakeMatAttributes = Cast<UMaterialExpressionMakeMaterialAttributes>(Expression))
+		else if (UMaterialExpressionFunctionInput* ExpressionFunctionInput = Cast<UMaterialExpressionFunctionInput>(Expression))
 		{
-			auto SetMatAttributeConditionally = [&](EMaterialProperty InMaterialProperty, bool InIsConnected)
+			return EvaluateStaticBoolExpression(ExpressionFunctionInput->Preview.Expression, OutValue, Association, ParameterIndex);
+		}
+		else if (UMaterialExpressionStaticBool* ExpressionStaticBool = Cast<UMaterialExpressionStaticBool>(Expression))
+		{
+			OutValue = ExpressionStaticBool->Value;
+			return true;
+		}
+
+		return false;
+	}
+};
+
+void FMaterialCachedExpressionData::AnalyzeMaterial(UMaterial& Material)
+{
+	// First run a pre-pass to determine referenced textures and fix-up material layers and functions, add parameters, etc.
+	// This is executed for all expressions in the material.
+	{
+		FMaterialCachedExpressionContext Context;
+		UpdateForExpressions(Context, Material.GetExpressions(), EMaterialParameterAssociation::GlobalParameter, -1);
+	}
+
+	// Now crawl the graph from outputs to inputs skipping unused subgraphs ("compiled out" by static switches).
+	FMaterialConnectedPropertiesAnalyzer Analyzer{ *this };
+
+	// Add material outputs depending on whether material uses material attributes.
+	if (Material.bUseMaterialAttributes)
+	{
+		const FExpressionInput* Input = Material.GetExpressionInputForProperty(MP_MaterialAttributes);
+		if (Input && Input->IsConnected())
+		{
+			Analyzer.PushUnexploredExpression(Input->Expression);
+		}
+	}
+	else
+	{
+		for (int32 PropertyIndex = 0; PropertyIndex < MP_MAX; ++PropertyIndex)
+		{
+			const EMaterialProperty Property = (EMaterialProperty)PropertyIndex;
+			const FExpressionInput* Input = Material.GetExpressionInputForProperty(Property);
+			if (Input && Input->IsConnected())
 			{
-				if (InIsConnected)
-				{
-					SetPropertyConnected(InMaterialProperty);
-				}
-			};
-
-			SetMatAttributeConditionally(EMaterialProperty::MP_BaseColor, MakeMatAttributes->BaseColor.IsConnected());
-			SetMatAttributeConditionally(EMaterialProperty::MP_Metallic, MakeMatAttributes->Metallic.IsConnected());
-			SetMatAttributeConditionally(EMaterialProperty::MP_Specular, MakeMatAttributes->Specular.IsConnected());
-			SetMatAttributeConditionally(EMaterialProperty::MP_Roughness, MakeMatAttributes->Roughness.IsConnected());
-			SetMatAttributeConditionally(EMaterialProperty::MP_Anisotropy, MakeMatAttributes->Anisotropy.IsConnected());
-			SetMatAttributeConditionally(EMaterialProperty::MP_EmissiveColor, MakeMatAttributes->EmissiveColor.IsConnected());
-			SetMatAttributeConditionally(EMaterialProperty::MP_Opacity, MakeMatAttributes->Opacity.IsConnected());
-			SetMatAttributeConditionally(EMaterialProperty::MP_OpacityMask, MakeMatAttributes->OpacityMask.IsConnected());
-			SetMatAttributeConditionally(EMaterialProperty::MP_Normal, MakeMatAttributes->Normal.IsConnected());
-			SetMatAttributeConditionally(EMaterialProperty::MP_Tangent, MakeMatAttributes->Tangent.IsConnected());
-			SetMatAttributeConditionally(EMaterialProperty::MP_WorldPositionOffset, MakeMatAttributes->WorldPositionOffset.IsConnected());
-			SetMatAttributeConditionally(EMaterialProperty::MP_SubsurfaceColor, MakeMatAttributes->SubsurfaceColor.IsConnected());
-			SetMatAttributeConditionally(EMaterialProperty::MP_CustomData0, MakeMatAttributes->ClearCoat.IsConnected());
-			SetMatAttributeConditionally(EMaterialProperty::MP_CustomData1, MakeMatAttributes->ClearCoatRoughness.IsConnected());
-			SetMatAttributeConditionally(EMaterialProperty::MP_AmbientOcclusion, MakeMatAttributes->AmbientOcclusion.IsConnected());
-			SetMatAttributeConditionally(EMaterialProperty::MP_Refraction, MakeMatAttributes->Refraction.IsConnected());
-			SetMatAttributeConditionally(EMaterialProperty::MP_CustomizedUVs0, MakeMatAttributes->CustomizedUVs[0].IsConnected());
-			SetMatAttributeConditionally(EMaterialProperty::MP_CustomizedUVs1, MakeMatAttributes->CustomizedUVs[1].IsConnected());
-			SetMatAttributeConditionally(EMaterialProperty::MP_CustomizedUVs2, MakeMatAttributes->CustomizedUVs[2].IsConnected());
-			SetMatAttributeConditionally(EMaterialProperty::MP_CustomizedUVs3, MakeMatAttributes->CustomizedUVs[3].IsConnected());
-			SetMatAttributeConditionally(EMaterialProperty::MP_CustomizedUVs4, MakeMatAttributes->CustomizedUVs[4].IsConnected());
-			SetMatAttributeConditionally(EMaterialProperty::MP_CustomizedUVs5, MakeMatAttributes->CustomizedUVs[5].IsConnected());
-			SetMatAttributeConditionally(EMaterialProperty::MP_CustomizedUVs6, MakeMatAttributes->CustomizedUVs[6].IsConnected());
-			SetMatAttributeConditionally(EMaterialProperty::MP_CustomizedUVs7, MakeMatAttributes->CustomizedUVs[7].IsConnected());
-			SetMatAttributeConditionally(EMaterialProperty::MP_PixelDepthOffset, MakeMatAttributes->PixelDepthOffset.IsConnected());
-			SetMatAttributeConditionally(EMaterialProperty::MP_ShadingModel, MakeMatAttributes->ShadingModel.IsConnected());
+				Analyzer.PushUnexploredExpression(Input->Expression);
+				SetPropertyConnected(Property);
+			}
 		}
 	}
+
+	Analyzer.VisitUnexploredExpressions(GlobalParameter, -1);
 }
 
+
 namespace Private
 {
 void PrepareHLSLTree(UE::HLSLTree::FEmitContext& EmitContext,
diff --git a/Engine/Source/Runtime/Engine/Public/MaterialCachedData.h b/Engine/Source/Runtime/Engine/Public/MaterialCachedData.h
--- a/Engine/Source/Runtime/Engine/Public/MaterialCachedData.h
+++ b/Engine/Source/Runtime/Engine/Public/MaterialCachedData.h
@@ -187,6 +187,8 @@ struct FMaterialCachedExpressionData
 	void UpdateForExpressions(const FMaterialCachedExpressionContext& Context, TConstArrayView<TObjectPtr<UMaterialExpression>> Expressions, EMaterialParameterAssociation Association, int32 ParameterIndex);
 	void UpdateForFunction(const FMaterialCachedExpressionContext& Context, UMaterialFunctionInterface* Function, EMaterialParameterAssociation Association, int32 ParameterIndex);
 	void UpdateForLayerFunctions(const FMaterialCachedExpressionContext& Context, const FMaterialLayersFunctions& LayerFunctions);
+	void AnalyzeMaterial(UMaterial& Material);
+
 	ENGINE_API void UpdateForCachedHLSLTree(const FMaterialCachedHLSLTree& CachedTree, const FStaticParameterSet* StaticParameters);
 	void Validate();
 
