Replicated subobject
* Refactor subobject refcounting so it's lifetime is independant of the NetworkObjectInfo
* Fixes incorrect refcounts when ActorChannels are still not cleaned up while the NetworkObjectInfo of an actor was destroyed

diff --git a/Engine/Source/Runtime/Engine/Classes/Engine/NetConnection.h b/Engine/Source/Runtime/Engine/Classes/Engine/NetConnection.h
--- a/Engine/Source/Runtime/Engine/Classes/Engine/NetConnection.h
+++ b/Engine/Source/Runtime/Engine/Classes/Engine/NetConnection.h
@@ -1510,6 +1510,8 @@ protected:
 
 	void CleanupDormantActorState();
 
+	void ClearDormantReplicatorsReference();
+
 	/** During cleanup this will destroy the actor owned by this connection (generally a PlayerController) */
 	ENGINE_API virtual void DestroyOwningActor();
 
diff --git a/Engine/Source/Runtime/Engine/Classes/Engine/NetDriver.h b/Engine/Source/Runtime/Engine/Classes/Engine/NetDriver.h
--- a/Engine/Source/Runtime/Engine/Classes/Engine/NetDriver.h
+++ b/Engine/Source/Runtime/Engine/Classes/Engine/NetDriver.h
@@ -1757,7 +1757,7 @@ public:
 	/** Called when an actor channel is remotely opened for an actor. */
 	ENGINE_API virtual void NotifyActorChannelOpen(UActorChannel* Channel, AActor* Actor);
 	
-	/** Called when an actor channel is cleaned up foor an actor. */
+	/** Called when an actor channel is cleaned up for an actor. */
 	ENGINE_API virtual void NotifyActorChannelCleanedUp(UActorChannel* Channel, EChannelCloseReason CloseReason);
 
 	ENGINE_API virtual void NotifyActorTornOff(AActor* Actor);
diff --git a/Engine/Source/Runtime/Engine/Classes/Engine/NetworkObjectList.h b/Engine/Source/Runtime/Engine/Classes/Engine/NetworkObjectList.h
--- a/Engine/Source/Runtime/Engine/Classes/Engine/NetworkObjectList.h
+++ b/Engine/Source/Runtime/Engine/Classes/Engine/NetworkObjectList.h
@@ -73,73 +73,6 @@ struct FNetworkObjectInfo
 	*/
 	TSet<TWeakObjectPtr<UNetConnection>, FNetConnectionKeyFuncs> RecentlyDormantConnections;
 
-	/** 
-	* Keeps track of the number of channels that have replicated a subobject.
-	* When the status changes from Active its expected of existing references to be gradually removed as the actor replicates itself to each connection.
-	*/
-	struct FSubObjectChannelReference
-	{
-		/** The replicated subobject */
-		TWeakObjectPtr<UObject> SubObjectPtr;
-		
-		/** Number of channels that replicated the subobject */
-		uint16 ChannelRefCount = 0;
-
-		/** Current status of the subobject */
-		ENetSubObjectStatus Status = ENetSubObjectStatus::Active;
-
-#if DO_REPLICATED_OBJECT_CHANNELREF_CHECKS
-		/** 
-		* Debug array to keep track of every individual references. Helps to trigger an ensure early if something wrong is detected.
-		* Consists of UActorChannel for active replicators or UNetConnection's for dormant replicators
-		*/
-		TArray<UObject*> RegisteredOwners;
-#endif
-		inline bool operator==(const FSubObjectChannelReference& rhs) const { return SubObjectPtr.HasSameIndexAndSerialNumber(rhs.SubObjectPtr); }
-		inline bool operator==(const TWeakObjectPtr<UObject>& rhs) const { return SubObjectPtr.HasSameIndexAndSerialNumber(rhs); }
-
-		friend uint32 GetTypeHash(const FSubObjectChannelReference& SubObjChannelRef)
-		{
-			return GetTypeHash(SubObjChannelRef.SubObjectPtr);
-		}
-
-		inline bool IsTearOff() const { return Status == ENetSubObjectStatus::TearOff; }
-		inline bool IsDelete() const  { return Status == ENetSubObjectStatus::Delete; }
-		inline bool IsActive() const  { return Status == ENetSubObjectStatus::Active; }
-
-		FSubObjectChannelReference() = default;
-		explicit FSubObjectChannelReference(const TWeakObjectPtr<UObject>& InSubObject)
-			: SubObjectPtr(InSubObject)
-			, ChannelRefCount(1)
-			, Status(ENetSubObjectStatus::Active)
-		{
-		}
-	};
-
-	/** 
-	* Key definitions for TSet/TMap that works with invalided weak pointers 
-	*/
-	struct FSubObjectChannelRefKeyFuncs : BaseKeyFuncs<FSubObjectChannelReference, const TWeakObjectPtr<UObject>&, false>
-	{
-		static KeyInitType GetSetKey(ElementInitType& Element)			{ return Element.SubObjectPtr; }
-		static bool Matches(KeyInitType A, KeyInitType B)				{ return A.HasSameIndexAndSerialNumber(B); }
-		static uint32 GetKeyHash(KeyInitType Key)						{ return GetTypeHash(Key); }
-	};
-
-#if UE_REPLICATED_OBJECT_REFCOUNTING
-	/** The set of active replicated subobjects for an actor */
-	TSet<FSubObjectChannelReference, FSubObjectChannelRefKeyFuncs> ActiveSubObjectChannelReferences;
-
-	/** The list of replicated subobjects that need to be torn off or deleted by the channels still referencing it */
-	TArray<FSubObjectChannelReference> InvalidSubObjectChannelReferences;
-#endif
-
-	/** 
-	* This variable gets increased when a replicated subobject gets added to the invalid list.
-	* Actor channels will test against this count and when it differs will check if they have any replicated subobjects to remove.
-	*/
-	uint16 InvalidSubObjectDirtyCount = 0;
-
 	/** Is this object still pending a full net update due to clients that weren't able to replicate the actor at the time of LastNetUpdateTime */
 	uint8 bPendingNetUpdate : 1;
 
@@ -296,7 +229,11 @@ public:
 	/** Marks any actors in the given package/level active if they were fully dormant or dormant for the passed in connection */
 	void FlushDormantActors(UNetConnection* const Connection, const FName& PackageName);
 
+	/** Called when the netdriver gets notified that an actor is destroyed */
+	void OnActorDestroyed(AActor* DestroyedActor);
+
 #if UE_REPLICATED_OBJECT_REFCOUNTING
+
 	/** Set the subobject to be flagged for deletion */
 	void SetSubObjectForDeletion(AActor* Actor, UObject* SubObject);
 
@@ -313,13 +250,18 @@ public:
 	void RemoveSubObjectChannelReference(AActor* OwnerActor, const TWeakObjectPtr<UObject>& ReplicatedSubObject, UObject* ReferenceOwner);
 
 	/** Called when multiple subobjects need to remove their reference from either the active or inactive list */
-	void RemoveMultipleSubObjectChannelReference(AActor* OwnerActor, const TArrayView<TWeakObjectPtr<UObject>>& SubObjectsToRemove, UObject* ReferenceOwner);
+	void RemoveMultipleSubObjectChannelReference(FObjectKey OwnerActorKey, const TArrayView<TWeakObjectPtr<UObject>>& SubObjectsToRemove, UObject* ReferenceOwner);
 
 	/** Called when multiple subobjects that were flagged torn off or delete have removed their channel reference */
-	void RemoveMultipleInvalidSubObjectChannelReference(FNetworkObjectInfo* ActorNetInfo, const TArrayView<TWeakObjectPtr<UObject>>& SubObjectsToRemove, UObject* ReferenceOwner);
+	void RemoveMultipleInvalidSubObjectChannelReference(FObjectKey OwnerActorKey, const TArrayView<TWeakObjectPtr<UObject>>& SubObjectsToRemove, UObject* ReferenceOwner);
 
 	/** Called when multiple subobjects that were still considered active have removed their channel reference */
-	void RemoveMultipleActiveSubObjectChannelReference(FNetworkObjectInfo* ActorNetInfo, const TArrayView<TWeakObjectPtr<UObject>>& SubObjectsToRemove, UObject* ReferenceOwner);
+	void RemoveMultipleActiveSubObjectChannelReference(FObjectKey OwnerActorKey, const TArrayView<TWeakObjectPtr<UObject>>& SubObjectsToRemove, UObject* ReferenceOwner);
+
+	struct FActorInvalidSubObjectView;
+
+	/** Returns a struct holding the dirty count and a possible list of invalid subobjects who still have references to specific connections.*/
+	FNetworkObjectList::FActorInvalidSubObjectView FindActorInvalidSubObjects(AActor* OwnerActor) const;
 
 	/** 
 	* Keep track of the transfer of ownership from the channel to the connection when the actor becomes dormant.
@@ -339,14 +281,14 @@ private:
 #if UE_REPLICATED_OBJECT_REFCOUNTING
 	void InvalidateSubObject(AActor* Actor, UObject* SubObject, ENetSubObjectStatus InvalidStatus);
 
-	void HandleRemoveAnySubObjectChannelRef(FNetworkObjectInfo* ActorNetInfo, const TWeakObjectPtr<UObject>& ReplicatedSubObject, UObject* ReferenceOwner);
-	bool HandleRemoveActiveSubObjectRef(FNetworkObjectInfo* ActorNetInfo, const TWeakObjectPtr<UObject>& ReplicatedSubObject, UObject* ReferenceOwner);
-	bool HandleRemoveInvalidSubObjectRef(FNetworkObjectInfo* ActorNetInfo, const TWeakObjectPtr<UObject>& ReplicatedSubObject, UObject* ReferenceOwner);
+	struct FActorSubObjectReferences;
+	void HandleRemoveAnySubObjectChannelRef(FActorSubObjectReferences& SubObjectsRefInfo, const TWeakObjectPtr<UObject>& ReplicatedSubObject, UObject* ReferenceOwner);
+	bool HandleRemoveActiveSubObjectRef(FActorSubObjectReferences& SubObjectsRefInfo, const TWeakObjectPtr<UObject>& ReplicatedSubObject, UObject* ReferenceOwner);
+	bool HandleRemoveInvalidSubObjectRef(FActorSubObjectReferences& SubObjectsRefInfo, const TWeakObjectPtr<UObject>& ReplicatedSubObject, UObject* ReferenceOwner);
 #endif
 
-
 #if DO_REPLICATED_OBJECT_CHANNELREF_CHECKS
-	void HandleSwapReferenceForDormancy(FNetworkObjectInfo* ActorNetInfo, const TWeakObjectPtr<UObject>& SubObjectPtr, UObject* PreviousRefOwner, UObject* NewRefOwner);
+	void HandleSwapReferenceForDormancy(FActorSubObjectReferences* ActorNetInfo, const TWeakObjectPtr<UObject>& SubObjectPtr, UObject* PreviousRefOwner, UObject* NewRefOwner);
 #endif
 
 	FNetworkObjectSet AllNetworkObjects;
@@ -360,4 +302,132 @@ private:
 
 	TMap<FName, FNetworkObjectSet> FullyDormantObjectsByLevel;
 	TMap<TObjectKey<UNetConnection>, TMap<FName, FNetworkObjectSet>> DormantObjectsPerConnection;
+
+public:
+
+	/**
+	* Keeps track of the number of channels that have replicated a subobject.
+	* When the status is not Active anymore it is expected of existing references to be gradually removed as the actor replicates itself to each connection.
+	*/
+	struct FSubObjectChannelReference
+	{
+		/** The replicated subobject */
+		TWeakObjectPtr<UObject> SubObjectPtr;
+
+		/** Number of channels that replicated the subobject */
+		uint16 ChannelRefCount = 0;
+
+		/** Current status of the subobject */
+		ENetSubObjectStatus Status = ENetSubObjectStatus::Active;
+
+#if DO_REPLICATED_OBJECT_CHANNELREF_CHECKS
+		/**
+		* Debug array to keep track of every individual references. Helps to trigger an ensure early if something wrong is detected.
+		* Consists of UActorChannel for active replicators or UNetConnection's for dormant replicators
+		*/
+		TArray<UObject*> RegisteredOwners;
+#endif
+
+		inline bool operator==(const FSubObjectChannelReference& rhs) const { return SubObjectPtr.HasSameIndexAndSerialNumber(rhs.SubObjectPtr); }
+		inline bool operator==(const TWeakObjectPtr<UObject>& rhs) const { return SubObjectPtr.HasSameIndexAndSerialNumber(rhs); }
+
+		friend uint32 GetTypeHash(const FSubObjectChannelReference& SubObjChannelRef)
+		{
+			return GetTypeHash(SubObjChannelRef.SubObjectPtr);
+		}
+
+		inline bool IsTearOff() const { return Status == ENetSubObjectStatus::TearOff; }
+		inline bool IsDelete() const { return Status == ENetSubObjectStatus::Delete; }
+		inline bool IsActive() const { return Status == ENetSubObjectStatus::Active; }
+
+		FSubObjectChannelReference() = default;
+		explicit FSubObjectChannelReference(const TWeakObjectPtr<UObject>& InSubObject)
+			: SubObjectPtr(InSubObject)
+			, ChannelRefCount(1)
+			, Status(ENetSubObjectStatus::Active)
+		{
+		}
+	};
+
+	/** Key definitions for TSet that works with invalided weak pointers */
+	struct FSubObjectChannelRefKeyFuncs : BaseKeyFuncs<FSubObjectChannelReference, const TWeakObjectPtr<UObject>&, false>
+	{
+		static KeyInitType	GetSetKey(ElementInitType& Element)		{ return Element.SubObjectPtr; }
+		static bool			Matches(KeyInitType A, KeyInitType B)	{ return A.HasSameIndexAndSerialNumber(B); }
+		static uint32		GetKeyHash(KeyInitType Key)				{ return GetTypeHash(Key); }
+	};
+
+#if UE_REPLICATED_OBJECT_REFCOUNTING
+
+	/** Structure giving const-only access to the list of invalid subobjects of a given actor */
+	struct FActorInvalidSubObjectView
+	{
+	public:
+		explicit FActorInvalidSubObjectView(uint16 InDirtyCount, const TArray<FSubObjectChannelReference>* InArrayPtr=nullptr)
+			: InvalidSubObjectDirtyCount(InvalidSubObjectDirtyCount)
+			, InvalidSubObjectsPtr(InArrayPtr)
+		{ }
+
+		FActorInvalidSubObjectView() = delete;
+
+		inline uint16 GetDirtyCount() const { return InvalidSubObjectDirtyCount; }
+		inline bool HasInvalidSubObjects() const { return InvalidSubObjectsPtr != nullptr; }
+		inline const TArray<FSubObjectChannelReference>& GetInvalidSubObjects() const 
+		{ 
+			check(HasInvalidSubObjects());
+			return *InvalidSubObjectsPtr;
+		}
+
+	private:
+		/** The current dirty count */
+		uint16 InvalidSubObjectDirtyCount = 0;
+
+		/** The list of invalid subobjects that were set to be torn off or deleted */
+		const TArray<FSubObjectChannelReference>* InvalidSubObjectsPtr;
+	};
+
+private:
+
+	struct FActorSubObjectReferences
+	{
+		/** The actor who is replicating the subobjects */
+		FObjectKey ActorKey;
+
+		/** The set of active replicated subobjects for an actor */
+		TSet<FSubObjectChannelReference, FSubObjectChannelRefKeyFuncs> ActiveSubObjectChannelReferences;
+
+		/** The list of replicated subobjects that need to be torn off or deleted by the channels still referencing it */
+		TArray<FSubObjectChannelReference> InvalidSubObjectChannelReferences;
+
+		/**
+		* This variable gets increased when a replicated subobject gets added to the invalid list.
+		* Actor channels will test against this count and when it differs will check if they have any replicated subobjects to remove.
+		*/
+		uint16 InvalidSubObjectDirtyCount = 0;
+
+		explicit FActorSubObjectReferences(AActor* InActor) : ActorKey(InActor) {}
+
+		/** Returns true when the actor's subobjects have no more references and this entry itself can be deleted */
+		bool HasNoSubObjects() const 
+		{
+			return ActiveSubObjectChannelReferences.IsEmpty() && InvalidSubObjectChannelReferences.IsEmpty();
+		}
+		
+		void CountBytes(FArchive& Ar) const;
+	};
+
+	/** Definition to use the struct's actor as the key */
+	struct FActorSubObjectRefKeyFuncs : BaseKeyFuncs<FActorSubObjectReferences, FObjectKey, false>
+	{
+		static KeyInitType	GetSetKey(ElementInitType& Element)		{ return Element.ActorKey; }
+		static bool			Matches(KeyInitType A, KeyInitType B)	{ return A == B; }
+		static uint32		GetKeyHash(KeyInitType Key)				{ return GetTypeHash(Key); }
+	};
+
+private:
+
+	/** Map keeping track of the number of connections that have a reference to each actor's subobjects */
+	TSet<FActorSubObjectReferences, FActorSubObjectRefKeyFuncs> SubObjectChannelReferences;
+
+#endif //UE_REPLICATED_OBJECT_REFCOUNTING
 };
diff --git a/Engine/Source/Runtime/Engine/Private/DataChannel.cpp b/Engine/Source/Runtime/Engine/Private/DataChannel.cpp
--- a/Engine/Source/Runtime/Engine/Private/DataChannel.cpp
+++ b/Engine/Source/Runtime/Engine/Private/DataChannel.cpp
@@ -3158,6 +3158,10 @@ void UActorChannel::ProcessBunch( FInBunch & Bunch )
 		}
 	}
 
+#if UE_REPLICATED_OBJECT_REFCOUNTING
+	TArray<TWeakObjectPtr<UObject>, TInlineAllocator<16>> ReferencesToRemove;
+#endif
+
 	for (auto RepComp = ReplicationMap.CreateIterator(); RepComp; ++RepComp)
 	{
 		TSharedRef<FObjectReplicator>& ObjectReplicator = RepComp.Value();
@@ -3166,8 +3170,9 @@ void UActorChannel::ProcessBunch( FInBunch & Bunch )
 			if (!Connection->Driver->IsServer())
 			{
 #if UE_REPLICATED_OBJECT_REFCOUNTING
-				Connection->Driver->GetNetworkObjectList().RemoveSubObjectChannelReference(Actor, ObjectReplicator.Get().GetWeakObjectPtr(), this);
+				ReferencesToRemove.Add(ObjectReplicator.Get().GetWeakObjectPtr());
 #endif
+
 				RepComp.RemoveCurrent(); // This should cause the replicator to be cleaned up as there should be no outstandings refs. 
 			}
 			continue;
@@ -3176,6 +3181,10 @@ void UActorChannel::ProcessBunch( FInBunch & Bunch )
 		ObjectReplicator->PostReceivedBunch();
 	}
 
+#if UE_REPLICATED_OBJECT_REFCOUNTING
+	Connection->Driver->GetNetworkObjectList().RemoveMultipleSubObjectChannelReference(Actor, ReferencesToRemove, this);
+#endif
+
 	// After all properties have been initialized, call PostNetInit. This should call BeginPlay() so initialization can be done with proper starting values.
 	if (Actor && bSpawnedNewActor)
 	{
@@ -3641,39 +3650,36 @@ bool UActorChannel::UpdateDeletedSubObjects(FOutBunch& Bunch)
 #if UE_REPLICATED_OBJECT_REFCOUNTING
 	TArray< TWeakObjectPtr<UObject>, TInlineAllocator<16> > SubObjectsRemoved;
 
-	FNetworkObjectInfo* ActorInfo = Connection->Driver->FindNetworkObjectInfo(Actor);
-
 	// Check if the dirty count is different from our last update
-	if (ActorInfo)
+	FNetworkObjectList::FActorInvalidSubObjectView InvalidSubObjects = Connection->Driver->GetNetworkObjectList().FindActorInvalidSubObjects(Actor);
+	if (InvalidSubObjects.HasInvalidSubObjects() && ChannelSubObjectDirtyCount != InvalidSubObjects.GetDirtyCount())
 	{
-		if (ChannelSubObjectDirtyCount != ActorInfo->InvalidSubObjectDirtyCount)
+		ChannelSubObjectDirtyCount = InvalidSubObjects.GetDirtyCount();
+
+		for (const FNetworkObjectList::FSubObjectChannelReference& SubObjectRef : InvalidSubObjects.GetInvalidSubObjects())
 		{
-			ChannelSubObjectDirtyCount = ActorInfo->InvalidSubObjectDirtyCount;
+			ensure(SubObjectRef.Status != ENetSubObjectStatus::Active);
 
-			for (const FNetworkObjectInfo::FSubObjectChannelReference& SubObjectRef : ActorInfo->InvalidSubObjectChannelReferences)
+			// The TearOff won't be sent if the Object pointer is fully deleted once we get here.
+			// This would be fixed by converting the ReplicationMap to stop using raw pointers as the key.
+            // Instead the ReplicationMap iteration below should pick this deleted object and send for it to be destroyed.
+			if (UObject* ObjectToRemove = SubObjectRef.SubObjectPtr.GetEvenIfUnreachable())
 			{
-				ensure(SubObjectRef.Status != ENetSubObjectStatus::Active);
-
-				// The TearOff won't be sent if the Object pointer is fully deleted once we get here.
-				// This will be fixed by converting the ReplicationMap to stop using raw pointers as the key
-				if (UObject* ObjectToRemove = SubObjectRef.SubObjectPtr.GetEvenIfUnreachable())
+				if (TSharedRef<FObjectReplicator>* SubObjectReplicator = ReplicationMap.Find(ObjectToRemove))
 				{
-					if (TSharedRef<FObjectReplicator>* SubObjectReplicator = ReplicationMap.Find(ObjectToRemove))
-					{
-						const ESubObjectDeleteFlag DeleteFlag = SubObjectRef.IsTearOff() ? ESubObjectDeleteFlag::TearOff : ESubObjectDeleteFlag::ForceDelete;
-						
-						DeleteSubObject(*SubObjectReplicator, SubObjectRef.SubObjectPtr, DeleteFlag);
-						ReplicationMap.Remove(ObjectToRemove);
-					}
+					const ESubObjectDeleteFlag DeleteFlag = SubObjectRef.IsTearOff() ? ESubObjectDeleteFlag::TearOff : ESubObjectDeleteFlag::ForceDelete;
+					SubObjectsRemoved.Add(SubObjectRef.SubObjectPtr);
+					DeleteSubObject(*SubObjectReplicator, SubObjectRef.SubObjectPtr, DeleteFlag);
+					ReplicationMap.Remove(ObjectToRemove);
 				}
-
-				SubObjectsRemoved.Add(SubObjectRef.SubObjectPtr);
 			}
-
-			Connection->Driver->GetNetworkObjectList().RemoveMultipleInvalidSubObjectChannelReference(ActorInfo, SubObjectsRemoved, this);
 		}
 
-		SubObjectsRemoved.Reset();
+		if (!SubObjectsRemoved.IsEmpty())
+		{
+			Connection->Driver->GetNetworkObjectList().RemoveMultipleInvalidSubObjectChannelReference(Actor, SubObjectsRemoved, this);
+			SubObjectsRemoved.Reset();
+		}
 	}
 #endif //#if UE_REPLICATED_OBJECT_REFCOUNTING
 
@@ -3703,9 +3709,9 @@ bool UActorChannel::UpdateDeletedSubObjects(FOutBunch& Bunch)
 	}
 
 #if UE_REPLICATED_OBJECT_REFCOUNTING
-	if (SubObjectsRemoved.Num() > 0)
+	if (!SubObjectsRemoved.IsEmpty())
 	{
-		Connection->Driver->GetNetworkObjectList().RemoveMultipleActiveSubObjectChannelReference(ActorInfo, SubObjectsRemoved, this);
+		Connection->Driver->GetNetworkObjectList().RemoveMultipleSubObjectChannelReference(Actor, SubObjectsRemoved, this);
 	}
 #endif
 
diff --git a/Engine/Source/Runtime/Engine/Private/DemoNetDriver.cpp b/Engine/Source/Runtime/Engine/Private/DemoNetDriver.cpp
--- a/Engine/Source/Runtime/Engine/Private/DemoNetDriver.cpp
+++ b/Engine/Source/Runtime/Engine/Private/DemoNetDriver.cpp
@@ -3011,7 +3011,7 @@ void UDemoNetDriver::FinalizeFastForward(const double StartTime)
 			}
 		}
 
-		auto CallRepNotifies = [](AActor* OwnerActor, FObjectKey ObjectKey, const TSharedRef<FObjectReplicator>& ReplicatorRef)
+		auto CallRepNotifies = [](FObjectKey OwnerActorKey, FObjectKey ObjectKey, const TSharedRef<FObjectReplicator>& ReplicatorRef)
 		{
 			ReplicatorRef->CallRepNotifies(true);
 		};
@@ -3914,7 +3914,7 @@ bool UDemoNetDriver::FastForwardLevels(const FGotoResult& GotoResult)
 			}
 		}
 
-		auto CallRepNotifies = [](AActor* OwnerActor, FObjectKey ObjectKey, const TSharedRef<FObjectReplicator>& ReplicatorRef)
+		auto CallRepNotifies = [](FObjectKey OwnerActorKey, FObjectKey ObjectKey, const TSharedRef<FObjectReplicator>& ReplicatorRef)
 		{
 			ReplicatorRef->CallRepNotifies(true);
 		};
diff --git a/Engine/Source/Runtime/Engine/Private/Net/NetDormantHolder.cpp b/Engine/Source/Runtime/Engine/Private/Net/NetDormantHolder.cpp
--- a/Engine/Source/Runtime/Engine/Private/Net/NetDormantHolder.cpp
+++ b/Engine/Source/Runtime/Engine/Private/Net/NetDormantHolder.cpp
@@ -121,31 +121,11 @@ void FDormantReplicatorHolder::CleanupAllReplicatorsOfActor(AActor* DormantActor
 	ActorReplicatorSet.Remove(DormantActor);
 }
 
-void FDormantReplicatorHolder::CleanupStaleObjects()
-{
-	for (FActorReplicatorSet::TIterator ActorSetIt = ActorReplicatorSet.CreateIterator(); ActorSetIt; ++ActorSetIt )
-	{
-		for (FActorDormantReplicators::FObjectReplicatorSet::TIterator ReplicatorSetIt = ActorSetIt->DormantReplicators.CreateIterator(); ReplicatorSetIt; ++ReplicatorSetIt)
-		{
-			FDormantObjectReplicator& DormantReplicator = *ReplicatorSetIt;
-
-			if (!DormantReplicator.Replicator->GetWeakObjectPtr().IsValid())
-			{
-				ReplicatorSetIt.RemoveCurrent();
-			}
-		}
-		
-		if (ActorSetIt->DormantReplicators.IsEmpty())
-		{
-			ActorSetIt.RemoveCurrent();
-		}
-	}
-}
-
 void FDormantReplicatorHolder::CleanupStaleObjects(FNetworkObjectList& NetworkObjectList, UObject* ReferenceOwner)
 {
 #if UE_REPLICATED_OBJECT_REFCOUNTING
 	TArray<TWeakObjectPtr<UObject>, TInlineAllocator<16>> CleanedUpObjects;
+#endif
 
 	for (FActorReplicatorSet::TIterator ActorSetIt = ActorReplicatorSet.CreateIterator(); ActorSetIt; ++ActorSetIt)
 	{
@@ -156,11 +136,13 @@ void FDormantReplicatorHolder::CleanupStaleObjects(FNetworkObjectList& NetworkOb
 
 			if (!DormantObjectPtr.IsValid())
 			{
-				// If it's a subobject
-				if (ActorSetIt->OwnerActor != DormantReplicator.Replicator->GetObject())
+#if UE_REPLICATED_OBJECT_REFCOUNTING
+				// Keep track of the cleaned up object if it's a subobject and not the main actor
+				if (ActorSetIt->OwnerActorKey != DormantReplicator.ObjectKey)
 				{
 					CleanedUpObjects.Add(DormantObjectPtr);
 				}
+#endif
 
 				ReplicatorSetIt.RemoveCurrent();
 			}
@@ -171,15 +153,14 @@ void FDormantReplicatorHolder::CleanupStaleObjects(FNetworkObjectList& NetworkOb
 			ActorSetIt.RemoveCurrent();
 		}
 
+#if UE_REPLICATED_OBJECT_REFCOUNTING
 		if (CleanedUpObjects.Num() > 0)
 		{
-			NetworkObjectList.RemoveMultipleSubObjectChannelReference(ActorSetIt->OwnerActor, CleanedUpObjects, ReferenceOwner);
+			NetworkObjectList.RemoveMultipleSubObjectChannelReference(ActorSetIt->OwnerActorKey, CleanedUpObjects, ReferenceOwner);
 			CleanedUpObjects.Reset();
 		}
-	}
-#else
-	CleanupStaleObjects();
 #endif
+	}
 }
 
 void FDormantReplicatorHolder::ForEachDormantReplicator(UE::Net::FExecuteForEachDormantReplicator Function)
@@ -188,7 +169,7 @@ void FDormantReplicatorHolder::ForEachDormantReplicator(UE::Net::FExecuteForEach
 	{
 		for (const FDormantObjectReplicator& DormantReplicator : ActorReplicators.DormantReplicators)
 		{
-			Function(ActorReplicators.OwnerActor, DormantReplicator.ObjectKey, DormantReplicator.Replicator);
+			Function(ActorReplicators.OwnerActorKey, DormantReplicator.ObjectKey, DormantReplicator.Replicator);
 		}
 	}
 }
@@ -199,7 +180,7 @@ void FDormantReplicatorHolder::ForEachDormantReplicatorOfActor(AActor* DormantAc
 	{
 		for (const FDormantObjectReplicator& DormantReplicator : ActorReplicators->DormantReplicators)
 		{
-			Function(ActorReplicators->OwnerActor, DormantReplicator.ObjectKey, DormantReplicator.Replicator);
+			Function(ActorReplicators->OwnerActorKey, DormantReplicator.ObjectKey, DormantReplicator.Replicator);
 		}
 	}
 }
diff --git a/Engine/Source/Runtime/Engine/Private/NetConnection.cpp b/Engine/Source/Runtime/Engine/Private/NetConnection.cpp
--- a/Engine/Source/Runtime/Engine/Private/NetConnection.cpp
+++ b/Engine/Source/Runtime/Engine/Private/NetConnection.cpp
@@ -4812,9 +4812,46 @@ void UNetConnection::ResetGameWorldState()
 
 void UNetConnection::CleanupDormantActorState()
 {
+	ClearDormantReplicatorsReference();
+
 	DormantReplicatorSet.EmptySet();
 }
 
+void UNetConnection::ClearDormantReplicatorsReference()
+{
+	using namespace UE::Net::Private;
+
+#if UE_REPLICATED_OBJECT_REFCOUNTING
+	if (!Driver)
+	{
+		return;
+	}
+
+	TArray<TWeakObjectPtr<UObject>, TInlineAllocator<16>> SubObjectsToRemove;
+
+	// Find all the dormant subobject replicators still held by this connection and remove our reference to them.
+	for (const FActorDormantReplicators& ActorDormantReplicators : DormantReplicatorSet.ActorReplicatorSet)
+	{
+		FObjectKey OwnerKey = ActorDormantReplicators.OwnerActorKey;
+
+		for (const FDormantObjectReplicator& DormantObject : ActorDormantReplicators.DormantReplicators)
+		{
+			// Only if it's a subobject and not the actor itself
+			if (DormantObject.ObjectKey != OwnerKey)
+			{
+				SubObjectsToRemove.Add(DormantObject.Replicator->GetWeakObjectPtr());
+			}
+		}
+
+		if (!SubObjectsToRemove.IsEmpty())
+		{
+			Driver->GetNetworkObjectList().RemoveMultipleSubObjectChannelReference(ActorDormantReplicators.OwnerActorKey, SubObjectsToRemove, this);
+			SubObjectsToRemove.Reset();
+		}
+	}
+#endif
+}
+
 void UNetConnection::FlushDormancy(AActor* Actor)
 {
 	UE_LOG( LogNetDormancy, Verbose, TEXT( "FlushDormancy: %s. Connection: %s" ), *Actor->GetName(), *GetName() );
@@ -5154,9 +5191,10 @@ void UNetConnection::CleanupDormantReplicatorsForActor(AActor* Actor)
 		if (Driver)
 		{
 			TArray<TWeakObjectPtr<UObject>, TInlineAllocator<16>> RemovedObjects;
-			auto ExecuteFunction = [&RemovedObjects](AActor* OwnerActor, FObjectKey ObjectKey, const TSharedRef<FObjectReplicator>& ReplicatorRef)
+			auto ExecuteFunction = [&RemovedObjects](FObjectKey OwnerActorKey, FObjectKey ObjectKey, const TSharedRef<FObjectReplicator>& ReplicatorRef)
 			{
-				if (OwnerActor != ReplicatorRef->GetObject())
+				// If it's the replicator of a subobject and not the main actor
+				if (OwnerActorKey != ObjectKey)
 				{
 					RemovedObjects.Add(ReplicatorRef->GetWeakObjectPtr());
 				}
@@ -5174,18 +5212,10 @@ void UNetConnection::CleanupDormantReplicatorsForActor(AActor* Actor)
 
 void UNetConnection::CleanupStaleDormantReplicators()
 {
-#if UE_REPLICATED_OBJECT_REFCOUNTING
-	if (Driver)	
+	if (ensure(Driver))	
 	{
 		DormantReplicatorSet.CleanupStaleObjects(Driver->GetNetworkObjectList(), this);
 	}
-	else
-	{
-		DormantReplicatorSet.CleanupStaleObjects();
-	}
-#else
-	DormantReplicatorSet.CleanupStaleObjects();
-#endif
 }
 
 void UNetConnection::SetPendingCloseDueToReplicationFailure()
diff --git a/Engine/Source/Runtime/Engine/Private/NetDriver.cpp b/Engine/Source/Runtime/Engine/Private/NetDriver.cpp
--- a/Engine/Source/Runtime/Engine/Private/NetDriver.cpp
+++ b/Engine/Source/Runtime/Engine/Private/NetDriver.cpp
@@ -3329,6 +3329,8 @@ void UNetDriver::NotifyActorDestroyed( AActor* ThisActor, bool IsSeamlessTravel
 		ServerConnection->NotifyActorDestroyed(ThisActor);
 	}
 
+	NetworkObjects->OnActorDestroyed(ThisActor);
+
 	// Remove this actor from the network object list
 	RemoveNetworkActor( ThisActor );
 }
@@ -5154,7 +5156,7 @@ int32 UNetDriver::ServerReplicateActors(float DeltaSeconds)
 		// net.DormancyValidate can be set to 2 to validate all dormant actors against last known state before going dormant
 		if ( GNetDormancyValidate == 2 )
 		{
-			auto ValidateFunction = [](AActor* OwnerActor, FObjectKey ObjectKey, const TSharedRef<FObjectReplicator>& ReplicatorRef)
+			auto ValidateFunction = [](FObjectKey OwnerActorKey, FObjectKey ObjectKey, const TSharedRef<FObjectReplicator>& ReplicatorRef)
 			{
 				FObjectReplicator& Replicator = ReplicatorRef.Get();
 
diff --git a/Engine/Source/Runtime/Engine/Private/NetworkObjectList.cpp b/Engine/Source/Runtime/Engine/Private/NetworkObjectList.cpp
--- a/Engine/Source/Runtime/Engine/Private/NetworkObjectList.cpp
+++ b/Engine/Source/Runtime/Engine/Private/NetworkObjectList.cpp
@@ -28,6 +28,12 @@ namespace UE::Net::Private
 	}
 }
 
+void FNetworkObjectList::FActorSubObjectReferences::CountBytes(FArchive& Ar) const
+{
+	ActiveSubObjectChannelReferences.CountBytes(Ar);
+	InvalidSubObjectChannelReferences.CountBytes(Ar);
+}
+
 void FNetworkObjectList::AddInitialObjects(UWorld* const World, UNetDriver* NetDriver)
 {
 	if (World == nullptr || NetDriver == nullptr)
@@ -509,10 +515,6 @@ void FNetworkObjectInfo::CountBytes(FArchive& Ar) const
 
 	GRANULAR_NETWORK_MEMORY_TRACKING_TRACK("DormantConnections", DormantConnections.CountBytes(Ar));
 	GRANULAR_NETWORK_MEMORY_TRACKING_TRACK("RecentlyDormantConnections", RecentlyDormantConnections.CountBytes(Ar));
-#if UE_REPLICATED_OBJECT_REFCOUNTING
-	GRANULAR_NETWORK_MEMORY_TRACKING_TRACK("ActiveSubObjectChannelReferences", ActiveSubObjectChannelReferences.CountBytes(Ar));
-	GRANULAR_NETWORK_MEMORY_TRACKING_TRACK("InvalidSubObjectChannelReferences", InvalidSubObjectChannelReferences.CountBytes(Ar));
-#endif
 }
 
 void FNetworkObjectList::CountBytes(FArchive& Ar) const
@@ -558,6 +560,14 @@ void FNetworkObjectList::CountBytes(FArchive& Ar) const
 			}
 		}
 	);
+
+#if UE_REPLICATED_OBJECT_REFCOUNTING
+	GRANULAR_NETWORK_MEMORY_TRACKING_TRACK("SubObjectChannelReferences", SubObjectChannelReferences.CountBytes(Ar));
+	for (const FActorSubObjectReferences& SubObjectRef : SubObjectChannelReferences)
+	{
+		SubObjectRef.CountBytes(Ar);
+	}
+#endif
 }
 
 void FNetworkObjectList::FlushDormantActors(UNetConnection* const Connection, const FName& PackageName)
@@ -584,7 +594,28 @@ void FNetworkObjectList::FlushDormantActors(UNetConnection* const Connection, co
 		MarkActiveInternal(ActorInfo, Connection, Connection->Driver);
 	}
 }
+
+void FNetworkObjectList::OnActorDestroyed(AActor* DestroyedActor)
+{
+#if UE_REPLICATED_OBJECT_REFCOUNTING
+	SubObjectChannelReferences.Remove(DestroyedActor);
+#endif
+}
+
 #if UE_REPLICATED_OBJECT_REFCOUNTING
+
+FNetworkObjectList::FActorInvalidSubObjectView FNetworkObjectList::FindActorInvalidSubObjects(AActor* OwnerActor) const
+{
+	if (const FActorSubObjectReferences* SubObjectsRefInfo = SubObjectChannelReferences.Find(OwnerActor))
+	{
+		return FActorInvalidSubObjectView(SubObjectsRefInfo->InvalidSubObjectDirtyCount, &(SubObjectsRefInfo->InvalidSubObjectChannelReferences));
+	}
+	else
+	{
+		return FActorInvalidSubObjectView(0);
+	}
+}
+
 void FNetworkObjectList::SetSubObjectForDeletion(AActor* Actor, UObject* SubObject)
 {
 	InvalidateSubObject(Actor, SubObject, ENetSubObjectStatus::Delete);
@@ -599,163 +630,186 @@ void FNetworkObjectList::InvalidateSubObject(AActor* Actor, UObject* SubObject,
 {
 	ensure(InvalidStatus != ENetSubObjectStatus::Active);
 
-	if (TSharedPtr<FNetworkObjectInfo>* InfoPtr = AllNetworkObjects.Find(Actor))
+	if (FActorSubObjectReferences* SubObjectsRefInfo = SubObjectChannelReferences.Find(Actor))
 	{
-		FNetworkObjectInfo* ActorNetInfo = InfoPtr->Get();
-
 		const TWeakObjectPtr<UObject> SubObjectPtr = SubObject;
-		const FSetElementId FoundId = ActorNetInfo->ActiveSubObjectChannelReferences.FindId(SubObjectPtr);
+		const FSetElementId FoundId = SubObjectsRefInfo->ActiveSubObjectChannelReferences.FindId(SubObjectPtr);
 
 		if (FoundId.IsValidId())
 		{
 			// Flag its new state
-			ActorNetInfo->ActiveSubObjectChannelReferences[FoundId].Status = InvalidStatus;
+			SubObjectsRefInfo->ActiveSubObjectChannelReferences[FoundId].Status = InvalidStatus;
 
 			// Move the reference to the destroyed list.
-			ActorNetInfo->InvalidSubObjectChannelReferences.Emplace(MoveTemp(ActorNetInfo->ActiveSubObjectChannelReferences[FoundId]));
-			ActorNetInfo->ActiveSubObjectChannelReferences.Remove(FoundId);
+			SubObjectsRefInfo->InvalidSubObjectChannelReferences.Emplace(MoveTemp(SubObjectsRefInfo->ActiveSubObjectChannelReferences[FoundId]));
+			SubObjectsRefInfo->ActiveSubObjectChannelReferences.Remove(FoundId);
 
 			// Increase dirty count so channels can refresh this list
-			ActorNetInfo->InvalidSubObjectDirtyCount = FMath::Max(ActorNetInfo->InvalidSubObjectDirtyCount + 1, 1); // Wrap around to 1 because 0 is the default actor channel value
+			SubObjectsRefInfo->InvalidSubObjectDirtyCount = FMath::Max(SubObjectsRefInfo->InvalidSubObjectDirtyCount + 1, 1); // Wrap around to 1 because 0 is the default actor channel value
 		}
 	}
 }
 
 void FNetworkObjectList::AddSubObjectChannelReference(AActor* OwnerActor, UObject* ReplicatedSubObject, UObject* ReferenceOwner)
 {
-	TSharedPtr<FNetworkObjectInfo>* InfoPtr = AllNetworkObjects.Find(OwnerActor);
+	FActorSubObjectReferences* SubObjectsRefInfo = SubObjectChannelReferences.Find(OwnerActor);
 
-	if (ensureMsgf(InfoPtr, TEXT("The owner (0x%p) %s replicated %s (0x%p) without having any network objectinfo for %s"), ReferenceOwner, *GetNameSafe(ReferenceOwner), *GetNameSafe(ReplicatedSubObject), ReplicatedSubObject, *OwnerActor->GetName()))
+	if (SubObjectsRefInfo == nullptr)
 	{
-		FNetworkObjectInfo* ActorNetInfo = InfoPtr->Get();
-		const TWeakObjectPtr<UObject> SubObjectPtr = ReplicatedSubObject;
+		// Create an entry for this actor
+		FSetElementId Index = SubObjectChannelReferences.Emplace(FActorSubObjectReferences(OwnerActor));
+		SubObjectsRefInfo = &(SubObjectChannelReferences[Index]);
+	}
 
-#if DO_REPLICATED_OBJECT_CHANNELREF_CHECKS
-		// Make sure the object wasn't previously put in the invalid list  
-		if (FNetworkObjectInfo::FSubObjectChannelReference* DestroyedObjectRef = ActorNetInfo->InvalidSubObjectChannelReferences.FindByKey(SubObjectPtr))
-		{
-			ensureMsgf(DestroyedObjectRef == nullptr, TEXT("SubObject %s (0x%p) owned by %s was replicated by connection (0x%p) %s after it was set to (0x%p)"), *GetNameSafe(ReplicatedSubObject), ReplicatedSubObject, *GetNameSafe(OwnerActor), ReferenceOwner, *GetNameSafe(ReferenceOwner), UE::Net::Private::LexToString(DestroyedObjectRef->Status));
-		}
-#endif
+	const TWeakObjectPtr<UObject> SubObjectPtr = ReplicatedSubObject;
 
-		if (FNetworkObjectInfo::FSubObjectChannelReference* ChannelReferences = ActorNetInfo->ActiveSubObjectChannelReferences.Find(SubObjectPtr))
-		{
-#if DO_REPLICATED_OBJECT_CHANNELREF_CHECKS
-			ensureMsgf(!ChannelReferences->RegisteredOwners.Contains(ReferenceOwner), TEXT("SubObject %s (0x%p) owned by %s was already referenced by (0x%p) %s"), *GetNameSafe(ReplicatedSubObject), ReplicatedSubObject, *GetNameSafe(OwnerActor), ReferenceOwner, *GetNameSafe(ReferenceOwner));
-			ChannelReferences->RegisteredOwners.Add(ReferenceOwner);
-#endif
+	// Find the subobject in the active list first
+	FSubObjectChannelReference* ChannelReferences = SubObjectsRefInfo->ActiveSubObjectChannelReferences.Find(SubObjectPtr);
 
-			ChannelReferences->ChannelRefCount++;
-			check(ChannelReferences->ChannelRefCount < MAX_uint16);
+	if (ChannelReferences == nullptr)
+	{
+		// Look if the object could be in the invalid list
+		ChannelReferences = SubObjectsRefInfo->InvalidSubObjectChannelReferences.FindByKey(SubObjectPtr);
+	}
 
-			UE_LOG(LogNetSubObject, Verbose, TEXT("Adding ChannelRef (%d) to %s (0x%p) owned by %s for (0x%p) %s"), ChannelReferences->ChannelRefCount, *GetNameSafe(ReplicatedSubObject), ReplicatedSubObject,  *GetNameSafe(OwnerActor), ReferenceOwner, *GetNameSafe(ReferenceOwner));
-		}
-		else
-		{
-			bool bWasSet = false;
-			FSetElementId ElementId = ActorNetInfo->ActiveSubObjectChannelReferences.Emplace(FNetworkObjectInfo::FSubObjectChannelReference(SubObjectPtr), &bWasSet);
+	if (ChannelReferences)
+	{
+#if DO_REPLICATED_OBJECT_CHANNELREF_CHECKS
+		ensureMsgf(!ChannelReferences->RegisteredOwners.Contains(ReferenceOwner), TEXT("SubObject %s (0x%p) owned by %s was already referenced by (0x%p) %s"), *GetNameSafe(ReplicatedSubObject), ReplicatedSubObject, *GetNameSafe(OwnerActor), ReferenceOwner, *GetNameSafe(ReferenceOwner));
+		ChannelReferences->RegisteredOwners.Add(ReferenceOwner);
+#endif
 
-			UE_LOG(LogNetSubObject, Verbose, TEXT("Adding ChannelRef (new) to %s (0x%p) owned by %s for (0x%p) %s"), *GetNameSafe(ReplicatedSubObject), ReplicatedSubObject, *GetNameSafe(OwnerActor), ReferenceOwner, *GetNameSafe(ReferenceOwner));
+		ChannelReferences->ChannelRefCount++;
+		check(ChannelReferences->ChannelRefCount < MAX_uint16);
+	}
+	else
+	{
+		bool bWasSet = false;
+		FSetElementId ElementId = SubObjectsRefInfo->ActiveSubObjectChannelReferences.Emplace(FSubObjectChannelReference(SubObjectPtr), &bWasSet);
 
 #if DO_REPLICATED_OBJECT_CHANNELREF_CHECKS
-			ActorNetInfo->ActiveSubObjectChannelReferences[ElementId].RegisteredOwners.Add(ReferenceOwner);
+		SubObjectsRefInfo->ActiveSubObjectChannelReferences[ElementId].RegisteredOwners.Add(ReferenceOwner);
 #endif
-		}
 	}
 }
 
-void FNetworkObjectList::RemoveMultipleSubObjectChannelReference(AActor* OwnerActor, const TArrayView<TWeakObjectPtr<UObject>>& SubObjectsToRemove, UObject* ReferenceOwner)
+void FNetworkObjectList::RemoveMultipleSubObjectChannelReference(FObjectKey OwnerActorKey, const TArrayView<TWeakObjectPtr<UObject>>& SubObjectsToRemove, UObject* ReferenceOwner)
 {
-	if (TSharedPtr<FNetworkObjectInfo>* InfoPtr = AllNetworkObjects.Find(OwnerActor))
+	FSetElementId FoundIndex = SubObjectChannelReferences.FindId(OwnerActorKey);
+	if (FoundIndex.IsValidId())
 	{
-		check(InfoPtr->IsValid());
-
-		FNetworkObjectInfo* ActorNetInfo = InfoPtr->Get();
+		FActorSubObjectReferences& SubObjectsRefInfo = SubObjectChannelReferences[FoundIndex];
 
 		for (const TWeakObjectPtr<UObject>& SubObjectPtr : SubObjectsToRemove)
 		{
-			HandleRemoveAnySubObjectChannelRef(ActorNetInfo, SubObjectPtr, ReferenceOwner);
+			HandleRemoveAnySubObjectChannelRef(SubObjectsRefInfo, SubObjectPtr, ReferenceOwner);
+		}
+
+		if (SubObjectsRefInfo.HasNoSubObjects())
+		{
+			SubObjectChannelReferences.Remove(FoundIndex);
 		}
 	}
 }
 
-void FNetworkObjectList::RemoveMultipleInvalidSubObjectChannelReference(FNetworkObjectInfo* ActorNetInfo, const TArrayView<TWeakObjectPtr<UObject>>& SubObjectsToRemove, UObject* ReferenceOwner)
+void FNetworkObjectList::RemoveMultipleInvalidSubObjectChannelReference(FObjectKey OwnerActorKey, const TArrayView<TWeakObjectPtr<UObject>>& SubObjectsToRemove, UObject* ReferenceOwner)
 {
-	check(ActorNetInfo);
-
-	for (const TWeakObjectPtr<UObject>& SubObjectPtr : SubObjectsToRemove)
+	FSetElementId FoundIndex = SubObjectChannelReferences.FindId(OwnerActorKey);
+	if (FoundIndex.IsValidId())
 	{
-		HandleRemoveInvalidSubObjectRef(ActorNetInfo, SubObjectPtr, ReferenceOwner);
+		FActorSubObjectReferences& SubObjectsRefInfo = SubObjectChannelReferences[FoundIndex];
+	
+		for (const TWeakObjectPtr<UObject>& SubObjectPtr : SubObjectsToRemove)
+		{
+			HandleRemoveInvalidSubObjectRef(SubObjectsRefInfo, SubObjectPtr, ReferenceOwner);
+		}
+
+		if (SubObjectsRefInfo.HasNoSubObjects())
+		{
+			SubObjectChannelReferences.Remove(FoundIndex);
+		}
 	}
 }
 
-void FNetworkObjectList::RemoveMultipleActiveSubObjectChannelReference(FNetworkObjectInfo* ActorNetInfo, const TArrayView<TWeakObjectPtr<UObject>>& SubObjectsToRemove, UObject* ReferenceOwner)
+void FNetworkObjectList::RemoveMultipleActiveSubObjectChannelReference(FObjectKey OwnerActorKey, const TArrayView<TWeakObjectPtr<UObject>>& SubObjectsToRemove, UObject* ReferenceOwner)
 {
-	check(ActorNetInfo);
-
-	for (const TWeakObjectPtr<UObject>& SubObjectPtr : SubObjectsToRemove)
+	FSetElementId FoundIndex = SubObjectChannelReferences.FindId(OwnerActorKey);
+	if (FoundIndex.IsValidId())
 	{
-		bool bWasRemoved = HandleRemoveActiveSubObjectRef(ActorNetInfo, SubObjectPtr, ReferenceOwner);
+		FActorSubObjectReferences& SubObjectsRefInfo = SubObjectChannelReferences[FoundIndex];
+
+		for (const TWeakObjectPtr<UObject>& SubObjectPtr : SubObjectsToRemove)
+		{
+			bool bWasRemoved = HandleRemoveActiveSubObjectRef(SubObjectsRefInfo, SubObjectPtr, ReferenceOwner);
 
-		// If somehow the object wasn't in the active list, check if it was in the invalid list just to be sure.
-		// It's possible for an object set to be deleted to have it's uobject ptr become innacessible before we replicate it's owner.
-		// Calling this here would act as a safeguard to ensure its reference is removed even if the tear off/force delete command was never sent.
-		if (!bWasRemoved)
+			// If somehow the object wasn't in the active list, check if it was in the invalid list just to be sure.
+			// It's possible for an object set to be deleted to have it's uobject ptr become innacessible before we replicate it's owner.
+			// Calling this here would act as a safeguard to ensure its reference is removed even if the tear off/force delete command was never sent.
+			if (!bWasRemoved)
+			{
+				HandleRemoveInvalidSubObjectRef(SubObjectsRefInfo, SubObjectPtr, ReferenceOwner);
+			}
+		}
+
+		if (SubObjectsRefInfo.HasNoSubObjects())
 		{
-			HandleRemoveInvalidSubObjectRef(ActorNetInfo, SubObjectPtr, ReferenceOwner);
+			SubObjectChannelReferences.Remove(FoundIndex);
 		}
 	}
 }
 
 void FNetworkObjectList::RemoveSubObjectChannelReference(AActor* OwnerActor, const TWeakObjectPtr<UObject>& SubObjectPtr, UObject* ReferenceOwner)
 {
-	if (TSharedPtr<FNetworkObjectInfo>* InfoPtr = AllNetworkObjects.Find(OwnerActor))
+	FSetElementId FoundIndex = SubObjectChannelReferences.FindId(OwnerActor);
+	if (FoundIndex.IsValidId())
 	{
-		check(InfoPtr->IsValid());
-		HandleRemoveAnySubObjectChannelRef(InfoPtr->Get(), SubObjectPtr, ReferenceOwner);
+		FActorSubObjectReferences& SubObjectsRefInfo = SubObjectChannelReferences[FoundIndex];
+		HandleRemoveAnySubObjectChannelRef(SubObjectsRefInfo, SubObjectPtr, ReferenceOwner);
+
+		if (SubObjectsRefInfo.HasNoSubObjects())
+		{
+			SubObjectChannelReferences.Remove(FoundIndex);
+		}
 	}
 }
 
-void FNetworkObjectList::HandleRemoveAnySubObjectChannelRef(FNetworkObjectInfo* ActorNetInfo, const TWeakObjectPtr<UObject>& SubObjectPtr, UObject* ReferenceOwner)
+void FNetworkObjectList::HandleRemoveAnySubObjectChannelRef(FActorSubObjectReferences& SubObjectsRefInfo, const TWeakObjectPtr<UObject>& SubObjectPtr, UObject* ReferenceOwner)
 {
-	check(ActorNetInfo);
 	check(SubObjectPtr.IsExplicitlyNull() == false);
 
 	// Look in the active set
-	bool bWasRemoved = HandleRemoveActiveSubObjectRef(ActorNetInfo, SubObjectPtr, ReferenceOwner);
+	bool bWasRemoved = HandleRemoveActiveSubObjectRef(SubObjectsRefInfo, SubObjectPtr, ReferenceOwner);
 
 	// Look in the destroyed array
 	if (!bWasRemoved)
 	{
-		bWasRemoved = HandleRemoveInvalidSubObjectRef(ActorNetInfo, SubObjectPtr, ReferenceOwner);
+		bWasRemoved = HandleRemoveInvalidSubObjectRef(SubObjectsRefInfo, SubObjectPtr, ReferenceOwner);
 	}
 
-	ensureMsgf(bWasRemoved, TEXT("HandleRemoveAnySubObjectChannelRef could not find any references for %s (0x%p) owned by %s"), *GetNameSafe(SubObjectPtr.GetEvenIfUnreachable()), SubObjectPtr.GetEvenIfUnreachable(), *GetNameSafe(ActorNetInfo->Actor));
+#if DO_REPLICATED_OBJECT_CHANNELREF_CHECKS
+	ensureMsgf(bWasRemoved, TEXT("HandleRemoveAnySubObjectChannelRef could not find any references for %s (0x%p) owned by %s"), *GetNameSafe(SubObjectPtr.GetEvenIfUnreachable()), SubObjectPtr.GetEvenIfUnreachable(), *GetNameSafe(SubObjectsRefInfo.ActorKey.ResolveObjectPtrEvenIfUnreachable()));
+#endif
 }
 
-bool FNetworkObjectList::HandleRemoveActiveSubObjectRef(FNetworkObjectInfo* ActorNetInfo, const TWeakObjectPtr<UObject>& SubObjectPtr, UObject* ReferenceOwner)
+bool FNetworkObjectList::HandleRemoveActiveSubObjectRef(FActorSubObjectReferences& SubObjectsRefInfo, const TWeakObjectPtr<UObject>& SubObjectPtr, UObject* ReferenceOwner)
 {
-	check(ActorNetInfo);
-
-	FSetElementId FoundId = ActorNetInfo->ActiveSubObjectChannelReferences.FindId(SubObjectPtr);
+	FSetElementId FoundId = SubObjectsRefInfo.ActiveSubObjectChannelReferences.FindId(SubObjectPtr);
 	if (FoundId.IsValidId())
 	{
-		FNetworkObjectInfo::FSubObjectChannelReference& ActiveReference = ActorNetInfo->ActiveSubObjectChannelReferences[FoundId];
+		FSubObjectChannelReference& ActiveReference = SubObjectsRefInfo.ActiveSubObjectChannelReferences[FoundId];
 
 		check(ActiveReference.SubObjectPtr.HasSameIndexAndSerialNumber(SubObjectPtr));
 
 		ActiveReference.ChannelRefCount--;
 
-		UE_LOG(LogNetSubObject, Verbose, TEXT("Removed ACTIVE ChannelRef (%d) for %s (0x%p) owned by %s for connection (0x%p) %s"), ActiveReference.ChannelRefCount, *GetNameSafe(SubObjectPtr.GetEvenIfUnreachable()), SubObjectPtr.GetEvenIfUnreachable(), *GetNameSafe(ActorNetInfo->Actor), ReferenceOwner, *GetNameSafe(ReferenceOwner));
-
 #if DO_REPLICATED_OBJECT_CHANNELREF_CHECKS
 		int32 Removed = ActiveReference.RegisteredOwners.RemoveSingleSwap(ReferenceOwner);
-		ensureMsgf(Removed > 0, TEXT("Removed ACTIVE ref for Subobject %s (0x%p) owned by %s but it was never referenced by the connection (0x%p) %s"), *GetNameSafe(SubObjectPtr.GetEvenIfUnreachable()), SubObjectPtr.GetEvenIfUnreachable(), *GetNameSafe(ActorNetInfo->Actor), ReferenceOwner, *GetNameSafe(ReferenceOwner));
+		ensureMsgf(Removed > 0, TEXT("Removed ACTIVE ref for Subobject %s (0x%p) owned by %s but it was never referenced by the connection (0x%p) %s"), *GetNameSafe(SubObjectPtr.GetEvenIfUnreachable()), SubObjectPtr.GetEvenIfUnreachable(), *GetNameSafe(SubObjectsRefInfo.ActorKey.ResolveObjectPtrEvenIfUnreachable()), ReferenceOwner, *GetNameSafe(ReferenceOwner));
 #endif
 
 		if (ActiveReference.ChannelRefCount == 0)
 		{
-			ActorNetInfo->ActiveSubObjectChannelReferences.Remove(FoundId);
+			SubObjectsRefInfo.ActiveSubObjectChannelReferences.Remove(FoundId);
 		}
 
 		return true;
@@ -764,25 +818,21 @@ bool FNetworkObjectList::HandleRemoveActiveSubObjectRef(FNetworkObjectInfo* Acto
 	return false;
 }
 
-bool FNetworkObjectList::HandleRemoveInvalidSubObjectRef(FNetworkObjectInfo* ActorNetInfo, const TWeakObjectPtr<UObject>& SubObjectPtr, UObject* ReferenceOwner)
+bool FNetworkObjectList::HandleRemoveInvalidSubObjectRef(FActorSubObjectReferences& SubObjectsRefInfo, const TWeakObjectPtr<UObject>& SubObjectPtr, UObject* ReferenceOwner)
 {
-	check(ActorNetInfo);
-
-	const int32 Index = ActorNetInfo->InvalidSubObjectChannelReferences.IndexOfByKey(SubObjectPtr);
+	const int32 Index = SubObjectsRefInfo.InvalidSubObjectChannelReferences.IndexOfByKey(SubObjectPtr);
 	if (Index != INDEX_NONE)
 	{
-		FNetworkObjectInfo::FSubObjectChannelReference& InvalidReference = ActorNetInfo->InvalidSubObjectChannelReferences[Index];
+		FSubObjectChannelReference& InvalidReference = SubObjectsRefInfo.InvalidSubObjectChannelReferences[Index];
 		InvalidReference.ChannelRefCount--;
 
-		UE_LOG(LogNetSubObject, Verbose, TEXT("Removed INVALID ChannelRef (%d) for %s (0x%p) owned by %s for connection (0x%p) %s"), InvalidReference.ChannelRefCount,*GetNameSafe(SubObjectPtr.GetEvenIfUnreachable()), SubObjectPtr.GetEvenIfUnreachable(), *GetNameSafe(ActorNetInfo->Actor), ReferenceOwner, *GetNameSafe(ReferenceOwner));
-
 #if DO_REPLICATED_OBJECT_CHANNELREF_CHECKS
 		int32 Removed = InvalidReference.RegisteredOwners.RemoveSingleSwap(ReferenceOwner);
-		ensureMsgf(Removed > 0, TEXT("Removed INVALID ref for Subobject %s (0x%p) owned by %s but it was never referenced by the connection (0x%p) %s"), *GetNameSafe(SubObjectPtr.GetEvenIfUnreachable()), SubObjectPtr.GetEvenIfUnreachable(), *GetNameSafe(ActorNetInfo->Actor), ReferenceOwner, *GetNameSafe(ReferenceOwner));
+		ensureMsgf(Removed > 0, TEXT("Removed INVALID ref for Subobject %s (0x%p) owned by %s but it was never referenced by the connection (0x%p) %s"), *GetNameSafe(SubObjectPtr.GetEvenIfUnreachable()), SubObjectPtr.GetEvenIfUnreachable(), *GetNameSafe(SubObjectsRefInfo.ActorKey.ResolveObjectPtrEvenIfUnreachable()), ReferenceOwner, *GetNameSafe(ReferenceOwner));
 #endif
 		if (InvalidReference.ChannelRefCount == 0)
 		{
-			ActorNetInfo->InvalidSubObjectChannelReferences.RemoveAtSwap(Index);
+			SubObjectsRefInfo.InvalidSubObjectChannelReferences.RemoveAtSwap(Index);
 		}
 
 		return true;
@@ -794,56 +844,53 @@ bool FNetworkObjectList::HandleRemoveInvalidSubObjectRef(FNetworkObjectInfo* Act
 #if DO_REPLICATED_OBJECT_CHANNELREF_CHECKS
 void FNetworkObjectList::SwapMultipleReferencesForDormancy(AActor* OwnerActor, const TArrayView<TWeakObjectPtr<UObject>>& SubObjectsToSwap, UActorChannel* PreviousChannelRefOwner, UNetConnection* NewConnectionRefOwner)
 {
-	if (TSharedPtr<FNetworkObjectInfo>* InfoPtr = AllNetworkObjects.Find(OwnerActor))
+	if (FActorSubObjectReferences* SubObjectsRefInfo = SubObjectChannelReferences.Find(OwnerActor))
 	{
-		check(InfoPtr->IsValid());
-
-		FNetworkObjectInfo* ActorNetInfo = InfoPtr->Get();
-
 		for (const TWeakObjectPtr<UObject>& SubObjectPtr : SubObjectsToSwap)
 		{
-			HandleSwapReferenceForDormancy(ActorNetInfo, SubObjectPtr, PreviousChannelRefOwner, NewConnectionRefOwner);
+			HandleSwapReferenceForDormancy(SubObjectsRefInfo, SubObjectPtr, PreviousChannelRefOwner, NewConnectionRefOwner);
 		}
 	}
 }
 
 void FNetworkObjectList::SwapReferenceForDormancy(AActor* OwnerActor, UObject* ReplicatedSubObject, UNetConnection* PreviousConnectionRefOwner, UActorChannel* NewChannelRefOwner)
 {
-	if (TSharedPtr<FNetworkObjectInfo>* InfoPtr = AllNetworkObjects.Find(OwnerActor))
+	if (FActorSubObjectReferences* SubObjectsRefInfo = SubObjectChannelReferences.Find(OwnerActor))
 	{
-		check(InfoPtr->IsValid());
-
 		const TWeakObjectPtr<UObject> SubObjectPtr = ReplicatedSubObject;
-		HandleSwapReferenceForDormancy(InfoPtr->Get(), SubObjectPtr, PreviousConnectionRefOwner, NewChannelRefOwner);
+		HandleSwapReferenceForDormancy(SubObjectsRefInfo, SubObjectPtr, PreviousConnectionRefOwner, NewChannelRefOwner);
 	}
 }
 
-void FNetworkObjectList::HandleSwapReferenceForDormancy(FNetworkObjectInfo* ActorNetInfo, const TWeakObjectPtr<UObject>& SubObjectPtr, UObject* PreviousRefOwner, UObject* NewRefOwner)
+void FNetworkObjectList::HandleSwapReferenceForDormancy(FActorSubObjectReferences* SubObjectsRefInfo, const TWeakObjectPtr<UObject>& SubObjectPtr, UObject* PreviousRefOwner, UObject* NewRefOwner)
 {
-	check(ActorNetInfo);
+	check(SubObjectsRefInfo);
 
 	// First check the active list
-	FNetworkObjectInfo::FSubObjectChannelReference* SubObjectChannelRef = ActorNetInfo->ActiveSubObjectChannelReferences.Find(SubObjectPtr);
-	
+	FSubObjectChannelReference* SubObjectChannelRef = SubObjectsRefInfo->ActiveSubObjectChannelReferences.Find(SubObjectPtr);
+
 	if (!SubObjectChannelRef)
 	{
 		// Then check the inactive list
-		SubObjectChannelRef = ActorNetInfo->InvalidSubObjectChannelReferences.FindByKey(SubObjectPtr);
+		SubObjectChannelRef = SubObjectsRefInfo->InvalidSubObjectChannelReferences.FindByKey(SubObjectPtr);
 	}
 
 	if (SubObjectChannelRef)
 	{
-		UE_LOG(LogNetSubObject, Verbose, TEXT("Swapped ChannelRef (%d) for SubObject %s (0x%p) owned by %s. From (0x%p) %s to (0x%p) %s"), SubObjectChannelRef->ChannelRefCount, *GetNameSafe(SubObjectPtr.GetEvenIfUnreachable()), SubObjectPtr.GetEvenIfUnreachable(), *GetNameSafe(ActorNetInfo->Actor), PreviousRefOwner, *GetNameSafe(PreviousRefOwner), NewRefOwner, *GetNameSafe(NewRefOwner));
-
 		int32 Removed = SubObjectChannelRef->RegisteredOwners.RemoveSingleSwap(PreviousRefOwner);
-		ensureMsgf(Removed > 0, TEXT("SwapReferencesForDormancy could not find reference to previous reference (0x%p) %s for subobject %s (0x%p) owned by %s. Swapping to (0x%p) %s"), PreviousRefOwner, *GetNameSafe(PreviousRefOwner), *GetNameSafe(SubObjectPtr.GetEvenIfUnreachable()), SubObjectPtr.GetEvenIfUnreachable(), *GetNameSafe(ActorNetInfo->Actor), NewRefOwner, *GetNameSafe(NewRefOwner));
 
-		ensureMsgf(!SubObjectChannelRef->RegisteredOwners.Contains(NewRefOwner), TEXT("SwapReferencesForDormancy found new reference (0x%p) %s was already registered to %s (0x%p) owned by %s"), NewRefOwner, *GetNameSafe(NewRefOwner), *GetNameSafe(SubObjectPtr.GetEvenIfUnreachable()), SubObjectPtr.GetEvenIfUnreachable(), *GetNameSafe(ActorNetInfo->Actor));
+		ensureMsgf(Removed > 0, TEXT("SwapReferencesForDormancy could not find reference to previous reference (0x%p) %s for subobject %s (0x%p) owned by %s. Swapping to (0x%p) %s"), 
+			PreviousRefOwner, *GetNameSafe(PreviousRefOwner), *GetNameSafe(SubObjectPtr.GetEvenIfUnreachable()), SubObjectPtr.GetEvenIfUnreachable(), *GetNameSafe(SubObjectsRefInfo->ActorKey.ResolveObjectPtrEvenIfUnreachable()), NewRefOwner, *GetNameSafe(NewRefOwner));
+
+		ensureMsgf(!SubObjectChannelRef->RegisteredOwners.Contains(NewRefOwner), TEXT("SwapReferencesForDormancy found new reference (0x%p) %s was already registered to %s (0x%p) owned by %s"), 
+			NewRefOwner, *GetNameSafe(NewRefOwner), *GetNameSafe(SubObjectPtr.GetEvenIfUnreachable()), SubObjectPtr.GetEvenIfUnreachable(), *GetNameSafe(SubObjectsRefInfo->ActorKey.ResolveObjectPtrEvenIfUnreachable()));
+
 		SubObjectChannelRef->RegisteredOwners.Add(NewRefOwner);
 	}
 	else
 	{
-		ensureMsgf(false, TEXT("SwapReferencesForDormancy could not find any references to %s (0x%p) owned by %s. Swapping from (0x%p) %s to (0x%p) %s"), *GetNameSafe(SubObjectPtr.GetEvenIfUnreachable()), SubObjectPtr.GetEvenIfUnreachable(), *GetNameSafe(ActorNetInfo->Actor), PreviousRefOwner, *GetNameSafe(PreviousRefOwner), NewRefOwner, *GetNameSafe(NewRefOwner));
+		ensureMsgf(false, TEXT("SwapReferencesForDormancy could not find any references to %s (0x%p) owned by %s. Swapping from (0x%p) %s to (0x%p) %s"), 
+			*GetNameSafe(SubObjectPtr.GetEvenIfUnreachable()), SubObjectPtr.GetEvenIfUnreachable(), *GetNameSafe(SubObjectsRefInfo->ActorKey.ResolveObjectPtrEvenIfUnreachable()), PreviousRefOwner, *GetNameSafe(PreviousRefOwner), NewRefOwner, *GetNameSafe(NewRefOwner));
 	}
 }
 #endif //#if DO_REPLICATED_OBJECT_CHANNELREF_CHECKS
diff --git a/Engine/Source/Runtime/Engine/Public/Net/NetDormantHolder.h b/Engine/Source/Runtime/Engine/Public/Net/NetDormantHolder.h
--- a/Engine/Source/Runtime/Engine/Public/Net/NetDormantHolder.h
+++ b/Engine/Source/Runtime/Engine/Public/Net/NetDormantHolder.h
@@ -13,7 +13,7 @@ class FNetworkObjectList;
 namespace UE::Net
 {
 	/** Function definition allowed in the ForEach functions */
-	typedef TFunctionRef<void(AActor* OwnerActor, FObjectKey ObjectKey, const TSharedRef<FObjectReplicator>& ReplicatorRef)> FExecuteForEachDormantReplicator;
+	typedef TFunctionRef<void(FObjectKey OwnerActorKey, FObjectKey ObjectKey, const TSharedRef<FObjectReplicator>& ReplicatorRef)> FExecuteForEachDormantReplicator;
 }
 
 namespace UE::Net::Private
@@ -47,10 +47,10 @@ struct FDormantObjectReplicator
 */
 struct FActorDormantReplicators
 {
-	explicit FActorDormantReplicators(AActor* InOwnerActor) : OwnerActor(InOwnerActor) {}
+	explicit FActorDormantReplicators(AActor* InOwnerActor) : OwnerActorKey(InOwnerActor) {}
 
-	bool operator==(const FActorDormantReplicators& rhs) const { return OwnerActor == rhs.OwnerActor; }
-	bool operator==(AActor* rhs) const { return OwnerActor == rhs; }
+	bool operator==(const FActorDormantReplicators& rhs) const { return OwnerActorKey == rhs.OwnerActorKey; }
+	bool operator==(FObjectKey rhs) const { return OwnerActorKey == rhs; }
 
 	void CountBytes(FArchive& Ar) const { DormantReplicators.CountBytes(Ar); }
 
@@ -63,7 +63,7 @@ struct FActorDormantReplicators
 	};
 
 	/** The dormant actor who owns all the different object replicators */
-	AActor* OwnerActor = nullptr;
+	FObjectKey OwnerActorKey;
 
 	typedef TSet<FDormantObjectReplicator, FDormantObjectReplicatorKeyFuncs> FObjectReplicatorSet;
 	/** List of all object replicates stored for this dormant actor */
@@ -139,11 +139,6 @@ struct FDormantReplicatorHolder
 	*/
 	void CleanupAllReplicatorsOfActor(AActor* DormantActor);
 
-	/**
-	* Iterate over all the stored object replicators and destroy any that are tied to a replicated object that is now considered invalid.
-	*/
-	void CleanupStaleObjects();
-
 	/**
 	* Iterate over all the stored object replicators and destroy any that are tied to a replicated object that is now considered invalid.
 	* This version will decrement the subobject references stored in the networkobject list too.
@@ -171,9 +166,9 @@ struct FDormantReplicatorHolder
 	void CountBytes(FArchive& Ar) const;
 
 	/** KeyFuncs that make it so the TSet can only needs a simple AActor pointer for the Key */
-	struct FActorDormantReplicatorsKeyFuncs : BaseKeyFuncs<FActorDormantReplicators, AActor*, false>
+	struct FActorDormantReplicatorsKeyFuncs : BaseKeyFuncs<FActorDormantReplicators, FObjectKey, false>
 	{
-		static KeyInitType GetSetKey(ElementInitType Element) { return Element.OwnerActor; }
+		static KeyInitType GetSetKey(ElementInitType Element) { return Element.OwnerActorKey; }
 		static bool Matches(KeyInitType lhs, KeyInitType rhs) { return lhs == rhs; }
 		static uint32 GetKeyHash(KeyInitType Key) { return GetTypeHash(Key); }
 	};
