Integrate RenderResourceViewer from //UE5/Main
Part13 - Integrate CL23526201 as edigrate only skipping RenderResource.cpp change
Remove game thread check on RHI resource set owner name. EngineTest can call on FSkinWeightVertexBuffer from other thread on consoles

diff --git a/Engine/Source/Runtime/Engine/Private/Rendering/SkinWeightVertexBuffer.cpp b/Engine/Source/Runtime/Engine/Private/Rendering/SkinWeightVertexBuffer.cpp
--- a/Engine/Source/Runtime/Engine/Private/Rendering/SkinWeightVertexBuffer.cpp
+++ b/Engine/Source/Runtime/Engine/Private/Rendering/SkinWeightVertexBuffer.cpp
@@ -803,7 +803,6 @@ void FSkinWeightVertexBuffer::ResetVertexBoneWeights(uint32 VertexIndex)
 
 void FSkinWeightVertexBuffer::SetOwnerName(const FName& OwnerName)
 {
-	check(IsInGameThread());
 	LookupVertexBuffer.SetOwnerName(OwnerName);
 	DataVertexBuffer.SetOwnerName(OwnerName);
 }
