Add missing define when refcount is not used.

diff --git a/Engine/Source/Runtime/Engine/Private/NetworkObjectList.cpp b/Engine/Source/Runtime/Engine/Private/NetworkObjectList.cpp
--- a/Engine/Source/Runtime/Engine/Private/NetworkObjectList.cpp
+++ b/Engine/Source/Runtime/Engine/Private/NetworkObjectList.cpp
@@ -28,11 +28,13 @@ namespace UE::Net::Private
 	}
 }
 
+#if UE_REPLICATED_OBJECT_REFCOUNTING
 void FNetworkObjectList::FActorSubObjectReferences::CountBytes(FArchive& Ar) const
 {
 	ActiveSubObjectChannelReferences.CountBytes(Ar);
 	InvalidSubObjectChannelReferences.CountBytes(Ar);
 }
+#endif
 
 void FNetworkObjectList::AddInitialObjects(UWorld* const World, UNetDriver* NetDriver)
 {
