Fix member variable initialized with itself.

diff --git a/Engine/Source/Runtime/Engine/Classes/Engine/NetworkObjectList.h b/Engine/Source/Runtime/Engine/Classes/Engine/NetworkObjectList.h
--- a/Engine/Source/Runtime/Engine/Classes/Engine/NetworkObjectList.h
+++ b/Engine/Source/Runtime/Engine/Classes/Engine/NetworkObjectList.h
@@ -364,7 +364,7 @@ public:
 	{
 	public:
 		explicit FActorInvalidSubObjectView(uint16 InDirtyCount, const TArray<FSubObjectChannelReference>* InArrayPtr=nullptr)
-			: InvalidSubObjectDirtyCount(InvalidSubObjectDirtyCount)
+			: InvalidSubObjectDirtyCount(InDirtyCount)
 			, InvalidSubObjectsPtr(InArrayPtr)
 		{ }
 
@@ -373,7 +373,7 @@ public:
 		inline uint16 GetDirtyCount() const { return InvalidSubObjectDirtyCount; }
 		inline bool HasInvalidSubObjects() const { return InvalidSubObjectsPtr != nullptr; }
 		inline const TArray<FSubObjectChannelReference>& GetInvalidSubObjects() const 
-		{ 
+		{
 			check(HasInvalidSubObjects());
 			return *InvalidSubObjectsPtr;
 		}
