Removed check in UActorDescContainer::GetWorld() to prevent console command "ke *" from crashing.

diff --git a/Engine/Source/Runtime/Engine/Public/WorldPartition/ActorDescContainer.h b/Engine/Source/Runtime/Engine/Public/WorldPartition/ActorDescContainer.h
--- a/Engine/Source/Runtime/Engine/Public/WorldPartition/ActorDescContainer.h
+++ b/Engine/Source/Runtime/Engine/Public/WorldPartition/ActorDescContainer.h
@@ -117,7 +117,7 @@ protected:
 
 private:
 	// GetWorld() should never be called on an ActorDescContainer to avoid any confusion as it can be used as a template
-	virtual UWorld* GetWorld() const override { check(false); return nullptr; }
+	virtual UWorld* GetWorld() const override { return nullptr; }
 
 #if WITH_EDITOR
 	bool ShouldRegisterDelegates();
