Fixing AddTimedDisplay crash when changing level

diff --git a/Engine/Source/Runtime/Engine/Private/Engine/ViewportStatsSubsystem.cpp b/Engine/Source/Runtime/Engine/Private/Engine/ViewportStatsSubsystem.cpp
--- a/Engine/Source/Runtime/Engine/Private/Engine/ViewportStatsSubsystem.cpp
+++ b/Engine/Source/Runtime/Engine/Private/Engine/ViewportStatsSubsystem.cpp
@@ -124,15 +124,15 @@ void UViewportStatsSubsystem::AddTimedDisplay(FText Text, FLinearColor Color, fl
 		FTimerDelegate TimerDel;
 		FTimerHandle TimerHandle;
 
-		auto RemoveAfterSecondsLambda = [](TSharedPtr<FUniqueDisplayData> DisplayItem, TArray<TSharedPtr<FUniqueDisplayData>>* ConditionArray)
+		auto RemoveAfterSecondsLambda = [](TWeakPtr<FUniqueDisplayData> DisplayItem, TArray<TSharedPtr<FUniqueDisplayData>>* ConditionArray)
 		{
-			if(DisplayItem.IsValid() && ConditionArray)
+			if (DisplayItem.IsValid() && ConditionArray)
 			{
-				ConditionArray->Remove(DisplayItem);
+				ConditionArray->Remove(DisplayItem.Pin());
 			}
 		};
 
-		TimerDel.BindLambda(RemoveAfterSecondsLambda, Message, &UniqueDisplayMessages);
+		TimerDel.BindLambda(RemoveAfterSecondsLambda, TWeakPtr<FUniqueDisplayData>(Message), &UniqueDisplayMessages);
 		MyWorld->GetTimerManager().SetTimer(TimerHandle, TimerDel, Duration, /* bInLoop= */ false);
 	}
 }
