Add min/max/average info of shader code size to shader compilation stats.

#rnx

diff --git a/Engine/Source/Runtime/Engine/Private/ShaderCompiler/ShaderCompiler.cpp b/Engine/Source/Runtime/Engine/Private/ShaderCompiler/ShaderCompiler.cpp
--- a/Engine/Source/Runtime/Engine/Private/ShaderCompiler/ShaderCompiler.cpp
+++ b/Engine/Source/Runtime/Engine/Private/ShaderCompiler/ShaderCompiler.cpp
@@ -3257,6 +3257,14 @@ void FShaderCompilerStats::WriteStatSummary()
 		UE_LOG(LogShaderCompilers, Display, TEXT("Job life time (pending + execution): average %.2f s, max %.2f"), AccumulatedJobLifeTime / JobsCompleted, MaxJobLifeTime);
 	}
 
+	if (NumAccumulatedShaderCodes > 0)
+	{
+		const FString AvgCodeSizeStr = FText::AsMemory((uint64)((double)AccumulatedShaderCodeSize / (double)NumAccumulatedShaderCodes)).ToString();
+		const FString MinCodeSizeStr = FText::AsMemory((uint64)MinShaderCodeSize).ToString();
+		const FString MaxCodeSizeStr = FText::AsMemory((uint64)MaxShaderCodeSize).ToString();
+		UE_LOG(LogShaderCompilers, Display, TEXT("Shader code size: average %s, min %s, max %s"), *AvgCodeSizeStr, *MinCodeSizeStr, *MaxCodeSizeStr);
+	}
+
 	UE_LOG(LogShaderCompilers, Display, TEXT("Time at least one job was in flight (either pending or executed): %.2f s"), TotalTimeAtLeastOneJobWasInFlight);
 
 	// print stats about the batches
@@ -3450,6 +3458,16 @@ void FShaderCompilerStats::RegisterFinishedJob(FShaderCommonCompileJob& Job)
 
 	if (const FShaderCompileJob* SingleJob = Job.GetSingleShaderJob())
 	{
+		// Register min/max/average shader code sizes for single job output
+		const int32 ShaderCodeSize = SingleJob->Output.ShaderCode.GetShaderCodeSize();
+		if (ShaderCodeSize > 0)
+		{
+			MinShaderCodeSize = (MinShaderCodeSize > 0 ? FMath::Min(MinShaderCodeSize, ShaderCodeSize) : ShaderCodeSize);
+			MaxShaderCodeSize = (MaxShaderCodeSize > 0 ? FMath::Max(MaxShaderCodeSize, ShaderCodeSize) : ShaderCodeSize);
+			AccumulatedShaderCodeSize += (uint64)ShaderCodeSize;
+			++NumAccumulatedShaderCodes;
+		}
+
 		const FString ShaderName(SingleJob->Key.ShaderType->GetName());
 		if (FShaderTimings* Existing = ShaderTimings.Find(ShaderName))
 		{
diff --git a/Engine/Source/Runtime/Engine/Public/ShaderCompiler.h b/Engine/Source/Runtime/Engine/Public/ShaderCompiler.h
--- a/Engine/Source/Runtime/Engine/Public/ShaderCompiler.h
+++ b/Engine/Source/Runtime/Engine/Public/ShaderCompiler.h
@@ -552,6 +552,18 @@ private:
 	/** Total jobs in local job batches. */
 	int64 TotalJobsReportedInDistributedJobBatches = 0;
 
+	/** Size of the smallest output shader code. */
+	int32 MinShaderCodeSize = 0;
+
+	/** Size of the largest output shader code. */
+	int32 MaxShaderCodeSize = 0;
+
+	/** Total accumulated size of all output shader codes. */
+	uint64 AccumulatedShaderCodeSize = 0;
+
+	/** Number of accumulated output shader codes. */
+	uint64 NumAccumulatedShaderCodes = 0;
+
 	/** Accumulates the job lifetimes without overlaps */
 	TArray<TInterval<double>> JobLifeTimeIntervals;
 
