Fix monolithic build with unresolved FObjectKey constructor

diff --git a/Engine/Source/Runtime/Engine/Public/Net/NetDormantHolder.h b/Engine/Source/Runtime/Engine/Public/Net/NetDormantHolder.h
--- a/Engine/Source/Runtime/Engine/Public/Net/NetDormantHolder.h
+++ b/Engine/Source/Runtime/Engine/Public/Net/NetDormantHolder.h
@@ -47,7 +47,7 @@ struct FDormantObjectReplicator
 */
 struct FActorDormantReplicators
 {
-	explicit FActorDormantReplicators(AActor* InOwnerActor) : OwnerActorKey(InOwnerActor) {}
+	explicit FActorDormantReplicators(FObjectKey InOwnerActorKey) : OwnerActorKey(InOwnerActorKey) {}
 
 	bool operator==(const FActorDormantReplicators& rhs) const { return OwnerActorKey == rhs.OwnerActorKey; }
 	bool operator==(FObjectKey rhs) const { return OwnerActorKey == rhs; }
