WorldPartition: Fix cooked FunctionalTest not getting discovered for WP maps

diff --git a/Engine/Source/Runtime/Engine/Private/WorldPartition/WorldPartitionRuntimeLevelStreamingCell.cpp b/Engine/Source/Runtime/Engine/Private/WorldPartition/WorldPartitionRuntimeLevelStreamingCell.cpp
--- a/Engine/Source/Runtime/Engine/Private/WorldPartition/WorldPartitionRuntimeLevelStreamingCell.cpp
+++ b/Engine/Source/Runtime/Engine/Private/WorldPartition/WorldPartitionRuntimeLevelStreamingCell.cpp
@@ -8,6 +8,7 @@
 #include "Engine/Level.h"
 #include "Misc/HierarchicalLogArchive.h"
 #include "Misc/Paths.h"
+#include "AssetRegistry/IAssetRegistry.h"
 
 #include UE_INLINE_GENERATED_CPP_BY_NAME(WorldPartitionRuntimeLevelStreamingCell)
 
@@ -247,6 +248,11 @@ bool UWorldPartitionRuntimeLevelStreamingCell::PopulateGeneratorPackageForCook(T
 
 		// Remap needed here for references to Actors that are inside a Container
 		FWorldPartitionLevelHelper::RemapLevelSoftObjectPaths(OuterWorld->PersistentLevel, WorldPartition);
+
+		// Make sure Asset Registry tags are updated here synchronously now that Package contains all its actors
+		// ex: AFunctionalTest actors need to be part of the Worlds asset tags once they are not longer external so they can be discovered at runtime
+		IAssetRegistry::Get()->AssetFullyUpdateTags(OuterWorld);
+
 		// Empty cell's package list (this ensures that no one can rely on cell's content).
 		Packages.Empty();
 	}
