Allow levels to turn off garbage collection clustering so that actors can be deleted

diff --git a/Engine/Source/Runtime/Engine/Classes/Engine/Level.h b/Engine/Source/Runtime/Engine/Classes/Engine/Level.h
--- a/Engine/Source/Runtime/Engine/Classes/Engine/Level.h
+++ b/Engine/Source/Runtime/Engine/Classes/Engine/Level.h
@@ -664,6 +664,9 @@ public:
 	uint8										bHasRerunConstructionScripts:1;
 	/** Whether the level had its actor cluster created. This doesn't mean that the creation was successful. */
 	uint8										bActorClusterCreated : 1;
+	/** If true, allows garbage collection clustering for the level */
+	uint8										bGarbageCollectionClusteringEnabled : 1;
+
 	/** Whether the level is partitioned or not. */
     UPROPERTY()
 	uint8										bIsPartitioned : 1;
diff --git a/Engine/Source/Runtime/Engine/Classes/Engine/LevelStreaming.h b/Engine/Source/Runtime/Engine/Classes/Engine/LevelStreaming.h
--- a/Engine/Source/Runtime/Engine/Classes/Engine/LevelStreaming.h
+++ b/Engine/Source/Runtime/Engine/Classes/Engine/LevelStreaming.h
@@ -299,6 +299,9 @@ public:
 	UPROPERTY(EditAnywhere, AdvancedDisplay, Category = LevelStreaming)
 	uint8 bDrawOnLevelStatusMap : 1;
 
+	/** If true, allows garbage collection clustering for the level. */
+	uint8 bGarbageCollectionClusteringEnabled : 1;
+
 #if WITH_EDITORONLY_DATA
 	/** Deprecated level color used for visualization. */
 	UPROPERTY()
diff --git a/Engine/Source/Runtime/Engine/Private/Level.cpp b/Engine/Source/Runtime/Engine/Private/Level.cpp
--- a/Engine/Source/Runtime/Engine/Private/Level.cpp
+++ b/Engine/Source/Runtime/Engine/Private/Level.cpp
@@ -445,6 +445,7 @@ ULevel::ULevel( const FObjectInitializer& ObjectInitializer )
 	bFixupActorFoldersAtLoad = IsActorFolderObjectsFeatureAvailable();
 #endif	
 	bActorClusterCreated = false;
+	bGarbageCollectionClusteringEnabled = true;
 	bIsPartitioned = false;
 	bStaticComponentsRegisteredInStreamingManager = false;
 	IncrementalComponentState = EIncrementalComponentState::Init;
@@ -466,7 +467,7 @@ void ULevel::AddReferencedObjects(UObject* InThis, FReferenceCollector& Collecto
 	ULevel* This = CastChecked<ULevel>(InThis);
 
 	// Let GC know that we're referencing some AActor objects
-	if (FPlatformProperties::RequiresCookedData() && GActorClusteringEnabled && This->bActorClusterCreated)
+	if (FPlatformProperties::RequiresCookedData() && GActorClusteringEnabled && This->bGarbageCollectionClusteringEnabled && This->bActorClusterCreated)
 	{
 		Collector.AddStableReferenceArray(&This->ActorsForGC);
 	}
@@ -1286,7 +1287,7 @@ void ULevel::CreateCluster()
 	// Also, we don't want the level to reference the actors that are clusters because that would
 	// make things work even slower (references to clustered objects are expensive). That's why
 	// we keep a separate array for referencing unclustered actors (ActorsForGC).
-	if (FPlatformProperties::RequiresCookedData() && GCreateGCClusters && GActorClusteringEnabled && !bActorClusterCreated)
+	if (FPlatformProperties::RequiresCookedData() && GCreateGCClusters && GActorClusteringEnabled && bGarbageCollectionClusteringEnabled && !bActorClusterCreated)
 	{
 		ActorsForGC.Reset();
 
diff --git a/Engine/Source/Runtime/Engine/Private/LevelStreaming.cpp b/Engine/Source/Runtime/Engine/Private/LevelStreaming.cpp
--- a/Engine/Source/Runtime/Engine/Private/LevelStreaming.cpp
+++ b/Engine/Source/Runtime/Engine/Private/LevelStreaming.cpp
@@ -387,6 +387,7 @@ ULevelStreaming::ULevelStreaming(const FObjectInitializer& ObjectInitializer)
 	CurrentState = ELevelStreamingState::Removed;
 	bSkipClientUseMakingInvisibleTransactionRequest = false;
 	bSkipClientUseMakingVisibleTransactionRequest = false;
+	bGarbageCollectionClusteringEnabled = true;
 }
 
 void ULevelStreaming::PostLoad()
@@ -1312,6 +1313,8 @@ void ULevelStreaming::PrepareLoadedLevel(ULevel* InLevel, UPackage* InLevelPacka
 	check(InLevel);
 	UWorld* LevelOwningWorld = InLevel->OwningWorld;
 
+	InLevel->bGarbageCollectionClusteringEnabled = bGarbageCollectionClusteringEnabled;
+
 	if (ensure(LevelOwningWorld))
 	{
 		ULevel* PendingLevelVisOrInvis = (LevelOwningWorld->GetCurrentLevelPendingVisibility() ? LevelOwningWorld->GetCurrentLevelPendingVisibility() : LevelOwningWorld->GetCurrentLevelPendingInvisibility());
