- Remove resources from deferred update list once UpdateDeferredResource is called
- This fixes an issue with double clearing the resource if UpdateResourceImmediate was called which can result in lost data

diff --git a/Engine/Source/Runtime/Engine/Private/TextureRenderTarget2DArray.cpp b/Engine/Source/Runtime/Engine/Private/TextureRenderTarget2DArray.cpp
--- a/Engine/Source/Runtime/Engine/Private/TextureRenderTarget2DArray.cpp
+++ b/Engine/Source/Runtime/Engine/Private/TextureRenderTarget2DArray.cpp
@@ -303,6 +303,8 @@ void FTextureRenderTarget2DArrayResource::UpdateDeferredResource(FRHICommandList
 {
 	LLM_SCOPED_TAG_WITH_OBJECT_IN_SET(Owner->GetOutermost(), ELLMTagSet::Assets);
 
+	RemoveFromDeferredUpdateList();
+
 	if (!bClearRenderTarget)
 	{
 		return;
diff --git a/Engine/Source/Runtime/Engine/Private/TextureRenderTargetCube.cpp b/Engine/Source/Runtime/Engine/Private/TextureRenderTargetCube.cpp
--- a/Engine/Source/Runtime/Engine/Private/TextureRenderTargetCube.cpp
+++ b/Engine/Source/Runtime/Engine/Private/TextureRenderTargetCube.cpp
@@ -345,6 +345,8 @@ void FTextureRenderTargetCubeResource::UpdateDeferredResource(FRHICommandListImm
 {
 	LLM_SCOPED_TAG_WITH_OBJECT_IN_SET(Owner->GetOutermost(), ELLMTagSet::Assets);
 
+	RemoveFromDeferredUpdateList();
+
 	if (!bClearRenderTarget)
 	{
 		return;
diff --git a/Engine/Source/Runtime/Engine/Private/TextureRenderTargetVolume.cpp b/Engine/Source/Runtime/Engine/Private/TextureRenderTargetVolume.cpp
--- a/Engine/Source/Runtime/Engine/Private/TextureRenderTargetVolume.cpp
+++ b/Engine/Source/Runtime/Engine/Private/TextureRenderTargetVolume.cpp
@@ -311,6 +311,8 @@ void FTextureRenderTargetVolumeResource::UpdateDeferredResource(FRHICommandListI
 {
 	LLM_SCOPED_TAG_WITH_OBJECT_IN_SET(Owner->GetOutermost(), ELLMTagSet::Assets);
 
+	RemoveFromDeferredUpdateList();
+
 	if (bClearRenderTarget)
 	{
 		RHICmdList.Transition(FRHITransitionInfo(TextureRHI, ERHIAccess::Unknown, ERHIAccess::RTV));
