Added `IsSlotActive()`method that allows to query if an animation instance's slot is active (Has any montage instance that is not blending out).

#rnx

diff --git a/Engine/Source/Runtime/Engine/Classes/Animation/AnimInstance.h b/Engine/Source/Runtime/Engine/Classes/Animation/AnimInstance.h
--- a/Engine/Source/Runtime/Engine/Classes/Animation/AnimInstance.h
+++ b/Engine/Source/Runtime/Engine/Classes/Animation/AnimInstance.h
@@ -547,6 +547,10 @@ public:
 	/** Return true if this instance playing the slot animation, also returning the montage it is playing on */
 	bool IsPlayingSlotAnimation(const UAnimSequenceBase* Asset, FName SlotNodeName, UAnimMontage*& OutMontage) const;
 
+	/** Return true if this instance has an active montage in the given slot. A UAnimMontage that is playing in the slot and blending out is not determined to be "active". */
+	UFUNCTION(BlueprintPure, Category="Animation|Montage")
+	bool IsSlotActive(FName SlotNodeName) const;
+	
 	/*********************************************************************************************
 	 * AnimMontage
 	 ********************************************************************************************* */
diff --git a/Engine/Source/Runtime/Engine/Private/Animation/AnimInstance.cpp b/Engine/Source/Runtime/Engine/Private/Animation/AnimInstance.cpp
--- a/Engine/Source/Runtime/Engine/Private/Animation/AnimInstance.cpp
+++ b/Engine/Source/Runtime/Engine/Private/Animation/AnimInstance.cpp
@@ -2107,6 +2107,22 @@ bool UAnimInstance::IsPlayingSlotAnimation(const UAnimSequenceBase* Asset, FName
 	return false;
 }
 
+bool UAnimInstance::IsSlotActive(FName SlotNodeName) const
+{
+	if (SlotNodeName != NAME_None)
+	{
+		for (const FAnimMontageInstance* Instance : MontageInstances)
+		{
+			if (Instance && Instance->Montage && Instance->Montage->IsValidSlot(SlotNodeName) && Instance->IsActive())
+			{
+				return true;
+			}
+		}
+	}
+	
+	return false;
+}
+
 float UAnimInstance::Montage_PlayInternal(UAnimMontage* MontageToPlay, const FMontageBlendSettings& BlendInSettings, float InPlayRate /*= 1.f*/, EMontagePlayReturnType ReturnValueType /*= EMontagePlayReturnType::MontageLength*/, float InTimeToStartMontageAt /*= 0.f*/, bool bStopAllMontages /*= true*/)
 {
 	LLM_SCOPE(ELLMTag::Animation);
