Remove UWorld and USystem object iterators from particle perf stats

diff --git a/Engine/Source/Runtime/Engine/Private/Particles/ParticlePerfStats.cpp b/Engine/Source/Runtime/Engine/Private/Particles/ParticlePerfStats.cpp
--- a/Engine/Source/Runtime/Engine/Private/Particles/ParticlePerfStats.cpp
+++ b/Engine/Source/Runtime/Engine/Private/Particles/ParticlePerfStats.cpp
@@ -183,9 +183,12 @@ void FParticlePerfStatsManager::Reset()
 
 	{
 		FScopeLock ScopeLock(&FParticlePerfStatsManager::WorldToPerfStatsGuard);
-		for (TObjectIterator<UWorld> WorldIt; WorldIt; ++WorldIt)
+		for (auto& WorldStatsPair : WorldToPerfStats)
 		{
-			WorldIt->ParticlePerfStats = nullptr;
+			if(const UWorld* World = WorldStatsPair.Key.GetEvenIfUnreachable())
+			{
+				World->ParticlePerfStats = nullptr;
+			}
 		}
 		WorldToPerfStats.Empty();
 	}
@@ -193,9 +196,12 @@ void FParticlePerfStatsManager::Reset()
 	#if WITH_PER_SYSTEM_PARTICLE_PERF_STATS
 	{
 		FScopeLock ScopeLock(&FParticlePerfStatsManager::SystemToPerfStatsGuard);
-		for (TObjectIterator<UFXSystemAsset> SystemIt; SystemIt; ++SystemIt)
+		for (auto& SystemStatsPair : SystemToPerfStats)
 		{
-			SystemIt->ParticlePerfStats = nullptr;
+			if (const UFXSystemAsset* System = SystemStatsPair.Key.GetEvenIfUnreachable())
+			{
+				System->ParticlePerfStats = nullptr;
+			}
 		}
 		SystemToPerfStats.Empty();
 	}
