Add coverage for the FAutomationTestBase::TestNearlyEqual Methods

diff --git a/Engine/Source/Runtime/Engine/Private/Tests/EngineAutomationTests.cpp b/Engine/Source/Runtime/Engine/Private/Tests/EngineAutomationTests.cpp
--- a/Engine/Source/Runtime/Engine/Private/Tests/EngineAutomationTests.cpp
+++ b/Engine/Source/Runtime/Engine/Private/Tests/EngineAutomationTests.cpp
@@ -278,6 +278,562 @@ bool FAutomationLogAddError::RunTest(const FString& Parameters)
 	return true;
 }
 
+class FAutomationNearlyEqualTest : public FAutomationTestBase
+{
+public:
+	FAutomationNearlyEqualTest(const FString& InName, const bool bInComplexTask)
+		: FAutomationTestBase(InName, bInComplexTask)
+	{
+	}
+
+protected:
+	template<typename T>
+	int32 RunSimpleTest(const FString& What, const T ActualValue, const T ExpectedValue, float Tolerance)
+	{
+		int32 CasesCheckedTotal = 0;
+
+		TestNearlyEqual(What, ActualValue, ExpectedValue, Tolerance);
+		++CasesCheckedTotal;
+
+		return CasesCheckedTotal;
+	}
+
+	// Note: this method is used to avoid rising of error C2666 2 overloads have similar conversions
+	template<>
+	int32 RunSimpleTest<double>(const FString& What, double ActualValue, double ExpectedValue, float Tolerance)
+	{
+		int32 CasesCheckedTotal = 0;
+
+		TestNearlyEqual(What, ActualValue, ExpectedValue, static_cast<double>(Tolerance));
+		++CasesCheckedTotal;
+
+		return CasesCheckedTotal;
+	}
+
+	int32 RunFloatMutationTest(const FString& WhatPrefix, float BaseValue, float Difference, float Tolerance)
+	{
+		check(Difference != 0.f);
+		check(Tolerance > 0.f);
+
+		int32 CasesCheckedTotal = 0;
+
+		// Perform tests with mutated values
+		TestNearlyEqual(FString::Format(*ActualValueIsIncreasedByFormatString, { *WhatPrefix, Difference }),
+			BaseValue + Difference, BaseValue, Tolerance);
+		++CasesCheckedTotal;
+
+		TestNearlyEqual(FString::Format(*ExpectedValueIsIncreasedByFormatString, { *WhatPrefix, Difference }),
+			BaseValue, BaseValue + Difference, Tolerance);
+		++CasesCheckedTotal;
+
+		return CasesCheckedTotal;
+	}
+
+	int32 RunDoubleMutationTest(const FString& WhatPrefix, double BaseValue, double Difference, float Tolerance)
+	{
+		check(Difference != 0.f);
+		check(Tolerance > 0.f);
+
+		int32 CasesCheckedTotal = 0;
+
+		// Perform tests with mutated values
+		TestNearlyEqual(FString::Format(*ActualValueIsIncreasedByFormatString, { *WhatPrefix, Difference }),
+			BaseValue + Difference, BaseValue, static_cast<double>(Tolerance));
+		++CasesCheckedTotal;
+		TestNearlyEqual(FString::Format(*ExpectedValueIsIncreasedByFormatString, { *WhatPrefix, Difference }),
+			BaseValue, BaseValue + Difference, static_cast<double>(Tolerance));
+		++CasesCheckedTotal;
+
+		return CasesCheckedTotal;
+	}
+
+	using GetWhatCallable = TFunction<FString(const FString& WhatPrefix, uint32 ActualValueMutationBitMask, 
+		uint32 ExpectedValueMutationBitMask, double Difference)>;
+
+	template<typename T>
+	using GetMutatedValueCallable = TFunction<T(const T& BaseValue, uint32 MutationBitMask, double Difference)>;
+
+	int32 RunFVectorMutationTest(const FString& WhatPrefix, const FVector& BaseValue, double Difference, float Tolerance)
+	{
+		GetWhatCallable GetWhatCallableImpl = [](const FString& WhatPrefix, uint32 ActualValueMutationBitMask, uint32 ExpectedValueMutationBitMask, double Difference)
+		{
+			return FString::Printf(
+				TEXT(
+					"%s: the actual FVector value is not nearly equal to the expected FVector value\n"
+					"(mutation mask for actual value is (%c, %c, %c), mutation mask for expected value is (%c, %c, %c), values were increased by %f)"
+				),
+				*WhatPrefix,
+				GetNthBitAsChar(ActualValueMutationBitMask, 2),
+				GetNthBitAsChar(ActualValueMutationBitMask, 1),
+				GetNthBitAsChar(ActualValueMutationBitMask, 0),
+				GetNthBitAsChar(ExpectedValueMutationBitMask, 2),
+				GetNthBitAsChar(ExpectedValueMutationBitMask, 1),
+				GetNthBitAsChar(ExpectedValueMutationBitMask, 0),
+				Difference);
+		};
+
+		GetMutatedValueCallable<FVector> GetMutatedValueCallableImpl = [](const FVector& BaseValue, uint32 MutationBitMask, double Difference)
+		{
+			return FVector(
+				BaseValue.X + GetNthBitAsUInt32(MutationBitMask, 2) * Difference,
+				BaseValue.Y + GetNthBitAsUInt32(MutationBitMask, 1) * Difference,
+				BaseValue.Z + GetNthBitAsUInt32(MutationBitMask, 0) * Difference
+			);
+		};
+
+		return RunMutationTestImpl<FVector>(WhatPrefix, BaseValue, MaxFVectorMutationBitMask, Difference, Tolerance,
+			GetWhatCallableImpl, GetMutatedValueCallableImpl);
+	}
+
+	int32 RunFRotatorMutationTest(const FString& WhatPrefix, const FRotator& BaseValue, double Difference, float Tolerance)
+	{
+		GetWhatCallable GetWhatCallableImpl = [](const FString& WhatPrefix, uint32 ActualValueMutationBitMask, uint32 ExpectedValueMutationBitMask, double Difference)
+		{
+			return FString::Printf(
+				TEXT(
+					"%s: the actual FRotator value is not nearly equal to the expected FRotator value\n"
+					"(mutation mask for actual value is (%c, %c, %c), mutation mask for expected value is (%c, %c, %c), values were increased by %f)"
+				),
+				*WhatPrefix,
+				GetNthBitAsChar(ActualValueMutationBitMask, 2),
+				GetNthBitAsChar(ActualValueMutationBitMask, 1),
+				GetNthBitAsChar(ActualValueMutationBitMask, 0),
+				GetNthBitAsChar(ExpectedValueMutationBitMask, 2),
+				GetNthBitAsChar(ExpectedValueMutationBitMask, 1),
+				GetNthBitAsChar(ExpectedValueMutationBitMask, 0),
+				Difference);
+		};
+
+		GetMutatedValueCallable<FRotator> GetMutatedValueCallableImpl = [](const FRotator& BaseValue, uint32 MutationBitMask, double Difference)
+		{
+			return FRotator(
+				BaseValue.Pitch + GetNthBitAsUInt32(MutationBitMask, 2) * Difference,
+				BaseValue.Yaw + GetNthBitAsUInt32(MutationBitMask, 1) * Difference,
+				BaseValue.Roll + GetNthBitAsUInt32(MutationBitMask, 0) * Difference
+			);
+		};
+
+		return RunMutationTestImpl<FRotator>(WhatPrefix, BaseValue, MaxFRotatorMutationBitMask, Difference, Tolerance,
+			GetWhatCallableImpl, GetMutatedValueCallableImpl);
+	}
+
+	int32 RunFTransformMutationTest(const FString& WhatPrefix, const FTransform& BaseValue, double Difference, float Tolerance)
+	{
+		GetWhatCallable GetWhatCallableImpl = [](const FString& WhatPrefix, uint32 ActualValueMutationBitMask, uint32 ExpectedValueMutationBitMask, double Difference)
+		{
+			return FString::Printf(
+				TEXT(
+					"%s: the actual FTransform value is not nearly equal to the expected FTransform value\n"
+					"(mutation mask for actual value is (%c, %c, %c, %c, %c, %c, %c, %c, %c), "
+					"mutation mask for expected value is (%c, %c, %c, %c, %c, %c, %c, %c, %c), values were increased by %f)"
+				),
+				*WhatPrefix,
+				GetNthBitAsChar(ActualValueMutationBitMask, 8),
+				GetNthBitAsChar(ActualValueMutationBitMask, 7),
+				GetNthBitAsChar(ActualValueMutationBitMask, 6),
+				GetNthBitAsChar(ActualValueMutationBitMask, 5),
+				GetNthBitAsChar(ActualValueMutationBitMask, 4),
+				GetNthBitAsChar(ActualValueMutationBitMask, 3),
+				GetNthBitAsChar(ActualValueMutationBitMask, 2),
+				GetNthBitAsChar(ActualValueMutationBitMask, 1),
+				GetNthBitAsChar(ActualValueMutationBitMask, 0),
+				GetNthBitAsChar(ExpectedValueMutationBitMask, 8),
+				GetNthBitAsChar(ExpectedValueMutationBitMask, 7),
+				GetNthBitAsChar(ExpectedValueMutationBitMask, 6),
+				GetNthBitAsChar(ExpectedValueMutationBitMask, 5),
+				GetNthBitAsChar(ExpectedValueMutationBitMask, 4),
+				GetNthBitAsChar(ExpectedValueMutationBitMask, 3),
+				GetNthBitAsChar(ExpectedValueMutationBitMask, 2),
+				GetNthBitAsChar(ExpectedValueMutationBitMask, 1),
+				GetNthBitAsChar(ExpectedValueMutationBitMask, 0),
+				Difference);
+		};
+
+		GetMutatedValueCallable<FTransform> GetMutatedValueCallableImpl = [](const FTransform& BaseValue, uint32 MutationBitMask, double Difference)
+		{
+			return FTransform(
+				FRotator(
+					BaseValue.Rotator().Pitch + GetNthBitAsUInt32(MutationBitMask, 8) * Difference,
+					BaseValue.Rotator().Yaw + GetNthBitAsUInt32(MutationBitMask, 7) * Difference,
+					BaseValue.Rotator().Roll + GetNthBitAsUInt32(MutationBitMask, 6) * Difference
+				),
+				FVector(
+					BaseValue.GetTranslation().X + GetNthBitAsUInt32(MutationBitMask, 5) * Difference,
+					BaseValue.GetTranslation().Y + GetNthBitAsUInt32(MutationBitMask, 4) * Difference,
+					BaseValue.GetTranslation().Z + GetNthBitAsUInt32(MutationBitMask, 3) * Difference
+				),
+				FVector(
+					BaseValue.GetScale3D().X + GetNthBitAsUInt32(MutationBitMask, 2) * Difference,
+					BaseValue.GetScale3D().Y + GetNthBitAsUInt32(MutationBitMask, 1) * Difference,
+					BaseValue.GetScale3D().Z + GetNthBitAsUInt32(MutationBitMask, 0) * Difference
+				)
+			);
+		};
+
+		return RunMutationTestImpl<FTransform>(WhatPrefix, BaseValue, MaxFTransformMutationBitMask, Difference, Tolerance,
+			GetWhatCallableImpl, GetMutatedValueCallableImpl);
+	}
+
+	static const float NullTolerance;
+	static const float PositiveTolerance;
+	static const float PositiveDifference;
+	static const float PositiveHalfDifference;
+
+	// Max mutation masks for complex classes/structs
+	// Each bit represents whether (value 1) or not (value 0) mutation will be applied to the object's constructor paramter
+	static const uint32 MaxFVectorMutationBitMask;
+	static const uint32 MaxFRotatorMutationBitMask;
+	static const uint32 MaxFTransformMutationBitMask;
+
+	static const FString TestFailMessage;
+
+	static const float BaseFloatValue;
+	static const float ActualFloatValue;
+	static const float ExpectedFloatValue;
+	static const float ExpectedFloatValueForNullTolerance;
+	static const float FloatDifferenceToGetOutOfTolerance;
+	static const float ExpectedFloatValueOutOfTolerance;
+
+	static const double BaseDoubleValue;
+	static const double ActualDoubleValue;
+	static const double ExpectedDoubleValue;
+	static const double ExpectedDoubleValueForNullTolerance;
+	static const double DoubleDifferenceToGetOutOfTolerance;
+	static const double ExpectedDoubleValueOutOfTolerance;
+
+	static const FVector ActualFVectorValue;
+	static const FVector& ExpectedFVectorValue;
+	static const FVector& BaseFVectorValue;
+
+	static const FRotator ActualFRotatorValue;
+	static const FRotator& ExpectedFRotatorValue;
+	static const FRotator& BaseFRotatorValue;
+
+	static const FTransform ActualFTransformValue;
+	static const FTransform& ExpectedFTransformValue;
+	static const FTransform& BaseFTransformValue;
+
+private:
+	template<typename T>
+	int32 RunMutationTestImpl(const FString& WhatPrefix, const T& BaseValue, uint32 MaxMutationBitMask, float Difference, float Tolerance,
+		GetWhatCallable GetWhatCallableImpl, GetMutatedValueCallable<T> GetMutatedValueCallableImpl)
+	{
+		check(Difference != 0.f);
+		check(Tolerance > 0.f);
+
+		int32 CasesCheckedTotal = 0;
+
+		for (uint32 ActualValueMutationBitMask = 0; ActualValueMutationBitMask <= MaxMutationBitMask; ++ActualValueMutationBitMask)
+		{
+			for (uint32 ExpectedValueMutationBitMask = 0; ExpectedValueMutationBitMask <= MaxMutationBitMask; ++ExpectedValueMutationBitMask)
+			{
+				if (ActualValueMutationBitMask == ExpectedValueMutationBitMask)
+				{
+					// The values' mutation submasks are the same, we should skip this combination
+					continue;
+				}
+
+				// Perform test with mutated values in accordance to the current MutationBitMask
+				const FString WhatMessage(GetWhatCallableImpl(WhatPrefix, ActualValueMutationBitMask, ExpectedValueMutationBitMask, Difference));
+				const T ActualValue(GetMutatedValueCallableImpl(BaseValue, ActualValueMutationBitMask, Difference));
+				const T ExpectedValue(GetMutatedValueCallableImpl(BaseValue, ExpectedValueMutationBitMask, Difference));
+
+				TestNearlyEqual(WhatMessage, ActualValue, ExpectedValue, Tolerance);
+				++CasesCheckedTotal;
+			}
+		}
+
+		return CasesCheckedTotal;
+	}
+
+	static uint32 GetNthBitAsUInt32(uint32 Value, uint32 BitIndex)
+	{
+		return ((Value & (1 << BitIndex)) == 0 ? 0 : 1);
+	}
+
+	static char GetNthBitAsChar(uint32 Value, uint32 BitIndex)
+	{
+		return (GetNthBitAsUInt32(Value, BitIndex) == 1 ? '1' : '0');
+	}
+
+	static const FString ActualValueIsIncreasedByFormatString;
+	static const FString ExpectedValueIsIncreasedByFormatString;
+	static const FString DifferenceAndOrToleranceAreNotValidFormatString;
+};
+
+const float FAutomationNearlyEqualTest::NullTolerance(0.f);
+const float FAutomationNearlyEqualTest::PositiveTolerance(1.e-4f);
+const float FAutomationNearlyEqualTest::PositiveDifference(1.e-4f);
+const float FAutomationNearlyEqualTest::PositiveHalfDifference((1.e-4f) / 2.f);
+const FString FAutomationNearlyEqualTest::TestFailMessage(TEXT("Total amount of errors is not equal to the expected amount"));
+const uint32 FAutomationNearlyEqualTest::MaxFVectorMutationBitMask(0b111);
+const uint32 FAutomationNearlyEqualTest::MaxFRotatorMutationBitMask(0b111);
+const uint32 FAutomationNearlyEqualTest::MaxFTransformMutationBitMask(0b111111111);
+
+const float FAutomationNearlyEqualTest::BaseFloatValue(0.f);
+const float FAutomationNearlyEqualTest::ActualFloatValue(BaseFloatValue);
+const float FAutomationNearlyEqualTest::ExpectedFloatValue(BaseFloatValue);
+const float FAutomationNearlyEqualTest::ExpectedFloatValueForNullTolerance(0.1f);
+const float FAutomationNearlyEqualTest::FloatDifferenceToGetOutOfTolerance(PositiveTolerance + 0.1f);
+const float FAutomationNearlyEqualTest::ExpectedFloatValueOutOfTolerance(ActualFloatValue + FloatDifferenceToGetOutOfTolerance + 0.1f);
+
+const double FAutomationNearlyEqualTest::BaseDoubleValue(0.0);
+const double FAutomationNearlyEqualTest::ActualDoubleValue(BaseDoubleValue);
+const double FAutomationNearlyEqualTest::ExpectedDoubleValue(BaseDoubleValue);
+const double FAutomationNearlyEqualTest::ExpectedDoubleValueForNullTolerance(0.1);
+const double FAutomationNearlyEqualTest::DoubleDifferenceToGetOutOfTolerance(PositiveTolerance + 0.1);
+const double FAutomationNearlyEqualTest::ExpectedDoubleValueOutOfTolerance(ActualDoubleValue + DoubleDifferenceToGetOutOfTolerance);
+
+const FVector FAutomationNearlyEqualTest::ActualFVectorValue(0.f, -1.f, 1.f);
+const FVector& FAutomationNearlyEqualTest::ExpectedFVectorValue(ActualFVectorValue);
+const FVector& FAutomationNearlyEqualTest::BaseFVectorValue(ActualFVectorValue);
+
+const FRotator FAutomationNearlyEqualTest::ActualFRotatorValue(0.001f, -1.002f, 1.003f);
+const FRotator& FAutomationNearlyEqualTest::ExpectedFRotatorValue(ActualFRotatorValue);
+const FRotator& FAutomationNearlyEqualTest::BaseFRotatorValue(ActualFRotatorValue);
+
+const FTransform FAutomationNearlyEqualTest::ActualFTransformValue(FRotator(0.f, -1.f, 1.f), FVector(0.1f, -1.2f, 1.3f), FVector(0.01f, -1.02f, 1.03f));
+const FTransform& FAutomationNearlyEqualTest::ExpectedFTransformValue(ActualFTransformValue);
+const FTransform& FAutomationNearlyEqualTest::BaseFTransformValue(ActualFTransformValue);
+
+const FString FAutomationNearlyEqualTest::ActualValueIsIncreasedByFormatString(TEXT("{0} (actual value is increased by {1})"));
+const FString FAutomationNearlyEqualTest::ExpectedValueIsIncreasedByFormatString(TEXT("{0} (expected value is increased by {1})"));
+const FString FAutomationNearlyEqualTest::DifferenceAndOrToleranceAreNotValidFormatString(TEXT("Difference and/or Tolerance are not valid. Difference: {0}, Tolerance: {1}"));
+
+IMPLEMENT_CUSTOM_SIMPLE_AUTOMATION_TEST(FAutomationTestNearlyEqualFloatPositive, FAutomationNearlyEqualTest, "System.Automation.Validation.TestNearlyEqualFloatPositive", EAutomationTestFlags::EditorContext | EAutomationTestFlags::EngineFilter);
+bool FAutomationTestNearlyEqualFloatPositive::RunTest(const FString& Parameters)
+{
+	//** TEST **//
+	RunSimpleTest<float>(TEXT("The same float values with null tolerance"),
+		ActualFloatValue, ExpectedFloatValue, NullTolerance);
+	RunSimpleTest<float>(TEXT("The same float values with positive tolerance"),
+		ActualFloatValue, ExpectedFloatValue, PositiveTolerance);
+	RunFloatMutationTest(TEXT("Mutation of base float value with the same positive difference and tolerance (edge case)"),
+		BaseFloatValue, PositiveDifference, PositiveTolerance);
+	RunFloatMutationTest(TEXT("Mutation of base float value with negative difference and positive tolerance that are equal after being placed in Abs"),
+		BaseFloatValue, -PositiveDifference, PositiveTolerance);
+	RunFloatMutationTest(TEXT("Mutation of base float value with positive half difference and positive tolerance"),
+		BaseFloatValue, PositiveHalfDifference, PositiveTolerance);
+	RunFloatMutationTest(TEXT("Mutation of base float value with negative half difference and positive tolerance"),
+		BaseFloatValue, -PositiveHalfDifference, PositiveTolerance);
+
+	//** VERIFY **//
+	const int32 ErrorTotal = ExecutionInfo.GetErrorTotal();
+	const int32 ExpectedErrorTotal = 0;
+
+	ExecutionInfo.RemoveAllEvents(EAutomationEventType::Error);
+	TestEqual(TestFailMessage, ErrorTotal, ExpectedErrorTotal);
+
+	return true;
+}
+
+IMPLEMENT_CUSTOM_SIMPLE_AUTOMATION_TEST(FAutomationTestNearlyEqualFloatNegative, FAutomationNearlyEqualTest, "System.Automation.Validation.TestNearlyEqualFloatNegative", EAutomationTestFlags::EditorContext | EAutomationTestFlags::EngineFilter);
+bool FAutomationTestNearlyEqualFloatNegative::RunTest(const FString& Parameters)
+{
+	int32 CasesCheckedTotal = 0;
+
+	//** TEST **//
+	CasesCheckedTotal += RunSimpleTest<float>(TEXT("Different float values with null tolerance"),
+		ActualFloatValue, ExpectedFloatValueForNullTolerance, NullTolerance);
+	CasesCheckedTotal += RunSimpleTest<float>(TEXT("Different float values with positive tolerance"),
+		ActualFloatValue, ExpectedFloatValueOutOfTolerance, PositiveTolerance);
+	CasesCheckedTotal += RunFloatMutationTest(TEXT("Mutation of base float value with positive difference that is greater than positive tolerance"),
+		BaseFloatValue, FloatDifferenceToGetOutOfTolerance, PositiveTolerance);
+	CasesCheckedTotal += RunFloatMutationTest(TEXT("Mutation of base float value with negative difference which absolute value is greater than positive tolerance"),
+		BaseFloatValue, -FloatDifferenceToGetOutOfTolerance, PositiveTolerance);
+
+	//** VERIFY **//
+	const int32 ErrorTotal = ExecutionInfo.GetErrorTotal();
+	ExecutionInfo.RemoveAllEvents(EAutomationEventType::Error);
+	TestEqual(TestFailMessage, ErrorTotal, CasesCheckedTotal);
+
+	return true;
+}
+
+IMPLEMENT_CUSTOM_SIMPLE_AUTOMATION_TEST(FAutomationTestNearlyEqualDoublePositive, FAutomationNearlyEqualTest, "System.Automation.Validation.TestNearlyEqualDoublePositive", EAutomationTestFlags::EditorContext | EAutomationTestFlags::EngineFilter);
+bool FAutomationTestNearlyEqualDoublePositive::RunTest(const FString& Parameters)
+{
+	//** TEST **//
+	RunSimpleTest<double>(TEXT("The same double values with null tolerance"),
+		ActualDoubleValue, ExpectedDoubleValue, NullTolerance);
+	RunSimpleTest<double>(TEXT("The same double values with positive tolerance"),
+		ActualDoubleValue, ExpectedDoubleValue, PositiveTolerance);
+	RunDoubleMutationTest(TEXT("Mutation of base double value with the same positive difference and tolerance (edge case)"),
+		BaseDoubleValue, PositiveDifference, PositiveTolerance);
+	RunDoubleMutationTest(TEXT("Mutation of base double value with negative difference and positive tolerance that are equal after being placed in Abs"),
+		BaseDoubleValue, -PositiveDifference, PositiveTolerance);
+	RunDoubleMutationTest(TEXT("Mutation of base double value with positive half difference and positive tolerance"),
+		BaseDoubleValue, PositiveHalfDifference, PositiveTolerance);
+	RunDoubleMutationTest(TEXT("Mutation of base double value with negative half difference and positive tolerance"),
+		BaseDoubleValue, -PositiveHalfDifference, PositiveTolerance);
+
+	//** VERIFY **//
+	const int32 ErrorTotal = ExecutionInfo.GetErrorTotal();
+	const int32 ExpectedErrorTotal = 0;
+
+	TestEqual(TestFailMessage, ErrorTotal, ExpectedErrorTotal);
+
+	return true;
+}
+
+IMPLEMENT_CUSTOM_SIMPLE_AUTOMATION_TEST(FAutomationTestNearlyEqualDoubleNegative, FAutomationNearlyEqualTest, "System.Automation.Validation.TestNearlyEqualDoubleNegative", EAutomationTestFlags::EditorContext | EAutomationTestFlags::EngineFilter);
+bool FAutomationTestNearlyEqualDoubleNegative::RunTest(const FString& Parameters)
+{
+	int32 CasesCheckedTotal = 0;
+
+	//** TEST **//
+	CasesCheckedTotal += RunSimpleTest<double>(TEXT("Different double values with null tolerance"),
+		ActualDoubleValue, ExpectedDoubleValueForNullTolerance, NullTolerance);
+	CasesCheckedTotal += RunSimpleTest<double>(TEXT("Different double values with positive tolerance"),
+		ActualDoubleValue, ExpectedDoubleValueOutOfTolerance, PositiveTolerance);
+	CasesCheckedTotal += RunDoubleMutationTest(TEXT("Mutation of base double value with positive difference that is greater than positive tolerance"),
+		BaseDoubleValue, DoubleDifferenceToGetOutOfTolerance, PositiveTolerance);
+	CasesCheckedTotal += RunDoubleMutationTest(TEXT("Mutation of base double value with negative difference which absolute value is greater than positive tolerance"),
+		BaseDoubleValue, -DoubleDifferenceToGetOutOfTolerance, PositiveTolerance);
+
+	//** VERIFY **//
+	const int32 ErrorTotal = ExecutionInfo.GetErrorTotal();
+	ExecutionInfo.RemoveAllEvents(EAutomationEventType::Error);
+	TestEqual(TestFailMessage, ErrorTotal, CasesCheckedTotal);
+
+	return true;
+}
+
+IMPLEMENT_CUSTOM_SIMPLE_AUTOMATION_TEST(FAutomationTestNearlyEqualFVectorPositive, FAutomationNearlyEqualTest, "System.Automation.Validation.TestNearlyEqualFVectorPositive", EAutomationTestFlags::EditorContext | EAutomationTestFlags::EngineFilter);
+bool FAutomationTestNearlyEqualFVectorPositive::RunTest(const FString& Parameters)
+{
+	//** TEST **//
+	RunSimpleTest<FVector>(TEXT("The same FVector values with null tolerance"),
+		ActualFVectorValue, ExpectedFVectorValue, NullTolerance);
+	RunSimpleTest<FVector>(TEXT("The same FVector values with positive tolerance"),
+		ActualFVectorValue, ExpectedFVectorValue, PositiveTolerance);
+	RunFVectorMutationTest(TEXT("Mutation of base FVector value with the same positive difference and tolerance (edge case)"),
+		BaseFVectorValue, PositiveDifference, PositiveTolerance);
+	RunFVectorMutationTest(TEXT("Mutation of base FVector value with negative difference and positive tolerance that are equal after being placed in Abs"),
+		BaseFVectorValue, -PositiveDifference, PositiveTolerance);
+	RunFVectorMutationTest(TEXT("Mutation of base FVector value with positive half difference and positive tolerance"),
+		BaseFVectorValue, PositiveHalfDifference, PositiveTolerance);
+	RunFVectorMutationTest(TEXT("Mutation of base FVector value with negative half difference and positive tolerance"),
+		BaseFVectorValue, -PositiveHalfDifference, PositiveTolerance);
+
+	//** VERIFY **//
+	const int32 ErrorTotal = ExecutionInfo.GetErrorTotal();
+	const int32 ExpectedErrorTotal = 0;
+
+	TestEqual(TestFailMessage, ErrorTotal, ExpectedErrorTotal);
+
+	return true;
+}
+
+IMPLEMENT_CUSTOM_SIMPLE_AUTOMATION_TEST(FAutomationTestNearlyEqualFVectorNegative, FAutomationNearlyEqualTest, "System.Automation.Validation.TestNearlyEqualFVectorNegative", EAutomationTestFlags::EditorContext | EAutomationTestFlags::EngineFilter);
+bool FAutomationTestNearlyEqualFVectorNegative::RunTest(const FString& Parameters)
+{
+	int32 CasesCheckedTotal = 0;
+
+	//** TEST **//
+	CasesCheckedTotal += RunFVectorMutationTest(TEXT("Mutation of base FVector value with positive difference that is greater than positive tolerance"),
+		BaseFVectorValue, PositiveDifference + 0.1f, PositiveTolerance);
+	CasesCheckedTotal += RunFVectorMutationTest(TEXT("Mutation of base FVector value with negative difference which absolute value is greater than positive tolerance"),
+		BaseFVectorValue, -PositiveDifference - 0.1f, PositiveTolerance);
+
+	//** VERIFY **//
+	const int32 ErrorTotal = ExecutionInfo.GetErrorTotal();
+	ExecutionInfo.RemoveAllEvents(EAutomationEventType::Error);
+	TestEqual(TestFailMessage, ErrorTotal, CasesCheckedTotal);
+
+	return true;
+}
+
+IMPLEMENT_CUSTOM_SIMPLE_AUTOMATION_TEST(FAutomationTestNearlyEqualFRotatorPositive, FAutomationNearlyEqualTest, "System.Automation.Validation.TestNearlyEqualFRotatorPositive", EAutomationTestFlags::EditorContext | EAutomationTestFlags::EngineFilter);
+bool FAutomationTestNearlyEqualFRotatorPositive::RunTest(const FString& Parameters)
+{
+	//** TEST **//
+	RunSimpleTest<FRotator>(TEXT("The same FRotator values with null tolerance"),
+		ActualFRotatorValue, ExpectedFRotatorValue, NullTolerance);
+	RunSimpleTest<FRotator>(TEXT("The same FRotator values with positive tolerance"),
+		ActualFRotatorValue, ExpectedFRotatorValue, PositiveTolerance);
+	RunFRotatorMutationTest(TEXT("Mutation of base FRotator value with the same positive difference and tolerance (edge case)"),
+		BaseFRotatorValue, PositiveDifference, PositiveTolerance);
+	RunFRotatorMutationTest(TEXT("Mutation of base FRotator value with negative difference and positive tolerance that are equal after being placed in Abs"),
+		BaseFRotatorValue, -PositiveDifference, PositiveTolerance);
+	RunFRotatorMutationTest(TEXT("Mutation of base FRotator value with positive half difference and positive tolerance"),
+		BaseFRotatorValue, PositiveHalfDifference, PositiveTolerance);
+	RunFRotatorMutationTest(TEXT("Mutation of base FRotator value with negative half difference and positive tolerance"),
+		BaseFRotatorValue, -PositiveHalfDifference, PositiveTolerance);
+
+	//** VERIFY **//
+	const int32 ErrorTotal = ExecutionInfo.GetErrorTotal();
+	const int32 ExpectedErrorTotal = 0;
+
+	TestEqual(TestFailMessage, ErrorTotal, ExpectedErrorTotal);
+
+	return true;
+}
+
+IMPLEMENT_CUSTOM_SIMPLE_AUTOMATION_TEST(FAutomationTestNearlyEqualFRotatorNegative, FAutomationNearlyEqualTest, "System.Automation.Validation.TestNearlyEqualFRotatorNegative", EAutomationTestFlags::EditorContext | EAutomationTestFlags::EngineFilter);
+bool FAutomationTestNearlyEqualFRotatorNegative::RunTest(const FString& Parameters)
+{
+	int32 CasesCheckedTotal = 0;
+
+	//** TEST **//
+	CasesCheckedTotal += RunFRotatorMutationTest(TEXT("Mutation of base FRotator value with positive difference that is greater than positive tolerance"),
+		BaseFRotatorValue, PositiveDifference + 1, PositiveTolerance);
+	CasesCheckedTotal += RunFRotatorMutationTest(TEXT("Mutation of base FRotator value with negative difference which absolute value is greater than positive tolerance"),
+		BaseFRotatorValue, -PositiveDifference - 1, PositiveTolerance);
+
+	//** VERIFY **//
+	const int32 ErrorTotal = ExecutionInfo.GetErrorTotal();
+	ExecutionInfo.RemoveAllEvents(EAutomationEventType::Error);
+	TestEqual(TestFailMessage, ErrorTotal, CasesCheckedTotal);
+
+	return true;
+}
+
+IMPLEMENT_CUSTOM_SIMPLE_AUTOMATION_TEST(FAutomationTestNearlyEqualFTransformPositive, FAutomationNearlyEqualTest, "System.Automation.Validation.TestNearlyEqualFTransformPositive", EAutomationTestFlags::EditorContext | EAutomationTestFlags::EngineFilter);
+bool FAutomationTestNearlyEqualFTransformPositive::RunTest(const FString& Parameters)
+{
+	//** TEST **//
+	RunSimpleTest<FTransform>(TEXT("The same FTransform values with null tolerance"),
+		ActualFTransformValue, ExpectedFTransformValue, NullTolerance);
+	RunSimpleTest<FTransform>(TEXT("The same FTransform values with positive tolerance"),
+		ActualFTransformValue, ExpectedFTransformValue, PositiveTolerance);
+	RunFTransformMutationTest(TEXT("Mutation of base FTransform value with the same positive difference and tolerance (edge case)"),
+		BaseFTransformValue, PositiveDifference, PositiveTolerance);
+	RunFTransformMutationTest(TEXT("Mutation of base FTransform value with negative difference and positive tolerance that are equal after being placed in Abs"),
+		BaseFTransformValue, -PositiveDifference, PositiveTolerance);
+	RunFTransformMutationTest(TEXT("Mutation of base FTransform value with positive half difference and positive tolerance"),
+		BaseFTransformValue, PositiveHalfDifference, PositiveTolerance);
+	RunFTransformMutationTest(TEXT("Mutation of base FTransform value with negative half difference and positive tolerance"),
+		BaseFTransformValue, -PositiveHalfDifference, PositiveTolerance);
+
+	//** VERIFY **//
+	const int32 ErrorTotal = ExecutionInfo.GetErrorTotal();
+	const int32 ExpectedErrorTotal = 0;
+
+	TestEqual(TestFailMessage, ErrorTotal, ExpectedErrorTotal);
+
+	return true;
+}
+
+IMPLEMENT_CUSTOM_SIMPLE_AUTOMATION_TEST(FAutomationTestNearlyEqualFTransformNegative, FAutomationNearlyEqualTest, "System.Automation.Validation.TestNearlyEqualFTransformNegative", EAutomationTestFlags::EditorContext | EAutomationTestFlags::EngineFilter);
+bool FAutomationTestNearlyEqualFTransformNegative::RunTest(const FString& Parameters)
+{
+	int32 CasesCheckedTotal = 0;
+
+	//** TEST **//
+	CasesCheckedTotal += RunFTransformMutationTest(TEXT("Mutation of base FTransform value with positive difference that is greater than positive tolerance"),
+		BaseFTransformValue, PositiveDifference + 0.1f, PositiveTolerance);
+	CasesCheckedTotal += RunFTransformMutationTest(TEXT("Mutation of base FTransform value with negative difference which absolute value is greater than positive tolerance"),
+		BaseFTransformValue, -PositiveDifference - 0.1f, PositiveTolerance);
+
+	//** VERIFY **//
+	const int32 ErrorTotal = ExecutionInfo.GetErrorTotal();
+	ExecutionInfo.RemoveAllEvents(EAutomationEventType::Error);
+	TestEqual(TestFailMessage, ErrorTotal, CasesCheckedTotal);
+
+	return true;
+}
+
 IMPLEMENT_SIMPLE_AUTOMATION_TEST(FAutomationAttachment, "System.Engine.Attachment", EAutomationTestFlags::EditorContext | EAutomationTestFlags::EngineFilter)
 
 #define DUMP_EXPECTED_TRANSFORMS 0
