RayTracing: Remove RayTracingDebug from the list of debug view modes

This allows manipulators to be visible and restores the post-processing on the view which matches what the buffer visualization mode does (which is not classified as debug view mode either).

diff --git a/Engine/Source/Editor/UnrealEd/Private/EditorBuildUtils.cpp b/Engine/Source/Editor/UnrealEd/Private/EditorBuildUtils.cpp
--- a/Engine/Source/Editor/UnrealEd/Private/EditorBuildUtils.cpp
+++ b/Engine/Source/Editor/UnrealEd/Private/EditorBuildUtils.cpp
@@ -1369,8 +1369,6 @@ EDebugViewShaderMode ViewModeIndexToDebugViewShaderMode(EViewModeIndex SelectedV
 		return DVSM_RequiredTextureResolution;
 	case VMI_VirtualTexturePendingMips:
 		return DVSM_VirtualTexturePendingMips;
-	case VMI_RayTracingDebug:
-		return DVSM_RayTracingDebug;
 	case VMI_LODColoration:
 	case VMI_HLODColoration:
 		return DVSM_LODColoration;
diff --git a/Engine/Source/Runtime/Engine/Private/DebugViewModeHelpers.cpp b/Engine/Source/Runtime/Engine/Private/DebugViewModeHelpers.cpp
--- a/Engine/Source/Runtime/Engine/Private/DebugViewModeHelpers.cpp
+++ b/Engine/Source/Runtime/Engine/Private/DebugViewModeHelpers.cpp
@@ -47,8 +47,6 @@ const TCHAR* DebugViewShaderModeToString(EDebugViewShaderMode InShaderMode)
 		return TEXT("DVSM_RequiredTextureResolution");
 	case DVSM_VirtualTexturePendingMips:
 		return TEXT("DVSM_VirtualTexturePendingMips");
-	case DVSM_RayTracingDebug:
-		return TEXT("DVSM_RayTracingDebug");
 	case DVSM_LODColoration:
 		return TEXT("DVSM_LODColoration");
 	case DVSM_VisualizeGPUSkinCache:
@@ -98,8 +96,6 @@ bool AllowDebugViewShaderMode(EDebugViewShaderMode ShaderMode, EShaderPlatform P
 	case DVSM_OutputMaterialTextureScales:
 	case DVSM_VirtualTexturePendingMips:
 		return FeatureLevel >= ERHIFeatureLevel::SM5 && (bForceTextureStreamingBuild || PlatformSupportsDebugViewShaders(Platform));
-	case DVSM_RayTracingDebug:
-		return FeatureLevel >= ERHIFeatureLevel::SM5 ;
 	case DVSM_VisualizeGPUSkinCache:
 		return PlatformSupportsDebugViewShaders(Platform);
 	default:
diff --git a/Engine/Source/Runtime/Engine/Private/SceneView.cpp b/Engine/Source/Runtime/Engine/Private/SceneView.cpp
--- a/Engine/Source/Runtime/Engine/Private/SceneView.cpp
+++ b/Engine/Source/Runtime/Engine/Private/SceneView.cpp
@@ -3061,10 +3061,6 @@ EDebugViewShaderMode FSceneViewFamily::ChooseDebugViewShaderMode() const
 	{
 		return DVSM_VirtualTexturePendingMips;
 	}
-	else if (EngineShowFlags.RayTracingDebug)
-	{
-		return DVSM_RayTracingDebug;
-	}
 	else if (EngineShowFlags.LODColoration || EngineShowFlags.HLODColoration)
 	{
 		return DVSM_LODColoration;
diff --git a/Engine/Source/Runtime/Engine/Public/DebugViewModeHelpers.h b/Engine/Source/Runtime/Engine/Public/DebugViewModeHelpers.h
--- a/Engine/Source/Runtime/Engine/Public/DebugViewModeHelpers.h
+++ b/Engine/Source/Runtime/Engine/Public/DebugViewModeHelpers.h
@@ -38,7 +38,6 @@ enum EDebugViewShaderMode
 	DVSM_OutputMaterialTextureScales,  // Outputs the material texture scales.
 	DVSM_RequiredTextureResolution, // Visualize the accuracy of the streamed texture resolution.
 	DVSM_VirtualTexturePendingMips,	// Visualize the pending virtual texture mips.
-	DVSM_RayTracingDebug,			// Visualize ray tracing debug modes.
 	DVSM_LODColoration,				// Visualize primitive LOD .
 	DVSM_VisualizeGPUSkinCache,		// Visualize various properties of Skin Cache.
 	DVSM_MAX
diff --git a/Engine/Source/Runtime/Renderer/Private/PostProcess/PostProcessing.cpp b/Engine/Source/Runtime/Renderer/Private/PostProcess/PostProcessing.cpp
--- a/Engine/Source/Runtime/Renderer/Private/PostProcess/PostProcessing.cpp
+++ b/Engine/Source/Runtime/Renderer/Private/PostProcess/PostProcessing.cpp
@@ -1590,7 +1590,6 @@ void AddDebugViewPostProcessingPasses(FRDGBuilder& GraphBuilder, const FViewInfo
 			SceneColor = AddStreamingAccuracyLegendPass(GraphBuilder, View, PassInputs);
 			break;
 		}
-		case DVSM_RayTracingDebug:
 		case DVSM_VisualizeGPUSkinCache:
 		{
 			FTAAPassParameters Parameters(View);
