Backout CL 23957362 to re-enable DeprecatedFunction in DataLayerSubsystem now that data has been fixed.

diff --git a/Engine/Source/Runtime/Engine/Public/WorldPartition/DataLayer/DataLayerSubsystem.h b/Engine/Source/Runtime/Engine/Public/WorldPartition/DataLayer/DataLayerSubsystem.h
--- a/Engine/Source/Runtime/Engine/Public/WorldPartition/DataLayer/DataLayerSubsystem.h
+++ b/Engine/Source/Runtime/Engine/Public/WorldPartition/DataLayer/DataLayerSubsystem.h
@@ -36,20 +36,20 @@ public:
 
 	/** Find a Data Layer by its asset. */
 	UE_DEPRECATED(5.3, "Use UDataLayerManager::GetDataLayerInstanceFromAsset instead.")
-	UFUNCTION(BlueprintCallable, Category = DataLayers)
+	UFUNCTION(BlueprintCallable, Category = DataLayers, meta = (DeprecatedFunction, DeprecationMessage = "DataLayerSubsystem is deprecated, use the equivalent function in DataLayerManager"))
 	UDataLayerInstance* GetDataLayerInstanceFromAsset(const UDataLayerAsset* InDataLayerAsset) const;
 
 	UE_DEPRECATED(5.3, "Use UDataLayerManager::GetDataLayerInstanceRuntimeState instead.")
-	UFUNCTION(BlueprintCallable, Category = DataLayers)
+	UFUNCTION(BlueprintCallable, Category = DataLayers, meta = (DeprecatedFunction, DeprecationMessage = "DataLayerSubsystem is deprecated, use the equivalent function in DataLayerManager"))
 	EDataLayerRuntimeState GetDataLayerInstanceRuntimeState(const UDataLayerAsset* InDataLayerAsset) const;
 
 	UE_DEPRECATED(5.3, "Use UDataLayerManager::GetDataLayerInstanceEffectiveRuntimeState instead.")
-	UFUNCTION(BlueprintCallable, Category = DataLayers)
+	UFUNCTION(BlueprintCallable, Category = DataLayers, meta = (DeprecatedFunction, DeprecationMessage = "DataLayerSubsystem is deprecated, use the equivalent function in DataLayerManager"))
 	EDataLayerRuntimeState GetDataLayerInstanceEffectiveRuntimeState(const UDataLayerAsset* InDataLayerAsset) const;
 
 	/** Set the Data Layer state using its name. */
 	UE_DEPRECATED(5.3, "Use UDataLayerManager::SetDataLayerInstanceRuntimeState instead.")
-	UFUNCTION(BlueprintCallable, Category = DataLayers)
+	UFUNCTION(BlueprintCallable, Category = DataLayers, BlueprintAuthorityOnly, meta = (DeprecatedFunction, DeprecationMessage = "DataLayerSubsystem is deprecated, use the equivalent function in DataLayerManager"))
 	void SetDataLayerInstanceRuntimeState(const UDataLayerAsset* InDataLayerAsset, EDataLayerRuntimeState InState, bool bInIsRecursive = false);
 
 	/** Called when a Data Layer changes state. */
@@ -198,39 +198,39 @@ public:
 	const TSet<FName>& GetLoadedDataLayerNames() const { static TSet<FName> EmptySet; return EmptySet; }
 
 	UE_DEPRECATED(5.1, "Use GetDataLayerInstanceFromAsset instead.")
-	UFUNCTION(BlueprintCallable, Category = DataLayers)
+	UFUNCTION(BlueprintCallable, Category = DataLayers, meta = (DeprecatedFunction, DeprecationMessage = "DataLayerSubsystem is deprecated, use GetDataLayerInstanceFromAsset in DataLayerManager"))
 	UDataLayerInstance* GetDataLayer(const FActorDataLayer& InDataLayer) const;
 
 	UE_DEPRECATED(5.1, "Use GetDataLayerInstanceFromAsset instead.")
-	UFUNCTION(BlueprintCallable, Category = DataLayers)
+	UFUNCTION(BlueprintCallable, Category = DataLayers, meta = (DeprecatedFunction, DeprecationMessage = "DataLayerSubsystem is deprecated, use GetDataLayerInstanceFromAsset in DataLayerManager"))
 	UDataLayerInstance* GetDataLayerFromName(FName InDataLayerName) const;
 
 	UE_DEPRECATED(5.1, "Use GetDataLayerInstanceFromAsset instead.")
-	UFUNCTION(BlueprintCallable, Category = DataLayers)
+	UFUNCTION(BlueprintCallable, Category = DataLayers, meta = (DeprecatedFunction, DeprecationMessage = "DataLayerSubsystem is deprecated, use GetDataLayerInstanceFromAsset in DataLayerManager"))
 	UDataLayerInstance* GetDataLayerFromLabel(FName InDataLayerLabel) const;
 
 	UE_DEPRECATED(5.1, "Use GetDataLayerInstanceRuntimeState instead.")
-	UFUNCTION(BlueprintCallable, Category = DataLayers)
+	UFUNCTION(BlueprintCallable, Category = DataLayers, meta = (DeprecatedFunction, DeprecationMessage = "DataLayerSubsystem is deprecated, use GetDataLayerInstanceRuntimeState in DataLayerManager"))
 	EDataLayerRuntimeState GetDataLayerRuntimeState(const FActorDataLayer& InDataLayer) const;
 
 	UE_DEPRECATED(5.1, "Use GetDataLayerInstanceRuntimeState instead.")
-	UFUNCTION(BlueprintCallable, Category = DataLayers)
+	UFUNCTION(BlueprintCallable, Category = DataLayers, meta = (DeprecatedFunction, DeprecationMessage = "DataLayerSubsystem is deprecated, use GetDataLayerInstanceRuntimeState in DataLayerManager"))
 	EDataLayerRuntimeState GetDataLayerRuntimeStateByLabel(const FName& InDataLayerLabel) const;
 
 	UE_DEPRECATED(5.1, "Use GetDataLayerInstanceEffectiveRuntimeState instead.")
-	UFUNCTION(BlueprintCallable, Category = DataLayers)
+	UFUNCTION(BlueprintCallable, Category = DataLayers, meta = (DeprecatedFunction, DeprecationMessage = "DataLayerSubsystem is deprecated, use GetDataLayerInstanceEffectiveRuntimeState in DataLayerManager"))
 	EDataLayerRuntimeState GetDataLayerEffectiveRuntimeState(const FActorDataLayer& InDataLayer) const;
 
 	UE_DEPRECATED(5.1, "Use GetDataLayerInstanceEffectiveRuntimeState instead.")
-	UFUNCTION(BlueprintCallable, Category = DataLayers)
+	UFUNCTION(BlueprintCallable, Category = DataLayers, meta = (DeprecatedFunction, DeprecationMessage = "DataLayerSubsystem is deprecated, use GetDataLayerInstanceEffectiveRuntimeState in DataLayerManager"))
 	EDataLayerRuntimeState GetDataLayerEffectiveRuntimeStateByLabel(const FName& InDataLayerLabel) const;
 
 	UE_DEPRECATED(5.1, "Use SetDataLayerRuntimeState() with UDataLayerAsset* overload instead.")
-	UFUNCTION(BlueprintCallable, Category = DataLayers, BlueprintAuthorityOnly)
+	UFUNCTION(BlueprintCallable, Category = DataLayers, BlueprintAuthorityOnly, meta = (DeprecatedFunction, DeprecationMessage = "DataLayerSubsystem is deprecated, use SetDataLayerInstanceRuntimeState in DataLayerManager"))
 	void SetDataLayerRuntimeState(const FActorDataLayer& InDataLayer, EDataLayerRuntimeState InState, bool bInIsRecursive = false);
 
 	UE_DEPRECATED(5.1, "Use SetDataLayerInstanceRuntimeState instead.")
-	UFUNCTION(BlueprintCallable, Category = DataLayers, BlueprintAuthorityOnly)
+	UFUNCTION(BlueprintCallable, Category = DataLayers, BlueprintAuthorityOnly, meta = (DeprecatedFunction, DeprecationMessage = "DataLayerSubsystem is deprecated, use SetDataLayerInstanceRuntimeState in DataLayerManager"))
 	void SetDataLayerRuntimeStateByLabel(const FName& InDataLayerLabel, EDataLayerRuntimeState InState, bool bInIsRecursive = false);
 
 private:
