Removed DeprecatedFunction on DataLayerSubsystem SetDataLayerRuntimeState/SetDataLayerRuntimeStateByLabel as it is still used by Lyra.

diff --git a/Engine/Source/Runtime/Engine/Public/WorldPartition/DataLayer/DataLayerSubsystem.h b/Engine/Source/Runtime/Engine/Public/WorldPartition/DataLayer/DataLayerSubsystem.h
--- a/Engine/Source/Runtime/Engine/Public/WorldPartition/DataLayer/DataLayerSubsystem.h
+++ b/Engine/Source/Runtime/Engine/Public/WorldPartition/DataLayer/DataLayerSubsystem.h
@@ -226,11 +226,11 @@ public:
 	EDataLayerRuntimeState GetDataLayerEffectiveRuntimeStateByLabel(const FName& InDataLayerLabel) const;
 
 	UE_DEPRECATED(5.1, "Use SetDataLayerRuntimeState() with UDataLayerAsset* overload instead.")
-	UFUNCTION(BlueprintCallable, Category = DataLayers, BlueprintAuthorityOnly, meta = (DeprecatedFunction, DeprecationMessage = "DataLayerSubsystem is deprecated, use SetDataLayerInstanceRuntimeState in DataLayerManager"))
+	UFUNCTION(BlueprintCallable, Category = DataLayers, BlueprintAuthorityOnly)
 	void SetDataLayerRuntimeState(const FActorDataLayer& InDataLayer, EDataLayerRuntimeState InState, bool bInIsRecursive = false);
 
 	UE_DEPRECATED(5.1, "Use SetDataLayerInstanceRuntimeState instead.")
-	UFUNCTION(BlueprintCallable, Category = DataLayers, BlueprintAuthorityOnly, meta = (DeprecatedFunction, DeprecationMessage = "DataLayerSubsystem is deprecated, use SetDataLayerInstanceRuntimeState in DataLayerManager"))
+	UFUNCTION(BlueprintCallable, Category = DataLayers, BlueprintAuthorityOnly)
 	void SetDataLayerRuntimeStateByLabel(const FName& InDataLayerLabel, EDataLayerRuntimeState InState, bool bInIsRecursive = false);
 
 private:
