Fixed erroneous deprecation of BreakVector3f in CL 23992905.

For float-based math functions, we've normally been using UE_DEPRECATED on the function, along with a redirect entry, which should've been done initially.

This should correctly update any existing BP content that uses BreakVector3f.

diff --git a/Engine/Config/BaseEngine.ini b/Engine/Config/BaseEngine.ini
--- a/Engine/Config/BaseEngine.ini
+++ b/Engine/Config/BaseEngine.ini
@@ -1126,6 +1126,9 @@ DefaultWorldPartitionHLODLayer=/Engine/Maps/Templates/HLODs/HLODLayer_Instanced.
 +PropertyRedirects=(OldName="/Script/NavigationSystem.NavigationSystemV1.GeometryExportVertexCountWarningThreshold",NewName="/Script/NavigationSystem.NavigationSystemV1.GeometryExportTriangleCountWarningThreshold")
 +PropertyRedirects=(OldName="FISMComponentDescriptor.bIsLocalToWorldDeterminantNegative", NewName="FISMComponentDescriptor.bReverseCulling")
 
+; 5.3
++FunctionRedirects=(OldName="KismetMathLibrary.BreakVector3f", NewName="BreakVector")
+
 [CoreUObject.Metadata]
 ; Note: UnrealHeaderTool should be rerun after making changes to MetadataRedirects to catch any keys specified in class headers
 +MetadataRedirects=(OldKey="K2Protected", NewKey="BlueprintProtected")
diff --git a/Engine/Source/Runtime/Engine/Classes/Kismet/KismetMathLibrary.h b/Engine/Source/Runtime/Engine/Classes/Kismet/KismetMathLibrary.h
--- a/Engine/Source/Runtime/Engine/Classes/Kismet/KismetMathLibrary.h
+++ b/Engine/Source/Runtime/Engine/Classes/Kismet/KismetMathLibrary.h
@@ -1334,7 +1334,7 @@ class ENGINE_API UKismetMathLibrary : public UBlueprintFunctionLibrary
 	UFUNCTION(BlueprintPure, Category="Math|Vector", meta=(NativeBreakFunc))
 	static void BreakVector(FVector InVec, double& X, double& Y, double& Z);
 
-	UFUNCTION(BlueprintPure, Category = "Math|Vector", meta=(NativeBreakFunc, DeprecatedFunction, DeprecationMessage = "Explicit conversions between floats and doubles are not necessary. Please remove node."))
+	UE_DEPRECATED(5.3, "This method has been deprecated and will be removed. Use the double version instead.")
 	static void BreakVector3f(FVector3f InVec, float& X, float& Y, float& Z);
 
 	/** Converts a vector to LinearColor */
