Fix Regression: Non-WP OFPA Levels do not have a DataLayerManager

diff --git a/Engine/Source/Runtime/Engine/Private/WorldPartition/WorldPartitionActorDesc.cpp b/Engine/Source/Runtime/Engine/Private/WorldPartition/WorldPartitionActorDesc.cpp
--- a/Engine/Source/Runtime/Engine/Private/WorldPartition/WorldPartitionActorDesc.cpp
+++ b/Engine/Source/Runtime/Engine/Private/WorldPartition/WorldPartitionActorDesc.cpp
@@ -13,6 +13,7 @@
 #include "UObject/UE5ReleaseStreamObjectVersion.h"
 #include "UObject/FortniteNCBranchObjectVersion.h"
 #include "UObject/FortniteMainBranchObjectVersion.h"
+#include "WorldPartition/WorldPartitionHelpers.h"
 #include "WorldPartition/WorldPartitionActorDescView.h"
 #include "WorldPartition/WorldPartitionLog.h"
 #include "WorldPartition/ActorDescContainer.h"
@@ -69,50 +70,60 @@ void FWorldPartitionActorDesc::Init(const AActor* InActor)
 		TArray<FName> LocalDataLayerAssetPaths;
 		TArray<FName> LocalDataLayerInstanceNames;
 		// Use Actor's DataLayerManager since the fixup is relative to this level
-		UDataLayerManager* DataLayerManager = UDataLayerManager::GetDataLayerManager(InActor);
-		if (ensure(DataLayerManager))
+		if (UWorldPartition* WorldPartition = FWorldPartitionHelpers::GetWorldPartition(InActor))
 		{
-			LocalDataLayerAssetPaths.Reserve(InActor->GetDataLayerAssets().Num());
-			for (const TObjectPtr<const UDataLayerAsset>& DataLayerAsset : InActor->GetDataLayerAssets())
+			UDataLayerManager* DataLayerManager = WorldPartition->GetDataLayerManager();
+			if (ensure(DataLayerManager))
 			{
-				// Pass Actor Level when resolving the DataLayerInstance as FWorldPartitionActorDesc always represents the state of the actor local to its outer level
-				if (DataLayerAsset && DataLayerManager->GetDataLayerInstance(DataLayerAsset))
+				LocalDataLayerAssetPaths.Reserve(InActor->GetDataLayerAssets().Num());
+				for (const TObjectPtr<const UDataLayerAsset>& DataLayerAsset : InActor->GetDataLayerAssets())
 				{
-					LocalDataLayerAssetPaths.Add(*DataLayerAsset->GetPathName());
+					// Pass Actor Level when resolving the DataLayerInstance as FWorldPartitionActorDesc always represents the state of the actor local to its outer level
+					if (DataLayerAsset && DataLayerManager->GetDataLayerInstance(DataLayerAsset))
+					{
+						LocalDataLayerAssetPaths.Add(*DataLayerAsset->GetPathName());
+					}
 				}
-			}
 
-			PRAGMA_DISABLE_DEPRECATION_WARNINGS
-			// Pass Actor Level when resolving the DataLayerInstance as FWorldPartitionActorDesc always represents the state of the actor local to its outer level
-			LocalDataLayerInstanceNames = DataLayerManager->GetDataLayerInstanceNames(InActor->GetActorDataLayers());
-			PRAGMA_ENABLE_DEPRECATION_WARNINGS
-		}
+				PRAGMA_DISABLE_DEPRECATION_WARNINGS
+					// Pass Actor Level when resolving the DataLayerInstance as FWorldPartitionActorDesc always represents the state of the actor local to its outer level
+					LocalDataLayerInstanceNames = DataLayerManager->GetDataLayerInstanceNames(InActor->GetActorDataLayers());
+				PRAGMA_ENABLE_DEPRECATION_WARNINGS
+			}
 
-		// Validation
-		const bool bHasDataLayerAssets = LocalDataLayerAssetPaths.Num() > 0;
-		const bool bHasDeprecatedDataLayers = LocalDataLayerInstanceNames.Num() > 0;
-		check((!bHasDataLayerAssets && !bHasDeprecatedDataLayers) || (bHasDataLayerAssets != bHasDeprecatedDataLayers));
+			// Validation
+			const bool bHasDataLayerAssets = LocalDataLayerAssetPaths.Num() > 0;
+			const bool bHasDeprecatedDataLayers = LocalDataLayerInstanceNames.Num() > 0;
+			check((!bHasDataLayerAssets && !bHasDeprecatedDataLayers) || (bHasDataLayerAssets != bHasDeprecatedDataLayers));
 
-		// Init DataLayers persistent info
-		bIsUsingDataLayerAsset = bHasDataLayerAssets;
-		DataLayers = bIsUsingDataLayerAsset ? MoveTemp(LocalDataLayerAssetPaths) : MoveTemp(LocalDataLayerInstanceNames);
+			// Init DataLayers persistent info
+			bIsUsingDataLayerAsset = bHasDataLayerAssets;
+			DataLayers = bIsUsingDataLayerAsset ? MoveTemp(LocalDataLayerAssetPaths) : MoveTemp(LocalDataLayerInstanceNames);
 
-		// Init DataLayers transient info
-		if (ensure(DataLayerManager))
-		{
-			// There are valid cases where the DataLayerManager can't resolve transient data layer instance names.
-			// One case is when committing changes on a level instance, FWorldPartitionActorDesc::Init is called when saving 
-			// the child actors.
-			if (DataLayerManager->CanResolveDataLayers())
+			// Init DataLayers transient info
+			if (ensure(DataLayerManager))
 			{
-				// Here we process a loaded actor, so resolving makes sense as long as the actordesc is not reused as a template
-				ResolvedDataLayerInstanceNames = FDataLayerUtils::ResolvedDataLayerInstanceNames(DataLayerManager, this);
-			}
-			else
-			{
-				check(!HasResolvedDataLayerInstanceNames());
+				// There are valid cases where the DataLayerManager can't resolve transient data layer instance names.
+				// One case is when committing changes on a level instance, FWorldPartitionActorDesc::Init is called when saving 
+				// the child actors.
+				if (DataLayerManager->CanResolveDataLayers())
+				{
+					// Here we process a loaded actor, so resolving makes sense as long as the actordesc is not reused as a template
+					ResolvedDataLayerInstanceNames = FDataLayerUtils::ResolvedDataLayerInstanceNames(DataLayerManager, this);
+				}
+				else
+				{
+					check(!HasResolvedDataLayerInstanceNames());
+				}
 			}
 		}
+		else
+		{
+			// Possible there is no World Partition for regular OFPA levels that haven't been converted to support Data Layers
+			bIsUsingDataLayerAsset = true;
+			DataLayers.Empty();
+			ResolvedDataLayerInstanceNames = TArray<FName>();
+		}
 	}
 
 	Tags = InActor->Tags;
