Sequencer: Allow holdout and visibility flags to be controlled by sequencer/blueprints

diff --git a/Engine/Source/Runtime/Engine/Classes/Components/PrimitiveComponent.h b/Engine/Source/Runtime/Engine/Classes/Components/PrimitiveComponent.h
--- a/Engine/Source/Runtime/Engine/Classes/Components/PrimitiveComponent.h
+++ b/Engine/Source/Runtime/Engine/Classes/Components/PrimitiveComponent.h
@@ -435,7 +435,7 @@ public:
 	uint8 bReceivesDecals:1;
 
 	/** If this is True, this primitive will render black with an alpha of 0, but all secondary effects (shadows, reflections, indirect lighting) remain. This feature is currently only implemented in the Path Tracer. */
-	UPROPERTY(EditAnywhere, BlueprintReadWrite, Category = PathTracing)
+	UPROPERTY(EditAnywhere, BlueprintReadOnly, Category = PathTracing, Interp)
 	uint8 bHoldout : 1;
 
 	/** If this is True, this component won't be visible when the view actor is the component's owner, directly or indirectly. */
@@ -479,7 +479,7 @@ public:
 	// Lighting flags
 	
 	/** Controls whether the primitive component should cast a shadow or not. */
-	UPROPERTY(EditAnywhere, BlueprintReadOnly, Category=Lighting)
+	UPROPERTY(EditAnywhere, BlueprintReadOnly, Category=Lighting, Interp)
 	uint8 CastShadow:1;
 
 	/** Whether the primitive will be used as an emissive light source. */
@@ -487,11 +487,11 @@ public:
 	uint8 bEmissiveLightSource:1;
 
 	/** Controls whether the primitive should influence indirect lighting. */
-	UPROPERTY(EditAnywhere, BlueprintReadOnly, Category=Lighting, AdvancedDisplay)
+	UPROPERTY(EditAnywhere, BlueprintReadOnly, Category=Lighting, AdvancedDisplay, Interp)
 	uint8 bAffectDynamicIndirectLighting:1;
 
 	/** Controls whether the primitive should affect indirect lighting when hidden. This flag is only used if bAffectDynamicIndirectLighting is true. */
-	UPROPERTY(EditAnywhere, AdvancedDisplay, BlueprintReadOnly, Category=Lighting, meta=(EditCondition="bAffectDynamicIndirectLighting", DisplayName = "Affect Indirect Lighting While Hidden"))
+	UPROPERTY(EditAnywhere, AdvancedDisplay, BlueprintReadOnly, Category=Lighting, meta=(EditCondition="bAffectDynamicIndirectLighting", DisplayName = "Affect Indirect Lighting While Hidden"), Interp)
 	uint8 bAffectIndirectLightingWhileHidden:1;
 
 	/** Controls whether the primitive should affect dynamic distance field lighting methods.  This flag is only used if CastShadow is true. **/
@@ -553,7 +553,7 @@ public:
 	 *	Controls whether the primitive should cast shadows when hidden.
 	 *	This flag is only used if CastShadow is true.
 	 */
-	UPROPERTY(EditAnywhere, AdvancedDisplay, BlueprintReadOnly, Category=Lighting, meta=(EditCondition="CastShadow", DisplayName = "Hidden Shadow"))
+	UPROPERTY(EditAnywhere, AdvancedDisplay, BlueprintReadOnly, Category=Lighting, meta=(EditCondition="CastShadow", DisplayName = "Hidden Shadow"), Interp)
 	uint8 bCastHiddenShadow:1;
 
 	/** Whether this primitive should cast dynamic shadows as if it were a two sided material. */
@@ -1816,6 +1816,19 @@ public:
 	UFUNCTION(BlueprintCallable, Category = "Rendering")
 	void SetReceivesDecals(bool bNewReceivesDecals);
 
+    /** Changes the value of bHoldout (Path Tracing only feature)*/
+	UFUNCTION(BlueprintCallable, Category = "Rendering")
+	void SetHoldout(bool bNewHoldout);
+
+    /** Changes the value of bAffectDynamicIndirectLighting */
+	UFUNCTION(BlueprintCallable, Category = "Rendering")
+	void SetAffectDynamicIndirectLighting(bool bNewAffectDynamicIndirectLighting);
+
+    /** Changes the value of bAffectIndirectLightingWhileHidden */
+	UFUNCTION(BlueprintCallable, Category = "Rendering")
+	void SetAffectIndirectLightingWhileHidden(bool bNewAffectIndirectLightingWhileHidden);
+
+
 	/** Controls what kind of collision is enabled for this body */
 	UFUNCTION(BlueprintCallable, Category="Collision")
 	virtual void SetCollisionEnabled(ECollisionEnabled::Type NewType);
diff --git a/Engine/Source/Runtime/Engine/Private/Components/PrimitiveComponent.cpp b/Engine/Source/Runtime/Engine/Private/Components/PrimitiveComponent.cpp
--- a/Engine/Source/Runtime/Engine/Private/Components/PrimitiveComponent.cpp
+++ b/Engine/Source/Runtime/Engine/Private/Components/PrimitiveComponent.cpp
@@ -1796,6 +1796,34 @@ void UPrimitiveComponent::SetReceivesDecals(bool bNewReceivesDecals)
 	}
 }
 
+void UPrimitiveComponent::SetHoldout(bool bNewHoldout)
+{
+	if (bHoldout != bNewHoldout)
+	{
+		bHoldout = bNewHoldout;
+		MarkRenderStateDirty();
+	}
+}
+
+void UPrimitiveComponent::SetAffectDynamicIndirectLighting(bool bNewAffectDynamicIndirectLighting)
+{
+	if (bAffectDynamicIndirectLighting != bNewAffectDynamicIndirectLighting)
+	{
+		bAffectDynamicIndirectLighting = bNewAffectDynamicIndirectLighting;
+		MarkRenderStateDirty();
+	}
+}
+
+
+void UPrimitiveComponent::SetAffectIndirectLightingWhileHidden(bool bNewAffectIndirectLightingWhileHidden)
+{
+	if (bAffectIndirectLightingWhileHidden != bNewAffectIndirectLightingWhileHidden)
+	{
+		bAffectIndirectLightingWhileHidden = bNewAffectIndirectLightingWhileHidden;
+		MarkRenderStateDirty();
+	}
+}
+
 
 void UPrimitiveComponent::PushSelectionToProxy()
 {
