In PIE, UWorldPartitionLevelStreamingDynamic::RequestLevel now behaves like what ULevelStreaming::RequestLevel does.
ULevelStreaming sets ULevel::StreamedLevelsOwningWorld so that ULevel::PostLoad can set the level OwningWorld.
Since UWorldPartitionLevelStreamingDynamic creates a memory level, it sets the OwningWorld right before loading OFPA actors.

diff --git a/Engine/Source/Runtime/Engine/Private/WorldPartition/WorldPartitionLevelStreamingDynamic.cpp b/Engine/Source/Runtime/Engine/Private/WorldPartition/WorldPartitionLevelStreamingDynamic.cpp
--- a/Engine/Source/Runtime/Engine/Private/WorldPartition/WorldPartitionLevelStreamingDynamic.cpp
+++ b/Engine/Source/Runtime/Engine/Private/WorldPartition/WorldPartitionLevelStreamingDynamic.cpp
@@ -267,6 +267,10 @@ bool UWorldPartitionLevelStreamingDynamic::RequestLevel(UWorld* InPersistentWorl
 				((IsRunningGame() || IsRunningDedicatedServer()) && RuntimeLevel->OwningWorld->WorldType == EWorldType::Game) || 
 				(RuntimeLevel->OwningWorld->WorldType == EWorldType::Editor && GetShouldBeVisibleInEditor()));
 
+			// ULevelStreaming sets ULevel::StreamedLevelsOwningWorld so that ULevel::PostLoad sets the level OwningWorld
+			// Since this is a memory level, setup the OwningWorld right away to mimic the default behavior.
+			RuntimeLevel->OwningWorld = InPersistentWorld;
+
 			if (IssueLoadRequests())
 			{
 				// Editor immediately blocks on load and we also block if background level streaming is disabled.
