Remove an unnecessary ensure in force feedback when you attempt to use a null device property. This can happen if you have an empty property on an asset and then use it anyway.

#rnx

diff --git a/Engine/Source/Runtime/Engine/Private/GameFramework/ForceFeedbackEffect.cpp b/Engine/Source/Runtime/Engine/Private/GameFramework/ForceFeedbackEffect.cpp
--- a/Engine/Source/Runtime/Engine/Private/GameFramework/ForceFeedbackEffect.cpp
+++ b/Engine/Source/Runtime/Engine/Private/GameFramework/ForceFeedbackEffect.cpp
@@ -177,7 +177,7 @@ void FActiveForceFeedbackEffect::ActivateDeviceProperties()
 		{
 			for (TObjectPtr<UInputDeviceProperty> DeviceProp : ForceFeedbackEffect->DeviceProperties)
 			{
-				if (ensure(DeviceProp))
+				if (DeviceProp)
 				{
 					FActivateDevicePropertyParams Params = {};
 					Params.bIgnoreTimeDilation = Parameters.bIgnoreTimeDilation;
