Fixed crash when loading partitioned world with a missing AWorldDataLayers.
Moved repair logic of WorldDataLayers from WorldPartition into DataLayerManager.

diff --git a/Engine/Source/Runtime/Engine/Private/WorldPartition/DataLayer/DataLayerManager.cpp b/Engine/Source/Runtime/Engine/Private/WorldPartition/DataLayer/DataLayerManager.cpp
--- a/Engine/Source/Runtime/Engine/Private/WorldPartition/DataLayer/DataLayerManager.cpp
+++ b/Engine/Source/Runtime/Engine/Private/WorldPartition/DataLayer/DataLayerManager.cpp
@@ -19,9 +19,11 @@
 #include "WorldPartition/DataLayer/IDataLayerEditorModule.h"
 #include "WorldPartition/ActorDescContainer.h"
 #include "WorldPartition/WorldPartitionEditorPerProjectUserSettings.h"
+#include "WorldPartition/ActorDescList.h"
 #include "LevelInstance/LevelInstanceSubsystem.h"
 #include "Modules/ModuleManager.h"
 #include "UObject/UObjectIterator.h"
+#include "UObject/UObjectGlobals.h"
 #include "Engine/LevelStreaming.h"
 #include "EngineUtils.h"
 #include "LevelUtils.h"
@@ -146,9 +148,39 @@ UDataLayerManager::UDataLayerManager()
 void UDataLayerManager::Initialize()
 {
 #if WITH_EDITOR
+	UWorld* OuterWorld = GetTypedOuter<UWorld>();
+	UWorldPartition* OuterWorldPartition = GetOuterUWorldPartition();
+	UActorDescContainer* ActorDescContainer = OuterWorldPartition->GetActorDescContainer();
+	// In PIE, main world partition doesn't have an ActorDescContainer and doesn't need it, but instanced world partitions do have one.
+	check(!ActorDescContainer || ActorDescContainer->IsInitialized());
+
+	if (!GetWorldDataLayers())
+	{
+		if (ActorDescContainer)
+		{
+			// Try to find and load AWorldDataLayers actor
+			for (FActorDescList::TIterator<> ActorDescIterator(ActorDescContainer); ActorDescIterator; ++ActorDescIterator)
+			{
+				if (ActorDescIterator->GetActorNativeClass()->IsChildOf<AWorldDataLayers>())
+				{
+					WorldDataLayersActor = FWorldPartitionReference(OuterWorldPartition, ActorDescIterator->GetGuid());
+					break;
+				}
+			}
+		}
+
+		if (!GetWorldDataLayers())
+		{
+			// Create missing AWorldDataLayers actor
+			AWorldDataLayers* WorldDataLayers = AWorldDataLayers::Create(OuterWorld);
+			OuterWorld->SetWorldDataLayers(WorldDataLayers);
+			check(GetWorldDataLayers());
+		}
+	}
+
 	// Partitioned Level Instance's DataLayerManager will never resolve DataLayers (it's up to its owning WorldPartition DataLayerManager to do the job.
 	ULevelInstanceSubsystem* LevelInstanceSubsystem = !GetWorld()->IsGameWorld() ? UWorld::GetSubsystem<ULevelInstanceSubsystem>(GetWorld()) : nullptr;
-	ILevelInstanceInterface* LevelInstance = LevelInstanceSubsystem ? LevelInstanceSubsystem->GetOwningLevelInstance(GetTypedOuter<UWorld>()->PersistentLevel) : nullptr;
+	ILevelInstanceInterface* LevelInstance = LevelInstanceSubsystem ? LevelInstanceSubsystem->GetOwningLevelInstance(OuterWorld->PersistentLevel) : nullptr;
 	bCanResolveDataLayers = (LevelInstance == nullptr);
 
 	DataLayerLoadingPolicy = NewObject<UDataLayerLoadingPolicy>(this, GetDataLayerLoadingPolicyClass());
@@ -161,11 +193,17 @@ void UDataLayerManager::Initialize()
 	if (CanResolveDataLayers())
 	{
 		UActorDescContainer::OnActorDescContainerInitialized.AddUObject(this, &UDataLayerManager::OnActorDescContainerInitialized);
+
+		// Manually call OnActorDescContainerInitialized on already initialized outer world partition container
+		if (ActorDescContainer)
+		{
+			OnActorDescContainerInitialized(ActorDescContainer);
+		}
 	}
 
 	// SaveAs of a partition world will duplicate actors which will trigger AActor::FixupDataLayer and DataLayerManager is not yet created.
 	// Here we make sure to revisit loaded actors of the outer world of this DataLayerManager (editor only)
-	for (TActorIterator<AActor> It(GetTypedOuter<UWorld>()); It; ++It)
+	for (TActorIterator<AActor> It(OuterWorld); It; ++It)
 	{
 		check(IsValid(*It));
 		It->FixupDataLayers();
@@ -177,6 +215,8 @@ void UDataLayerManager::DeInitialize()
 {
 #if WITH_EDITOR
 	UActorDescContainer::OnActorDescContainerInitialized.RemoveAll(this);
+
+	WorldDataLayersActor = FWorldPartitionReference();
 #endif
 }
 
diff --git a/Engine/Source/Runtime/Engine/Private/WorldPartition/DataLayer/WorldDataLayers.cpp b/Engine/Source/Runtime/Engine/Private/WorldPartition/DataLayer/WorldDataLayers.cpp
--- a/Engine/Source/Runtime/Engine/Private/WorldPartition/DataLayer/WorldDataLayers.cpp
+++ b/Engine/Source/Runtime/Engine/Private/WorldPartition/DataLayer/WorldDataLayers.cpp
@@ -442,7 +442,6 @@ AWorldDataLayers* AWorldDataLayers::Create(UWorld* World, FName InWorldDataLayer
 
 AWorldDataLayers* AWorldDataLayers::Create(const FActorSpawnParameters& SpawnParameters)
 {
-	check(SpawnParameters.OverrideLevel != nullptr && SpawnParameters.OverrideLevel->IsPersistentLevel());
 	check(SpawnParameters.Name != NAME_None);
 	check(SpawnParameters.NameMode == FActorSpawnParameters::ESpawnActorNameMode::Required_Fatal);
 
diff --git a/Engine/Source/Runtime/Engine/Private/WorldPartition/WorldPartition.cpp b/Engine/Source/Runtime/Engine/Private/WorldPartition/WorldPartition.cpp
--- a/Engine/Source/Runtime/Engine/Private/WorldPartition/WorldPartition.cpp
+++ b/Engine/Source/Runtime/Engine/Private/WorldPartition/WorldPartition.cpp
@@ -435,9 +435,6 @@ void UWorldPartition::Initialize(UWorld* InWorld, const FTransform& InTransform)
 	check(InWorld);
 	World = InWorld;
 
-	DataLayerManager = NewObject<UDataLayerManager>(this, TEXT("DataLayerManager"), RF_Transient);
-	DataLayerManager->Initialize();
-
 	if (!InTransform.Equals(FTransform::Identity))
 	{
 		InstanceTransform = InTransform;
@@ -495,6 +492,7 @@ void UWorldPartition::Initialize(UWorld* InWorld, const FTransform& InTransform)
 	check(RuntimeHash);
 	RuntimeHash->SetFlags(RF_Transactional);
 
+	TArray<FGuid> ForceLoadedActorGuids;
 	if (bIsEditor || bIsGame || bPIEWorldTravel || bIsDedicatedServer)
 	{
 		UPackage* LevelPackage = OuterWorld->PersistentLevel->GetOutermost();
@@ -522,8 +520,6 @@ void UWorldPartition::Initialize(UWorld* InWorld, const FTransform& InTransform)
 		ActorDescContainer = RegisterActorDescContainer(PackageName);
 
 		{
-			TArray<FGuid> ForceLoadedActorGuids;
-			
 			TRACE_CPUPROFILER_EVENT_SCOPE(UActorDescContainer::Hash);
 			for (FActorDescContainerCollection::TIterator<> ActorDescIterator(this); ActorDescIterator; ++ActorDescIterator)
 			{
@@ -549,12 +545,6 @@ void UWorldPartition::Initialize(UWorld* InWorld, const FTransform& InTransform)
 					HashActorDesc(*ActorDescIterator);
 				}
 			}
-
-			if (ForceLoadedActors && ForceLoadedActorGuids.Num() > 0)
-			{
-				TRACE_CPUPROFILER_EVENT_SCOPE(UActorDescContainer::ForceLoadedActors);
-				ForceLoadedActors->AddActors(ForceLoadedActorGuids);
-			}
 		}
 	}
 
@@ -562,34 +552,24 @@ void UWorldPartition::Initialize(UWorld* InWorld, const FTransform& InTransform)
 	{
 		// Here we need to flush any async loading before starting any synchronous load (mixing synchronous and asynchronous load on the same package is not handled properly).
 		// Also, this will ensure that FindObject will work on an external actor that was loading asynchronously.
-		{
-			TRACE_CPUPROFILER_EVENT_SCOPE(FlushAsyncLoading);
-			FlushAsyncLoading();
-		}
+		TRACE_CPUPROFILER_EVENT_SCOPE(FlushAsyncLoading);
+		FlushAsyncLoading();
+	}
+#endif
 
-		// Make sure to preload only AWorldDataLayers actor first
-		for (FActorDescList::TIterator<> ActorDescIterator(ActorDescContainer); ActorDescIterator; ++ActorDescIterator)
-		{
-			if (ActorDescIterator->GetActorNativeClass()->IsChildOf<AWorldDataLayers>())
-			{
-				WorldDataLayersActor = FWorldPartitionReference(this, ActorDescIterator->GetGuid());
-				break;
-			}
-		}
+	// Here's it's safe to initialize the DataLayerManager
+	DataLayerManager = NewObject<UDataLayerManager>(this, TEXT("DataLayerManager"), RF_Transient);
+	DataLayerManager->Initialize();
+
+#if WITH_EDITOR
+	if (ForceLoadedActors && ForceLoadedActorGuids.Num() > 0)
+	{
+		TRACE_CPUPROFILER_EVENT_SCOPE(UActorDescContainer::ForceLoadedActors);
+		ForceLoadedActors->AddActors(ForceLoadedActorGuids);
 	}
 	
 	if (bIsEditor)
 	{
-		// Repair existing maps with no UWorld::WorldDataLayers actor.
-		AWorldDataLayers* WorldDataLayers = OuterWorld->GetWorldDataLayers();
-		if (!WorldDataLayers)
-		{
-			// WorldDataLayersActor is invalid here orelse its PostLoad would have set itself as the world's WorldDataLayers
-			check(!WorldDataLayersActor.IsValid());
-			WorldDataLayers = AWorldDataLayers::Create(OuterWorld);
-			OuterWorld->SetWorldDataLayers(WorldDataLayers);
-		}
-
 		// Apply level transform on actors already part of the level
 		if (!GetInstanceTransform().Equals(FTransform::Identity))
 		{
@@ -722,8 +702,6 @@ void UWorldPartition::Uninitialize()
 			RegisteredEditorLoaderAdapters.Empty();
 		}
 
-		WorldDataLayersActor = FWorldPartitionReference();
-
 		DirtyActors.Empty();
 
 		UninitializeActorDescContainers();
diff --git a/Engine/Source/Runtime/Engine/Public/WorldPartition/DataLayer/DataLayerManager.h b/Engine/Source/Runtime/Engine/Public/WorldPartition/DataLayer/DataLayerManager.h
--- a/Engine/Source/Runtime/Engine/Public/WorldPartition/DataLayer/DataLayerManager.h
+++ b/Engine/Source/Runtime/Engine/Public/WorldPartition/DataLayer/DataLayerManager.h
@@ -3,6 +3,7 @@
 #pragma once
 
 #include "WorldPartition/WorldPartition.h"
+#include "WorldPartition/WorldPartitionHandle.h"
 #include "WorldPartition/DataLayer/ActorDataLayer.h"
 #include "WorldPartition/DataLayer/DataLayerInstance.h"
 #include "WorldPartition/DataLayer/DataLayerEditorContext.h"
@@ -192,6 +193,8 @@ private:
 
 	UPROPERTY(Transient)
 	TObjectPtr<UDataLayerLoadingPolicy> DataLayerLoadingPolicy;
+
+	FWorldPartitionReference WorldDataLayersActor;
 #endif
 };
 
diff --git a/Engine/Source/Runtime/Engine/Public/WorldPartition/WorldPartition.h b/Engine/Source/Runtime/Engine/Public/WorldPartition/WorldPartition.h
--- a/Engine/Source/Runtime/Engine/Public/WorldPartition/WorldPartition.h
+++ b/Engine/Source/Runtime/Engine/Public/WorldPartition/WorldPartition.h
@@ -371,8 +371,6 @@ private:
 
 #if WITH_EDITORONLY_DATA
 	FLinkerInstancingContext InstancingContext;
-
-	FWorldPartitionReference WorldDataLayersActor;
 #endif
 
 #if WITH_EDITOR
