Add r.DynamicRes.ThrottlingMaxScreenPercentage cvar to allow us to limit screen resolution without reallocating rendertargets. This is useful for energy saving scenarios, where modifying r.dynamicres.maxscreenpercentage can result in an obvious pop.

Submitted on behalf of Guillaume.Abadie

diff --git a/Engine/Source/Runtime/Engine/Private/DynamicResolution.cpp b/Engine/Source/Runtime/Engine/Private/DynamicResolution.cpp
--- a/Engine/Source/Runtime/Engine/Private/DynamicResolution.cpp
+++ b/Engine/Source/Runtime/Engine/Private/DynamicResolution.cpp
@@ -26,7 +26,15 @@ static TAutoConsoleVariable<float> CVarDynamicResMinSP(
 static TAutoConsoleVariable<float> CVarDynamicResMaxSP(
 	TEXT("r.DynamicRes.MaxScreenPercentage"),
 	DynamicRenderScaling::FractionToPercentage(DynamicRenderScaling::FHeuristicSettings::kDefaultMaxResolutionFraction),
-	TEXT("Maximal primary screen percentage."),
+	TEXT("Maximal primary screen percentage. Importantly this setting controls the preallocated video memory needed by the renderer to render."),
+	ECVF_Default);
+
+static TAutoConsoleVariable<float> CVarDynamicResThrottlingMaxSP(
+	TEXT("r.DynamicRes.ThrottlingMaxScreenPercentage"),
+	DynamicRenderScaling::FractionToPercentage(DynamicRenderScaling::FHeuristicSettings::kDefaultThrottlingMaxResolutionFraction),
+	TEXT("Throttle the primary screen percentage allowed by the heuristic to this max value when enabled. This has no effect on preallocated video memory.\n")
+	TEXT("This is for instance useful when the video game wants to trottle power consumption when inactive without resizing internal renderer's render targets\n")
+	TEXT("(which can result in popping)"),
 	ECVF_Default);
 
 // TODO: Seriously need a centralized engine perf manager.
@@ -142,6 +150,7 @@ DynamicRenderScaling::FHeuristicSettings GetPrimaryDynamicResolutionSettings()
 	BudgetSetting.Model = DynamicRenderScaling::EHeuristicModel::Quadratic;
 	BudgetSetting.MinResolutionFraction      = DynamicRenderScaling::GetPercentageCVarToFraction(CVarDynamicResMinSP);
 	BudgetSetting.MaxResolutionFraction      = DynamicRenderScaling::GetPercentageCVarToFraction(CVarDynamicResMaxSP);
+	BudgetSetting.ThrottlingMaxResolutionFraction = DynamicRenderScaling::GetPercentageCVarToFraction(CVarDynamicResThrottlingMaxSP);
 	BudgetSetting.UpperBoundQuantization     = CVarUpperBoundQuantization.GetValueOnAnyThread();
 	BudgetSetting.BudgetMs                   = CVarFrameTimeBudget.GetValueOnAnyThread() * (1.0f - DynamicRenderScaling::GetPercentageCVarToFraction(CVarOverBudgetGPUHeadRoomPercentage));
 	BudgetSetting.ChangeThreshold            = DynamicRenderScaling::GetPercentageCVarToFraction(CVarChangeThreshold);
@@ -518,11 +527,21 @@ void FDynamicResolutionHeuristicProxy::RefreshCurrentFrameResolutionFraction_Ren
 			}
 		}
 
+		float FinalMaxResolutionFraction = BudgetSettings.MaxResolutionFraction;
+		if (BudgetSettings.ThrottlingMaxResolutionFraction > 0.0f)
+		{
+			// Don't allow the throttling to resolution to mess up with the primary MinResolutionFraction and MaxResolutionFraction settings.
+			FinalMaxResolutionFraction = FMath::Clamp(
+				BudgetSettings.ThrottlingMaxResolutionFraction,
+				BudgetSettings.MinResolutionFraction,
+				BudgetSettings.MaxResolutionFraction);
+		}
+
 		// Clamp resolution fraction.
 		NewFrameResolutionFraction = FMath::Clamp(
 			NewFrameResolutionFraction,
 			BudgetSettings.MinResolutionFraction,
-			BudgetSettings.MaxResolutionFraction);
+			FinalMaxResolutionFraction);
 
 		NewFrameResolutionFractions[Budget] = NewFrameResolutionFraction;
 	}
diff --git a/Engine/Source/Runtime/RenderCore/Public/DynamicRenderScaling.h b/Engine/Source/Runtime/RenderCore/Public/DynamicRenderScaling.h
--- a/Engine/Source/Runtime/RenderCore/Public/DynamicRenderScaling.h
+++ b/Engine/Source/Runtime/RenderCore/Public/DynamicRenderScaling.h
@@ -34,6 +34,7 @@ struct RENDERCORE_API FHeuristicSettings final
 {
 	static constexpr float kDefaultMinResolutionFraction = 0.5f;
 	static constexpr float kDefaultMaxResolutionFraction = 1.0f;
+	static constexpr float kDefaultThrottlingMaxResolutionFraction = 0.0f;
 	static constexpr float kBudgetMsDisabled = 0.0f;
 	static constexpr float kDefaultChangeThreshold = 0.02f;
 	static constexpr float kDefaultTargetedHeadRoom = 0.05f;
@@ -46,6 +47,7 @@ struct RENDERCORE_API FHeuristicSettings final
 
 	float MinResolutionFraction      = kDefaultMinResolutionFraction;
 	float MaxResolutionFraction      = kDefaultMaxResolutionFraction;
+	float ThrottlingMaxResolutionFraction = kDefaultThrottlingMaxResolutionFraction;
 	float BudgetMs                   = kBudgetMsDisabled;
 	float ChangeThreshold            = kDefaultChangeThreshold;
 	float TargetedHeadRoom           = kDefaultTargetedHeadRoom;
