Duplicate nanite override material if necessary when duplicating materials for HLODs.

diff --git a/Engine/Source/Runtime/Engine/Private/ISMPartition/ISMComponentDescriptor.cpp b/Engine/Source/Runtime/Engine/Private/ISMPartition/ISMComponentDescriptor.cpp
--- a/Engine/Source/Runtime/Engine/Private/ISMPartition/ISMComponentDescriptor.cpp
+++ b/Engine/Source/Runtime/Engine/Private/ISMPartition/ISMComponentDescriptor.cpp
@@ -3,6 +3,7 @@
 #include "ISMPartition/ISMComponentDescriptor.h"
 #include "Concepts/StaticStructProvider.h"
 #include "Materials/MaterialInterface.h"
+#include "Materials/MaterialInstanceDynamic.h"
 
 #include UE_INLINE_GENERATED_CPP_BY_NAME(ISMComponentDescriptor)
 
@@ -179,6 +180,14 @@ void FISMComponentDescriptor::InitComponent(UInstancedStaticMeshComponent* ISMCo
 			// As override materials are normally outered to their owner component, we need to duplicate them here to make sure we don't create
 			// references to actors in other levels (for packed level instances or HLOD actors).
 			OverrideMaterial = DuplicateObject<UMaterialInterface>(OverrideMaterial, ISMComponent);
+
+			// If the MID we just duplicated has a nanite override that's also not an asset, duplicate that too
+			UMaterialInstanceDynamic* OverrideMID = Cast<UMaterialInstanceDynamic>(OverrideMaterial);
+			UMaterialInterface* NaniteOverride = OverrideMID ? OverrideMID->GetNaniteOverride() : nullptr; 
+			if (NaniteOverride && !NaniteOverride->IsAsset())
+			{
+				OverrideMID->SetNaniteOverride(DuplicateObject<UMaterialInterface>(NaniteOverride, ISMComponent));
+			}
 		}
 
 		ISMComponent->OverrideMaterials.Add(OverrideMaterial);
