Assert occurs when Compressing multiple Anim Sequences at once

fix converted ensure into error log message, and make sure the system deals with failed compression

Ensure condition failed: false [File:D:\build\U5M+Inc\Sync\Engine\Source\Runtime\CoreUObject\Private\UObject\LinkerSave.cpp] [Line: 358]
fix run SynchronousAnimatedBoneAttributesCompression when requesting compression
fix fixed/improved behaviour of FCompressedAnimSequence::IsValid
fix always validate compressed data, and ensure no tasks are currently running during UAnimSequence::IsCachedCookedPlatformDataLoaded

misc deprecate public access to bUseRawDataOnly, not sure why this existed
misc update bUseRawDataOnly uses to go trhough BoneContainer.SetUseRAWData

diff --git a/Engine/Plugins/Animation/MLDeformer/NearestNeighborModel/Source/NearestNeighborModelEditor/Private/NearestNeighborEditorModel.cpp b/Engine/Plugins/Animation/MLDeformer/NearestNeighborModel/Source/NearestNeighborModelEditor/Private/NearestNeighborEditorModel.cpp
--- a/Engine/Plugins/Animation/MLDeformer/NearestNeighborModel/Source/NearestNeighborModelEditor/Private/NearestNeighborEditorModel.cpp
+++ b/Engine/Plugins/Animation/MLDeformer/NearestNeighborModel/Source/NearestNeighborModelEditor/Private/NearestNeighborEditorModel.cpp
@@ -19,6 +19,7 @@
 #include "Animation/DebugSkelMeshComponent.h"
 #include "Animation/MorphTarget.h"
 #include "Animation/AnimSequence.h"
+#include "Animation/AnimInstance.h"
 #include "PackageTools.h"
 #include "ObjectTools.h"
 #include "UObject/SavePackage.h"
@@ -313,7 +314,6 @@ namespace UE::NearestNeighborModel
 				const int32 NumFrames = FMath::Min(NumNeighborsFromGeomCache, NumNeighborsFromAnimSequence);
 				if (NumFrames > 0)
 				{
-					AnimSequence->bUseRawDataOnly = true;
 					AnimSequence->Interpolation = EAnimInterpolationType::Step;
 
 					SkeletalMeshComponent->SetAnimationMode(EAnimationMode::AnimationSingleNode);
@@ -322,6 +322,10 @@ namespace UE::NearestNeighborModel
 					SkeletalMeshComponent->SetPlayRate(1.0f);
 					SkeletalMeshComponent->Play(false);
 					SkeletalMeshComponent->RefreshBoneTransforms();
+					if (SkeletalMeshComponent->GetAnimInstance())
+					{
+						SkeletalMeshComponent->GetAnimInstance()->GetRequiredBones().SetUseRAWData(true);
+					}
 
 					// assuming MeshMappings do not change
 					GeometryCacheComponent->SetGeometryCache(GeometryCache);
diff --git a/Engine/Plugins/Enterprise/GLTFExporter/Source/GLTFExporter/Private/Converters/GLTFBoneUtilities.cpp b/Engine/Plugins/Enterprise/GLTFExporter/Source/GLTFExporter/Private/Converters/GLTFBoneUtilities.cpp
--- a/Engine/Plugins/Enterprise/GLTFExporter/Source/GLTFExporter/Private/Converters/GLTFBoneUtilities.cpp
+++ b/Engine/Plugins/Enterprise/GLTFExporter/Source/GLTFExporter/Private/Converters/GLTFBoneUtilities.cpp
@@ -52,18 +52,12 @@ void FGLTFBoneUtilities::GetBoneIndices(const USkeleton* Skeleton, TArray<FBoneI
 
 void FGLTFBoneUtilities::GetBoneTransformsByFrame(const UAnimSequence* AnimSequence, const TArray<float>& FrameTimestamps, const TArray<FBoneIndexType>& BoneIndices, TArray<TArray<FTransform>>& OutBoneTransformsByFrame)
 {
-#if WITH_EDITOR
-	const bool bUseRawData = AnimSequence->OnlyUseRawData();
-#else
-	const bool bUseRawData = false;
-#endif
 
 	// Make sure to free stack allocations made by FCompactPose, FBlendedCurve, and FStackCustomAttributes when end of scope
 	FMemMark Mark(FMemStack::Get());
 
 	FBoneContainer BoneContainer;
 	BoneContainer.InitializeTo(BoneIndices, FCurveEvaluationOption(true), *AnimSequence->GetSkeleton());
-	BoneContainer.SetUseRAWData(bUseRawData);
 
 	const int32 FrameCount = FrameTimestamps.Num();
 	OutBoneTransformsByFrame.AddDefaulted(FrameCount);
@@ -80,7 +74,7 @@ void FGLTFBoneUtilities::GetBoneTransformsByFrame(const UAnimSequence* AnimSeque
 	for (int32 Frame = 0; Frame < FrameCount; ++Frame)
 	{
 		const FAnimExtractContext ExtractionContext(static_cast<double>(FrameTimestamps[Frame])); // TODO: set bExtractRootMotion?
-		AnimSequence->GetBonePose(PoseData, ExtractionContext, bUseRawData);
+		AnimSequence->GetBonePose(PoseData, ExtractionContext);
 		Pose.CopyBonesTo(OutBoneTransformsByFrame[Frame]);
 	}
 
diff --git a/Engine/Source/Runtime/Engine/Classes/Animation/AnimSequence.h b/Engine/Source/Runtime/Engine/Classes/Animation/AnimSequence.h
--- a/Engine/Source/Runtime/Engine/Classes/Animation/AnimSequence.h
+++ b/Engine/Source/Runtime/Engine/Classes/Animation/AnimSequence.h
@@ -702,8 +702,20 @@ public:
 	FGuid GenerateGuidFromRawData() const;
 
 	// Should we be always using our raw data (i.e is our compressed data stale)
-	bool OnlyUseRawData() const { return bUseRawDataOnly; }
-	void SetUseRawDataOnly(bool bInUseRawDataOnly) { bUseRawDataOnly = bInUseRawDataOnly; }
+	UE_DEPRECATED(5.2, "OnlyUseRawData has been deprecated")
+	bool OnlyUseRawData() const
+	{
+		PRAGMA_DISABLE_DEPRECATION_WARNINGS
+		return bUseRawDataOnly;
+		PRAGMA_ENABLE_DEPRECATION_WARNINGS
+	}
+	UE_DEPRECATED(5.2, "SetUseRawDataOnly has been deprecated")
+	void SetUseRawDataOnly(bool bInUseRawDataOnly)
+	{
+		PRAGMA_DISABLE_DEPRECATION_WARNINGS
+		bUseRawDataOnly = bInUseRawDataOnly;
+		PRAGMA_ENABLE_DEPRECATION_WARNINGS
+	}
 
 	// Return this animations guid for the raw data
 	UE_DEPRECATED(5.1, "GetRawDataGuid has been deprecated use GenerateGuidFromModel instead")
@@ -754,7 +766,9 @@ public:
 	/** Take a set of marker positions and validates them against a requested start position, updating them as desired */
 	void ValidateCurrentPosition(const FMarkerSyncAnimPosition& Position, bool bPlayingForwards, bool bLooping, float&CurrentTime, FMarkerPair& PreviousMarker, FMarkerPair& NextMarker, const UMirrorDataTable* MirrorTable = nullptr) const;
 	bool UseRawDataForPoseExtraction(const FBoneContainer& RequiredBones) const;
+
 	// Should we be always using our raw data (i.e is our compressed data stale)
+	UE_DEPRECATED(5.2, "bUseRawDataOnly public access will be deprecated")
 	bool bUseRawDataOnly;
 
 public:
@@ -857,7 +871,6 @@ protected:
 	FIoHash DataKeyHash;
 	TMap<FIoHash, TUniquePtr<FCompressedAnimSequence>> DataByPlatformKeyHash;
 	TMap<FIoHash, TPimplPtr<UE::Anim::FAnimationSequenceAsyncCacheTask>> CacheTasksByKeyHash;
-
 protected:
 	bool TryCancelAsyncTasks();
 	void FinishAsyncTasks();
diff --git a/Engine/Source/Runtime/Engine/Private/Animation/AnimCompressionTypes.cpp b/Engine/Source/Runtime/Engine/Private/Animation/AnimCompressionTypes.cpp
--- a/Engine/Source/Runtime/Engine/Private/Animation/AnimCompressionTypes.cpp
+++ b/Engine/Source/Runtime/Engine/Private/Animation/AnimCompressionTypes.cpp
@@ -1468,8 +1468,9 @@ bool FCompressedAnimSequence::IsValid(const UAnimSequence* AnimSequence) const
 	const bool bHasBoneTracks = AnimSequence->GetDataModelInterface()->GetNumBoneTracks() != 0;
 	const bool bHasCurveData = AnimSequence->GetDataModelInterface()->GetNumberOfFloatCurves() != 0;
 
-	const bool bValidCompressedData = (!bHasBoneTracks || CompressedCurveByteStream.Num() != 0) &&
-		(!bHasCurveData || CompressedCurveByteStream.Num() != 0 || AnimSequence->IsValidAdditive());
+	const bool bHasValidCompressedBoneData = CompressedDataStructure != nullptr || CompressedTrackToSkeletonMapTable.Num() == 0;
+	const bool bHasValidCompressedCurveData = CompressedCurveByteStream.Num() != 0;
+	const bool bValidCompressedData = (!bHasBoneTracks || bHasValidCompressedBoneData) && (!bHasCurveData || bHasValidCompressedCurveData || AnimSequence->IsValidAdditive());
 
 	return bValidCompressedData;
 #else
diff --git a/Engine/Source/Runtime/Engine/Private/Animation/AnimSequence.cpp b/Engine/Source/Runtime/Engine/Private/Animation/AnimSequence.cpp
--- a/Engine/Source/Runtime/Engine/Private/Animation/AnimSequence.cpp
+++ b/Engine/Source/Runtime/Engine/Private/Animation/AnimSequence.cpp
@@ -459,7 +459,9 @@ void UAnimSequence::WillNeverCacheCookedPlatformDataAgain()
 		
 	CacheTasksByKeyHash.Empty();
 	DataByPlatformKeyHash.Empty();
+	PRAGMA_DISABLE_DEPRECATION_WARNINGS
 	bUseRawDataOnly = true;
+	PRAGMA_ENABLE_DEPRECATION_WARNINGS
 }
 
 void UAnimSequence::ClearCachedCookedPlatformData(const ITargetPlatform* TargetPlatform)
@@ -476,7 +478,9 @@ void UAnimSequence::ClearCachedCookedPlatformData(const ITargetPlatform* TargetP
 	{
 		CompressedData.Reset();
 		DataKeyHash = FIoHash::Zero;
+		PRAGMA_DISABLE_DEPRECATION_WARNINGS
 		bUseRawDataOnly = true;
+		PRAGMA_ENABLE_DEPRECATION_WARNINGS
 	}
 }
 
@@ -491,7 +495,9 @@ void UAnimSequence::ClearAllCachedCookedPlatformData()
 	CompressedData.Reset();
 	DataKeyHash = FIoHash::Zero;
 	
+	PRAGMA_DISABLE_DEPRECATION_WARNINGS
 	bUseRawDataOnly = true;
+	PRAGMA_ENABLE_DEPRECATION_WARNINGS
 }
 
 int32 UAnimSequence::GetUncompressedRawSize() const
@@ -699,8 +705,8 @@ void UAnimSequence::Serialize(FArchive& Ar)
 		{
 			PRAGMA_DISABLE_DEPRECATION_WARNINGS
 			SerializeCompressedData(Ar,false);
-			PRAGMA_ENABLE_DEPRECATION_WARNINGS
 			Ar << bUseRawDataOnly;
+			PRAGMA_ENABLE_DEPRECATION_WARNINGS
 		}
 	}
 
@@ -1089,18 +1095,17 @@ bool UAnimSequence::IsCachedCookedPlatformDataLoaded(const ITargetPlatform* Targ
 	if (PollCacheDerivedData(KeyHash))
 	{
 		EndCacheDerivedData(KeyHash);
-		return true;
 	}
 
 	if (KeyHash == DataKeyHash)
 	{
-		return CompressedData.IsValid(this);
+		return CompressedData.IsValid(this) && !CacheTasksByKeyHash.Contains(KeyHash);
 	}
 	else
 	{
 		if(const TUniquePtr<FCompressedAnimSequence>* CompressedDataPtr = DataByPlatformKeyHash.Find(KeyHash))
 		{
-			return (*CompressedDataPtr)->IsValid(this);
+			return (*CompressedDataPtr)->IsValid(this) && !CacheTasksByKeyHash.Contains(KeyHash);
 		}
 	}
 
@@ -1237,7 +1242,11 @@ FTransform UAnimSequence::ExtractRootTrackTransform(float Time, const FBoneConta
 	const bool bContainsRootBoneTrack = [this, RootBoneIndex]() 
 	{
 #if WITH_EDITOR
-		if (CanEvaluateRawAnimationData() && bUseRawDataOnly)
+		if (CanEvaluateRawAnimationData() &&
+			PRAGMA_DISABLE_DEPRECATION_WARNINGS
+			bUseRawDataOnly
+			PRAGMA_ENABLE_DEPRECATION_WARNINGS
+			)
 		{
 			ValidateModel();
 
@@ -1265,7 +1274,10 @@ FTransform UAnimSequence::ExtractRootTrackTransform(float Time, const FBoneConta
 	{
 		// if we do have root data, then return root data
 		FTransform RootTransform;
-		GetBoneTransform(RootTransform, FSkeletonPoseBoneIndex(RootBoneIndex), static_cast<double>(Time), bUseRawDataOnly);
+		GetBoneTransform(RootTransform, FSkeletonPoseBoneIndex(RootBoneIndex), static_cast<double>(Time), 
+			PRAGMA_DISABLE_DEPRECATION_WARNINGS
+			bUseRawDataOnly
+			PRAGMA_ENABLE_DEPRECATION_WARNINGS);
 		return RootTransform;
 	}
 
@@ -1975,7 +1987,9 @@ void UAnimSequence::FlagDependentAnimationsAsRawDataOnly() const
 		UAnimSequence* Seq = *Iter;
 		if (Seq->RefPoseSeq == this)
 		{
+			PRAGMA_DISABLE_DEPRECATION_WARNINGS
 			Seq->bUseRawDataOnly = true;
+			PRAGMA_ENABLE_DEPRECATION_WARNINGS
 		}
 	}
 }
@@ -3193,7 +3207,11 @@ uint8* UAnimSequence::FindSyncMarkerPropertyData(int32 SyncMarkerIndex, FArrayPr
 
 bool UAnimSequence::DoesNeedRecompress() const
 {
-	return GetSkeleton() && (bUseRawDataOnly || (GetSkeletonVirtualBoneGuid() != GetSkeleton()->GetVirtualBoneGuid()));
+	return GetSkeleton() && (
+		PRAGMA_DISABLE_DEPRECATION_WARNINGS
+		bUseRawDataOnly
+		PRAGMA_ENABLE_DEPRECATION_WARNINGS
+		|| (GetSkeletonVirtualBoneGuid() != GetSkeleton()->GetVirtualBoneGuid()));
 }
 
 bool UAnimSequence::CreateAnimation(USkeletalMesh* Mesh)
@@ -3333,7 +3351,11 @@ void UAnimSequence::EvaluateCurveData(FBlendedCurve& OutCurve, float CurrentTime
 		return;
 	}
 
-	const bool bEvaluateRawData = bUseRawDataOnly || bForceUseRawData || !IsCurveCompressedDataValid();
+	const bool bEvaluateRawData =
+	    PRAGMA_DISABLE_DEPRECATION_WARNINGS
+	    bUseRawDataOnly
+        PRAGMA_ENABLE_DEPRECATION_WARNINGS
+	    || bForceUseRawData || !IsCurveCompressedDataValid();
 	check(!bForceUseRawData || CanEvaluateRawAnimationData());
 	if (CanEvaluateRawAnimationData() && bEvaluateRawData)
 	{
@@ -3366,7 +3388,11 @@ float UAnimSequence::EvaluateCurveData(SmartName::UID_Type CurveUID, float Curre
 {
 	SCOPE_CYCLE_COUNTER(STAT_AnimSeq_EvalCurveData);
 
-	const bool bEvaluateRawData = bUseRawDataOnly || bForceUseRawData || !IsCurveCompressedDataValid();
+	const bool bEvaluateRawData =
+	 PRAGMA_DISABLE_DEPRECATION_WARNINGS
+     bUseRawDataOnly
+     PRAGMA_ENABLE_DEPRECATION_WARNINGS
+	  || bForceUseRawData || !IsCurveCompressedDataValid();
 	check(!bForceUseRawData || CanEvaluateRawAnimationData());
 	if (CanEvaluateRawAnimationData() && bEvaluateRawData)
 	{
@@ -3382,7 +3408,11 @@ float UAnimSequence::EvaluateCurveData(SmartName::UID_Type CurveUID, float Curre
 
 bool UAnimSequence::HasCurveData(SmartName::UID_Type CurveUID, bool bForceUseRawData) const
 {
-	const bool bEvaluateRawData = bUseRawDataOnly || bForceUseRawData || !IsCurveCompressedDataValid();
+	const bool bEvaluateRawData =
+	PRAGMA_DISABLE_DEPRECATION_WARNINGS
+		bUseRawDataOnly
+	PRAGMA_ENABLE_DEPRECATION_WARNINGS
+		|| bForceUseRawData || !IsCurveCompressedDataValid();
 	check(!bForceUseRawData || CanEvaluateRawAnimationData());
 	if (CanEvaluateRawAnimationData() && bEvaluateRawData)
 	{
@@ -3696,7 +3726,11 @@ void UAnimSequence::ValidateCurrentPosition(const FMarkerSyncAnimPosition& Posit
 
 bool UAnimSequence::UseRawDataForPoseExtraction(const FBoneContainer& RequiredBones) const
 {
-	return CanEvaluateRawAnimationData() && (bUseRawDataOnly || (GetSkeletonVirtualBoneGuid() != GetSkeleton()->GetVirtualBoneGuid()) || RequiredBones.GetDisableRetargeting() || RequiredBones.ShouldUseRawData() ||
+	return CanEvaluateRawAnimationData() && (
+	PRAGMA_DISABLE_DEPRECATION_WARNINGS
+	bUseRawDataOnly
+	PRAGMA_ENABLE_DEPRECATION_WARNINGS
+	|| (GetSkeletonVirtualBoneGuid() != GetSkeleton()->GetVirtualBoneGuid()) || RequiredBones.GetDisableRetargeting() || RequiredBones.ShouldUseRawData() ||
  #if WITH_EDITOR
   	GForceRawData == 1 ||
  #endif // WITH_EDITOR
@@ -4669,7 +4703,9 @@ void UAnimSequence::OnModelModified(const EAnimDataModelNotifyType& NotifyType,
 		PRAGMA_DISABLE_DEPRECATION_WARNINGS
 		RawDataGuid = bForceNewRawDataGuid ? FGuid::NewGuid() : Model->GenerateGuid();
 		PRAGMA_ENABLE_DEPRECATION_WARNINGS
-		bUseRawDataOnly = true;
+		PRAGMA_DISABLE_DEPRECATION_WARNINGS
+        bUseRawDataOnly = true;
+        PRAGMA_ENABLE_DEPRECATION_WARNINGS
 
 		FlagDependentAnimationsAsRawDataOnly();
 		UpdateDependentStreamingAnimations();
@@ -5041,12 +5077,14 @@ FIoHash UAnimSequence::BeginCacheDerivedData(const ITargetPlatform* TargetPlatfo
 	// Wait for any in-flight requests
 	UE::Anim::FAnimSequenceCompilingManager::Get().FinishCompilation({this});
 	
-	bUseRawDataOnly = true;
-
 	// We should wait here for any previous compilations?
 	FCompressedAnimSequence* TargetData = nullptr;
 	if (TargetPlatform->IsRunningPlatform())
 	{
+		PRAGMA_DISABLE_DEPRECATION_WARNINGS
+		bUseRawDataOnly = true;
+		PRAGMA_ENABLE_DEPRECATION_WARNINGS
+		
 		DataKeyHash = KeyHash;
 		TargetData = &CompressedData;
 	}
@@ -5079,6 +5117,11 @@ FIoHash UAnimSequence::BeginCacheDerivedData(const ITargetPlatform* TargetPlatfo
 		bPerformFrameStripping = ShouldPerformStripping(bPerformFrameStripping, bPerformFrameStrippingOnOddNumberedFrames);
 		PRAGMA_ENABLE_DEPRECATION_WARNINGS
 	}
+
+	// Always keep AnimatedBoneAttributes in lock-step when requesting compression
+	{		
+		SynchronousAnimatedBoneAttributesCompression();
+	}
 	
 	{		
 		// Data does not exist, need to build it.
@@ -5108,25 +5151,17 @@ bool UAnimSequence::PollCacheDerivedData(const FIoHash& KeyHash) const
 		return (*Task)->Poll();
 	}
 
-	return true;
+	return false;
 }
 
 void UAnimSequence::EndCacheDerivedData(const FIoHash& KeyHash)
 {
-	//UE_LOG(LogAnimation, Warning, TEXT("EndCacheDerivedData %s"), *GetName());
 	if (KeyHash.IsZero())
 	{
 		return;
 	}
 
 	UE::Anim::FAnimSequenceCompilingManager::Get().FinishCompilation({this});
-
-	// QQ calling FinishCompilation after this, makes it so that in ::FinishAsyncTasks we do not apply the skeleton GUID etc
-	// TPimplPtr<UE::Anim::FAnimationSequenceAsyncCacheTask> Task;
-	// if (CacheTasksByKeyHash.RemoveAndCopyValue(KeyHash, Task))
-	// {
-	// 	Task->Wait();	// 	
-	// }
 }
 	
 FCompressedAnimSequence& UAnimSequence::CacheDerivedData(const ITargetPlatform* TargetPlatform)
@@ -5194,32 +5229,54 @@ bool UAnimSequence::IsAsyncTaskComplete() const
 
 void UAnimSequence::FinishAsyncTasks()
 {	
-	//UE_LOG(LogAnimation, Warning, TEXT("FinishAsyncTasks %s [%i]"), *GetName(), CacheTasksByKeyHash.Num());
 	const bool bHasInflightTasks = CacheTasksByKeyHash.Num() > 0;
 	if (bHasInflightTasks)
 	{
 		COOK_STAT(auto Timer = UE::Anim::AnimSequenceCookStats::UsageStats.TimeAsyncWait());
+		COOK_STAT(Timer.TrackCyclesOnly());
+		
 		for (auto It = CacheTasksByKeyHash.CreateIterator(); It; ++It)
 		{
 			It->Value->Wait();
+
+			const FCompressedAnimSequence* CompressedAnimData = [Hash = It->Key, this]()
+			{
+				if(Hash == DataKeyHash)
+				{
+					return &CompressedData;
+				}
+				
+				return DataByPlatformKeyHash.FindChecked(Hash).Get();
+			}();
+
+			if (CompressedAnimData->IsValid(this))
+			{
+				#if WITH_EDITOR
+				//This is only safe during sync anim compression
+				if (GetSkeleton())
+				{
+					SetSkeletonVirtualBoneGuid(GetSkeleton()->GetVirtualBoneGuid());
+				}
+#endif
+				if (It->Key == DataKeyHash)
+				{
+					
+					FAssetRegistryModule& AssetRegistryModule = FModuleManager::LoadModuleChecked<FAssetRegistryModule>(TEXT("AssetRegistry"));
+					AssetRegistryModule.Get().AssetTagsFinalized(*this);
+					
+					PRAGMA_DISABLE_DEPRECATION_WARNINGS
+					check(bUseRawDataOnly);
+					bUseRawDataOnly = false;
+                    PRAGMA_ENABLE_DEPRECATION_WARNINGS
+				}
+			}
+			else
+			{
+				// Failed to compress
+			}
+
 			It.RemoveCurrent();
 		}
-		COOK_STAT(Timer.TrackCyclesOnly());
-	
-#if WITH_EDITOR
-		//This is only safe during sync anim compression
-		if (GetSkeleton())
-		{
-			SynchronousAnimatedBoneAttributesCompression();
-			SetSkeletonVirtualBoneGuid(GetSkeleton()->GetVirtualBoneGuid());
-		}
-#endif
-		
-		FAssetRegistryModule& AssetRegistryModule = FModuleManager::LoadModuleChecked<FAssetRegistryModule>(TEXT("AssetRegistry"));
-		AssetRegistryModule.Get().AssetTagsFinalized(*this);
-		
-		check(bUseRawDataOnly);
-		bUseRawDataOnly = false;
 	}
 }
 
diff --git a/Engine/Source/Runtime/Engine/Private/Animation/AnimationCompressionDerivedData.cpp b/Engine/Source/Runtime/Engine/Private/Animation/AnimationCompressionDerivedData.cpp
--- a/Engine/Source/Runtime/Engine/Private/Animation/AnimationCompressionDerivedData.cpp
+++ b/Engine/Source/Runtime/Engine/Private/Animation/AnimationCompressionDerivedData.cpp
@@ -9,6 +9,7 @@
 #include "Animation/AnimCompress.h"
 #include "DerivedDataCache.h"
 #include "Animation/Skeleton.h"
+#include "Interfaces/ITargetPlatform.h"
 
 namespace UE::Anim
 {	
@@ -219,9 +220,6 @@ bool FAnimationSequenceAsyncCacheTask::BuildData() const
 				
 	const bool bCompressionSuccessful = bBoneCompressionOk && bCurveCompressionOk;
 	const FString CompressionName = DataToCompress.BoneCompressionSettings->GetFullName();
-	ensureMsgf(bCompressionSuccessful, TEXT("Anim Compression failed for Sequence '%s' with compression scheme '%s': compressed data empty"), 
-		*DataToCompress.FullName,
-		*CompressionName);
 	
 	if (bCompressionSuccessful && !Owner.IsCanceled())
 	{
@@ -232,7 +230,11 @@ bool FAnimationSequenceAsyncCacheTask::BuildData() const
 		OutData.OwnerName = DataToCompress.AnimFName;
 		
 		return true;
-	}	
+	}
+	else
+	{
+		UE_LOG(LogAnimationCompression, Error, TEXT("Failed to generate compressed animation data for %s with compression scheme %s for target platform %s"), *CompressibleAnimPtr->FullName, *CompressionName, *TargetPlatform->DisplayName().ToString());
+	}
 	
 	return false;
 }
