LevelInstance: Fix check when breaking Level Instance that is assigned to a DataLayer

diff --git a/Engine/Source/Runtime/Engine/Private/LevelInstance/LevelInstanceSubsystem.cpp b/Engine/Source/Runtime/Engine/Private/LevelInstance/LevelInstanceSubsystem.cpp
--- a/Engine/Source/Runtime/Engine/Private/LevelInstance/LevelInstanceSubsystem.cpp
+++ b/Engine/Source/Runtime/Engine/Private/LevelInstance/LevelInstanceSubsystem.cpp
@@ -1326,7 +1326,16 @@ void ULevelInstanceSubsystem::BreakLevelInstance_Impl(ILevelInstanceInterface* L
 					{
 						for (const UDataLayerInstance* DataLayerInstance : LevelInstanceDataLayerInstances)
 						{
-							Actor->AddDataLayer(DataLayerInstance);
+							if (const UDataLayerInstanceWithAsset* DataLayerInstanceWithAsset = Cast<UDataLayerInstanceWithAsset>(DataLayerInstance))
+							{
+								// For DataLayerInstanceAsset we add the Asset to the actor instead of the DataLayerInstance because Actor hasn't moved yet and this
+								// will fail on a UDataLayerInstanceWithAsset when comparing the Actor's outer to the DataLayerInstance outer.
+								Actor->AddDataLayer(DataLayerInstanceWithAsset->GetAsset());
+							}
+							else
+							{
+								Actor->AddDataLayer(DataLayerInstance);
+							}
 						}
 					}
 
