Fix "Path does not map to any roots" fatal error from FHLODActorDesc::Serialize()

diff --git a/Engine/Source/Runtime/Engine/Private/WorldPartition/HLOD/HLODActorDesc.cpp b/Engine/Source/Runtime/Engine/Private/WorldPartition/HLOD/HLODActorDesc.cpp
--- a/Engine/Source/Runtime/Engine/Private/WorldPartition/HLOD/HLODActorDesc.cpp
+++ b/Engine/Source/Runtime/Engine/Private/WorldPartition/HLOD/HLODActorDesc.cpp
@@ -102,7 +102,8 @@ static int64 GetPackageSize(const FString& InPackageFileName)
 
 int64 FHLODActorDesc::GetPackageSize() const
 {
-	const FString PackageFileName = FPackageName::LongPackageNameToFilename(GetActorPackage().ToString(), FPackageName::GetAssetPackageExtension());
+	FString PackageFileName = GetActorPackage().ToString();
+	FPackageName::TryConvertLongPackageNameToFilename(GetActorPackage().ToString(), PackageFileName, FPackageName::GetAssetPackageExtension());
 	return ::GetPackageSize(PackageFileName);
 }
 
