UMaterialParameterCollectionInstance Lazy Creation Optimization.
- Added const to a few UMaterialParameterCollectionInstance getters.
- Removes full UWorld UMaterialParameterCollectionInstance creation on UMaterialParameterCollection::PostLoad.
- Saves up to 5ms in some cases (hitch).

diff --git a/Engine/Source/Editor/UnrealEd/Private/Factories/EditorFactories.cpp b/Engine/Source/Editor/UnrealEd/Private/Factories/EditorFactories.cpp
--- a/Engine/Source/Editor/UnrealEd/Private/Factories/EditorFactories.cpp
+++ b/Engine/Source/Editor/UnrealEd/Private/Factories/EditorFactories.cpp
@@ -666,12 +666,6 @@ UMaterialParameterCollectionFactoryNew::UMaterialParameterCollectionFactoryNew(c
 UObject* UMaterialParameterCollectionFactoryNew::FactoryCreateNew(UClass* Class,UObject* InParent,FName Name,EObjectFlags Flags,UObject* Context,FFeedbackContext* Warn)
 {
 	UMaterialParameterCollection* MaterialParameterCollection = NewObject<UMaterialParameterCollection>(InParent, Class, Name, Flags);
-	
-	if (MaterialParameterCollection)
-	{
-		MaterialParameterCollection->SetupWorldParameterCollectionInstances();
-	}
-
 	return MaterialParameterCollection;
 }
 
diff --git a/Engine/Source/Runtime/Engine/Classes/Engine/World.h b/Engine/Source/Runtime/Engine/Classes/Engine/World.h
--- a/Engine/Source/Runtime/Engine/Classes/Engine/World.h
+++ b/Engine/Source/Runtime/Engine/Classes/Engine/World.h
@@ -3033,6 +3033,10 @@ public:
 	/** Updates this world's scene with the list of instances, and optionally updates each instance's uniform buffer. */
 	void UpdateParameterCollectionInstances(bool bUpdateInstanceUniformBuffers, bool bRecreateUniformBuffer);
 
+private:
+	UMaterialParameterCollectionInstance* CreateParameterCollectionInstance(int32 ExistingIndex, UMaterialParameterCollection* Collection, bool bUpdateScene);
+public:
+
 	/** Gets the canvas object for rendering to a render target.  Will allocate one if needed. */
 	UCanvas* GetCanvasForRenderingToTarget();
 	UCanvas* GetCanvasForDrawMaterialToRenderTarget();
diff --git a/Engine/Source/Runtime/Engine/Classes/Materials/MaterialParameterCollection.h b/Engine/Source/Runtime/Engine/Classes/Materials/MaterialParameterCollection.h
--- a/Engine/Source/Runtime/Engine/Classes/Materials/MaterialParameterCollection.h
+++ b/Engine/Source/Runtime/Engine/Classes/Materials/MaterialParameterCollection.h
@@ -107,18 +107,18 @@ class UMaterialParameterCollection : public UObject
 #endif // WITH_EDITOR
 
 	/** Get the index in the ScalarParameters array of the parameter matching the input parameter name, returns -1 if not found **/
-	int32 GetScalarParameterIndexByName(FName ParameterName);
+	int32 GetScalarParameterIndexByName(FName ParameterName) const;
 
 	/** Get the index in the VectorParameters array of the parameter matching the input parameter name, returns -1 if not found **/
-	int32 GetVectorParameterIndexByName(FName ParameterName);
+	int32 GetVectorParameterIndexByName(FName ParameterName) const;
 
 	/** Returns an array of the names of all the scalar parameters in this Material Parameter Collection **/
 	UFUNCTION(BlueprintCallable, Category="Rendering|Material", meta=(Keywords="GetScalarParameterNames"))
-	TArray<FName> GetScalarParameterNames();
+	TArray<FName> GetScalarParameterNames() const;
 
 	/** Returns an array of the names of all the vector parameters in this Material Parameter Collection **/
 	UFUNCTION(BlueprintCallable, Category="Rendering|Material", meta=(Keywords="GetVectorParameterNames"))
-	TArray<FName> GetVectorParameterNames();
+	TArray<FName> GetVectorParameterNames() const;
 
 	/** Gets the default value of a scalar parameter from a material collection.
 	 * @param ParameterName - The name of the value to get the value of
@@ -126,7 +126,7 @@ class UMaterialParameterCollection : public UObject
 	 * @returns the value of the parameter
 	 **/
 	UFUNCTION(BlueprintCallable, Category="Rendering|Material", meta=(Keywords="GetFloatParameterDefaultValue"))
-	float GetScalarParameterDefaultValue(FName ParameterName, bool& bParameterFound);
+	float GetScalarParameterDefaultValue(FName ParameterName, bool& bParameterFound) const;
 
 	/** Gets the default value of a scalar parameter from a material collection.
 	 * @param ParameterName - The name of the value to get the value of
@@ -134,7 +134,7 @@ class UMaterialParameterCollection : public UObject
 	 * @returns the value of the parameter
 	 **/
 	UFUNCTION(BlueprintCallable, Category="Rendering|Material", meta=(Keywords="GetFloatParameterDefaultValue"))
-	FLinearColor GetVectorParameterDefaultValue(FName ParameterName, bool& bParameterFound);
+	FLinearColor GetVectorParameterDefaultValue(FName ParameterName, bool& bParameterFound) const;
 	
 	//~ Begin UObject Interface.
 #if WITH_EDITOR
@@ -173,8 +173,6 @@ class UMaterialParameterCollection : public UObject
 		return *UniformBufferStruct;
 	}
 
-	ENGINE_API void SetupWorldParameterCollectionInstances();
-
 private:
 	virtual ENGINE_API void FinishDestroy() override;
 	virtual ENGINE_API bool IsReadyForFinishDestroy() override;
diff --git a/Engine/Source/Runtime/Engine/Private/Materials/ParameterCollection.cpp b/Engine/Source/Runtime/Engine/Private/Materials/ParameterCollection.cpp
--- a/Engine/Source/Runtime/Engine/Private/Materials/ParameterCollection.cpp
+++ b/Engine/Source/Runtime/Engine/Private/Materials/ParameterCollection.cpp
@@ -47,20 +47,9 @@ void UMaterialParameterCollection::PostLoad()
 	}
 
 	CreateBufferStruct();
-	SetupWorldParameterCollectionInstances();
 	UpdateDefaultResource(true);
 }
 
-void UMaterialParameterCollection::SetupWorldParameterCollectionInstances()
-{
-	// Create an instance for this collection in every world
-	for (TObjectIterator<UWorld> It; It; ++It)
-	{
-		UWorld* CurrentWorld = *It;
-		CurrentWorld->AddParameterCollectionInstance(this, true);
-	}
-}
-
 void UMaterialParameterCollection::BeginDestroy()
 {
 	if (DefaultResource)
@@ -325,7 +314,7 @@ void UMaterialParameterCollection::PostEditChangeProperty(FPropertyChangedEvent&
 
 #endif // WITH_EDITOR
 
-int32 UMaterialParameterCollection::GetScalarParameterIndexByName(FName ParameterName)
+int32 UMaterialParameterCollection::GetScalarParameterIndexByName(FName ParameterName) const
 {
 	// loop over all the available scalar parameters and look for a name match
 	for (int32 ParameterIndex = 0; ParameterIndex < ScalarParameters.Num(); ParameterIndex++)
@@ -339,7 +328,7 @@ int32 UMaterialParameterCollection::GetScalarParameterIndexByName(FName Paramete
 	return -1;
 }
 
-int32 UMaterialParameterCollection::GetVectorParameterIndexByName(FName ParameterName)
+int32 UMaterialParameterCollection::GetVectorParameterIndexByName(FName ParameterName) const
 {
 	// loop over all the available vector parameters and look for a name match
 	for (int32 ParameterIndex = 0; ParameterIndex < VectorParameters.Num(); ParameterIndex++)
@@ -353,21 +342,21 @@ int32 UMaterialParameterCollection::GetVectorParameterIndexByName(FName Paramete
 	return -1;
 }
 
-TArray<FName> UMaterialParameterCollection::GetScalarParameterNames()
+TArray<FName> UMaterialParameterCollection::GetScalarParameterNames() const
 {
 	TArray<FName> Names;
 	GetParameterNames(Names, false);
 	return Names;
 }
 
-TArray<FName> UMaterialParameterCollection::GetVectorParameterNames()
+TArray<FName> UMaterialParameterCollection::GetVectorParameterNames() const
 {
 	TArray<FName> Names;
 	GetParameterNames(Names, true);
 	return Names;
 }
 
-float UMaterialParameterCollection::GetScalarParameterDefaultValue(FName ParameterName, bool& bParameterFound)
+float UMaterialParameterCollection::GetScalarParameterDefaultValue(FName ParameterName, bool& bParameterFound) const
 {
 	const int32 ParameterIndex = GetScalarParameterIndexByName(ParameterName);
 	bParameterFound = true;
@@ -380,7 +369,7 @@ float UMaterialParameterCollection::GetScalarParameterDefaultValue(FName Paramet
 	return ScalarParameters[ParameterIndex].DefaultValue;
 }
 
-FLinearColor UMaterialParameterCollection::GetVectorParameterDefaultValue(FName ParameterName, bool& bParameterFound)
+FLinearColor UMaterialParameterCollection::GetVectorParameterDefaultValue(FName ParameterName, bool& bParameterFound) const
 {
 	const int32 ParameterIndex = GetVectorParameterIndexByName(ParameterName);
 	bParameterFound = true;
diff --git a/Engine/Source/Runtime/Engine/Private/World.cpp b/Engine/Source/Runtime/Engine/Private/World.cpp
--- a/Engine/Source/Runtime/Engine/Private/World.cpp
+++ b/Engine/Source/Runtime/Engine/Private/World.cpp
@@ -1225,9 +1225,6 @@ void UWorld::PostLoad()
 	// Add the garbage collection callbacks
 	FLevelStreamingGCHelper::AddGarbageCollectorCallback();
 
-	// Initially set up the parameter collection list. This may be run again in UWorld::InitWorld but it's required here for some editor and streaming cases
-	SetupParameterCollectionInstances();
-
 #if WITH_EDITOR
 	if (GIsEditor)
 	{
@@ -1388,31 +1385,7 @@ void UWorld::AddParameterCollectionInstance(UMaterialParameterCollection* Collec
 		}
 	}
 
-	UMaterialParameterCollectionInstance* NewInstance = NewObject<UMaterialParameterCollectionInstance>();
-	NewInstance->SetCollection(Collection, this);
-
-	if (ExistingIndex != INDEX_NONE)
-	{
-		// Overwrite an existing instance
-		ParameterCollectionInstances[ExistingIndex] = NewInstance;
-	}
-	else
-	{
-		// Add a new instance
-		ParameterCollectionInstances.Add(NewInstance);
-	}
-
-	// Ensure the new instance creates initial render thread resources
-	// This needs to happen right away, so they can be picked up by any cached shader bindings
-	NewInstance->UpdateRenderState(true);
-
-	if (bUpdateScene)
-	{
-		// Update the scene's list of instances, needs to happen to prevent a race condition with GC 
-		// (rendering thread still uses the FMaterialParameterCollectionInstanceResource when GC deletes the UMaterialParameterCollectionInstance)
-		// However, if UpdateParameterCollectionInstances is going to be called after many AddParameterCollectionInstance's, this can be skipped for now.
-		UpdateParameterCollectionInstances(false, false);
-	}
+	CreateParameterCollectionInstance(ExistingIndex, Collection, bUpdateScene);
 }
 
 UMaterialParameterCollectionInstance* UWorld::GetParameterCollectionInstance(const UMaterialParameterCollection* Collection) const
@@ -1425,9 +1398,8 @@ UMaterialParameterCollectionInstance* UWorld::GetParameterCollectionInstance(con
 		}
 	}
 
-	// Instance should always exist due to SetupParameterCollectionInstances() and UMaterialParameterCollection::PostLoad()
-	check(0);
-	return nullptr;
+	// Lazy create one if not found
+	return const_cast<UWorld*>(this)->CreateParameterCollectionInstance(INDEX_NONE, const_cast<UMaterialParameterCollection*>(Collection), true);
 }
 
 void UWorld::UpdateParameterCollectionInstances(bool bUpdateInstanceUniformBuffers, bool bRecreateUniformBuffer)
@@ -1456,6 +1428,37 @@ void UWorld::UpdateParameterCollectionInstances(bool bUpdateInstanceUniformBuffe
 	}
 }
 
+UMaterialParameterCollectionInstance* UWorld::CreateParameterCollectionInstance(int32 ExistingIndex, UMaterialParameterCollection* Collection, bool bUpdateScene)
+{
+	UMaterialParameterCollectionInstance* NewInstance = NewObject<UMaterialParameterCollectionInstance>();
+	NewInstance->SetCollection(Collection, this);
+
+	if (ExistingIndex != INDEX_NONE)
+	{
+		// Overwrite an existing instance
+		ParameterCollectionInstances[ExistingIndex] = NewInstance;
+	}
+	else
+	{
+		// Add a new instance
+		ParameterCollectionInstances.Add(NewInstance);
+	}
+
+	// Ensure the new instance creates initial render thread resources
+	// This needs to happen right away, so they can be picked up by any cached shader bindings
+	NewInstance->UpdateRenderState(true);
+
+	if (bUpdateScene)
+	{
+		// Update the scene's list of instances, needs to happen to prevent a race condition with GC 
+		// (rendering thread still uses the FMaterialParameterCollectionInstanceResource when GC deletes the UMaterialParameterCollectionInstance)
+		// However, if UpdateParameterCollectionInstances is going to be called after many AddParameterCollectionInstance's, this can be skipped for now.
+		UpdateParameterCollectionInstances(false, false);
+	}
+
+	return NewInstance;
+}
+
 UCanvas* UWorld::GetCanvasForRenderingToTarget()
 {
 	if (!CanvasForRenderingToTarget)
@@ -1873,8 +1876,6 @@ void UWorld::InitWorld(const InitializationValues IVS)
 		AvoidanceManager = NewObject<UAvoidanceManager>(this, GEngine->AvoidanceManagerClass);
 	}
 
-	SetupParameterCollectionInstances();
-
 	if (PersistentLevel->GetOuter() != this)
 	{
 		// Move persistent level into world so the world object won't get garbage collected in the multi- level
