Fix security vulnerabilities in FBatchedElements::DrawPointElements, FBatchedElements::Draw, and DynamicRHIResourceArray::GetResourceDataSize() by disallowing integer overflow when calculating how much memory to allocate for a vertex buffer.

diff --git a/Engine/Source/Runtime/Core/Public/Containers/DynamicRHIResourceArray.h b/Engine/Source/Runtime/Core/Public/Containers/DynamicRHIResourceArray.h
--- a/Engine/Source/Runtime/Core/Public/Containers/DynamicRHIResourceArray.h
+++ b/Engine/Source/Runtime/Core/Public/Containers/DynamicRHIResourceArray.h
@@ -66,6 +66,11 @@ public:
 	*/
 	virtual uint32 GetResourceDataSize() const
 	{
+		if (this->Num() > UINT32_MAX / sizeof(ElementType))
+		{
+			UE_LOG(LogCore, Fatal, TEXT("Resource data size too large for uint32, will overflow. Calculate with larger data type or use fewer elements. sizeof(ElementType): %d"), sizeof(ElementType));
+			return 0;
+		}
 		return this->Num() * sizeof(ElementType);
 	}
 
diff --git a/Engine/Source/Runtime/Engine/Private/BatchedElements.cpp b/Engine/Source/Runtime/Engine/Private/BatchedElements.cpp
--- a/Engine/Source/Runtime/Engine/Private/BatchedElements.cpp
+++ b/Engine/Source/Runtime/Engine/Private/BatchedElements.cpp
@@ -837,9 +837,16 @@ void FBatchedElements::DrawPointElements(FRHICommandList& RHICmdList, const FMat
 	if( Points.Num() > 0 )
 	{
 		// preallocate some memory to directly fill out
-		const int32 NumPoints = Points.Num();
-		const int32 NumTris = NumPoints * 2;
-		const int32 NumVertices = NumTris * 3;
+		const uint32 NumTris = ((uint32)Points.Num()) * 2; // even if Points.Num() == INT32_MAX, this won't overflow a uint32
+		const uint32 NumVertices = NumTris * 3; // but this could
+
+		// Prevent integer overflow to buffer overflow.
+		if (NumTris > (UINT32_MAX / 3) ||
+			NumVertices > UINT32_MAX / sizeof(FSimpleElementVertex))
+		{
+			UE_LOG(LogBatchedElements, Error, TEXT("Too many points. Will overflow uint32 buffer size. NumPoints: %d"), Points.Num());
+			return;
+		}
 
 		FRHIResourceCreateInfo CreateInfo(TEXT("FBatchedElements_Points"));
 		FBufferRHIRef VertexBufferRHI = RHICmdList.CreateBuffer(sizeof(FSimpleElementVertex) * NumVertices, BUF_VertexBuffer | BUF_Volatile, 0, ERHIAccess::VertexOrIndexBuffer, CreateInfo);
@@ -847,8 +854,8 @@ void FBatchedElements::DrawPointElements(FRHICommandList& RHICmdList, const FMat
 
 		FSimpleElementVertex* PointVertices = (FSimpleElementVertex*)VerticesPtr;
 
-		int32 VertIdx = 0;
-		for(int32 PointIndex = 0;PointIndex < NumPoints;PointIndex++)
+		uint32 VertIdx = 0;
+		for(int32 PointIndex = 0; PointIndex < Points.Num(); PointIndex++)
 		{
 			// TODO: Support quad primitives here
 			const FBatchedPoint& Point = Points[PointIndex];
@@ -924,6 +931,13 @@ bool FBatchedElements::Draw(FRHICommandList& RHICmdList, const FMeshPassProcesso
 			// Draw the line elements.
 			if( LineVertices.Num() > 0 )
 			{
+				// Prevent integer overflow to buffer overflow.
+				if (LineVertices.Num() > UINT32_MAX / sizeof(FSimpleElementVertex))
+				{
+					UE_LOG(LogBatchedElements, Error, TEXT("Too many line vertices. Will overflow uint32 buffer size. LineVertices: %d"), LineVertices.Num());
+					return false;
+				}
+
 				GraphicsPSOInit.PrimitiveType = PT_LineList;
 
 				// Set the appropriate pixel shader parameters & shader state for the non-textured elements.
@@ -1171,6 +1185,13 @@ bool FBatchedElements::Draw(FRHICommandList& RHICmdList, const FMeshPassProcesso
 
 			if (ValidSpriteCount > 0)
 			{
+				// Prevent integer overflow to buffer overflow.
+				if (ValidSpriteCount > UINT32_MAX / sizeof(FSimpleElementVertex) * 6)
+				{
+					UE_LOG(LogBatchedElements, Error, TEXT("Too many sprites. Will overflow uint32 buffer size. ValidSpriteCount: %d"), ValidSpriteCount);
+					return false;
+				}
+
 				FRHIResourceCreateInfo CreateInfo(TEXT("Sprites"));
 				FBufferRHIRef VertexBufferRHI = RHICmdList.CreateBuffer(sizeof(FSimpleElementVertex) * ValidSpriteCount * 6, BUF_VertexBuffer | BUF_Volatile, 0, ERHIAccess::VertexOrIndexBuffer, CreateInfo);
 				void* VoidPtr = RHICmdList.LockBuffer(VertexBufferRHI, 0, sizeof(FSimpleElementVertex) * ValidSpriteCount * 6, RLM_WriteOnly);
@@ -1238,6 +1259,13 @@ bool FBatchedElements::Draw(FRHICommandList& RHICmdList, const FMeshPassProcesso
 
 		if( MeshElements.Num() > 0)
 		{
+			// Prevent integer overflow to buffer overflow.
+			if (MeshVertices.Num() > UINT32_MAX / sizeof(FSimpleElementVertex))
+			{
+				UE_LOG(LogBatchedElements, Error, TEXT("Too many mesh vertices. Will overflow uint32 buffer size. MeshVertices.Num(): %d"), MeshVertices.Num());
+				return false;
+			}
+
 			FRHIResourceCreateInfo CreateInfo(TEXT("MeshElements"));
 			FBufferRHIRef VertexBufferRHI = RHICmdList.CreateBuffer(sizeof(FSimpleElementVertex) * MeshVertices.Num(), BUF_VertexBuffer | BUF_Volatile, 0, ERHIAccess::VertexOrIndexBuffer, CreateInfo);
 			void* VoidPtr = RHICmdList.LockBuffer(VertexBufferRHI, 0, sizeof(FSimpleElementVertex) * MeshVertices.Num(), RLM_WriteOnly);
