Disable GatherImplementers work from ActorSoundParameterInterface

diff --git a/Engine/Source/Runtime/Engine/Private/Audio/ActorSoundParameterInterface.cpp b/Engine/Source/Runtime/Engine/Private/Audio/ActorSoundParameterInterface.cpp
--- a/Engine/Source/Runtime/Engine/Private/Audio/ActorSoundParameterInterface.cpp
+++ b/Engine/Source/Runtime/Engine/Private/Audio/ActorSoundParameterInterface.cpp
@@ -1,16 +1,43 @@
 // Copyright Epic Games, Inc. All Rights Reserved.
+
 #include "Audio/ActorSoundParameterInterface.h"
+
 #include "AudioParameter.h"
+#include "HAL/IConsoleManager.h"
 #include "GameFramework/Actor.h"
 
 #include UE_INLINE_GENERATED_CPP_BY_NAME(ActorSoundParameterInterface)
 
+namespace ActorSoundParameterInterfaceConsoleVariables
+{
+	bool bGatherImplementers = false;
+	FAutoConsoleVariableRef CVarGatherImplementers(
+		TEXT("au.ActorSoundParameterInterface.GatherImplementers"),
+		bGatherImplementers,
+		TEXT("When true, allows the interface to search for attached components and actors that implement the interface."),
+		ECVF_Default);
+
+} // namespace ActorSoundParameterInterfaceConsoleVariables
 
 void UActorSoundParameterInterface::Fill(const AActor* OwningActor, TArray<FAudioParameter>& OutParams)
 {
 	TArray<const AActor*> Actors;
 	TArray<const UActorComponent*> Components;
-	GetImplementers(OwningActor, Actors, Components);
+
+	if (ActorSoundParameterInterfaceConsoleVariables::bGatherImplementers)
+	{
+		// This is prohibitively expensive, as it goes through our owning actor and all attached actors, looking for 
+		// components that implement the IActorSoundParameterInterface
+		GetImplementers(OwningActor, Actors, Components);
+	}
+	else
+	{
+		// Prior to the GetImplementers change, we only considered the owning actor, and no attached actors or components.
+		if (OwningActor && OwningActor->Implements<UActorSoundParameterInterface>())
+		{
+			Actors.Add(OwningActor);
+		}
+	}
 
 	TArray<FAudioParameter> TempParams;
 
