World Partition Handle Unit Test: make sure to create the world with type None to avoid initializing World Partition during tests (broken since we initialize World Partition for inactive worlds).

diff --git a/Engine/Source/Editor/UnrealEd/Private/EditorWorldUtils.cpp b/Engine/Source/Editor/UnrealEd/Private/EditorWorldUtils.cpp
--- a/Engine/Source/Editor/UnrealEd/Private/EditorWorldUtils.cpp
+++ b/Engine/Source/Editor/UnrealEd/Private/EditorWorldUtils.cpp
@@ -9,30 +9,30 @@ FScopedEditorWorld::FScopedEditorWorld()
 {
 }
 
-FScopedEditorWorld::FScopedEditorWorld(UWorld* InWorld, const UWorld::InitializationValues& InInitializationValues)
+FScopedEditorWorld::FScopedEditorWorld(UWorld* InWorld, const UWorld::InitializationValues& InInitializationValues, EWorldType::Type InWorldType)
 	: FScopedEditorWorld()
 {
-	Init(InWorld, InInitializationValues);
+	Init(InWorld, InInitializationValues, InWorldType);
 }
 
-FScopedEditorWorld::FScopedEditorWorld(const FStringView InLongPackageName, const UWorld::InitializationValues& InInitializationValues)
+FScopedEditorWorld::FScopedEditorWorld(const FStringView InLongPackageName, const UWorld::InitializationValues& InInitializationValues, EWorldType::Type InWorldType)
 	: FScopedEditorWorld()
 {
 	if (UPackage* WorldPackage = LoadWorldPackageForEditor(InLongPackageName))
 	{
 		if (UWorld* RuntimeWorld = UWorld::FindWorldInPackage(WorldPackage))
 		{
-			Init(RuntimeWorld, InInitializationValues);
+			Init(RuntimeWorld, InInitializationValues, InWorldType);
 		}
 	}
 }
 
-FScopedEditorWorld::FScopedEditorWorld(const TSoftObjectPtr<UWorld>& InSoftWorld, const UWorld::InitializationValues& InInitializationValues)
-	: FScopedEditorWorld(InSoftWorld.ToSoftObjectPath().GetLongPackageName(), InInitializationValues)
+FScopedEditorWorld::FScopedEditorWorld(const TSoftObjectPtr<UWorld>& InSoftWorld, const UWorld::InitializationValues& InInitializationValues, EWorldType::Type InWorldType)
+	: FScopedEditorWorld(InSoftWorld.ToSoftObjectPath().GetLongPackageName(), InInitializationValues, InWorldType)
 {
 }
 
-void FScopedEditorWorld::Init(UWorld* InWorld, const UWorld::InitializationValues& InInitializationValues)
+void FScopedEditorWorld::Init(UWorld* InWorld, const UWorld::InitializationValues& InInitializationValues, EWorldType::Type InWorldType)
 {
 	check(InWorld);
 	check(!InWorld->bIsWorldInitialized);
@@ -49,7 +49,7 @@ void FScopedEditorWorld::Init(UWorld* InWorld, const UWorld::InitializationValue
 	GWorld = World;
 
 	// Initialize the world
-	World->WorldType = EWorldType::Editor;
+	World->WorldType = InWorldType;
 	World->InitWorld(InInitializationValues);
 	World->PersistentLevel->UpdateModelComponents();
 	World->UpdateWorldComponents(true /*bRerunConstructionScripts*/, false /*bCurrentLevelOnly*/);
diff --git a/Engine/Source/Editor/UnrealEd/Public/EditorWorldUtils.h b/Engine/Source/Editor/UnrealEd/Public/EditorWorldUtils.h
--- a/Engine/Source/Editor/UnrealEd/Public/EditorWorldUtils.h
+++ b/Engine/Source/Editor/UnrealEd/Public/EditorWorldUtils.h
@@ -23,21 +23,21 @@ public:
 	 * @param InWorld					The world to manage.
 	 * @param InInitializationValues	The initialization values to use for the world.
 	 */
-	FScopedEditorWorld(UWorld* InWorld, const UWorld::InitializationValues& InInitializationValues);
+	FScopedEditorWorld(UWorld* InWorld, const UWorld::InitializationValues& InInitializationValues, EWorldType::Type InWorldType = EWorldType::Editor);
 
 	/**
 	 * FStringView Constructor - Load the specified package & initialize the world as an editor world.
 	 * @param InLongPackageName			Path to a package containing a world.
 	 * @param InInitializationValues	The initialization values to use for the world.
 	 */
-	FScopedEditorWorld(const FStringView InLongPackageName, const UWorld::InitializationValues& InInitializationValues);
+	FScopedEditorWorld(const FStringView InLongPackageName, const UWorld::InitializationValues& InInitializationValues, EWorldType::Type InWorldType = EWorldType::Editor);
 
 	/**
 	 * SoftObjectPtr Constructor - Initialize the provided world as an editor world.
 	 * @param InWorld					World soft object pointer.
 	 * @param InInitializationValues	The initialization values to use for the world.
 	 */
-	FScopedEditorWorld(const TSoftObjectPtr<UWorld>& InSoftWorld, const UWorld::InitializationValues& InInitializationValues);
+	FScopedEditorWorld(const TSoftObjectPtr<UWorld>& InSoftWorld, const UWorld::InitializationValues& InInitializationValues, EWorldType::Type InWorldType = EWorldType::Editor);
 
 	/**
 	 * Destructor - Destroy the provided world.
@@ -51,14 +51,13 @@ public:
 
 private:
 	FScopedEditorWorld();
-	void Init(UWorld* InWorld, const UWorld::InitializationValues& InInitializationValues);
+	void Init(UWorld* InWorld, const UWorld::InitializationValues& InInitializationValues, EWorldType::Type InWorldType = EWorldType::Editor);
 
 	UWorld* World;
 	UWorld* PrevGWorld;
 };
 
-
 /**
  * Load a world package, managing the WorldTypePreLoadMap to ensure the correct world type is specified in UWorld::PostLoad()
  */
-UNREALED_API UPackage* LoadWorldPackageForEditor(const FStringView InLongPackageName, EWorldType::Type InWorldType = EWorldType::Editor, uint32 InLoadFlags = LOAD_None);
+UNREALED_API UPackage* LoadWorldPackageForEditor(const FStringView InLongPackageName, EWorldType::Type InWorldType = EWorldType::Editor, uint32 InLoadFlags = LOAD_None);
\ No newline at end of file
diff --git a/Engine/Source/Runtime/Engine/Private/Tests/WorldPartition/WorldPartitionHandleTests.cpp b/Engine/Source/Runtime/Engine/Private/Tests/WorldPartition/WorldPartitionHandleTests.cpp
--- a/Engine/Source/Runtime/Engine/Private/Tests/WorldPartition/WorldPartitionHandleTests.cpp
+++ b/Engine/Source/Runtime/Engine/Private/Tests/WorldPartition/WorldPartitionHandleTests.cpp
@@ -1,10 +1,13 @@
 // Copyright Epic Games, Inc. All Rights Reserved.
 
 #include "Misc/AutomationTest.h"
-#include "WorldPartition/ActorDescContainer.h"
-#include "Engine/World.h"
 
 #if WITH_EDITOR
+#include "Engine/World.h"
+#include "WorldPartition/WorldPartition.h"
+#include "WorldPartition/WorldPartitionHelpers.h"
+#include "WorldPartition/ActorDescContainer.h"
+#include "EditorWorldUtils.h"
 #include "PackageTools.h"
 #endif
 
@@ -38,6 +41,27 @@ namespace WorldPartitionTests
 		auto TestTrue = [Test](const TCHAR* What, bool Value) -> bool { return Test->TestTrue(What, Value); };
 		auto TestFalse = [Test](const TCHAR* What, bool Value) -> bool { return Test->TestFalse(What, Value); };
 
+		TUniquePtr<FScopedEditorWorld> ScopedEditorWorld = MakeUnique<FScopedEditorWorld>(
+			TEXTVIEW("/Engine/WorldPartition/WorldPartitionUnitTest"),
+			UWorld::InitializationValues()
+				.RequiresHitProxies(false)
+				.ShouldSimulatePhysics(false)
+				.EnableTraceCollision(false)
+				.CreateNavigation(false)
+				.CreateAISystem(false)
+				.AllowAudioPlayback(false)
+				.CreatePhysicsScene(true),
+			EWorldType::None
+		);
+		
+		UWorld* World = ScopedEditorWorld->GetWorld();
+		if (!Test->TestNotNull(TEXT("Missing World Object"), World))
+		{
+			return;
+		}
+
+		TestTrue(TEXT("World type"), World->WorldType == EWorldType::None);
+
 		UActorDescContainer* ActorDescContainer = NewObject<UActorDescContainer>(GetTransientPackage());
 		ActorDescContainer->Initialize({ nullptr, TEXT("/Engine/WorldPartition/WorldPartitionUnitTest") });
 
@@ -270,15 +294,11 @@ namespace WorldPartitionTests
 			TestTrue(TEXT("Invalid container test"), Reference.IsValid());
 
 			// Make sure to cleanup world before collecting garbage so it gets uninitialized
-			UPackage* Package = FindPackage(NULL, *ActorDescContainer->GetContainerPackage().ToString());
-			check(Package);
-			UWorld* World = UWorld::FindWorldInPackage(Package);
-			check(World->IsInitialized());
-			World->CleanupWorld();
-			
+			ScopedEditorWorld.Reset();
 			CollectGarbage(GARBAGE_COLLECTION_KEEPFLAGS);
 
 			// Unload Package so test can run twice without issues.
+			UPackage* Package = FindPackage(NULL, *ActorDescContainer->GetContainerPackage().ToString());
 			UPackageTools::UnloadPackages({ Package });
 
 			TestFalse(TEXT("Invalid container test"), Handle.IsValid());
