[NO NON-PII COMMIT MESSAGE]

diff --git a/Engine/Source/Runtime/Engine/Classes/Components/SplineComponent.h b/Engine/Source/Runtime/Engine/Classes/Components/SplineComponent.h
--- a/Engine/Source/Runtime/Engine/Classes/Components/SplineComponent.h
+++ b/Engine/Source/Runtime/Engine/Classes/Components/SplineComponent.h
@@ -403,6 +403,10 @@ public:
 	UFUNCTION(BlueprintCallable, Category = Spline)
 	float GetDistanceAlongSplineAtSplineInputKey(float InKey) const;
 
+	/** Get distance along the spline at closest point of the provided input location */
+	UFUNCTION(BlueprintCallable, Category = Spline)
+	float GetDistanceAlongSplineAtLocation(const FVector& InLocation, ESplineCoordinateSpace::Type CoordinateSpace) const;
+
 	/** Get a metadata property float value along the spline at spline input key */
 	UFUNCTION(BlueprintCallable, Category = Spline)
 	float GetFloatPropertyAtSplineInputKey(float InKey, FName PropertyName) const;
diff --git a/Engine/Source/Runtime/Engine/Private/Components/SplineComponent.cpp b/Engine/Source/Runtime/Engine/Private/Components/SplineComponent.cpp
--- a/Engine/Source/Runtime/Engine/Private/Components/SplineComponent.cpp
+++ b/Engine/Source/Runtime/Engine/Private/Components/SplineComponent.cpp
@@ -513,6 +513,13 @@ float USplineComponent::GetDistanceAlongSplineAtSplineInputKey(float InKey) cons
 	return 0.0f;
 }
 
+float USplineComponent::GetDistanceAlongSplineAtLocation(const FVector& InLocation, ESplineCoordinateSpace::Type CoordinateSpace) const
+{
+	const FVector LocalLocation = (CoordinateSpace == ESplineCoordinateSpace::World) ? GetComponentTransform().InverseTransformPosition(InLocation) : InLocation;
+	float Dummy;
+	float Key = SplineCurves.Position.InaccurateFindNearest(LocalLocation, Dummy);
+	return GetDistanceAlongSplineAtSplineInputKey(Key);
+}
 
 template<class T>
 T GetPropertyValueAtSplineInputKey(const USplineMetadata* Metadata, float InKey, FName PropertyName)
