Remove requirement that component level PSO collection can only run on GameThread by adding a shared critical section on the PrecachedPSORequestIDs which are stored in the FMaterial - all other access should already be threadsafe

#rnx

diff --git a/Engine/Source/Runtime/Engine/Private/Components/PrimitiveComponent.cpp b/Engine/Source/Runtime/Engine/Private/Components/PrimitiveComponent.cpp
--- a/Engine/Source/Runtime/Engine/Private/Components/PrimitiveComponent.cpp
+++ b/Engine/Source/Runtime/Engine/Private/Components/PrimitiveComponent.cpp
@@ -4277,8 +4277,6 @@ void UPrimitiveComponent::PrecachePSOs()
 		return;
 	}
 
-	check(IsInGameThread());
-
 	// clear the current request data
 	MaterialPSOPrecacheRequestIDs.Empty();
 	PSOPrecacheCompileEvent = nullptr;
diff --git a/Engine/Source/Runtime/Engine/Private/Materials/MaterialShared.cpp b/Engine/Source/Runtime/Engine/Private/Materials/MaterialShared.cpp
--- a/Engine/Source/Runtime/Engine/Private/Materials/MaterialShared.cpp
+++ b/Engine/Source/Runtime/Engine/Private/Materials/MaterialShared.cpp
@@ -2143,9 +2143,15 @@ uint32 FMaterial::Release() const
 bool FMaterial::PrepareDestroy_GameThread()
 {
 	check(IsInGameThread());
-
-	ReleasePSOPrecacheData(PrecachedPSORequestIDs);
-	PrecachedPSORequestIDs.Empty();
+	
+	// Make local copy to make sure lock is held at short as possible
+	TArray<FMaterialPSOPrecacheRequestID> TmpPrecachedPSORequestIDs;	
+	{
+		FScopeLock ScopeLock(&PrecachedPSORequestIDsCS);
+		TmpPrecachedPSORequestIDs = MoveTemp(PrecachedPSORequestIDs);
+		PrecachedPSORequestIDs.Empty();
+	}
+	ReleasePSOPrecacheData(TmpPrecachedPSORequestIDs);
 
 #if WITH_EDITOR
 	const bool bReleasedCompilingId = ReleaseGameThreadCompilingShaderMap();
@@ -2810,12 +2816,12 @@ FString FMaterial::GetUniqueAssetName(EShaderPlatform Platform, const FMaterialS
 }
 #endif // WITH_EDITOR
 
+FCriticalSection FMaterial::PrecachedPSORequestIDsCS;
+
 FGraphEventArray FMaterial::CollectPSOs(ERHIFeatureLevel::Type InFeatureLevel, const FPSOPrecacheVertexFactoryDataList& VertexFactoryDataList, const FPSOPrecacheParams& PreCacheParams, EPSOPrecachePriority Priority, TArray<FMaterialPSOPrecacheRequestID>& OutMaterialPSORequestIDs)
 {
 	TRACE_CPUPROFILER_EVENT_SCOPE(FMaterial::CollectPSOs);
 	
-	check(IsInGameThread());
-
 	FGraphEventArray GraphEvents;
 	for (const FPSOPrecacheVertexFactoryData& VFData : VertexFactoryDataList)
 	{
@@ -2836,6 +2842,7 @@ FGraphEventArray FMaterial::CollectPSOs(ERHIFeatureLevel::Type InFeatureLevel, c
 			OutMaterialPSORequestIDs.AddUnique(RequestID);
 
 			// Verified in game thread above
+			FScopeLock ScopeLock(&PrecachedPSORequestIDsCS);
 			PrecachedPSORequestIDs.AddUnique(RequestID);
 		}
 	}
diff --git a/Engine/Source/Runtime/Engine/Public/MaterialShared.h b/Engine/Source/Runtime/Engine/Public/MaterialShared.h
--- a/Engine/Source/Runtime/Engine/Public/MaterialShared.h
+++ b/Engine/Source/Runtime/Engine/Public/MaterialShared.h
@@ -2408,6 +2408,9 @@ private:
 	/** MaterialPSO request IDs currently active for this material */
 	TArray<FMaterialPSOPrecacheRequestID> PrecachedPSORequestIDs;
 
+	/** Shared critical section on PrecachedPSORequestIDs because contention should be very limited */
+	static FCriticalSection PrecachedPSORequestIDsCS;
+
 #if WITH_EDITOR
 	/**
 	* Compiles this material for Platform, storing the result in OutShaderMap if the compile was synchronous
