GetTypedOuter<UWorld> && GetTypedOuter<UWorldPartition> no longer compiles when called from a cell
Remove references to GetTypedOuter<UWorld> && GetTypedOuter<UWorldPartition> when called from a cell.
Remove runtime outer change of UExternalStreamingObject
Fixes Ulevel::WorldPartitionCell not resolving in cooked build for cells stored in the UExternalStreamingObject

diff --git a/Engine/Source/Runtime/Engine/Private/WorldPartition/HLOD/HLODSubsystem.cpp b/Engine/Source/Runtime/Engine/Private/WorldPartition/HLOD/HLODSubsystem.cpp
--- a/Engine/Source/Runtime/Engine/Private/WorldPartition/HLOD/HLODSubsystem.cpp
+++ b/Engine/Source/Runtime/Engine/Private/WorldPartition/HLOD/HLODSubsystem.cpp
@@ -106,7 +106,7 @@ namespace FHLODSubsystem
 	    } // If not find it through the cell
 	    else if (const UWorldPartitionRuntimeCell* RuntimeCell = GetActorRuntimeCell(InWorldPartitionHLOD))
 	    {
-		    return RuntimeCell->GetTypedOuter<UWorldPartition>();
+		    return RuntimeCell->GetOuterWorld()->GetWorldPartition();
 	    }
     
 	    return nullptr;
@@ -285,7 +285,7 @@ const UHLODSubsystem::FCellData* UHLODSubsystem::GetCellData(const UWorldPartiti
 
 UHLODSubsystem::FCellData* UHLODSubsystem::GetCellData(const UWorldPartitionRuntimeCell* InCell)
 {
-	const UWorldPartition* WorldPartition = InCell->GetTypedOuter<UWorldPartition>();
+	const UWorldPartition* WorldPartition = InCell->GetOuterWorld()->GetWorldPartition();
 	check(WorldPartition);
 	FWorldPartitionHLODRuntimeData* WorldPartitionHLODRuntimeData = WorldPartitionsHLODRuntimeData.Find(WorldPartition);
 	if (WorldPartitionHLODRuntimeData)
diff --git a/Engine/Source/Runtime/Engine/Private/WorldPartition/WorldPartitionHelpers.cpp b/Engine/Source/Runtime/Engine/Private/WorldPartition/WorldPartitionHelpers.cpp
--- a/Engine/Source/Runtime/Engine/Private/WorldPartition/WorldPartitionHelpers.cpp
+++ b/Engine/Source/Runtime/Engine/Private/WorldPartition/WorldPartitionHelpers.cpp
@@ -4,7 +4,7 @@
 
 #include "WorldPartition/WorldPartition.h"
 #include "WorldPartition/WorldPartitionEditorHash.h"
-#include "WorldPartition/WorldPartitionRuntimeCellInterface.h"
+#include "WorldPartition/WorldPartitionRuntimeCell.h"
 #include "Engine/World.h"
 #include "Engine/Level.h"
 #include "Algo/AnyOf.h"
@@ -38,6 +38,10 @@ UWorldPartition* FWorldPartitionHelpers::GetWorldPartition(const UObject* InObje
 	{
 		return GetWorldPartition(OuterLevel);
 	}
+	else if (const UWorldPartitionRuntimeCell* Cell = Cast<const UWorldPartitionRuntimeCell>(InObject))
+	{
+		OuterWorld = Cell->GetOuterWorld();
+	}
 
 	if (!OuterWorld)
 	{
diff --git a/Engine/Source/Runtime/Engine/Private/WorldPartition/WorldPartitionLevelStreamingDynamic.cpp b/Engine/Source/Runtime/Engine/Private/WorldPartition/WorldPartitionLevelStreamingDynamic.cpp
--- a/Engine/Source/Runtime/Engine/Private/WorldPartition/WorldPartitionLevelStreamingDynamic.cpp
+++ b/Engine/Source/Runtime/Engine/Private/WorldPartition/WorldPartitionLevelStreamingDynamic.cpp
@@ -50,7 +50,7 @@ void UWorldPartitionLevelStreamingDynamic::Initialize(const UWorldPartitionRunti
 	bShouldBeAlwaysLoaded = InCell.IsAlwaysLoaded();
 	StreamingPriority = 0;
 
-	UWorld* CellOuterWorld = InCell.GetTypedOuter<UWorld>();
+	UWorld* CellOuterWorld = InCell.GetOuterWorld();
 #if WITH_EDITOR
 	check(ChildPackages.Num() == 0);
 
@@ -647,17 +647,20 @@ bool UWorldPartitionLevelStreamingDynamic::CanChangeVisibility(bool bMakeVisible
 
 			if (!bAlwaysAllowVisibilityChange)
 			{
-				if (const UWorldPartition* WorldPartition = RuntimeLevelStreamingCell->GetTypedOuter<UWorldPartition>())
+				if (const UWorld* OuterWorld = RuntimeLevelStreamingCell->GetOuterWorld())
 				{
-					if (UHLODSubsystem* HLODSubsystem = GetWorld()->GetSubsystem<UHLODSubsystem>())
+					if (const UWorldPartition* WorldPartition = OuterWorld->GetWorldPartition())
 					{
-						if (bMakeVisible)
+						if (UHLODSubsystem* HLODSubsystem = GetWorld()->GetSubsystem<UHLODSubsystem>())
 						{
-							return HLODSubsystem->CanMakeVisible(RuntimeLevelStreamingCell);
-						}
-						else
-						{
-							return HLODSubsystem->CanMakeInvisible(RuntimeLevelStreamingCell);
+							if (bMakeVisible)
+							{
+								return HLODSubsystem->CanMakeVisible(RuntimeLevelStreamingCell);
+							}
+							else
+							{
+								return HLODSubsystem->CanMakeInvisible(RuntimeLevelStreamingCell);
+							}
 						}
 					}
 				}
diff --git a/Engine/Source/Runtime/Engine/Private/WorldPartition/WorldPartitionRuntimeCell.cpp b/Engine/Source/Runtime/Engine/Private/WorldPartition/WorldPartitionRuntimeCell.cpp
--- a/Engine/Source/Runtime/Engine/Private/WorldPartition/WorldPartitionRuntimeCell.cpp
+++ b/Engine/Source/Runtime/Engine/Private/WorldPartition/WorldPartitionRuntimeCell.cpp
@@ -10,6 +10,7 @@
 #include "WorldPartition/ContentBundle/ContentBundleDescriptor.h"
 #include "WorldPartition/DataLayer/WorldDataLayers.h"
 #include "WorldPartition/WorldPartitionStreamingSource.h"
+#include "WorldPartition/WorldPartitionRuntimeHash.h"
 
 #include UE_INLINE_GENERATED_CPP_BY_NAME(WorldPartitionRuntimeCell)
 
@@ -23,6 +24,24 @@ UWorldPartitionRuntimeCell::UWorldPartitionRuntimeCell(const FObjectInitializer&
 	, RuntimeCellData(nullptr)
 {}
 
+UWorld* UWorldPartitionRuntimeCell::GetOwningWorld() const
+{
+	if (URuntimeHashExternalStreamingObjectBase* StreamingObjectOuter = GetTypedOuter<URuntimeHashExternalStreamingObjectBase>())
+	{
+		return StreamingObjectOuter->GetOwningWorld();
+	}
+	return UObject::GetTypedOuter<UWorldPartition>()->GetWorld();
+}
+
+UWorld* UWorldPartitionRuntimeCell::GetOuterWorld() const
+{
+	if (URuntimeHashExternalStreamingObjectBase* StreamingObjectOuter = GetTypedOuter<URuntimeHashExternalStreamingObjectBase>())
+	{
+		return StreamingObjectOuter->GetOuterWorld();
+	}
+	return UObject::GetTypedOuter<UWorld>();
+}
+
 #if WITH_EDITOR
 void UWorldPartitionRuntimeCell::PostDuplicate(bool bDuplicateForPIE)
 {
diff --git a/Engine/Source/Runtime/Engine/Private/WorldPartition/WorldPartitionRuntimeHash.cpp b/Engine/Source/Runtime/Engine/Private/WorldPartition/WorldPartitionRuntimeHash.cpp
--- a/Engine/Source/Runtime/Engine/Private/WorldPartition/WorldPartitionRuntimeHash.cpp
+++ b/Engine/Source/Runtime/Engine/Private/WorldPartition/WorldPartitionRuntimeHash.cpp
@@ -36,9 +36,6 @@ void URuntimeHashExternalStreamingObjectBase::OnStreamingObjectLoaded(UWorld* In
 	bool bIsACookedObject = !CellToLevelStreamingPackage.IsEmpty();
 	if (bIsACookedObject)
 	{
-		// Outer the streaming object back to the runtime hash.
-		Rename(nullptr, InjectedWorld->GetWorldPartition()->RuntimeHash);
-
 		// Cooked streaming object's Cells do not have LevelStreaming.
 		ForEachStreamingCells([this](UWorldPartitionRuntimeCell& Cell)
 		{
diff --git a/Engine/Source/Runtime/Engine/Private/WorldPartition/WorldPartitionRuntimeLevelStreamingCell.cpp b/Engine/Source/Runtime/Engine/Private/WorldPartition/WorldPartitionRuntimeLevelStreamingCell.cpp
--- a/Engine/Source/Runtime/Engine/Private/WorldPartition/WorldPartitionRuntimeLevelStreamingCell.cpp
+++ b/Engine/Source/Runtime/Engine/Private/WorldPartition/WorldPartitionRuntimeLevelStreamingCell.cpp
@@ -170,24 +170,6 @@ bool UWorldPartitionRuntimeLevelStreamingCell::IsLoading() const
 	return Super::IsLoading();
 }
 
-UWorld* UWorldPartitionRuntimeLevelStreamingCell::GetOwningWorld() const
-{
-	if (URuntimeHashExternalStreamingObjectBase* StreamingObjectOuter = GetTypedOuter<URuntimeHashExternalStreamingObjectBase>())
-	{
-		return StreamingObjectOuter->GetOwningWorld();
-	}
-	return GetTypedOuter<UWorldPartition>()->GetWorld();
-}
-
-UWorld* UWorldPartitionRuntimeLevelStreamingCell::GetOuterWorld() const
-{
-	if (URuntimeHashExternalStreamingObjectBase* StreamingObjectOuter = GetTypedOuter<URuntimeHashExternalStreamingObjectBase>())
-	{
-		return StreamingObjectOuter->GetOuterWorld();
-	}
-	return GetTypedOuter<UWorld>();
-}
-
 FLinearColor UWorldPartitionRuntimeLevelStreamingCell::GetDebugColor(EWorldPartitionRuntimeCellVisualizeMode VisualizeMode) const
 {
 	switch (VisualizeMode)
@@ -321,7 +303,7 @@ int32 UWorldPartitionRuntimeLevelStreamingCell::GetActorCount() const
 
 FString UWorldPartitionRuntimeLevelStreamingCell::GetPackageNameToCreate() const
 {
-	return UWorldPartitionLevelStreamingPolicy::GetCellPackagePath(GetFName(), GetTypedOuter<UWorld>());
+	return UWorldPartitionLevelStreamingPolicy::GetCellPackagePath(GetFName(), GetOuterWorld());
 }
 
 void UWorldPartitionRuntimeLevelStreamingCell::DumpStateLog(FHierarchicalLogArchive& Ar)
@@ -359,7 +341,7 @@ UWorldPartitionLevelStreamingDynamic* UWorldPartitionRuntimeLevelStreamingCell::
 	if (LevelStreaming)
 	{
 		// Setup pre-created LevelStreaming's outer to the WorldPartition owning world
-		const UWorldPartition* WorldPartition = GetTypedOuter<UWorldPartition>();
+		const UWorldPartition* WorldPartition = GetOuterWorld()->GetWorldPartition();
 		UWorld* OwningWorld = GetOwningWorld();
 		if (LevelStreaming->GetWorld() != OwningWorld)
 		{
@@ -500,16 +482,18 @@ void UWorldPartitionRuntimeLevelStreamingCell::Deactivate() const
 
 void UWorldPartitionRuntimeLevelStreamingCell::OnLevelShown()
 {
-	if (UWorldPartition* WorldPartition = GetTypedOuter<UWorldPartition>(); WorldPartition && WorldPartition->IsInitialized())
+	UWorldPartition* OuterWorldPartition = GetOuterWorld()->GetWorldPartition();
+	if (OuterWorldPartition && OuterWorldPartition->IsInitialized())
 	{
-		WorldPartition->OnCellShown(this);
+		OuterWorldPartition->OnCellShown(this);
 	}
 }
 
 void UWorldPartitionRuntimeLevelStreamingCell::OnLevelHidden()
 {
-	if (UWorldPartition* WorldPartition = GetTypedOuter<UWorldPartition>(); WorldPartition && WorldPartition->IsInitialized())
+	UWorldPartition* OuterWorldPartition = GetOuterWorld()->GetWorldPartition();
+	if (OuterWorldPartition && OuterWorldPartition->IsInitialized())
 	{
-		WorldPartition->OnCellHidden(this);
+		OuterWorldPartition->OnCellHidden(this);
 	}
 }
diff --git a/Engine/Source/Runtime/Engine/Public/WorldPartition/WorldPartitionRuntimeCell.h b/Engine/Source/Runtime/Engine/Public/WorldPartition/WorldPartitionRuntimeCell.h
--- a/Engine/Source/Runtime/Engine/Public/WorldPartition/WorldPartitionRuntimeCell.h
+++ b/Engine/Source/Runtime/Engine/Public/WorldPartition/WorldPartitionRuntimeCell.h
@@ -157,6 +157,17 @@ class ENGINE_API UWorldPartitionRuntimeCell : public UObject, public IWorldParti
 	virtual void PostDuplicate(bool bDuplicateForPIE) override;
 #endif
 
+	template<class T>
+	T* GetTypedOuter() const
+	{
+		static_assert(!std::is_same<T, UWorld>::value, "Use GetOuterWorld instead");
+		static_assert(!std::is_same<T, UWorldPartition>::value, "Use GetOuterWorld()->GetWorldPartition() instead");
+		return Super::GetTypedOuter<T>();
+	}
+
+	UWorld* GetOwningWorld() const;
+	UWorld* GetOuterWorld() const;
+
 	virtual void Load() const PURE_VIRTUAL(UWorldPartitionRuntimeCell::Load,);
 	virtual void Unload() const PURE_VIRTUAL(UWorldPartitionRuntimeCell::Unload,);
 	virtual bool CanUnload() const PURE_VIRTUAL(UWorldPartitionRuntimeCell::CanUnload, return true;);
diff --git a/Engine/Source/Runtime/Engine/Public/WorldPartition/WorldPartitionRuntimeLevelStreamingCell.h b/Engine/Source/Runtime/Engine/Public/WorldPartition/WorldPartitionRuntimeLevelStreamingCell.h
--- a/Engine/Source/Runtime/Engine/Public/WorldPartition/WorldPartitionRuntimeLevelStreamingCell.h
+++ b/Engine/Source/Runtime/Engine/Public/WorldPartition/WorldPartitionRuntimeLevelStreamingCell.h
@@ -32,9 +32,6 @@ class ENGINE_API UWorldPartitionRuntimeLevelStreamingCell : public UWorldPartiti
 	virtual bool IsLoading() const override;
 	//~End UWorldPartitionRuntimeCell Interface
 
-	UWorld* GetOwningWorld() const;
-	UWorld* GetOuterWorld() const;
-
 	//~Begin IWorldPartitionCell Interface
 	FName GetLevelPackageName() const override;
 	//~End IWorldPartitionCell Interface
