World Partition: fix crash when unloading data layers after deleting an actor.

diff --git a/Engine/Source/Runtime/Engine/Private/WorldPartition/WorldPartitionActorDesc.cpp b/Engine/Source/Runtime/Engine/Private/WorldPartition/WorldPartitionActorDesc.cpp
--- a/Engine/Source/Runtime/Engine/Private/WorldPartition/WorldPartitionActorDesc.cpp
+++ b/Engine/Source/Runtime/Engine/Private/WorldPartition/WorldPartitionActorDesc.cpp
@@ -38,6 +38,7 @@ FWorldPartitionActorDesc::FWorldPartitionActorDesc()
 
 void FWorldPartitionActorDesc::Init(const AActor* InActor)
 {	
+	check(IsValid(InActor));
 	check(InActor->IsPackageExternal());
 
 	Guid = InActor->GetActorGuid();
diff --git a/Engine/Source/Runtime/Engine/Private/WorldPartition/WorldPartitionActorDescViewProxy.cpp b/Engine/Source/Runtime/Engine/Private/WorldPartition/WorldPartitionActorDescViewProxy.cpp
--- a/Engine/Source/Runtime/Engine/Private/WorldPartition/WorldPartitionActorDescViewProxy.cpp
+++ b/Engine/Source/Runtime/Engine/Private/WorldPartition/WorldPartitionActorDescViewProxy.cpp
@@ -12,7 +12,7 @@
 FWorldPartitionActorViewProxy::FWorldPartitionActorViewProxy(const FWorldPartitionActorDesc* InActorDesc)
 	: FWorldPartitionActorDescView(InActorDesc)
 {
-	if (AActor* Actor = ActorDesc->GetActor())
+	if (AActor* Actor = ActorDesc->GetActor(false))
 	{
 		if (Actor->GetPackage()->IsDirty())
 		{
