Fix an ensure that could happen if you played a trigger effect without any affected triggers set. Make the default value of triggers "All" so that it works without neding to change anything, and also log if you try to apply a trigger effect that has no affected triggers.

Set the UI Min and Max of trigger resistance data so that you can't possibly get an ensure if you go higher then what the platform supports.

#rnx

diff --git a/Engine/Source/Runtime/Engine/Classes/GameFramework/InputDeviceProperties.h b/Engine/Source/Runtime/Engine/Classes/GameFramework/InputDeviceProperties.h
--- a/Engine/Source/Runtime/Engine/Classes/GameFramework/InputDeviceProperties.h
+++ b/Engine/Source/Runtime/Engine/Classes/GameFramework/InputDeviceProperties.h
@@ -276,7 +276,7 @@ struct FDeviceTriggerBaseData
 
 	/** Which trigger this property should effect */
 	UPROPERTY(EditAnywhere, BlueprintReadWrite, Category = "Triggers")
-	EInputDeviceTriggerMask AffectedTriggers = EInputDeviceTriggerMask::None;
+	EInputDeviceTriggerMask AffectedTriggers = EInputDeviceTriggerMask::All;
 
 	/** True if the triggers should be reset after the duration of this device property */
 	UPROPERTY(EditAnywhere, BlueprintReadWrite, Category = "Triggers")
@@ -293,6 +293,7 @@ public:
 
 	virtual FInputDeviceProperty* GetInternalDeviceProperty() override;
 	virtual void ResetDeviceProperty_Implementation(const FPlatformUserId PlatformUser, const FInputDeviceId DeviceId) override;
+	virtual void ApplyDeviceProperty(const FPlatformUserId UserId, const FInputDeviceId DeviceId) override;
 
 	UPROPERTY(EditAnywhere, BlueprintReadWrite, Category = "Triggers")
 	FDeviceTriggerBaseData BaseTriggerData;
@@ -365,22 +366,22 @@ struct FDeviceTriggerTriggerResistanceData
 	GENERATED_BODY()
 
 	/** The position that the trigger should start providing resistance */
-	UPROPERTY(EditAnywhere, BlueprintReadWrite, Category = "DeviceProperty", meta = (UIMin = "0"))
+	UPROPERTY(EditAnywhere, BlueprintReadWrite, Category = "DeviceProperty", meta = (UIMin = "0", UIMAX = "9"))
 	int32 StartPosition = 0;
 
 	/** How strong the resistance is */
-	UPROPERTY(EditAnywhere, BlueprintReadWrite, Category = "DeviceProperty", meta = (UIMin = "0"))
+	UPROPERTY(EditAnywhere, BlueprintReadWrite, Category = "DeviceProperty", meta = (UIMin = "0", UIMAX = "8"))
 	int32 StartStrengh = 0;
 
-	UPROPERTY(EditAnywhere, BlueprintReadWrite, Category = "DeviceProperty", meta = (UIMin = "0"))
+	UPROPERTY(EditAnywhere, BlueprintReadWrite, Category = "DeviceProperty", meta = (UIMin = "0", UIMAX = "9"))
 	int32 EndPosition = 0;
 
-	UPROPERTY(EditAnywhere, BlueprintReadWrite, Category = "DeviceProperty", meta = (UIMin = "0"))
+	UPROPERTY(EditAnywhere, BlueprintReadWrite, Category = "DeviceProperty", meta = (UIMin = "0", UIMAX = "8"))
 	int32 EndStrengh = 0;
 };
 
 /** 
-* Provides resistance to a trigger while it is being pressed between a start and end value
+* Provides linear resistance to a trigger while it is being pressed between a start and end value
 * 
 * NOTE: This property has platform specific implementations and may behave differently per platform.
 * See the docs for more details on each platform.
diff --git a/Engine/Source/Runtime/Engine/Private/GameFramework/InputDeviceProperties.cpp b/Engine/Source/Runtime/Engine/Private/GameFramework/InputDeviceProperties.cpp
--- a/Engine/Source/Runtime/Engine/Private/GameFramework/InputDeviceProperties.cpp
+++ b/Engine/Source/Runtime/Engine/Private/GameFramework/InputDeviceProperties.cpp
@@ -205,6 +205,19 @@ void UInputDeviceTriggerEffect::ResetDeviceProperty_Implementation(const FPlatfo
 	}	
 }
 
+void UInputDeviceTriggerEffect::ApplyDeviceProperty(const FPlatformUserId UserId, const FInputDeviceId DeviceId)
+{
+	// If the affected triggers is set to none then nothing will happen and the platform may throw an ensure, so just skip it and log here
+	if (BaseTriggerData.AffectedTriggers != EInputDeviceTriggerMask::None)
+	{
+		Super::ApplyDeviceProperty(UserId, DeviceId);
+	}
+	else
+	{
+		UE_LOG(LogInputDevices, Warning, TEXT("The Affected Triggers is set to none on '%s', nothing will happen!"), *GetFName().ToString());
+	}
+}
+
 ///////////////////////////////////////////////////////////////////////
 // UInputDeviceTriggerFeedbackProperty
 
@@ -473,4 +486,4 @@ const FAudioBasedVibrationData* UInputDeviceAudioBasedVibrationProperty::GetRele
 		DataToUse = OverrideData;
 	}
 	return DataToUse;
-}
\ No newline at end of file
+}
