Fixup Data Layer asset path when loading WorldPartitionActorDescs.

diff --git a/Engine/Source/Runtime/Engine/Private/WorldPartition/DataLayer/WorldDataLayersActorDesc.cpp b/Engine/Source/Runtime/Engine/Private/WorldPartition/DataLayer/WorldDataLayersActorDesc.cpp
--- a/Engine/Source/Runtime/Engine/Private/WorldPartition/DataLayer/WorldDataLayersActorDesc.cpp
+++ b/Engine/Source/Runtime/Engine/Private/WorldPartition/DataLayer/WorldDataLayersActorDesc.cpp
@@ -32,6 +32,9 @@ FArchive& operator<<(FArchive& Ar, FDataLayerInstanceDesc& Desc)
 			Ar << Desc.bIsRuntime << Desc.ShortName;
 		}
 	}
+
+	Desc.FixupAssetPath();
+
 	return Ar;
 }
 
@@ -70,6 +73,15 @@ bool operator < (const FDataLayerInstanceDesc& Lhs, const FDataLayerInstanceDesc
 	return Lhs.Name.LexicalLess(Rhs.Name);
 }
 
+// Fixup asset path (in the case where the asset was renamed)
+void FDataLayerInstanceDesc::FixupAssetPath()
+{
+	if (const UDataLayerAsset* DataLayerAsset = GetAsset())
+	{
+		AssetPath = FName(DataLayerAsset->GetPathName());
+	}
+}
+
 UDataLayerAsset* FDataLayerInstanceDesc::GetAsset() const
 {
 	if (IsUsingAsset() && !AssetPath.IsNone())
diff --git a/Engine/Source/Runtime/Engine/Private/WorldPartition/WorldPartitionActorDesc.cpp b/Engine/Source/Runtime/Engine/Private/WorldPartition/WorldPartitionActorDesc.cpp
--- a/Engine/Source/Runtime/Engine/Private/WorldPartition/WorldPartitionActorDesc.cpp
+++ b/Engine/Source/Runtime/Engine/Private/WorldPartition/WorldPartitionActorDesc.cpp
@@ -535,6 +535,18 @@ PRAGMA_ENABLE_DEPRECATION_WARNINGS
 			}
 		}
 	}
+
+	if (bIsUsingDataLayerAsset)
+	{
+		// Fixup asset path (in the case where the asset was renamed)
+		for (FName& DataLayer : DataLayers)
+		{
+			if (UDataLayerAsset* DataLayerAsset = LoadObject<UDataLayerAsset>(nullptr, *DataLayer.ToString()))
+			{
+				DataLayer = FName(DataLayerAsset->GetPathName());
+			}
+		}
+	}
 }
 
 FBox FWorldPartitionActorDesc::GetEditorBounds() const
diff --git a/Engine/Source/Runtime/Engine/Public/WorldPartition/DataLayer/WorldDataLayersActorDesc.h b/Engine/Source/Runtime/Engine/Public/WorldPartition/DataLayer/WorldDataLayersActorDesc.h
--- a/Engine/Source/Runtime/Engine/Public/WorldPartition/DataLayer/WorldDataLayersActorDesc.h
+++ b/Engine/Source/Runtime/Engine/Public/WorldPartition/DataLayer/WorldDataLayersActorDesc.h
@@ -25,6 +25,9 @@ public:
 	FString GetShortName() const;
 
 private:
+
+	void FixupAssetPath();
+
 	// DataLayerInstance Name
 	FName Name;
 	// Parent DataLayerInstance Name
