Improve tooltip for UDirectionalLightComponent::bCastShadowsOnClouds property

diff --git a/Engine/Source/Runtime/Engine/Classes/Components/DirectionalLightComponent.h b/Engine/Source/Runtime/Engine/Classes/Components/DirectionalLightComponent.h
--- a/Engine/Source/Runtime/Engine/Classes/Components/DirectionalLightComponent.h
+++ b/Engine/Source/Runtime/Engine/Classes/Components/DirectionalLightComponent.h
@@ -182,7 +182,7 @@ class ENGINE_API UDirectionalLightComponent : public ULightComponent
 	uint32 bPerPixelAtmosphereTransmittance : 1;
 
 	/**
-	 * Whether the light should cast any shadows from opaque meshes onto clouds. This is disabled for AtmosphereLight1.
+	 * Whether the light should cast any shadows from opaque meshes onto clouds. This is disabled when 'Atmosphere Sun Light Index' is set to 1.
 	 */
 	UPROPERTY(EditAnywhere, BlueprintReadOnly, Category = AtmosphereAndCloud)
 	uint32 bCastShadowsOnClouds : 1;
