Fix function overloading with select_internal by explicitly casting MCT_ShadingModel to uint. Appendix to CL 22134993.

#rnx

diff --git a/Engine/Source/Runtime/Engine/Private/Materials/HLSLMaterialDerivativeAutogen.cpp b/Engine/Source/Runtime/Engine/Private/Materials/HLSLMaterialDerivativeAutogen.cpp
--- a/Engine/Source/Runtime/Engine/Private/Materials/HLSLMaterialDerivativeAutogen.cpp
+++ b/Engine/Source/Runtime/Engine/Private/Materials/HLSLMaterialDerivativeAutogen.cpp
@@ -1209,10 +1209,22 @@ int32 FMaterialDerivativeAutogen::GenerateIfFunc(FHLSLMaterialTranslator& Transl
 		}
 		else
 		{
-			CodeFinite = FString::Printf(
-				TEXT("select(%s, select(%s, %s, %s), %s)"),
-				*CompareNotEqual, *CompareGreaterEqual,
-				*GreaterFinite, *LessFinite, *Translator.GetParameterCode(Equal));
+			if (ResultType == MCT_ShadingModel)
+			{
+				// @lh-todo: Workaround SPIR-V bug in DXC: Wrong literal type is deduced during implicit type deduction from int to uint
+				// GitHub PR: [REDACTED_PROFILE_URL]/DirectXShaderCompiler/pull/4626
+				CodeFinite = FString::Printf(
+					TEXT("select(%s, select(%s, (uint)%s, (uint)%s), (uint)%s)"),
+					*CompareNotEqual, *CompareGreaterEqual,
+					*GreaterFinite, *LessFinite, *Translator.GetParameterCode(Equal));
+			}
+			else
+			{
+				CodeFinite = FString::Printf(
+					TEXT("select(%s, select(%s, %s, %s), %s)"),
+					*CompareNotEqual, *CompareGreaterEqual,
+					*GreaterFinite, *LessFinite, *Translator.GetParameterCode(Equal));
+			}
 		}
 	}
 	else
