Fix for possible crash in FHLODActorDesc::Init() when source HLOD layer object is null

diff --git a/Engine/Source/Runtime/Engine/Private/WorldPartition/HLOD/HLODActorDesc.cpp b/Engine/Source/Runtime/Engine/Private/WorldPartition/HLOD/HLODActorDesc.cpp
--- a/Engine/Source/Runtime/Engine/Private/WorldPartition/HLOD/HLODActorDesc.cpp
+++ b/Engine/Source/Runtime/Engine/Private/WorldPartition/HLOD/HLODActorDesc.cpp
@@ -22,7 +22,12 @@ void FHLODActorDesc::Init(const AActor* InActor)
 	HLODSubActors.Reserve(HLODActor->GetSubActors().Num());
 	Algo::Transform(HLODActor->GetSubActors(), HLODSubActors, [](const FHLODSubActor& SubActor) { return FHLODSubActorDesc(SubActor.ActorGuid, SubActor.ContainerID); });
 
-	SourceHLODLayerName = HLODActor->GetSubActorsHLODLayer()->GetFName();
+	SourceHLODLayerName = NAME_None;
+	if (const UHLODLayer* HLODLayer = HLODActor->GetSubActorsHLODLayer())
+	{
+		SourceHLODLayerName = HLODLayer->GetFName();
+	}
+	
 	HLODStats = HLODActor->GetStats();
 }
 
