Iris:
* Destroy & TearOff subobject functions now work in Iris replication

diff --git a/Engine/Source/Runtime/Engine/Classes/GameFramework/Actor.h b/Engine/Source/Runtime/Engine/Classes/GameFramework/Actor.h
--- a/Engine/Source/Runtime/Engine/Classes/GameFramework/Actor.h
+++ b/Engine/Source/Runtime/Engine/Classes/GameFramework/Actor.h
@@ -3760,16 +3760,23 @@ public:
 	*/
 	void DestroyReplicatedSubObjectOnRemotePeers(UObject* SubObject);
 
+	/** Similar to the other destroy function but for subobjects owned by an ActorComponent */
+	void DestroyReplicatedSubObjectOnRemotePeers(UActorComponent* OwnerComponent, UObject* SubObject);
+
 	/**
 	* Stop replicating a subobject and tell actor channels who spawned a replica of this subobject to release ownership over it.
 	* This means that on the remote connection the network engine will stop holding a reference to the subobject and it's up to other systems 
 	* to keep that reference active or the subobject will get garbage collected.
+    * Note that the subobject won't be replicated anymore, so it's final state on the client will be the one from the last replication update sent.
 	* If you are using the legacy subobject replication method (ReplicateSubObjects() aka bReplicateUsingRegisteredSubObjectList=false) make sure the
 	* subobject doesn't get replicated there either.
 	* @param SubObject The SubObject to tear off
 	*/
 	void TearOffReplicatedSubObjectOnRemotePeers(UObject* SubObject);
 
+	/** Similar to the other tear off function but for subobjects owned by an ActorComponent */
+	void TearOffReplicatedSubObjectOnRemotePeers(UActorComponent* OwnerComponent, UObject* SubObject);
+
 	/**
 	* Register a SubObject that will get replicated along with the actor component owning it.
 	* The subobject needs to be manually removed from the list before it gets deleted.
@@ -3801,6 +3808,12 @@ private:
 	/** Array of replicated components and the list of replicated subobjects they own. Replaces the deprecated ReplicatedCompoments array. */
 	TArray<UE::Net::FReplicatedComponentInfo> ReplicatedComponentsInfo;
 
+	/** Remove the subobject from the registry list. Returns true if the object was found and removed from the list. */
+	bool RemoveReplicatedSubObjectFromList(UObject* SubObject);
+
+	/** Remove the subobject from a component's registry list. Returns true if the object was found and removed from the list. */
+	bool RemoveActorComponentReplicatedSubObjectFromList(UActorComponent* OwnerComponent, UObject* SubObject);
+
 	/** Check if a new component is replicated by the actor and must be registered in the list */
 	void AddComponentForReplication(UActorComponent* Component);
 
diff --git a/Engine/Source/Runtime/Engine/Private/ActorReplication.cpp b/Engine/Source/Runtime/Engine/Private/ActorReplication.cpp
--- a/Engine/Source/Runtime/Engine/Private/ActorReplication.cpp
+++ b/Engine/Source/Runtime/Engine/Private/ActorReplication.cpp
@@ -605,13 +605,18 @@ void AActor::AddReplicatedSubObject(UObject* SubObject, ELifetimeCondition NetCo
 #endif
 }
 
-void AActor::RemoveReplicatedSubObject(UObject* SubObject)
+bool AActor::RemoveReplicatedSubObjectFromList(UObject* SubObject)
 {
 	check(SubObject);
 	const bool bWasRemoved = ReplicatedSubObjects.RemoveSubObject(SubObject);
 
 	UE_CLOG(bWasRemoved, LogNetSubObject, Verbose, TEXT("%s (0x%p) removed replicated subobject %s (0x%p)"), *GetName(), this, *SubObject->GetName(), SubObject);
+	return bWasRemoved;
+}
 
+void AActor::RemoveReplicatedSubObject(UObject* SubObject)
+{
+	const bool bWasRemoved = RemoveReplicatedSubObjectFromList(SubObject);
 #if UE_WITH_IRIS
 	if (bWasRemoved)
 	{
@@ -620,7 +625,6 @@ void AActor::RemoveReplicatedSubObject(UObject* SubObject)
 #endif // UE_WITH_IRIS
 }
 
-
 void AActor::DestroyReplicatedSubObjectOnRemotePeers(UObject* SubObject)
 {
 	check(SubObject);
@@ -646,10 +650,67 @@ void AActor::DestroyReplicatedSubObjectOnRemotePeers(UObject* SubObject)
 
 	if (IsUsingRegisteredSubObjectList())
 	{
-		RemoveReplicatedSubObject(SubObject);
+		RemoveReplicatedSubObjectFromList(SubObject);
+	}
+}
+
+void AActor::DestroyReplicatedSubObjectOnRemotePeers(UActorComponent* OwnerComponent, UObject* SubObject)
+{
+	check(SubObject);
+
+	if (!HasAuthority())
+	{
+		// Only the authority can call this.
+		return;
+	}
+
+	UE_LOG(LogNetSubObject, Verbose, TEXT("%s::%s (0x%p) requested to Delete replicated subobject on clients %s (0x%p)"), *GetName(), *OwnerComponent->GetName(), OwnerComponent, *SubObject->GetName(), SubObject);
+
+	if (FWorldContext* const Context = GEngine->GetWorldContextFromWorld(GetWorld()))
+	{
+		for (const FNamedNetDriver& Driver : Context->ActiveNetDrivers)
+		{
+			if (Driver.NetDriver)
+			{
+				Driver.NetDriver->DeleteSubObjectOnClients(this, SubObject);
+			}
+		}
+	}
+
+	if (OwnerComponent->IsUsingRegisteredSubObjectList())
+	{
+		RemoveActorComponentReplicatedSubObjectFromList(OwnerComponent, SubObject);
 	}
 }
 
+void AActor::TearOffReplicatedSubObjectOnRemotePeers(UActorComponent* OwnerComponent, UObject* SubObject)
+{
+	check(SubObject);
+
+	if (!HasAuthority())
+	{
+		// Only the authority can call this.
+		return;
+	}
+
+	UE_LOG(LogNetSubObject, Verbose, TEXT("%s::%s (0x%p) requested to TearOff replicated subobject on clients %s (0x%p)"), *GetName(), *OwnerComponent->GetName(), this, *SubObject->GetName(), SubObject);
+
+	if (FWorldContext* const Context = GEngine->GetWorldContextFromWorld(GetWorld()))
+	{
+		for (const FNamedNetDriver& Driver : Context->ActiveNetDrivers)
+		{
+			if (Driver.NetDriver)
+			{
+				Driver.NetDriver->TearOffSubObjectOnClients(this, SubObject);
+			}
+		}
+	}
+
+	if (OwnerComponent->IsUsingRegisteredSubObjectList())
+	{
+		RemoveActorComponentReplicatedSubObjectFromList(OwnerComponent, SubObject);
+	}
+}
 
 void AActor::TearOffReplicatedSubObjectOnRemotePeers(UObject* SubObject)
 {
@@ -676,7 +737,7 @@ void AActor::TearOffReplicatedSubObjectOnRemotePeers(UObject* SubObject)
 
 	if (IsUsingRegisteredSubObjectList())
 	{
-		RemoveReplicatedSubObject(SubObject);
+		RemoveReplicatedSubObjectFromList(SubObject);
 	}
 
 }
@@ -722,21 +783,33 @@ void AActor::AddActorComponentReplicatedSubObject(UActorComponent* OwnerComponen
 	}
 }
 
-void AActor::RemoveActorComponentReplicatedSubObject(UActorComponent* OwnerComponent, UObject* SubObject)
+bool AActor::RemoveActorComponentReplicatedSubObjectFromList(UActorComponent* OwnerComponent, UObject* SubObject)
 {
 	check(OwnerComponent);
 	check(SubObject);
 
 	if (FReplicatedComponentInfo* ComponentInfo = ReplicatedComponentsInfo.FindByKey(OwnerComponent))
 	{
-		bool bWasRemoved = ComponentInfo->SubObjects.RemoveSubObject(SubObject);
+		const bool bWasRemoved = ComponentInfo->SubObjects.RemoveSubObject(SubObject);
 
 		UE_CLOG(bWasRemoved, LogNetSubObject, Verbose, TEXT("%s::%s (0x%p) removed replicated subobject %s (0x%p)"), *GetName(), *OwnerComponent->GetName(), OwnerComponent, *SubObject->GetName(), SubObject);
 
+		return bWasRemoved;
+	}
+
+	return false;
+}
+
+void AActor::RemoveActorComponentReplicatedSubObject(UActorComponent* OwnerComponent, UObject* SubObject)
+{
+	const bool bWasRemoved = RemoveActorComponentReplicatedSubObjectFromList(OwnerComponent, SubObject);
+	
 #if UE_WITH_IRIS
+	if (bWasRemoved)
+	{
 		UE::Net::FReplicationSystemUtil::EndReplicationForActorComponentSubObject(OwnerComponent, SubObject);
-#endif // UE_WITH_IRIS
 	}
+#endif
 }
 
 bool AActor::IsReplicatedSubObjectRegistered(const UObject* SubObject) const
diff --git a/Engine/Source/Runtime/Engine/Private/Components/ActorComponent.cpp b/Engine/Source/Runtime/Engine/Private/Components/ActorComponent.cpp
--- a/Engine/Source/Runtime/Engine/Private/Components/ActorComponent.cpp
+++ b/Engine/Source/Runtime/Engine/Private/Components/ActorComponent.cpp
@@ -2135,30 +2135,7 @@ void UActorComponent::DestroyReplicatedSubObjectOnRemotePeers(UObject* SubObject
 {
 	if (AActor* MyOwner = GetOwner())
 	{
-		if (!MyOwner->HasAuthority())
-		{
-			// Only the authority can call this.
-			return;
-		}
-
-		UE_LOG(LogNetSubObject, Verbose, TEXT("%s::%s (0x%p) requested to Delete replicated subobject on clients %s (0x%p)"), *MyOwner->GetName(), *GetName(), this, *SubObject->GetName(), SubObject);
-
-		if (FWorldContext* const Context = GEngine->GetWorldContextFromWorld(GetWorld()))
-		{
-			for (const FNamedNetDriver& Driver : Context->ActiveNetDrivers)
-			{
-				if (Driver.NetDriver)
-				{
-					Driver.NetDriver->DeleteSubObjectOnClients(MyOwner, SubObject);
-				}
-			}
-		}
-
-
-		if (IsUsingRegisteredSubObjectList())
-		{
-			MyOwner->RemoveActorComponentReplicatedSubObject(this, SubObject);
-		}
+		MyOwner->DestroyReplicatedSubObjectOnRemotePeers(this, SubObject);
 	}
 }
 
@@ -2166,29 +2143,7 @@ void UActorComponent::TearOffReplicatedSubObjectOnRemotePeers(UObject* SubObject
 {
 	if (AActor* MyOwner=GetOwner())
 	{
-		if (!MyOwner->HasAuthority())
-		{
-			// Only the authority can call this.
-			return;
-		}
-
-		UE_LOG(LogNetSubObject, Verbose, TEXT("%s::%s (0x%p) requested to TearOff replicated subobject on clients %s (0x%p)"), *MyOwner->GetName(), *GetName(), this, *SubObject->GetName(), SubObject);
-	
-		if (FWorldContext* const Context = GEngine->GetWorldContextFromWorld(GetWorld()))
-		{
-			for (const FNamedNetDriver& Driver : Context->ActiveNetDrivers)
-			{
-				if (Driver.NetDriver)
-				{
-					Driver.NetDriver->TearOffSubObjectOnClients(MyOwner, SubObject);
-				}
-			}
-		}
-
-		if (IsUsingRegisteredSubObjectList())
-		{
-			MyOwner->RemoveActorComponentReplicatedSubObject(this, SubObject);
-		}
+		MyOwner->TearOffReplicatedSubObjectOnRemotePeers(this, SubObject);
 	}
 }
 
diff --git a/Engine/Source/Runtime/Engine/Private/NetDriver.cpp b/Engine/Source/Runtime/Engine/Private/NetDriver.cpp
--- a/Engine/Source/Runtime/Engine/Private/NetDriver.cpp
+++ b/Engine/Source/Runtime/Engine/Private/NetDriver.cpp
@@ -3355,16 +3355,42 @@ void UNetDriver::RemoveNetworkActor(AActor* Actor)
 
 void UNetDriver::DeleteSubObjectOnClients(AActor* Actor, UObject* SubObject)
 {
+#if UE_WITH_IRIS
+	if (ReplicationSystem)
+	{
+		if (UActorReplicationBridge* Bridge = Cast<UActorReplicationBridge>(ReplicationSystem->GetReplicationBridge()))
+		{
+			constexpr EEndReplicationFlags EndReplicationFlags = EEndReplicationFlags::Destroy | EEndReplicationFlags::DestroyNetHandle | EEndReplicationFlags::ClearNetPushId;
+			Bridge->EndReplication(SubObject, EndReplicationFlags);
+		}
+	}
+	else
+#endif 
+	{
 #if UE_REPLICATED_OBJECT_REFCOUNTING
-	NetworkObjects->SetSubObjectForDeletion(Actor, SubObject);
+		NetworkObjects->SetSubObjectForDeletion(Actor, SubObject);
 #endif
+	}
 }
 
 void UNetDriver::TearOffSubObjectOnClients(AActor* Actor, UObject* SubObject)
 {
+#if UE_WITH_IRIS
+	if (ReplicationSystem)
+	{
+		if (UActorReplicationBridge* Bridge = Cast<UActorReplicationBridge>(ReplicationSystem->GetReplicationBridge()))
+		{
+			constexpr EEndReplicationFlags EndReplicationFlags = EEndReplicationFlags::TearOff;
+			Bridge->EndReplication(SubObject, EndReplicationFlags);
+		}
+	}
+	else
+#endif 
+	{
 #if UE_REPLICATED_OBJECT_REFCOUNTING
-	NetworkObjects->SetSubObjectForTearOff(Actor, SubObject);
+		NetworkObjects->SetSubObjectForTearOff(Actor, SubObject);
 #endif
+	}
 }
 
 void UNetDriver::NotifyActorRenamed(AActor* ThisActor, FName PreviousName)
