Fixed cis
#rnx
#buildfix

diff --git a/Engine/Source/Runtime/Engine/Private/WorldPartition/HLOD/HLODActorDesc.cpp b/Engine/Source/Runtime/Engine/Private/WorldPartition/HLOD/HLODActorDesc.cpp
--- a/Engine/Source/Runtime/Engine/Private/WorldPartition/HLOD/HLODActorDesc.cpp
+++ b/Engine/Source/Runtime/Engine/Private/WorldPartition/HLOD/HLODActorDesc.cpp
@@ -23,9 +23,9 @@ void FHLODActorDesc::Init(const AActor* InActor)
 	Algo::Transform(HLODActor->GetSubActors(), HLODSubActors, [](const FHLODSubActor& SubActor) { return FHLODSubActorDesc(SubActor.ActorGuid, SubActor.ContainerID); });
 
 	SourceHLODLayerName = NAME_None;
-	if (const UHLODLayer* HLODLayer = HLODActor->GetSubActorsHLODLayer())
+	if (const UHLODLayer* SourceHLODLayer = HLODActor->GetSubActorsHLODLayer())
 	{
-		SourceHLODLayerName = HLODLayer->GetFName();
+		SourceHLODLayerName = SourceHLODLayer->GetFName();
 	}
 	
 	HLODStats = HLODActor->GetStats();
