World Partition: don't report errors for invalid runtime grid references when streaming is disabled.

diff --git a/Engine/Plugins/Editor/DataValidation/Source/DataValidation/Private/WorldPartitionChangelistValidator.cpp b/Engine/Plugins/Editor/DataValidation/Source/DataValidation/Private/WorldPartitionChangelistValidator.cpp
--- a/Engine/Plugins/Editor/DataValidation/Source/DataValidation/Private/WorldPartitionChangelistValidator.cpp
+++ b/Engine/Plugins/Editor/DataValidation/Source/DataValidation/Private/WorldPartitionChangelistValidator.cpp
@@ -220,7 +220,7 @@ EDataValidationResult UWorldPartitionChangelistValidator::ValidateActorsAndDataL
 		// Invoke static WorldPartition Validation from the ActorDescContainer
 		UWorldPartition::FCheckForErrorsParams Params;
 		Params.ErrorHandler = this;
-		Params.bEnableStreaming = ULevel::GetIsStreamingDisabledFromPackage(MapPath.GetPackageName());;
+		Params.bEnableStreaming = !ULevel::GetIsStreamingDisabledFromPackage(MapPath.GetPackageName());
 
 		ContainersToValidate.ForEachActorDescContainer([&Params](const UActorDescContainer* ActorDescContainer)
 		{
diff --git a/Engine/Source/Runtime/Engine/Private/WorldPartition/WorldPartitionActorDescView.cpp b/Engine/Source/Runtime/Engine/Private/WorldPartition/WorldPartitionActorDescView.cpp
--- a/Engine/Source/Runtime/Engine/Private/WorldPartition/WorldPartitionActorDescView.cpp
+++ b/Engine/Source/Runtime/Engine/Private/WorldPartition/WorldPartitionActorDescView.cpp
@@ -13,8 +13,8 @@ FWorldPartitionActorDescView::FWorldPartitionActorDescView()
 FWorldPartitionActorDescView::FWorldPartitionActorDescView(const FWorldPartitionActorDesc* InActorDesc)
 	: ActorDesc(InActorDesc)
 	, bIsForcedNonSpatiallyLoaded(false)
-	, bInvalidDataLayers(false)
-	, bInvalidRuntimeGrid(false)
+	, bIsForcedNoRuntimeGrid(false)
+	, bInvalidDataLayers(false)	
 {}
 
 const FGuid& FWorldPartitionActorDescView::GetGuid() const
@@ -46,12 +46,7 @@ PRAGMA_ENABLE_DEPRECATION_WARNINGS
 
 FName FWorldPartitionActorDescView::GetRuntimeGrid() const
 {
-	if (bInvalidRuntimeGrid)
-	{
-		return FName();
-	}
-
-	return ActorDesc->GetRuntimeGrid();
+	return bIsForcedNoRuntimeGrid ? NAME_None : ActorDesc->GetRuntimeGrid();
 }
 
 bool FWorldPartitionActorDescView::GetActorIsEditorOnly() const
@@ -220,9 +215,9 @@ void FWorldPartitionActorDescView::SetForcedNonSpatiallyLoaded()
 	}
 }
 
-void FWorldPartitionActorDescView::SetInvalidRuntimeGrid()
+void FWorldPartitionActorDescView::SetForcedNoRuntimeGrid()
 {
-	bInvalidRuntimeGrid = true;	
+	bIsForcedNoRuntimeGrid = true;	
 }
 
 void FWorldPartitionActorDescView::SetInvalidDataLayers()
diff --git a/Engine/Source/Runtime/Engine/Private/WorldPartition/WorldPartitionStreamingGeneration.cpp b/Engine/Source/Runtime/Engine/Private/WorldPartition/WorldPartitionStreamingGeneration.cpp
--- a/Engine/Source/Runtime/Engine/Private/WorldPartition/WorldPartitionStreamingGeneration.cpp
+++ b/Engine/Source/Runtime/Engine/Private/WorldPartition/WorldPartitionStreamingGeneration.cpp
@@ -243,6 +243,14 @@ class FWorldPartitionStreamingGenerator
 		}
 	}
 
+	void ResolveRuntimeGrid(FWorldPartitionActorDescView& ActorDescView)
+	{
+		if (!bEnableStreaming)
+		{
+			ActorDescView.SetForcedNoRuntimeGrid();
+		}
+	}
+
 	void ResolveRuntimeDataLayers(FWorldPartitionActorDescView& ActorDescView, const FActorDescViewMap& ActorDescViewMap)
 	{
 		const UDataLayerManager* DataLayerManager = UDataLayerManager::GetDataLayerManager(WorldPartitionContext);
@@ -505,6 +513,7 @@ class FWorldPartitionStreamingGenerator
 		auto ResolveActorDescView = [this, &ContainerDescriptor](FWorldPartitionActorDescView& ActorDescView)
 		{
 			ResolveRuntimeSpatiallyLoaded(ActorDescView);
+			ResolveRuntimeGrid(ActorDescView);
 			ResolveRuntimeDataLayers(ActorDescView, ContainerDescriptor.ActorDescViewMap);
 		};
 
@@ -647,7 +656,7 @@ PRAGMA_ENABLE_DEPRECATION_WARNINGS
 					}
 					else
 					{
-						ActorDescView.SetInvalidRuntimeGrid();
+						ActorDescView.SetForcedNoRuntimeGrid();
 					}
 				}
 
@@ -761,8 +770,8 @@ PRAGMA_ENABLE_DEPRECATION_WARNINGS
 							}
 							else
 							{
-								RefererActorDescView->SetInvalidRuntimeGrid();
-								ReferenceActorDescView->SetInvalidRuntimeGrid();
+								RefererActorDescView->SetForcedNoRuntimeGrid();
+								ReferenceActorDescView->SetForcedNoRuntimeGrid();
 							}
 
 							NbErrorsDetected++;
diff --git a/Engine/Source/Runtime/Engine/Public/WorldPartition/WorldPartitionActorDescView.h b/Engine/Source/Runtime/Engine/Public/WorldPartition/WorldPartitionActorDescView.h
--- a/Engine/Source/Runtime/Engine/Public/WorldPartition/WorldPartitionActorDescView.h
+++ b/Engine/Source/Runtime/Engine/Public/WorldPartition/WorldPartitionActorDescView.h
@@ -76,7 +76,7 @@ PRAGMA_ENABLE_DEPRECATION_WARNINGS
 	bool ShouldValidateRuntimeGrid() const;
 
 	void SetForcedNonSpatiallyLoaded();
-	void SetInvalidRuntimeGrid();
+	void SetForcedNoRuntimeGrid();
 	void SetInvalidDataLayers();
 	void SetRuntimeDataLayerInstanceNames(TArray<FName>& InRuntimeDataLayerInstanceNames);
 	void SetRuntimeReferences(TArray<FGuid>& InRuntimeReferences);
@@ -99,8 +99,8 @@ PRAGMA_ENABLE_DEPRECATION_WARNINGS
 protected:
 	const FWorldPartitionActorDesc* ActorDesc;
 	bool bIsForcedNonSpatiallyLoaded;
-	bool bInvalidDataLayers;
-	bool bInvalidRuntimeGrid;
+	bool bIsForcedNoRuntimeGrid;
+	bool bInvalidDataLayers;	
 	TOptional<TArray<FName>> ResolvedDataLayerInstanceNames;
 	TOptional<TArray<FName>> RuntimeDataLayerInstanceNames;
 	TOptional<TArray<FGuid>> RuntimeReferences;
