Assign a NavInvoker to Camps linked to the presence of spawnable AIs

diff --git a/Engine/Source/Runtime/Engine/Classes/AI/Navigation/NavAgentSelector.h b/Engine/Source/Runtime/Engine/Classes/AI/Navigation/NavAgentSelector.h
--- a/Engine/Source/Runtime/Engine/Classes/AI/Navigation/NavAgentSelector.h
+++ b/Engine/Source/Runtime/Engine/Classes/AI/Navigation/NavAgentSelector.h
@@ -10,6 +10,7 @@ struct ENGINE_API FNavAgentSelector
 	GENERATED_USTRUCT_BODY()
 
 	static const uint32 InitializedBit = 0x80000000;
+	static const uint32 AllAgentsMask = 0x7fffffff;
 
 #if CPP
 	union
@@ -55,13 +56,19 @@ struct ENGINE_API FNavAgentSelector
 	};
 #endif
 
-	explicit FNavAgentSelector(const uint32 InBits = 0x7fffffff);
+	explicit FNavAgentSelector(const uint32 InBits = AllAgentsMask);
 
 	FORCEINLINE bool Contains(int32 AgentIndex) const
 	{
 		return (AgentIndex >= 0 && AgentIndex < 16) ? !!(PackedBits & (1 << AgentIndex)) : false;
 	}
 
+	FORCEINLINE bool ContainsAnyAgent() const
+	{
+		return PackedBits & AllAgentsMask;
+	}
+
+
 	FORCEINLINE void Set(int32 AgentIndex)
 	{
 		if (AgentIndex >= 0 && AgentIndex < 16)
