ZenLoader: Always run with GEventDrivenLoaderEnabled set to false

diff --git a/Engine/Source/Runtime/CoreUObject/Private/Serialization/AsyncLoading2.cpp b/Engine/Source/Runtime/CoreUObject/Private/Serialization/AsyncLoading2.cpp
--- a/Engine/Source/Runtime/CoreUObject/Private/Serialization/AsyncLoading2.cpp
+++ b/Engine/Source/Runtime/CoreUObject/Private/Serialization/AsyncLoading2.cpp
@@ -6958,11 +6958,6 @@ FAsyncLoadingThread2::FAsyncLoadingThread2(FIoDispatcher& InIoDispatcher, IAsync
 	, UncookedPackageLoader(InUncookedPackageLoader)
 	, PackageStore(FPackageStore::Get())
 {
-#if !WITH_EDITOR
-	IsEventDrivenLoaderEnabled(); // make sure the one time init inside runs
-	GEventDrivenLoaderEnabled = true;
-#endif
-
 	EventQueue.SetZenaphore(&AltZenaphore);
 
 	AsyncPackageLookup.Reserve(DefaultAsyncPackagesReserveCount);
diff --git a/Engine/Source/Runtime/Engine/Private/Materials/Material.cpp b/Engine/Source/Runtime/Engine/Private/Materials/Material.cpp
--- a/Engine/Source/Runtime/Engine/Private/Materials/Material.cpp
+++ b/Engine/Source/Runtime/Engine/Private/Materials/Material.cpp
@@ -557,15 +557,11 @@ void UMaterialInterface::InitDefaultMaterials()
 		}
 
 		RecursionLevel--;
-#if USE_EVENT_DRIVEN_ASYNC_LOAD_AT_BOOT_TIME
-		bInitialized = !GEventDrivenLoaderEnabled || RecursionLevel == 0;
-#else
-		bInitialized = true;
-#endif
+		bInitialized = RecursionLevel == 0;
 
 		// Now precache PSOs for all the default materials after the default materials are marked initialize
 		// PSO precaching can request default materials so they have to marked as initialized to avoid endless recursion
-		if (PipelineStateCache::IsPSOPrecachingEnabled())
+		if (bInitialized && PipelineStateCache::IsPSOPrecachingEnabled())
 		{
 			TArray<FMaterialPSOPrecacheRequestID> MaterialPrecacheRequestIDs;
 
@@ -657,7 +653,10 @@ void UMaterialInterface::PostLoadDefaultMaterials()
 		if (GPoolSpecialMaterialsCompileJobs == true)
 		{
 			GPoolSpecialMaterialsCompileJobs = false;
-			GShaderCompilingManager->FinishAllCompilation();
+			if (GShaderCompilingManager)
+			{
+				GShaderCompilingManager->FinishAllCompilation();
+			}
 		}
 	}
 }
