Fix console key not working on some international keyboard layouts

RESOLUTION: Fixed conflict with latest changes

diff --git a/Engine/Source/Runtime/Engine/Classes/GameFramework/InputSettings.h b/Engine/Source/Runtime/Engine/Classes/GameFramework/InputSettings.h
--- a/Engine/Source/Runtime/Engine/Classes/GameFramework/InputSettings.h
+++ b/Engine/Source/Runtime/Engine/Classes/GameFramework/InputSettings.h
@@ -188,9 +188,10 @@ public:
 	// UObject interface
 #if WITH_EDITOR
 	virtual void PostEditChangeChainProperty(struct FPropertyChangedChainEvent& PropertyChangedEvent) override;
-	virtual void PostReloadConfig( class FProperty* PropertyThatWasLoaded ) override;
 #endif
 
+	virtual void PostReloadConfig(class FProperty* PropertyThatWasLoaded) override;
+
 	void RemoveInvalidKeys();
 
 	virtual void PostInitProperties() override;
@@ -289,6 +290,7 @@ public:
 	
 private:
 	void PopulateAxisConfigs();
+	void AddInternationalConsoleKey();
 };
 
 /**
diff --git a/Engine/Source/Runtime/Engine/Private/UserInterface/InputSettings.cpp b/Engine/Source/Runtime/Engine/Private/UserInterface/InputSettings.cpp
--- a/Engine/Source/Runtime/Engine/Private/UserInterface/InputSettings.cpp
+++ b/Engine/Source/Runtime/Engine/Private/UserInterface/InputSettings.cpp
@@ -80,49 +80,7 @@
 	Super::PostInitProperties();
 
 	PopulateAxisConfigs();
+	AddInternationalConsoleKey();
 
-#if PLATFORM_WINDOWS
-	// If the console key is set to the default we'll see about adding the keyboard default
-	// If they've mapped any additional keys, we'll just assume they've set it up in a way they desire
-	if (ConsoleKeys.Num() == 1 && ConsoleKeys[0] == EKeys::Tilde)
-	{
-		FKey DefaultConsoleKey = EKeys::Tilde;
-		switch(PRIMARYLANGID(LOWORD(GetKeyboardLayout(0))))
-		{
-		case LANG_FRENCH:
-		case LANG_HUNGARIAN:
-			DefaultConsoleKey = FInputKeyManager::Get().GetKeyFromCodes(VK_OEM_7, 0);
-			break;
-
-		case LANG_GERMAN:
-			DefaultConsoleKey = EKeys::Caret;
-			break;
-
-		case LANG_ITALIAN:
-			DefaultConsoleKey = EKeys::Backslash;
-			break;
-
-		case LANG_SPANISH:
-			DefaultConsoleKey = FInputKeyManager::Get().GetKeyFromCodes(VK_OEM_5, 0);
-			break;
-
-		case LANG_SLOVAK:
-		case LANG_SWEDISH:
-			DefaultConsoleKey = EKeys::Section;
-			break;
-
-		case LANG_JAPANESE:
-		case LANG_RUSSIAN:
-			DefaultConsoleKey = FInputKeyManager::Get().GetKeyFromCodes(VK_OEM_3, 0);
-			break;
-		}
-
-		if (DefaultConsoleKey != EKeys::Tilde && DefaultConsoleKey.IsValid())
-		{
-			ConsoleKeys.Add(DefaultConsoleKey);
-		}
-	}
-#endif
-
 	for (const FInputActionKeyMapping& KeyMapping : ActionMappings)
 	{
@@ -176,13 +134,62 @@
 #endif
 }
 
-#if WITH_EDITOR
-void UInputSettings::PostReloadConfig( FProperty* PropertyThatWasLoaded )
+void UInputSettings::AddInternationalConsoleKey()
 {
+#if PLATFORM_WINDOWS
+	// If the console key is set to the default we'll see about adding the keyboard default
+	// If they've mapped any additional keys, we'll just assume they've set it up in a way they desire
+	if (ConsoleKeys.Num() == 1 && ConsoleKeys[0] == EKeys::Tilde)
+	{
+		FKey DefaultConsoleKey = EKeys::Tilde;
+		switch (PRIMARYLANGID(LOWORD(GetKeyboardLayout(0))))
+		{
+		case LANG_FRENCH:
+		case LANG_HUNGARIAN:
+			DefaultConsoleKey = FInputKeyManager::Get().GetKeyFromCodes(VK_OEM_7, 0);
+			break;
+
+		case LANG_GERMAN:
+			DefaultConsoleKey = EKeys::Caret;
+			break;
+
+		case LANG_ITALIAN:
+			DefaultConsoleKey = EKeys::Backslash;
+			break;
+
+		case LANG_SPANISH:
+			DefaultConsoleKey = FInputKeyManager::Get().GetKeyFromCodes(VK_OEM_5, 0);
+			break;
+
+		case LANG_SWEDISH:
+			DefaultConsoleKey = EKeys::Section;
+			break;
+
+		case LANG_JAPANESE:
+		case LANG_RUSSIAN:
+			DefaultConsoleKey = FInputKeyManager::Get().GetKeyFromCodes(VK_OEM_3, 0);
+			break;
+		}
+
+		if (DefaultConsoleKey != EKeys::Tilde && DefaultConsoleKey.IsValid())
+		{
+			ConsoleKeys.Add(DefaultConsoleKey);
+		}
+	}
+#endif
+}
+
+void UInputSettings::PostReloadConfig(FProperty* PropertyThatWasLoaded)
+{
 	Super::PostReloadConfig(PropertyThatWasLoaded);
+#if WITH_EDITOR
 	PopulateAxisConfigs();
+#endif
+	AddInternationalConsoleKey();
 }
 
+#if WITH_EDITOR
+
 void UInputSettings::PostEditChangeChainProperty(FPropertyChangedChainEvent& PropertyChangedEvent)
 {
 	Super::PostEditChangeChainProperty(PropertyChangedEvent);
