Solved bug in CL 24026507 which regards better graph traversal during MaterialCachedData update.

Issue was affecting the StaticSwitch case handling.

diff --git a/Engine/Source/Runtime/Engine/Private/Materials/MaterialCachedData.cpp b/Engine/Source/Runtime/Engine/Private/Materials/MaterialCachedData.cpp
--- a/Engine/Source/Runtime/Engine/Private/Materials/MaterialCachedData.cpp
+++ b/Engine/Source/Runtime/Engine/Private/Materials/MaterialCachedData.cpp
@@ -459,35 +459,9 @@ struct FMaterialConnectedPropertiesAnalyzer
 				break;
 			}
 
-			// Handle static switch nodes specially. Try to evaluate the condition static bool value and continue exploration along the active subgraph.
-			if (UMaterialExpressionStaticSwitchParameter* ExpressionStaticSwitchParameter = Cast<UMaterialExpressionStaticSwitchParameter>(Expression))
-			{
-				// Push the input expression that matches the state of the referenced parameter (ignoring the other branch).
-				bool ParamValue;
-				verify(EvaluateStaticBoolExpression(Expression, ParamValue, Association, ParameterIndex));
-				FExpressionInput* ExpressionInput = ParamValue ? &ExpressionStaticSwitchParameter->A : &ExpressionStaticSwitchParameter->B;
-				PushUnexploredExpression(ExpressionInput->Expression);
-				continue;
-			}
-			else if (UMaterialExpressionStaticSwitch* ExpressionStaticSwitch = Cast<UMaterialExpressionStaticSwitch>(Expression))
-			{
-				// Try to read back the input actual value.
-				bool ParamValue;
-				if (EvaluateStaticBoolExpression(ExpressionStaticSwitch->Value.Expression, ParamValue, Association, ParameterIndex))
-				{
-					// Push the active input expression.
-					FExpressionInput* ExpressionInput = ParamValue ? &ExpressionStaticSwitch->A : &ExpressionStaticSwitch->B;
-					PushUnexploredExpression(ExpressionInput->Expression);
-					continue;
-				}
-			}
-			
-			// Expression is either not a static switch node or it is but its conditional value could not be determined. We need to proceed exploring
-			// all input expressions to the current one (essentially walking backwards from outputs to inputs).
-			for (FExpressionInput* ExpressionInput : Expression->GetInputs())
-			{
-				PushUnexploredExpression(ExpressionInput->Expression);
-			}
+			// By default explore all input expressions. If static switch nodes are encountered, they
+			// may be able to continue exploration only along the active input and set this flag to false.
+			bool bVisitAllInputExpressions = true;
 
 			// Explore the current expression
 			if (UMaterialExpressionCollectionParameter* ExpressionCollectionParameter = Cast<UMaterialExpressionCollectionParameter>(Expression))
@@ -626,6 +600,39 @@ struct FMaterialConnectedPropertiesAnalyzer
 			{
 				VisitFunction(FunctionCall->MaterialFunction, Association, ParameterIndex);
 			}
+			else if (UMaterialExpressionStaticSwitchParameter* ExpressionStaticSwitchParameter = Cast<UMaterialExpressionStaticSwitchParameter>(Expression))
+			{
+				// Read the static switch parameter and only continue exploration along active branch.
+				FMaterialParameterMetadata ParameterMeta;
+				if (CachedExpressionData.GetParameterValue(EMaterialParameterType::StaticSwitch, FMemoryImageMaterialParameterInfo(Expression->GetParameterName()), ParameterMeta))
+				{
+					FExpressionInput* ExpressionInput = ParameterMeta.Value.AsStaticSwitch() ? &ExpressionStaticSwitchParameter->A : &ExpressionStaticSwitchParameter->B;
+					PushUnexploredExpression(ExpressionInput->Expression);
+					bVisitAllInputExpressions = false;
+				}
+			}
+			else if (UMaterialExpressionStaticSwitch* ExpressionStaticSwitch = Cast<UMaterialExpressionStaticSwitch>(Expression))
+			{
+				bool ParamValue = (bool)ExpressionStaticSwitch->DefaultValue;
+
+				// If no input expression specified, use the default value otherwise try to evaluate the input expression to the static bool.
+				if (!ExpressionStaticSwitch->Value.Expression || EvaluateStaticBoolExpression(ExpressionStaticSwitch->Value.Expression, ParamValue, Association, ParameterIndex))
+				{
+					// Push the active input expression.
+					FExpressionInput* ExpressionInput = ParamValue ? &ExpressionStaticSwitch->A : &ExpressionStaticSwitch->B;
+					PushUnexploredExpression(ExpressionInput->Expression);
+					bVisitAllInputExpressions = false;
+				}
+			}
+
+			// If otherwise specified, explore all input expressions to this node.
+			if (bVisitAllInputExpressions)
+			{
+				for (FExpressionInput* ExpressionInput : Expression->GetInputs())
+				{
+					PushUnexploredExpression(ExpressionInput->Expression);
+				}
+			}
 		}
 	}
 
@@ -658,49 +665,43 @@ struct FMaterialConnectedPropertiesAnalyzer
 	/** If expression is of type "static bool" this function walks the graph to find the actual value of the referenced static bool and writes the result to OutValue.
 	 * @return true if the boolean value could be determined.
 	 */
-	bool EvaluateStaticBoolExpression(UMaterialExpression* Expression, bool& OutValue, EMaterialParameterAssociation Association, int32 ParameterIndex)
+	bool EvaluateStaticBoolExpression(UMaterialExpression* Expression, bool& bOutValue, EMaterialParameterAssociation Association, int32 ParameterIndex)
 	{
 		if (!Expression)
 		{
 			return false;
 		}
-		else if (Expression->IsA<UMaterialExpressionStaticSwitchParameter>() || Expression->IsA<UMaterialExpressionStaticBoolParameter>())
+		else if (UMaterialExpressionStaticBoolParameter* ExpressionStaticBoolParameter = Cast<UMaterialExpressionStaticBoolParameter>(Expression))
 		{
-			// If this is a static switch expression, push the input expression that matches the state of the referenced parameter (ignoring the other branch).
-			// Read the parameter used by this expression if any.
 			FMaterialParameterMetadata ParameterMeta;
-			verify(Expression->GetParameterValue(ParameterMeta));
-
-			// If we're processing a function, give that a chance to override the parameter value
-			if (CurrentFunction)
+			if (CachedExpressionData.GetParameterValue(EMaterialParameterType::StaticSwitch, FMemoryImageMaterialParameterInfo(Expression->GetParameterName()), ParameterMeta))
 			{
-				FMaterialParameterMetadata OverrideParameterMeta;
-				if (CurrentFunction->GetParameterOverrideValue(ParameterMeta.Value.Type, Expression->GetParameterName(), OverrideParameterMeta))
-				{
-					ParameterMeta.Value = OverrideParameterMeta.Value;
-					ParameterMeta.ExpressionGuid = OverrideParameterMeta.ExpressionGuid;
-					ParameterMeta.bUsedAsAtlasPosition = OverrideParameterMeta.bUsedAsAtlasPosition;
-					ParameterMeta.ScalarAtlas = OverrideParameterMeta.ScalarAtlas;
-					ParameterMeta.ScalarCurve = OverrideParameterMeta.ScalarCurve;
-				}
+				bOutValue = ParameterMeta.Value.AsStaticSwitch();
+				return true;
 			}
-
-			// Read the parameter.
-			UObject* UnusedReferencedTexture = nullptr;
-			const FMaterialParameterInfo ParameterInfo(Expression->GetParameterName(), Association, ParameterIndex);
-			CachedExpressionData.AddParameter(ParameterInfo, ParameterMeta, UnusedReferencedTexture);
-
-			OutValue = ParameterMeta.Value.AsStaticSwitch();
-			return true;
+			return false;
+		}
+		else if (UMaterialExpressionStaticSwitchParameter* ExpressionStaticSwitchParameter = Cast<UMaterialExpressionStaticSwitchParameter>(Expression))
+		{
+			FMaterialParameterMetadata ParameterMeta;
+			if (CachedExpressionData.GetParameterValue(EMaterialParameterType::StaticSwitch, FMemoryImageMaterialParameterInfo(Expression->GetParameterName()), ParameterMeta))
+			{
+				FExpressionInput* OperandExpressionInput = ParameterMeta.Value.AsStaticSwitch() ? &ExpressionStaticSwitchParameter->A : &ExpressionStaticSwitchParameter->B;
+				return EvaluateStaticBoolExpression(OperandExpressionInput->Expression, bOutValue, Association, ParameterIndex);
+			}
+			return false;
 		}
 		else if (UMaterialExpressionStaticSwitch* ExpressionStaticSwitch = Cast<UMaterialExpressionStaticSwitch>(Expression))
 		{
-			if (!EvaluateStaticBoolExpression(ExpressionStaticSwitch->Value.Expression, OutValue, Association, ParameterIndex))
+			bool Value = (bool)ExpressionStaticSwitch->DefaultValue;
+			// If node has a bool expression connected, try to evaluate that subgraph to a static boolean value.
+			if (ExpressionStaticSwitch->Value.Expression && !EvaluateStaticBoolExpression(ExpressionStaticSwitch->Value.Expression, Value, Association, ParameterIndex))
 			{
+				// Node has the bool expression slot connected, but we could not evaluate it to a static boolean value.
 				return false;
 			}
-			FExpressionInput* OperandExpressionInput = OutValue ? &ExpressionStaticSwitch->A : &ExpressionStaticSwitch->B;
-			return EvaluateStaticBoolExpression(OperandExpressionInput->Expression, OutValue, Association, ParameterIndex);
+			FExpressionInput* OperandExpressionInput = Value ? &ExpressionStaticSwitch->A : &ExpressionStaticSwitch->B;
+			return EvaluateStaticBoolExpression(OperandExpressionInput->Expression, bOutValue, Association, ParameterIndex);
 		}
 		else if (UMaterialExpressionNamedRerouteUsage* ExpressionRerouteUsage = Cast<UMaterialExpressionNamedRerouteUsage>(Expression))
 		{
@@ -708,15 +709,15 @@ struct FMaterialConnectedPropertiesAnalyzer
 			{
 				return false;
 			}
-			return EvaluateStaticBoolExpression(ExpressionRerouteUsage->Declaration->Input.Expression, OutValue, Association, ParameterIndex);
+			return EvaluateStaticBoolExpression(ExpressionRerouteUsage->Declaration->Input.Expression, bOutValue, Association, ParameterIndex);
 		}
 		else if (UMaterialExpressionFunctionInput* ExpressionFunctionInput = Cast<UMaterialExpressionFunctionInput>(Expression))
 		{
-			return EvaluateStaticBoolExpression(ExpressionFunctionInput->Preview.Expression, OutValue, Association, ParameterIndex);
+			return EvaluateStaticBoolExpression(ExpressionFunctionInput->Preview.Expression, bOutValue, Association, ParameterIndex);
 		}
 		else if (UMaterialExpressionStaticBool* ExpressionStaticBool = Cast<UMaterialExpressionStaticBool>(Expression))
 		{
-			OutValue = ExpressionStaticBool->Value;
+			bOutValue = ExpressionStaticBool->Value;
 			return true;
 		}
 
