Support to create a custom UContentBundleTypeFactory via engine config file.
The type factory currently only creates the content bundle clients.
Allows projects to override UContentBundleTypeFactory & FContentBundleClient to customize injection behavior and react on content bundle events.

diff --git a/Engine/Source/Runtime/Engine/Private/WorldPartition/ContentBundle/ContentBundleEngineSubsystem.cpp b/Engine/Source/Runtime/Engine/Private/WorldPartition/ContentBundle/ContentBundleEngineSubsystem.cpp
--- a/Engine/Source/Runtime/Engine/Private/WorldPartition/ContentBundle/ContentBundleEngineSubsystem.cpp
+++ b/Engine/Source/Runtime/Engine/Private/WorldPartition/ContentBundle/ContentBundleEngineSubsystem.cpp
@@ -2,6 +2,7 @@
 #include "WorldPartition/ContentBundle/ContentBundleEngineSubsystem.h"
 
 #include "WorldPartition/ContentBundle/ContentBundleClient.h"
+#include "WorldPartition/ContentBundle/ContentBundleTypeFactory.h"
 #include "WorldPartition/ContentBundle/ContentBundleDescriptor.h"
 #include "WorldPartition/ContentBundle/ContentBundleLog.h"
 #include "WorldPartition/ContentBundle/ContentBundleWorldSubsystem.h"
@@ -13,6 +14,14 @@ void UContentBundleEngineSubsystem::Initialize(FSubsystemCollectionBase& Collect
 {
 	Super::Initialize(Collection);
 
+	UClass* ContentBundlTypeFactoryClassValue = ContentBundleTypeFactoryClass.Get();
+	if (!ContentBundlTypeFactoryClassValue)
+	{
+		ContentBundlTypeFactoryClassValue = UContentBundleTypeFactory::StaticClass();
+	}
+	ContentBundleTypeFactory = NewObject<UContentBundleTypeFactory>(this, ContentBundlTypeFactoryClassValue);
+	check(ContentBundleTypeFactory);
+
 	FWorldDelegates::OnPreWorldInitialization.AddUObject(this, &UContentBundleEngineSubsystem::OnPreWorldInitialization);
 	FWorldDelegates::OnPostWorldCleanup.AddUObject(this, &UContentBundleEngineSubsystem::OnWorldPostCleanup);
 }
@@ -23,6 +32,8 @@ void UContentBundleEngineSubsystem::Deinitialize()
 
 	FWorldDelegates::OnPreWorldInitialization.RemoveAll(this);
 	FWorldDelegates::OnPostWorldCleanup.RemoveAll(this);
+
+	ContentBundleTypeFactory = nullptr;
 }
 
 TSharedPtr<FContentBundleClient> UContentBundleEngineSubsystem::RegisterContentBundle(const UContentBundleDescriptor* Descriptor, const FString& ClientDisplayName)
@@ -39,7 +50,7 @@ TSharedPtr<FContentBundleClient> UContentBundleEngineSubsystem::RegisterContentB
 		{
 			check(FindRegisteredClient(Descriptor) == nullptr);;
 
-			TSharedPtr<FContentBundleClient>& ContentBundleClient = ContentBundleClients.Add_GetRef(MakeShared<FContentBundleClient>(Descriptor, ClientDisplayName));
+			TSharedPtr<FContentBundleClient>& ContentBundleClient = ContentBundleClients.Add_GetRef(ContentBundleTypeFactory->CreateClient(Descriptor, ClientDisplayName));
 
 			UE_LOG(LogContentBundle, Log, TEXT("[Client] New client registered. Client: %s, Descriptor: %s"), *ContentBundleClient->GetDisplayName(), *ContentBundleClient->GetDescriptor()->GetDisplayName());
 
diff --git a/Engine/Source/Runtime/Engine/Private/WorldPartition/ContentBundle/ContentBundleTypeFactory.cpp b/Engine/Source/Runtime/Engine/Private/WorldPartition/ContentBundle/ContentBundleTypeFactory.cpp
new file mode 100644
--- /dev/null
+++ b/Engine/Source/Runtime/Engine/Private/WorldPartition/ContentBundle/ContentBundleTypeFactory.cpp
@@ -0,0 +1,9 @@
+// Copyright Epic Games, Inc. All Rights Reserved.
+#include "WorldPartition/ContentBundle/ContentBundleTypeFactory.h"
+
+#include "WorldPartition/ContentBundle/ContentBundleClient.h"
+
+TSharedPtr<FContentBundleClient> UContentBundleTypeFactory::CreateClient(const UContentBundleDescriptor* Descriptor, const FString& ClientDisplayName)
+{
+	return MakeShared<FContentBundleClient>(Descriptor, ClientDisplayName);
+}
\ No newline at end of file
diff --git a/Engine/Source/Runtime/Engine/Public/WorldPartition/ContentBundle/ContentBundleEngineSubsystem.h b/Engine/Source/Runtime/Engine/Public/WorldPartition/ContentBundle/ContentBundleEngineSubsystem.h
--- a/Engine/Source/Runtime/Engine/Public/WorldPartition/ContentBundle/ContentBundleEngineSubsystem.h
+++ b/Engine/Source/Runtime/Engine/Public/WorldPartition/ContentBundle/ContentBundleEngineSubsystem.h
@@ -11,9 +11,10 @@
 
 class FContentBundleClient;
 class UContentBundleDescriptor;
+class UContentBundleTypeFactory;
 class URuntimeHashExternalStreamingObjectBase;
 
-UCLASS()
+UCLASS(Config = Engine)
 class ENGINE_API UContentBundleEngineSubsystem : public UEngineSubsystem
 {
 	GENERATED_BODY()
@@ -50,4 +51,10 @@ private:
 	TSharedPtr<FContentBundleClient>* FindRegisteredClient(FContentBundleClient& InClient);
 
 	TArray<TSharedPtr<FContentBundleClient>> ContentBundleClients;
+
+	UPROPERTY(Config)
+	TSoftClassPtr<UContentBundleTypeFactory> ContentBundleTypeFactoryClass;
+
+	UPROPERTY(Transient)
+	TObjectPtr<UContentBundleTypeFactory> ContentBundleTypeFactory;
 };
\ No newline at end of file
diff --git a/Engine/Source/Runtime/Engine/Public/WorldPartition/ContentBundle/ContentBundleTypeFactory.h b/Engine/Source/Runtime/Engine/Public/WorldPartition/ContentBundle/ContentBundleTypeFactory.h
new file mode 100644
--- /dev/null
+++ b/Engine/Source/Runtime/Engine/Public/WorldPartition/ContentBundle/ContentBundleTypeFactory.h
@@ -0,0 +1,20 @@
+// Copyright Epic Games, Inc. All Rights Reserved.
+#pragma once
+
+#include "CoreMinimal.h"
+#include "UObject/ObjectMacros.h"
+#include "UObject/Object.h"
+
+#include "ContentBundleTypeFactory.generated.h"
+
+class FContentBundleClient;
+class UContentBundleDescriptor;
+
+UCLASS()
+class ENGINE_API UContentBundleTypeFactory : public UObject
+{
+	GENERATED_BODY()
+
+public:
+	virtual TSharedPtr<FContentBundleClient> CreateClient(const UContentBundleDescriptor* Descriptor, const FString& ClientDisplayName);
+};
\ No newline at end of file
