Virtualize ContentBundleClient interface so it can be overriden per projet
Remove force injection code from ContentBundleClient
Add a notification to the ContentBundleClient when one of its bundle is registered
HasContentToRemove is now basd on ShouldRemoveContent

diff --git a/Engine/Source/Runtime/Engine/Private/WorldPartition/ContentBundle/ContentBundleBase.cpp b/Engine/Source/Runtime/Engine/Private/WorldPartition/ContentBundle/ContentBundleBase.cpp
--- a/Engine/Source/Runtime/Engine/Private/WorldPartition/ContentBundle/ContentBundleBase.cpp
+++ b/Engine/Source/Runtime/Engine/Private/WorldPartition/ContentBundle/ContentBundleBase.cpp
@@ -29,6 +29,11 @@ void FContentBundleBase::Initialize()
 	DoInitialize();
 
 	check(GetStatus() == EContentBundleStatus::Registered);
+
+	if (TSharedPtr<FContentBundleClient> PinnedClient = Client.Pin())
+	{
+		PinnedClient->OnContentRegisteredInWorld(GetStatus(), GetInjectedWorld());
+	}
 }
 
 void FContentBundleBase::Uninitialize()
diff --git a/Engine/Source/Runtime/Engine/Private/WorldPartition/ContentBundle/ContentBundleClient.cpp b/Engine/Source/Runtime/Engine/Private/WorldPartition/ContentBundle/ContentBundleClient.cpp
--- a/Engine/Source/Runtime/Engine/Private/WorldPartition/ContentBundle/ContentBundleClient.cpp
+++ b/Engine/Source/Runtime/Engine/Private/WorldPartition/ContentBundle/ContentBundleClient.cpp
@@ -37,23 +37,24 @@ void FContentBundleClient::RequestContentInjection()
 
 void FContentBundleClient::RequestRemoveContent()
 {
-	if (State == EContentBundleClientState::ContentInjectionRequested)
+	if (State == EContentBundleClientState::Unregistered)
 	{
-		if (HasContentToRemove())
-		{
-			SetState(EContentBundleClientState::ContentRemovalRequested);
-			GEngine->GetEngineSubsystem<UContentBundleEngineSubsystem>()->RequestContentRemoval(*this);
-		}
-		else
-		{
-			// No ContentBundleBundles were injected (Non-WP Worlds). Reset the state to Register.
-			// Since the client will never receive OnContentRemovedFromWorld.
-			SetState(EContentBundleClientState::Registered);
-		}
+		check(!HasContentToRemove());
+		UE_LOG(LogContentBundle, Warning, TEXT("[CB: %s] Requesting content removal on an unregistered client."), *GetDescriptor()->GetDisplayName());
+		return;
+	}
+
+	SetState(EContentBundleClientState::ContentRemovalRequested);
+
+	if (HasContentToRemove())
+	{
+		GEngine->GetEngineSubsystem<UContentBundleEngineSubsystem>()->RequestContentRemoval(*this);
 	}
 	else
 	{
-		UE_LOG(LogContentBundle, Log, TEXT("[CB: %s] Client failed to request content removal. Its state is %s"), *GetDescriptor()->GetDisplayName(), *UEnum::GetDisplayValueAsText(State).ToString());
+		// No ContentBundleBundles were injected (Non-WP Worlds). Reset the state to Register.
+		// Since the client will never receive OnContentRemovedFromWorld.
+		SetState(EContentBundleClientState::Registered);
 	}
 }
 
@@ -72,28 +73,19 @@ void FContentBundleClient::RequestUnregister()
 
 bool FContentBundleClient::ShouldInjectContent(UWorld* World) const
 {
-	bool bIsClientInjecting = GetState() == EContentBundleClientState::ContentInjectionRequested;
-
-#if WITH_EDITOR
-	bool bIsClientForceInjectingWorld = GetState() != EContentBundleClientState::Unregistered && ForceInjectedWorlds.Contains(World);
-#else
-	bool bIsClientForceInjectingWorld = false;
-#endif
-	
-	return (bIsClientInjecting || bIsClientForceInjectingWorld) && GetWorldContentState(World) == EWorldContentState::NoContent;
+	return GetState() == EContentBundleClientState::ContentInjectionRequested && GetWorldContentState(World) == EWorldContentState::NoContent;
 }
 
 bool FContentBundleClient::ShouldRemoveContent(UWorld* World) const
 {
-	bool bIsClientInjecting = GetState() == EContentBundleClientState::ContentInjectionRequested;
+	return GetState() != EContentBundleClientState::ContentInjectionRequested && GetWorldContentState(World) != EWorldContentState::NoContent;
+}
 
-#if WITH_EDITOR
-	bool bIsClientForceInjectingWorld = GetState() != EContentBundleClientState::Unregistered && ForceInjectedWorlds.Contains(World);
-#else
-	bool bIsClientForceInjectingWorld = false;
-#endif
+void FContentBundleClient::OnContentRegisteredInWorld(EContentBundleStatus ContentBundleStatus, UWorld* World)
+{
+	check(ContentBundleStatus == EContentBundleStatus::Registered);
 
-	return !bIsClientInjecting && !bIsClientForceInjectingWorld && GetWorldContentState(World) != EWorldContentState::NoContent;
+	DoOnContentRegisteredInWorld(World);
 }
 
 void FContentBundleClient::OnContentInjectedInWorld(EContentBundleStatus InjectionStatus, UWorld* InjectedWorld)
@@ -104,6 +96,8 @@ void FContentBundleClient::OnContentInjectedInWorld(EContentBundleStatus Injecti
 		|| InjectionStatus == EContentBundleStatus::FailedToInject);
 
 	SetWorldContentState(InjectedWorld, EWorldContentState::ContentBundleInjected);
+
+	DoOnContentInjectedInWorld(InjectionStatus, InjectedWorld);
 }
 
 void FContentBundleClient::OnContentRemovedFromWorld(EContentBundleStatus RemovalStatus, UWorld* InjectedWorld)
@@ -112,11 +106,9 @@ void FContentBundleClient::OnContentRemovedFromWorld(EContentBundleStatus Remova
 	check(ShouldRemoveContent(InjectedWorld) || !InjectedWorld->ContentBundleManager->CanInject());
 	check(RemovalStatus == EContentBundleStatus::Registered);
 
-	WorldContentStates.Remove(InjectedWorld);
+	SetWorldContentState(InjectedWorld, EWorldContentState::NoContent);
 
-#if WITH_EDITOR
-	ForceInjectedWorlds.Remove(InjectedWorld);
-#endif
+	DoOnContentRemovedFromWorld(InjectedWorld);
 
 	// If the client request content removal and removed everything set the client state as registered.
 	// If the client did not request content removal but removed everything keep the client state. We just changed levels. The next level will have its content bundle injected.
@@ -125,7 +117,6 @@ void FContentBundleClient::OnContentRemovedFromWorld(EContentBundleStatus Remova
 	{
 		SetState(EContentBundleClientState::Registered);
 	}
-
 }
 
 void FContentBundleClient::SetWorldContentState(UWorld* World, EWorldContentState NewState)
@@ -156,14 +147,7 @@ bool FContentBundleClient::HasContentToRemove() const
 {
 	for (auto& WorldState : WorldContentStates)
 	{
-#if WITH_EDITOR
-		if (ForceInjectedWorlds.Contains(WorldState.Key))
-		{
-			continue;
-		}
-#endif
-
-		if (WorldState.Value != EWorldContentState::NoContent)
+		if (ShouldRemoveContent(WorldState.Key.Get()))
 		{
 			return true;
 		}
@@ -178,46 +162,4 @@ void FContentBundleClient::SetState(EContentBundleClientState NewState)
 
 	UE_LOG(LogContentBundle, Log, TEXT("%s Client State changing from %s to %s"), *ContentBundle::Log::MakeDebugInfoString(*this), *UEnum::GetDisplayValueAsText(State).ToString(), *UEnum::GetDisplayValueAsText(NewState).ToString());
 	State = NewState;
-}
-
-#if WITH_EDITOR
-
-void FContentBundleClient::RequestForceInject(UWorld* WorldToInject)
-{
-	if (UContentBundleManager* ContentBundleManager = WorldToInject->ContentBundleManager)
-	{
-		UE_LOG(LogContentBundle, Log, TEXT("%s Client requested a forced injection in world."), *ContentBundle::Log::MakeDebugInfoString(*this, WorldToInject));
-
-		ForceInjectedWorlds.Add(WorldToInject);
-	
-		if (ShouldInjectContent(WorldToInject))
-		{
-			ContentBundleManager->TryInject(*this);
-		}
-	}
-	else
-	{
-		UE_LOG(LogContentBundle, Log, TEXT("%s Client cannot force inject in World. It does not support content bundles."), *ContentBundle::Log::MakeDebugInfoString(*this, WorldToInject));
-	}
-}
-
-void FContentBundleClient::RequestRemoveForceInjectedContent(UWorld* WorldToInject)
-{
-	if (UContentBundleManager* ContentBundleManager = WorldToInject->ContentBundleManager)
-	{
-		UE_LOG(LogContentBundle, Log, TEXT("%s Client requested removal of force injected content in world."), *ContentBundle::Log::MakeDebugInfoString(*this, WorldToInject));
-
-		ForceInjectedWorlds.Remove(WorldToInject);
-
-		if (ShouldRemoveContent(WorldToInject))
-		{
-			ContentBundleManager->Remove(*this);
-		}
-	}
-	else
-	{
-		UE_LOG(LogContentBundle, Log, TEXT("%s Client cannot remove force injected content from world. It does not support content bundles."), *ContentBundle::Log::MakeDebugInfoString(*this, WorldToInject));
-	}
-}
-
-#endif
+}
\ No newline at end of file
diff --git a/Engine/Source/Runtime/Engine/Private/WorldPartition/ContentBundle/ContentBundleWorldSubsystem.cpp b/Engine/Source/Runtime/Engine/Private/WorldPartition/ContentBundle/ContentBundleWorldSubsystem.cpp
--- a/Engine/Source/Runtime/Engine/Private/WorldPartition/ContentBundle/ContentBundleWorldSubsystem.cpp
+++ b/Engine/Source/Runtime/Engine/Private/WorldPartition/ContentBundle/ContentBundleWorldSubsystem.cpp
@@ -168,27 +168,6 @@ TSharedPtr<FContentBundleEditor> UContentBundleManager::GetEditorContentBundle(c
 	return nullptr;
 }
 
-bool UContentBundleManager::TryInject(FContentBundleClient& Client)
-{
-	if (CanInject())
-	{
-		if (TUniquePtr<FContentBundleContainer>* Container = GetContentBundleContainer(GetTypedOuter<UWorld>()))
-		{
-			return (*Container)->InjectContentBundle(Client);
-		}
-	}
-
-	return false;
-}
-
-void UContentBundleManager::Remove(FContentBundleClient& Client)
-{
-	if (TUniquePtr<FContentBundleContainer>* Container = GetContentBundleContainer(GetTypedOuter<UWorld>()))
-	{
-		(*Container)->RemoveContentBundle(Client);
-	}
-}
-
 #endif
 
 uint32 UContentBundleManager::GetContentBundleContainerIndex(const UWorld* InjectedWorld) const
diff --git a/Engine/Source/Runtime/Engine/Public/WorldPartition/ContentBundle/ContentBundleClient.h b/Engine/Source/Runtime/Engine/Public/WorldPartition/ContentBundle/ContentBundleClient.h
--- a/Engine/Source/Runtime/Engine/Public/WorldPartition/ContentBundle/ContentBundleClient.h
+++ b/Engine/Source/Runtime/Engine/Public/WorldPartition/ContentBundle/ContentBundleClient.h
@@ -49,36 +49,33 @@ public:
 	
 	void RequestUnregister();
 
-#if WITH_EDITOR
-	void RequestForceInject(UWorld* WorldToInject);
-	void RequestRemoveForceInjectedContent(UWorld *WorldToInject);
-#endif 
-
 	EContentBundleClientState GetState() const { return State; }
 
 	FString const& GetDisplayName() const { return DisplayName; }
 
-	bool ShouldInjectContent(UWorld* World) const;
-	bool ShouldRemoveContent(UWorld* World) const;
+	virtual bool ShouldInjectContent(UWorld* World) const;
+	virtual bool ShouldRemoveContent(UWorld* World) const;
+
+protected:
+	virtual void DoOnContentRegisteredInWorld(UWorld* InjectedWorld) {};
+	virtual void DoOnContentInjectedInWorld(EContentBundleStatus InjectionStatus, UWorld* InjectedWorld) {};
+	virtual void DoOnContentRemovedFromWorld(UWorld* InjectedWorld) {};
 
 private:
 	bool HasContentToRemove() const;
 
+	void OnContentRegisteredInWorld(EContentBundleStatus ContentBundleStatus, UWorld* World);
+	void OnContentInjectedInWorld(EContentBundleStatus InjectionStatus, UWorld* InjectedWorld);
+	void OnContentRemovedFromWorld(EContentBundleStatus RemovalStatus, UWorld* InjectedWorld);
+
 	void SetState(EContentBundleClientState State);
 	void SetWorldContentState(UWorld* World, EWorldContentState State);
 	EWorldContentState GetWorldContentState(UWorld* World) const;
 
-	void OnContentInjectedInWorld(EContentBundleStatus InjectionStatus, UWorld* InjectedWorld);
-	void OnContentRemovedFromWorld(EContentBundleStatus RemovalStatus, UWorld* InjectedWorld);
-
 	TWeakObjectPtr<const UContentBundleDescriptor> ContentBundleDescriptor;
 	
 	TMap<TWeakObjectPtr<UWorld>, EWorldContentState> WorldContentStates;
 
-#if WITH_EDITOR
-	TSet<TWeakObjectPtr<UWorld>> ForceInjectedWorlds;
-#endif
-
 	FString DisplayName;
 
 	EContentBundleClientState State;
diff --git a/Engine/Source/Runtime/Engine/Public/WorldPartition/ContentBundle/ContentBundleWorldSubsystem.h b/Engine/Source/Runtime/Engine/Public/WorldPartition/ContentBundle/ContentBundleWorldSubsystem.h
--- a/Engine/Source/Runtime/Engine/Public/WorldPartition/ContentBundle/ContentBundleWorldSubsystem.h
+++ b/Engine/Source/Runtime/Engine/Public/WorldPartition/ContentBundle/ContentBundleWorldSubsystem.h
@@ -26,8 +26,6 @@ class ENGINE_API UContentBundleManager : public UObject
 {
 	GENERATED_BODY()
 
-	friend class FContentBundleClient;
-
 public:
 	UContentBundleManager();
 
@@ -50,11 +48,6 @@ public:
 	void DrawContentBundlesStatus(const UWorld* InWorld, UCanvas* Canvas, FVector2D& Offset) const;
 
 private:
-#if WITH_EDITOR
-	bool TryInject(FContentBundleClient& Client);
-	void Remove(FContentBundleClient& Client);
-#endif
-
 	uint32 GetContentBundleContainerIndex(const UWorld* InjectedWorld) const;
 	const TUniquePtr<FContentBundleContainer>* GetContentBundleContainer(const UWorld* InjectedWorld) const;
 	TUniquePtr<FContentBundleContainer>* GetContentBundleContainer(const UWorld* InjectedWorld);
