Implemented first version of streaming for SparseVolumeTextures. Fixed PostEditChange() being called twice on newly imported SVT. Small refactors on UStreamableSVT, UStaticSVT and UAnimatedSVT.

#rnx

diff --git a/Engine/Source/Editor/SparseVolumeTexture/Private/SparseVolumeTextureFactory.cpp b/Engine/Source/Editor/SparseVolumeTexture/Private/SparseVolumeTextureFactory.cpp
--- a/Engine/Source/Editor/SparseVolumeTexture/Private/SparseVolumeTextureFactory.cpp
+++ b/Engine/Source/Editor/SparseVolumeTexture/Private/SparseVolumeTextureFactory.cpp
@@ -535,7 +535,6 @@ UObject* USparseVolumeTextureFactory::FactoryCreateFile(UClass* InClass, UObject
 		ImportTask.MakeDialog(true);
 
 		// Allocate space for each frame
-		AnimatedSVTexture->FrameCount = NumFrames;
 		AnimatedSVTexture->Frames.SetNum(NumFrames);
 
 		FBox VolumeBounds(FVector(FLT_MAX), FVector(-FLT_MAX));
@@ -628,19 +627,10 @@ UObject* USparseVolumeTextureFactory::FactoryCreateFile(UClass* InClass, UObject
 	}
 
 	// Now notify the system about the imported/updated/created assets
-	AdditionalImportedObjects.Reserve(ResultAssets.Num());
-	for (UObject* Object : ResultAssets)
-	{
-		if (Object)
-		{
-			GEditor->GetEditorSubsystem<UImportSubsystem>()->BroadcastAssetPostImport(this, Object);
-			Object->MarkPackageDirty();
-			Object->PostEditChange();
-			AdditionalImportedObjects.Add(Object);
-		}
-	}
+	check(ResultAssets.Num() == 1);
+	GEditor->GetEditorSubsystem<UImportSubsystem>()->BroadcastAssetPostImport(this, ResultAssets[0]);
 
-	return (ResultAssets.Num() > 0) ? ResultAssets[0] : nullptr;
+	return ResultAssets[0];
 
 #else // OPENVDB_AVAILABLE
 
diff --git a/Engine/Source/Runtime/Engine/Classes/SparseVolumeTexture/SparseVolumeTexture.h b/Engine/Source/Runtime/Engine/Classes/SparseVolumeTexture/SparseVolumeTexture.h
--- a/Engine/Source/Runtime/Engine/Classes/SparseVolumeTexture/SparseVolumeTexture.h
+++ b/Engine/Source/Runtime/Engine/Classes/SparseVolumeTexture/SparseVolumeTexture.h
@@ -101,12 +101,19 @@ public:
 	FSparseVolumeTextureSceneProxy();
 	virtual ~FSparseVolumeTextureSceneProxy() override;
 
-	void InitializeRuntimeData(FSparseVolumeTextureRuntime& SparseVolumeTextureRuntime);
+	FSparseVolumeTextureRuntime& GetRuntimeData()
+	{
+		return SparseVolumeTextureRuntime;
+	}
+
+	const FSparseVolumeTextureRuntime& GetRuntimeData() const
+	{
+		return SparseVolumeTextureRuntime;
+	}
 
 	const FSparseVolumeAssetHeader& GetHeader() const 
 	{
-		check(SparseVolumeTextureRuntime);
-		return SparseVolumeTextureRuntime->Header;
+		return SparseVolumeTextureRuntime.Header;
 	}
 
 	FTextureRHIRef GetPhysicalTileDataATextureRHI() const
@@ -122,12 +129,14 @@ public:
 		return PageTableTextureRHI;
 	}
 
+	void GetMemorySize(SIZE_T* SizeCPU, SIZE_T* SizeGPU) const;
+
 	virtual void InitRHI() override;
 	virtual void ReleaseRHI() override;
 
 private:
 
-	FSparseVolumeTextureRuntime*		SparseVolumeTextureRuntime;
+	FSparseVolumeTextureRuntime			SparseVolumeTextureRuntime;
 
 	FTextureRHIRef						PageTableTextureRHI;
 	FTextureRHIRef						PhysicalTileDataATextureRHI;
@@ -140,11 +149,8 @@ struct ENGINE_API FSparseVolumeTextureFrame
 	// The frame data that can be streamed in when in game.
 	FByteBulkData						RuntimeStreamedInData;
 
-	// The runtime frame data that is created.
-	FSparseVolumeTextureRuntime			SparseVolumeTextureRuntime;
-
 	// The render side proxy for the sparse volume texture asset.
-	FSparseVolumeTextureSceneProxy		SparseVolumeTextureSceneProxy;
+	FSparseVolumeTextureSceneProxy*		SparseVolumeTextureSceneProxy;
 
 #if WITH_EDITORONLY_DATA
 	/** The raw data that can be loaded when we want to update cook the data with different settings or updated code without re importing. */
@@ -153,7 +159,7 @@ struct ENGINE_API FSparseVolumeTextureFrame
 
 	FSparseVolumeTextureFrame();
 	virtual ~FSparseVolumeTextureFrame();
-	bool BuildRuntimeData();
+	bool BuildRuntimeData(FSparseVolumeTextureRuntime* OutRuntimeData);
 	void Serialize(FArchive& Ar, UStreamableSparseVolumeTexture* Owner, int32 FrameIndex);
 };
 
@@ -179,9 +185,8 @@ public:
 	virtual ~USparseVolumeTexture() = default;
 
 	virtual int32 GetFrameCount() const { return 0; }
-	virtual const FSparseVolumeAssetHeader* GetSparseVolumeTextureHeader() const { return nullptr; }
-	virtual const FSparseVolumeTextureSceneProxy* GetSparseVolumeTextureSceneProxy() const { return nullptr; }
 	virtual FBox GetVolumeBounds() const { return FBox(); }
+	virtual const FSparseVolumeTextureSceneProxy* GetSparseVolumeTextureSceneProxy() const { return nullptr; }
 
 	/** Getter for the shader uniform parameters with index as ESparseVolumeTextureShaderUniform. */
 	FVector4 GetUniformParameter(int32 Index) const;
@@ -204,26 +209,37 @@ class UStreamableSparseVolumeTexture : public USparseVolumeTexture
 
 public:
 
+	UPROPERTY(VisibleAnywhere, Category = "Rendering", meta = (DisplayName = "Volume Bounds"))
+	FBox VolumeBounds;
+
 	UStreamableSparseVolumeTexture();
 	virtual ~UStreamableSparseVolumeTexture() = default;
 
 	//~ Begin UObject Interface.
 	virtual void PostLoad() override;
+	virtual void FinishDestroy() override;
 	virtual void BeginDestroy() override;
 	virtual void Serialize(FArchive& Ar) override;
 #if WITH_EDITOR
 	virtual void PostEditChangeProperty(FPropertyChangedEvent& PropertyChangedEvent) override;
 #endif
+	virtual void GetResourceSizeEx(FResourceSizeEx& CumulativeResourceSize) override;
 	//~ End UObject Interface.
 
+	//~ Begin USparseVolumeTexture Interface.
+	int32 GetFrameCount() const override { return Frames.Num(); }
+	FBox GetVolumeBounds() const override { return VolumeBounds; };
+	const FSparseVolumeTextureSceneProxy* GetSparseVolumeTextureSceneProxy() const override { return GetStreamedFrameProxyOrFallback(0); };
+	//~ End USparseVolumeTexture Interface.
+
 	const FSparseVolumeTextureSceneProxy* GetStreamedFrameProxyOrFallback(int32 FrameIndex) const;
+	TArrayView<const FSparseVolumeTextureFrame> GetFrames() const;
 
 protected:
 	TArray<FSparseVolumeTextureFrame, TInlineAllocator<1u>> Frames;
 
 	void GenerateOrLoadDDCRuntimeDataAndCreateSceneProxy();
 	void GenerateOrLoadDDCRuntimeDataForFrame(FSparseVolumeTextureFrame& Frame, UE::DerivedData::FRequestOwner& DDCRequestOwner);
-	void ConvertRawSourceDataToSparseVolumeTextureRuntime(FSparseVolumeTextureFrame& Frame);
 };
 
 UCLASS(ClassGroup = Rendering, BlueprintType)//, hidecategories = (Object))
@@ -233,21 +249,11 @@ class ENGINE_API UStaticSparseVolumeTexture : public UStreamableSparseVolumeText
 
 public:
 
-	UPROPERTY(VisibleAnywhere, Category = "Rendering", meta = (DisplayName = "Volume Bounds"))
-	FBox VolumeBounds;
-
 	UStaticSparseVolumeTexture();
 	virtual ~UStaticSparseVolumeTexture() = default;
 
-	//~ Begin UObject Interface.
-	virtual void GetResourceSizeEx(FResourceSizeEx& CumulativeResourceSize) override;
-	//~ End UObject Interface.
-
 	//~ Begin USparseVolumeTexture Interface.
 	int32 GetFrameCount() const override { return 1; }
-	const FSparseVolumeAssetHeader* GetSparseVolumeTextureHeader() const override;
-	const FSparseVolumeTextureSceneProxy* GetSparseVolumeTextureSceneProxy() const override;
-	FBox GetVolumeBounds() const override;
 	//~ End USparseVolumeTexture Interface.
 
 private:
@@ -265,25 +271,11 @@ class ENGINE_API UAnimatedSparseVolumeTexture : public UStreamableSparseVolumeTe
 
 public:
 
-	// The asset frame count.
-	UPROPERTY(VisibleAnywhere, Category = "Animation", meta = (DisplayName = "Frame Count"))
-	int32 FrameCount = 0;
-
-	UPROPERTY(VisibleAnywhere, Category = "Rendering", meta = (DisplayName = "Volume Bounds"))
-	FBox VolumeBounds;
-
 	UAnimatedSparseVolumeTexture();
 	virtual ~UAnimatedSparseVolumeTexture() = default;
 
-	//~ Begin UObject Interface.
-	virtual void GetResourceSizeEx(FResourceSizeEx& CumulativeResourceSize) override;
-	//~ End UObject Interface.
-
 	//~ Begin USparseVolumeTexture Interface.
-	int32 GetFrameCount() const override { return FrameCount; }
-	const FSparseVolumeAssetHeader* GetSparseVolumeTextureHeader() const override;
 	const FSparseVolumeTextureSceneProxy* GetSparseVolumeTextureSceneProxy() const override;
-	FBox GetVolumeBounds() const override;
 	//~ End USparseVolumeTexture Interface.
 
 	// Used for debugging a specific frame of an animated sequence.
@@ -311,19 +303,17 @@ public:
 
 	static USparseVolumeTextureFrame* CreateFrame(USparseVolumeTexture* Texture, int32 FrameIndex);
 
-	void Initialize(const FSparseVolumeTextureSceneProxy* InSceneProxy, const FSparseVolumeAssetHeader* InAssetHeader, const FBox& InVolumeBounds);
+	void Initialize(const FSparseVolumeTextureSceneProxy* InSceneProxy, const FBox& InVolumeBounds);
 
 	//~ Begin USparseVolumeTexture Interface.
 	int32 GetFrameCount() const override { return 1; }
-	const FSparseVolumeAssetHeader* GetSparseVolumeTextureHeader() const override;
-	const FSparseVolumeTextureSceneProxy* GetSparseVolumeTextureSceneProxy() const override;
-	FBox GetVolumeBounds() const override;
+	FBox GetVolumeBounds() const override { return VolumeBounds; };
+	const FSparseVolumeTextureSceneProxy* GetSparseVolumeTextureSceneProxy() const override { return SceneProxy; };
 	//~ End USparseVolumeTexture Interface.
 
 private:
-	const FSparseVolumeTextureSceneProxy* SceneProxy;
-	const FSparseVolumeAssetHeader* AssetHeader;
 	FBox VolumeBounds;
+	const FSparseVolumeTextureSceneProxy* SceneProxy;
 };
 
 UCLASS(ClassGroup = Rendering, BlueprintType)//, hidecategories = (Object))
diff --git a/Engine/Source/Runtime/Engine/Private/ContentStreaming.cpp b/Engine/Source/Runtime/Engine/Private/ContentStreaming.cpp
--- a/Engine/Source/Runtime/Engine/Private/ContentStreaming.cpp
+++ b/Engine/Source/Runtime/Engine/Private/ContentStreaming.cpp
@@ -18,6 +18,7 @@
 #include "AudioCompressionSettingsUtils.h"
 #include "VT/VirtualTextureChunkManager.h"
 #include "Rendering/NaniteCoarseMeshStreamingManager.h"
+#include "SparseVolumeTexture/SparseVolumeTextureStreaming.h"
 
 #if WITH_EDITOR
 #include "AudioDevice.h"
@@ -793,6 +794,9 @@ FStreamingManagerCollection::FStreamingManagerCollection()
 
 	VirtualTextureStreamingManager = new FVirtualTextureChunkStreamingManager();
 	AddStreamingManager(VirtualTextureStreamingManager);
+
+	SparseVolumeTextureStreamingManager = new FSparseVolumeTextureStreamingManager();
+	AddStreamingManager(SparseVolumeTextureStreamingManager);
 }
 
 FStreamingManagerCollection::~FStreamingManagerCollection()
@@ -804,6 +808,10 @@ FStreamingManagerCollection::~FStreamingManagerCollection()
 		NaniteCoarseMeshStreamingManager = nullptr;
 	}
 
+	RemoveStreamingManager(SparseVolumeTextureStreamingManager);
+	delete SparseVolumeTextureStreamingManager;
+	SparseVolumeTextureStreamingManager = nullptr;
+
 	RemoveStreamingManager(VirtualTextureStreamingManager);
 	delete VirtualTextureStreamingManager;
 	VirtualTextureStreamingManager = nullptr;
@@ -1068,6 +1076,12 @@ Nanite::FCoarseMeshStreamingManager* FStreamingManagerCollection::GetNaniteCoars
 	return NaniteCoarseMeshStreamingManager;
 }
 
+ENGINE_API ISparseVolumeTextureStreamingManager& FStreamingManagerCollection::GetSparseVolumeTextureStreamingManager() const
+{
+	check(SparseVolumeTextureStreamingManager);
+	return *SparseVolumeTextureStreamingManager;
+}
+
 /** Don't stream world resources for the next NumFrames. */
 void FStreamingManagerCollection::SetDisregardWorldResourcesForFrames(int32 NumFrames )
 {
diff --git a/Engine/Source/Runtime/Engine/Private/SparseVolumeTexture/SparseVolumeTexture.cpp b/Engine/Source/Runtime/Engine/Private/SparseVolumeTexture/SparseVolumeTexture.cpp
--- a/Engine/Source/Runtime/Engine/Private/SparseVolumeTexture/SparseVolumeTexture.cpp
+++ b/Engine/Source/Runtime/Engine/Private/SparseVolumeTexture/SparseVolumeTexture.cpp
@@ -23,6 +23,8 @@
 #include "Serialization/BulkDataWriter.h"
 #include "Serialization/EditorBulkDataReader.h"
 
+#include "ContentStreaming.h"
+
 #define LOCTEXT_NAMESPACE "USparseVolumeTexture"
 
 DEFINE_LOG_CATEGORY_STATIC(LogSparseVolumeTexture, Log, All);
@@ -124,7 +126,7 @@ void FSparseVolumeTextureRuntime::SetAsDefaultTexture()
 
 FSparseVolumeTextureSceneProxy::FSparseVolumeTextureSceneProxy()
 	: FRenderResource()
-	, SparseVolumeTextureRuntime(nullptr)
+	, SparseVolumeTextureRuntime()
 	, PageTableTextureRHI(nullptr)
 	, PhysicalTileDataATextureRHI(nullptr)
 	, PhysicalTileDataBTextureRHI(nullptr)
@@ -135,9 +137,28 @@ FSparseVolumeTextureSceneProxy::~FSparseVolumeTextureSceneProxy()
 {
 }
 
-void FSparseVolumeTextureSceneProxy::InitializeRuntimeData(FSparseVolumeTextureRuntime& SparseVolumeTextureRuntimeIn)
+void FSparseVolumeTextureSceneProxy::GetMemorySize(SIZE_T* SizeCPU, SIZE_T* SizeGPU) const
 {
-	SparseVolumeTextureRuntime = &SparseVolumeTextureRuntimeIn;
+	*SizeCPU += sizeof(FSparseVolumeTextureSceneProxy);
+	*SizeCPU += SparseVolumeTextureRuntime.PageTable.GetAllocatedSize();
+	*SizeCPU += SparseVolumeTextureRuntime.PhysicalTileDataA.GetAllocatedSize();
+	*SizeCPU += SparseVolumeTextureRuntime.PhysicalTileDataB.GetAllocatedSize();
+
+#if RHI_ENABLE_RESOURCE_INFO
+	FRHIResourceInfo ResourceInfo;
+	if (PageTableTextureRHI && PageTableTextureRHI->GetResourceInfo(ResourceInfo))
+	{
+		*SizeGPU += ResourceInfo.VRamAllocation.AllocationSize;
+	}
+	if (PhysicalTileDataATextureRHI && PhysicalTileDataATextureRHI->GetResourceInfo(ResourceInfo))
+	{
+		*SizeGPU += ResourceInfo.VRamAllocation.AllocationSize;
+	}
+	if (PhysicalTileDataBTextureRHI && PhysicalTileDataBTextureRHI->GetResourceInfo(ResourceInfo))
+	{
+		*SizeGPU += ResourceInfo.VRamAllocation.AllocationSize;
+	}
+#endif
 }
 
 void FSparseVolumeTextureSceneProxy::InitRHI()
@@ -145,7 +166,7 @@ void FSparseVolumeTextureSceneProxy::InitRHI()
 	// Page table
 	{
 		EPixelFormat PageEntryFormat = PF_R32_UINT;
-		FIntVector3 PageTableVolumeResolution = SparseVolumeTextureRuntime->Header.PageTableVolumeResolution;
+		FIntVector3 PageTableVolumeResolution = SparseVolumeTextureRuntime.Header.PageTableVolumeResolution;
 		const FRHITextureCreateDesc Desc =
 			FRHITextureCreateDesc::Create3D(TEXT("SparseVolumeTexture.PageTable.RHITexture"),
 				PageTableVolumeResolution.X, PageTableVolumeResolution.Y, PageTableVolumeResolution.Z, PageEntryFormat)
@@ -155,14 +176,14 @@ void FSparseVolumeTextureSceneProxy::InitRHI()
 
 		const int32 FormatSize = GPixelFormats[PageEntryFormat].BlockBytes;
 		const FUpdateTextureRegion3D UpdateRegion(0, 0, 0, 0, 0, 0, PageTableVolumeResolution.X, PageTableVolumeResolution.Y, PageTableVolumeResolution.Z);
-		RHIUpdateTexture3D(PageTableTextureRHI, 0, UpdateRegion, UpdateRegion.Width * FormatSize, UpdateRegion.Width * UpdateRegion.Height * FormatSize, (const uint8*)SparseVolumeTextureRuntime->PageTable.GetData());
+		RHIUpdateTexture3D(PageTableTextureRHI, 0, UpdateRegion, UpdateRegion.Width * FormatSize, UpdateRegion.Width * UpdateRegion.Height * FormatSize, (const uint8*)SparseVolumeTextureRuntime.PageTable.GetData());
 	}
 
 	// Tile data
 	{
-		FIntVector3 TileDataVolumeResolution = SparseVolumeTextureRuntime->Header.TileDataVolumeResolution;
-		EPixelFormat VoxelFormatA = SparseVolumeTextureRuntime->Header.AttributesAFormat;
-		EPixelFormat VoxelFormatB = SparseVolumeTextureRuntime->Header.AttributesBFormat;
+		FIntVector3 TileDataVolumeResolution = SparseVolumeTextureRuntime.Header.TileDataVolumeResolution;
+		EPixelFormat VoxelFormatA = SparseVolumeTextureRuntime.Header.AttributesAFormat;
+		EPixelFormat VoxelFormatB = SparseVolumeTextureRuntime.Header.AttributesBFormat;
 		const FUpdateTextureRegion3D UpdateRegion(0, 0, 0, 0, 0, 0, TileDataVolumeResolution.X, TileDataVolumeResolution.Y, TileDataVolumeResolution.Z);
 
 		// A
@@ -176,7 +197,7 @@ void FSparseVolumeTextureSceneProxy::InitRHI()
 			PhysicalTileDataATextureRHI = RHICreateTexture(Desc);
 
 			const int32 FormatSize = GPixelFormats[VoxelFormatA].BlockBytes;
-			RHIUpdateTexture3D(PhysicalTileDataATextureRHI, 0, UpdateRegion, UpdateRegion.Width * FormatSize, UpdateRegion.Width * UpdateRegion.Height * FormatSize, (const uint8*)SparseVolumeTextureRuntime->PhysicalTileDataA.GetData());
+			RHIUpdateTexture3D(PhysicalTileDataATextureRHI, 0, UpdateRegion, UpdateRegion.Width * FormatSize, UpdateRegion.Width * UpdateRegion.Height * FormatSize, (const uint8*)SparseVolumeTextureRuntime.PhysicalTileDataA.GetData());
 		}
 
 		// B
@@ -190,7 +211,7 @@ void FSparseVolumeTextureSceneProxy::InitRHI()
 			PhysicalTileDataBTextureRHI = RHICreateTexture(Desc);
 
 			const int32 FormatSize = GPixelFormats[VoxelFormatB].BlockBytes;
-			RHIUpdateTexture3D(PhysicalTileDataBTextureRHI, 0, UpdateRegion, UpdateRegion.Width * FormatSize, UpdateRegion.Width * UpdateRegion.Height * FormatSize, (const uint8*)SparseVolumeTextureRuntime->PhysicalTileDataB.GetData());
+			RHIUpdateTexture3D(PhysicalTileDataBTextureRHI, 0, UpdateRegion, UpdateRegion.Width * FormatSize, UpdateRegion.Width * UpdateRegion.Height * FormatSize, (const uint8*)SparseVolumeTextureRuntime.PhysicalTileDataB.GetData());
 		}
 	}
 }
@@ -203,7 +224,6 @@ void FSparseVolumeTextureSceneProxy::ReleaseRHI()
 
 FSparseVolumeTextureFrame::FSparseVolumeTextureFrame()
 	: RuntimeStreamedInData()
-	, SparseVolumeTextureRuntime()
 	, SparseVolumeTextureSceneProxy()
 #if WITH_EDITORONLY_DATA
 	, RawData()
@@ -215,7 +235,7 @@ FSparseVolumeTextureFrame::~FSparseVolumeTextureFrame()
 {
 }
 
-bool FSparseVolumeTextureFrame::BuildRuntimeData()
+bool FSparseVolumeTextureFrame::BuildRuntimeData(FSparseVolumeTextureRuntime* OutRuntimeData)
 {
 #if WITH_EDITORONLY_DATA
 	// Check if the virtualized bulk data payload is available now
@@ -227,10 +247,10 @@ bool FSparseVolumeTextureFrame::BuildRuntimeData()
 		SparseVolumeRawSource.Serialize(RawDataArchiveReader);
 
 		// Then, convert the raw source data to SVT
-		SparseVolumeTextureRuntime.Header = SparseVolumeRawSource.Header;
-		SparseVolumeTextureRuntime.PageTable = MoveTemp(SparseVolumeRawSource.PageTable);
-		SparseVolumeTextureRuntime.PhysicalTileDataA = MoveTemp(SparseVolumeRawSource.PhysicalTileDataA);
-		SparseVolumeTextureRuntime.PhysicalTileDataB = MoveTemp(SparseVolumeRawSource.PhysicalTileDataB);
+		OutRuntimeData->Header = SparseVolumeRawSource.Header;
+		OutRuntimeData->PageTable = MoveTemp(SparseVolumeRawSource.PageTable);
+		OutRuntimeData->PhysicalTileDataA = MoveTemp(SparseVolumeRawSource.PhysicalTileDataA);
+		OutRuntimeData->PhysicalTileDataB = MoveTemp(SparseVolumeRawSource.PhysicalTileDataB);
 
 		// Now unload the raw data
 		RawData.UnloadData();
@@ -245,7 +265,7 @@ void FSparseVolumeTextureFrame::Serialize(FArchive& Ar, UStreamableSparseVolumeT
 {
 	FStripDataFlags StripFlags(Ar);
 
-	const bool bInlinePayload = (FrameIndex == 0) || true; // SVT_TODO: remove once streaming is working
+	const bool bInlinePayload = (FrameIndex == 0);
 	RuntimeStreamedInData.SetBulkDataFlags(bInlinePayload ? BULKDATA_ForceInlinePayload : BULKDATA_Force_NOT_InlinePayload);
 
 	if (StripFlags.IsEditorDataStripped() && Ar.IsLoadingFromCookedPackage())
@@ -255,30 +275,35 @@ void FSparseVolumeTextureFrame::Serialize(FArchive& Ar, UStreamableSparseVolumeT
 		// Read cooked bulk data from archive
 		RuntimeStreamedInData.Serialize(Ar, Owner);
 
-		// Create runtime data from cooked bulk data
+		if (bInlinePayload)
 		{
-			FBulkDataReader BulkDataReader(RuntimeStreamedInData);
-			SparseVolumeTextureRuntime.Serialize(BulkDataReader);
-		}
+			SparseVolumeTextureSceneProxy = new FSparseVolumeTextureSceneProxy();
+
+			// Create runtime data from cooked bulk data
+			{
+				FBulkDataReader BulkDataReader(RuntimeStreamedInData);
+				SparseVolumeTextureSceneProxy->GetRuntimeData().Serialize(BulkDataReader);
+			}
 
-		// The bulk data is no longer needed
-		RuntimeStreamedInData.RemoveBulkData();
+			// The bulk data is no longer needed
+			RuntimeStreamedInData.RemoveBulkData();
 
-		// Runtime data is now valid, create the render thread proxy
-		SparseVolumeTextureSceneProxy.InitializeRuntimeData(SparseVolumeTextureRuntime);
-		BeginInitResource(&SparseVolumeTextureSceneProxy);
+			// Runtime data is now valid, initialize the render thread proxy
+			BeginInitResource(SparseVolumeTextureSceneProxy);
+		}
 	}
 	else if (Ar.IsCooking())
 	{
 		// We are cooking the game, serialize the asset out.
 
-		const bool bBuiltRuntimeData = BuildRuntimeData();
+		FSparseVolumeTextureRuntime RuntimeData;
+		const bool bBuiltRuntimeData = BuildRuntimeData(&RuntimeData);
 		check(bBuiltRuntimeData); // SVT_TODO: actual error handling
 
 		// Write runtime data into RuntimeStreamedInData
 		{
 			FBulkDataWriter BulkDataWriter(RuntimeStreamedInData);
-			SparseVolumeTextureRuntime.Serialize(BulkDataWriter);
+			RuntimeData.Serialize(BulkDataWriter);
 		}
 
 		// And now write the cooked bulk data to the archive
@@ -388,9 +413,20 @@ UStreamableSparseVolumeTexture::UStreamableSparseVolumeTexture(const FObjectInit
 
 void UStreamableSparseVolumeTexture::PostLoad()
 {
+	Super::PostLoad();
+
+#if WITH_EDITORONLY_DATA
 	GenerateOrLoadDDCRuntimeDataAndCreateSceneProxy();
+#else
+	IStreamingManager::Get().GetSparseVolumeTextureStreamingManager().AddSparseVolumeTexture(this); // GenerateOrLoadDDCRuntimeDataAndCreateSceneProxy() handles this in editor builds
+#endif
+}
 
-	Super::PostLoad();
+void UStreamableSparseVolumeTexture::FinishDestroy()
+{
+	Super::FinishDestroy();
+
+	IStreamingManager::Get().GetSparseVolumeTextureStreamingManager().RemoveSparseVolumeTexture(this);
 }
 
 void UStreamableSparseVolumeTexture::BeginDestroy()
@@ -398,7 +434,16 @@ void UStreamableSparseVolumeTexture::BeginDestroy()
 	Super::BeginDestroy();
 	for (FSparseVolumeTextureFrame& Frame : Frames)
 	{
-		BeginReleaseResource(&Frame.SparseVolumeTextureSceneProxy);
+		if (Frame.SparseVolumeTextureSceneProxy)
+		{
+			ENQUEUE_RENDER_COMMAND(DeleteSVTProxy)(
+				[Proxy = Frame.SparseVolumeTextureSceneProxy](FRHICommandListImmediate& RHICmdList)
+				{
+					Proxy->ReleaseResource();
+					delete Proxy;
+				});
+			Frame.SparseVolumeTextureSceneProxy = nullptr;
+		}
 	}
 }
 
@@ -429,11 +474,48 @@ void UStreamableSparseVolumeTexture::PostEditChangeProperty(FPropertyChangedEven
 }
 #endif // WITH_EDITOR
 
+void UStreamableSparseVolumeTexture::GetResourceSizeEx(FResourceSizeEx& CumulativeResourceSize)
+{
+	Super::GetResourceSizeEx(CumulativeResourceSize);
+	SIZE_T SizeCPU = sizeof(Frames);
+	SIZE_T SizeGPU = 0;
+	SizeCPU += Frames.GetAllocatedSize();
+	for (const FSparseVolumeTextureFrame& Frame : Frames)
+	{
+		if (Frame.SparseVolumeTextureSceneProxy)
+		{
+			Frame.SparseVolumeTextureSceneProxy->GetMemorySize(&SizeCPU, &SizeGPU);
+		}
+	}
+	ISparseVolumeTextureStreamingManager& StreamingManager = IStreamingManager::Get().GetSparseVolumeTextureStreamingManager();
+	StreamingManager.GetMemorySizeForSparseVolumeTexture(this, &SizeCPU, &SizeGPU);
+	CumulativeResourceSize.AddDedicatedSystemMemoryBytes(SizeCPU);
+	CumulativeResourceSize.AddDedicatedVideoMemoryBytes(SizeGPU);
+}
+
 const FSparseVolumeTextureSceneProxy* UStreamableSparseVolumeTexture::GetStreamedFrameProxyOrFallback(int32 FrameIndex) const
 {
-	check(!Frames.IsEmpty());
-	FrameIndex = FrameIndex % Frames.Num();
-	return &Frames[FrameIndex].SparseVolumeTextureSceneProxy;
+	ISparseVolumeTextureStreamingManager& StreamingManager = IStreamingManager::Get().GetSparseVolumeTextureStreamingManager();
+	const FSparseVolumeTextureSceneProxy* Proxy = StreamingManager.GetSparseVolumeTextureSceneProxy(this, FrameIndex, true);
+
+	int32 FallbackFrameIndex = FrameIndex;
+	while (!Proxy)
+	{
+		FallbackFrameIndex = FallbackFrameIndex > 0 ? (FallbackFrameIndex - 1) : (Frames.Num() - 1);
+		if (FallbackFrameIndex == FrameIndex)
+		{
+			UE_LOG(LogSparseVolumeTexture, Warning, TEXT("Failed to get ANY streamed SparseVolumeTexture frame  SVT: %s, FrameIndex: %i"), *GetName(), FrameIndex);
+			return nullptr;
+		}
+		Proxy = StreamingManager.GetSparseVolumeTextureSceneProxy(this, FallbackFrameIndex, false);
+	}
+
+	return Proxy;
+}
+
+TArrayView<const FSparseVolumeTextureFrame> UStreamableSparseVolumeTexture::GetFrames() const
+{
+	return Frames;
 }
 
 void UStreamableSparseVolumeTexture::GenerateOrLoadDDCRuntimeDataAndCreateSceneProxy()
@@ -445,7 +527,14 @@ void UStreamableSparseVolumeTexture::GenerateOrLoadDDCRuntimeDataAndCreateSceneP
 		for (FSparseVolumeTextureFrame& Frame : Frames)
 		{
 			// Release any previously allocated render thread proxy
-			BeginReleaseResource(&Frame.SparseVolumeTextureSceneProxy);
+			if (Frame.SparseVolumeTextureSceneProxy)
+			{
+				BeginReleaseResource(Frame.SparseVolumeTextureSceneProxy);
+			}
+			else
+			{
+				Frame.SparseVolumeTextureSceneProxy = new FSparseVolumeTextureSceneProxy();
+			}
 
 			GenerateOrLoadDDCRuntimeDataForFrame(Frame, DDCRequestOwner);
 		}
@@ -456,22 +545,11 @@ void UStreamableSparseVolumeTexture::GenerateOrLoadDDCRuntimeDataAndCreateSceneP
 
 	for (FSparseVolumeTextureFrame& Frame : Frames)
 	{
-		// Runtime data is now valid, create the render thread proxy
-		Frame.SparseVolumeTextureSceneProxy.InitializeRuntimeData(Frame.SparseVolumeTextureRuntime);
-		BeginInitResource(&Frame.SparseVolumeTextureSceneProxy);
+		// Runtime data is now valid, initialize the render thread proxy
+		BeginInitResource(Frame.SparseVolumeTextureSceneProxy);
 	}
-	
-#else
-	// SVT_TODO do we really need to do this in the cooked case?
-	for (FSparseVolumeTextureFrame& Frame : Frames)
-	{
-		// Release any previously allocated render thread proxy
-		BeginReleaseResource(&Frame.SparseVolumeTextureSceneProxy);
 
-		// Create the render thread proxy
-		Frame.SparseVolumeTextureSceneProxy.InitializeRuntimeData(Frame.SparseVolumeTextureRuntime);
-		BeginInitResource(&Frame.SparseVolumeTextureSceneProxy);
-	}
+	IStreamingManager::Get().GetSparseVolumeTextureStreamingManager().AddSparseVolumeTexture(this);
 #endif
 }
 
@@ -496,7 +574,7 @@ void UStreamableSparseVolumeTexture::GenerateOrLoadDDCRuntimeDataForFrame(FSpars
 					{
 						FSharedBuffer Data = Value.GetData().Decompress();
 						FMemoryReaderView Ar(Data, true /*bIsPersistent*/);
-						Frame.SparseVolumeTextureRuntime.Serialize(Ar);
+						Frame.SparseVolumeTextureSceneProxy->GetRuntimeData().Serialize(Ar);
 					});
 			}
 			else if (Response.Status == EStatus::Error)
@@ -504,16 +582,28 @@ void UStreamableSparseVolumeTexture::GenerateOrLoadDDCRuntimeDataForFrame(FSpars
 				DDCRequestOwner.LaunchTask(TEXT("UStreamableSparseVolumeTexture_DerivedDataBuild"),
 					[this, &Frame, &DDCRequestOwner, Name = Response.Name, Key = Response.Key]()
 					{
-						ConvertRawSourceDataToSparseVolumeTextureRuntime(Frame);
+						FSparseVolumeTextureRuntime& RuntimeData = Frame.SparseVolumeTextureSceneProxy->GetRuntimeData();
+
+						// Check if the virtualized bulk data payload is available now
+						if (Frame.RawData.HasPayloadData())
+						{
+							const bool bSuccess = Frame.BuildRuntimeData(&RuntimeData);
+							ensure(bSuccess);
+						}
+						else
+						{
+							UE_LOG(LogSparseVolumeTexture, Error, TEXT("SparseVolumeTexture - Raw source data is not available for %s. Using default data."), *GetName());
+							RuntimeData.SetAsDefaultTexture();
+						}
 
 						// Using a LargeMemoryWriter for serialization since the data can be bigger than 2 GB
 						FLargeMemoryWriter LargeMemWriter(0, /*bIsPersistent=*/ true);
-						Frame.SparseVolumeTextureRuntime.Serialize(LargeMemWriter);
+						RuntimeData.Serialize(LargeMemWriter);
 
 						int64 UncompressedSize = LargeMemWriter.TotalSize();
 
 						// Since the DDC doesn't support data bigger than 2 GB, we only cache for such uncompressed size.
-						constexpr int64 SizeThreshold = 2147483648;	// 2GB
+						constexpr int64 SizeThreshold = 2147483648LL;	// 2GB
 						const bool bIsCacheable = UncompressedSize < SizeThreshold;
 						if (bIsCacheable)
 						{
@@ -530,23 +620,6 @@ void UStreamableSparseVolumeTexture::GenerateOrLoadDDCRuntimeDataForFrame(FSpars
 #endif // WITH_EDITORONLY_DATA
 }
 
-void UStreamableSparseVolumeTexture::ConvertRawSourceDataToSparseVolumeTextureRuntime(FSparseVolumeTextureFrame& Frame)
-{
-#if WITH_EDITORONLY_DATA
-	// Check if the virtualized bulk data payload is available now
-	if (Frame.RawData.HasPayloadData())
-	{
-		const bool bSuccess = Frame.BuildRuntimeData();
-		ensure(bSuccess);
-	}
-	else
-	{
-		UE_LOG(LogSparseVolumeTexture, Error, TEXT("SparseVolumeTexture - Raw source data is not available for %s. Using default data."), *GetName());
-		Frame.SparseVolumeTextureRuntime.SetAsDefaultTexture();
-	}
-#endif // WITH_EDITORONLY_DATA
-}
-
 ////////////////////////////////////////////////////////////////////////////////////////////////
 
 UStaticSparseVolumeTexture::UStaticSparseVolumeTexture(const FObjectInitializer& ObjectInitializer)
@@ -554,74 +627,34 @@ UStaticSparseVolumeTexture::UStaticSparseVolumeTexture(const FObjectInitializer&
 {
 }
 
-void UStaticSparseVolumeTexture::GetResourceSizeEx(FResourceSizeEx& CumulativeResourceSize)
-{
-	Super::GetResourceSizeEx(CumulativeResourceSize);
-	// CumulativeResourceSize.AddDedicatedSystemMemoryBytes but not the RawData size
-}
-
-const FSparseVolumeAssetHeader* UStaticSparseVolumeTexture::GetSparseVolumeTextureHeader() const
-{
-	return &GetStreamedFrameProxyOrFallback(0)->GetHeader();
-}
-
-const FSparseVolumeTextureSceneProxy* UStaticSparseVolumeTexture::GetSparseVolumeTextureSceneProxy() const
-{ 
-	return GetStreamedFrameProxyOrFallback(0);
-}
-
-FBox UStaticSparseVolumeTexture::GetVolumeBounds() const
-{
-	return VolumeBounds;
-}
-
 ////////////////////////////////////////////////////////////////////////////////////////////////
 
 UAnimatedSparseVolumeTexture::UAnimatedSparseVolumeTexture(const FObjectInitializer& ObjectInitializer)
 	: Super(ObjectInitializer)
-	, FrameCount(0)
 {
 }
 
-void UAnimatedSparseVolumeTexture::GetResourceSizeEx(FResourceSizeEx& CumulativeResourceSize)
-{
-	Super::GetResourceSizeEx(CumulativeResourceSize);
-	// CumulativeResourceSize.AddDedicatedSystemMemoryBytes but not the RawData size
-}
-
-const FSparseVolumeAssetHeader* UAnimatedSparseVolumeTexture::GetSparseVolumeTextureHeader() const
-{
-	// When an AnimatedSparseVolumeTexture is used as SparseVolumeTexture, it can only be previewed using its first frame.
-	check(FrameCount >= 1);
-	const int32 FrameIndex = PreviewFrameIndex % FrameCount;
-	return &GetStreamedFrameProxyOrFallback(FrameIndex)->GetHeader();
-}
-
 const FSparseVolumeTextureSceneProxy* UAnimatedSparseVolumeTexture::GetSparseVolumeTextureSceneProxy() const
 {
-	// When an AnimatedSparseVolumeTexture is used as SparseVolumeTexture, it can only be previewed using its first frame.
-	check(FrameCount >= 1);
-	const int32 FrameIndex = PreviewFrameIndex % FrameCount;
+	// When an AnimatedSparseVolumeTexture is used as SparseVolumeTexture, it can only be previewed using a single preview frame.
+	check(!Frames.IsEmpty());
+	const int32 FrameIndex = PreviewFrameIndex % Frames.Num();
 	return GetStreamedFrameProxyOrFallback(FrameIndex);
 }
 
-FBox UAnimatedSparseVolumeTexture::GetVolumeBounds() const
-{
-	return VolumeBounds;
-}
-
 const FSparseVolumeTextureSceneProxy* UAnimatedSparseVolumeTexture::GetSparseVolumeTextureFrameSceneProxy(int32 FrameIndex) const
 {
-	check(FrameCount >= 1);
-	FrameIndex = FrameIndex % FrameCount;
+	check(!Frames.IsEmpty());
+	FrameIndex = FrameIndex % Frames.Num();
 	return GetStreamedFrameProxyOrFallback(FrameIndex);
 }
 
 const FSparseVolumeAssetHeader* UAnimatedSparseVolumeTexture::GetSparseVolumeTextureFrameHeader(int32 FrameIndex) const
 {
-	check(FrameCount >= 1);
-	FrameIndex = FrameIndex % FrameCount;
-	return &GetStreamedFrameProxyOrFallback(FrameIndex)->GetHeader();
+	check(!Frames.IsEmpty());
+	FrameIndex = FrameIndex % Frames.Num();
+	const FSparseVolumeTextureSceneProxy* Proxy = GetStreamedFrameProxyOrFallback(FrameIndex);
+	return Proxy ? &Proxy->GetHeader() : nullptr;
 }
 
 ////////////////////////////////////////////////////////////////////////////////////////////////
@@ -633,7 +666,8 @@ USparseVolumeTextureFrame::USparseVolumeTextureFrame(const FObjectInitializer& O
 
 USparseVolumeTextureFrame* USparseVolumeTextureFrame::CreateFrame(USparseVolumeTexture* Texture, int32 FrameIndex)
 {
-	USparseVolumeTextureFrame* Frame = NewObject<USparseVolumeTextureFrame>();
+	check(Texture);
+	
 	const FSparseVolumeTextureSceneProxy* Proxy = nullptr;
 	if (Texture->IsA<UStreamableSparseVolumeTexture>())
 	{
@@ -645,35 +679,23 @@ USparseVolumeTextureFrame* USparseVolumeTextureFrame::CreateFrame(USparseVolumeT
 	{
 		Proxy = Texture->GetSparseVolumeTextureSceneProxy();
 	}
-	check(Proxy);
-
-	Frame->Initialize(Proxy, &Proxy->GetHeader(), Texture->GetVolumeBounds());
 
-	return Frame;
+	if (Proxy)
+	{
+		USparseVolumeTextureFrame* Frame = NewObject<USparseVolumeTextureFrame>();
+		Frame->Initialize(Proxy, Texture->GetVolumeBounds());
+		return Frame;
+	}
+	
+	return nullptr;
 }
 
-void USparseVolumeTextureFrame::Initialize(const FSparseVolumeTextureSceneProxy* InSceneProxy, const FSparseVolumeAssetHeader* InAssetHeader, const FBox& InVolumeBounds)
+void USparseVolumeTextureFrame::Initialize(const FSparseVolumeTextureSceneProxy* InSceneProxy, const FBox& InVolumeBounds)
 {
 	SceneProxy = InSceneProxy;
-	AssetHeader = InAssetHeader;
 	VolumeBounds = InVolumeBounds;
 }
 
-const FSparseVolumeAssetHeader* USparseVolumeTextureFrame::GetSparseVolumeTextureHeader() const
-{
-	return AssetHeader;
-}
-
-const FSparseVolumeTextureSceneProxy* USparseVolumeTextureFrame::GetSparseVolumeTextureSceneProxy() const
-{
-	return SceneProxy;
-}
-
-FBox USparseVolumeTextureFrame::GetVolumeBounds() const
-{
-	return VolumeBounds;
-}
-
 ////////////////////////////////////////////////////////////////////////////////////////////////
 
 UAnimatedSparseVolumeTextureController::UAnimatedSparseVolumeTextureController(const FObjectInitializer& ObjectInitializer)
diff --git a/Engine/Source/Runtime/Engine/Private/SparseVolumeTexture/SparseVolumeTextureStreaming.cpp b/Engine/Source/Runtime/Engine/Private/SparseVolumeTexture/SparseVolumeTextureStreaming.cpp
new file mode 100644
--- /dev/null
+++ b/Engine/Source/Runtime/Engine/Private/SparseVolumeTexture/SparseVolumeTextureStreaming.cpp
@@ -0,0 +1,518 @@
+// Copyright Epic Games, Inc. All Rights Reserved.
+
+#include "SparseVolumeTexture/SparseVolumeTextureStreaming.h"
+#include "SparseVolumeTexture/SparseVolumeTexture.h"
+
+DEFINE_LOG_CATEGORY_STATIC(LogSparseVolumeTextureStreaming, Log, All);
+
+int32 GSVTNumPrefetchFrames = 1;
+static FAutoConsoleVariableRef CVarSVTStreamingPrefetchCount(
+	TEXT("r.SparseVolumeTexture.Streaming.NumPrefetchFrames"), 
+	GSVTNumPrefetchFrames,
+	TEXT("Number of frames to prefetch when streaming animated SparseVolumeTexture frames."),
+	ECVF_Scalability);
+
+////////////////////////////////////////////////////////////////////////////////////////////////
+
+void FLoadedSparseVolumeTextureFrame::CleanUpIORequest()
+{
+	if (IORequest)
+	{
+		IORequest->WaitCompletion();
+		delete IORequest;
+		IORequest = nullptr;
+	}
+}
+
+////////////////////////////////////////////////////////////////////////////////////////////////
+
+FStreamingSparseVolumeTextureData::FStreamingSparseVolumeTextureData()
+{
+	ResetRequestedFrames();
+}
+
+FStreamingSparseVolumeTextureData::~FStreamingSparseVolumeTextureData()
+{
+}
+
+bool FStreamingSparseVolumeTextureData::Initialize(UStreamableSparseVolumeTexture* InSparseVolumeTexture, FSparseVolumeTextureStreamingManager* InStreamingManager)
+{
+	check(InSparseVolumeTexture);
+
+	if (InSparseVolumeTexture->GetFrames().IsEmpty())
+	{
+		UE_LOG(LogSparseVolumeTextureStreaming, Error, TEXT("Failed to initialize streaming SparseVolumeTexture due to lack of SVT or serialized stream frames. '%s'"), *InSparseVolumeTexture->GetFullName());
+		return false;
+	}
+
+	SparseVolumeTexture = InSparseVolumeTexture;
+	StreamingManager = InStreamingManager;
+
+	// Always get the first frame of data so we can play immediately
+	check(LoadedFrames.IsEmpty());
+	check(LoadedFrameIndices.IsEmpty());
+	check(!InSparseVolumeTexture->GetFrames().IsEmpty()); // Must hold at least the first frame
+
+	AddNewLoadedFrame(0, InSparseVolumeTexture->GetFrames()[0].SparseVolumeTextureSceneProxy);
+	LoadedFrameIndices.Add(0);
+
+	return true;
+}
+
+void FStreamingSparseVolumeTextureData::FreeResources()
+{
+	// Make sure there are no pending requests in flight.
+	for (int32 Pass = 0; Pass < 3; Pass++)
+	{
+		BlockTillAllRequestsFinished();
+		if (!UpdateStreamingStatus())
+		{
+			break;
+		}
+		check(Pass < 2); // we should be done after two passes. Pass 0 will start anything we need and pass 1 will complete those requests
+	}
+
+	for (FLoadedSparseVolumeTextureFrame& LoadedFrame : LoadedFrames)
+	{
+		FreeLoadedFrame(LoadedFrame);
+	}
+}
+
+bool FStreamingSparseVolumeTextureData::UpdateStreamingStatus()
+{
+	if (!SparseVolumeTexture)
+	{
+		return false;
+	}
+
+	// Handle failed frames first
+	for (int32 LoadedFrameIdx = 0; LoadedFrameIdx < LoadedFrames.Num(); ++LoadedFrameIdx)
+	{
+		FLoadedSparseVolumeTextureFrame& LoadedFrame = LoadedFrames[LoadedFrameIdx];
+		if (LoadFailedFrameIndices.Contains(LoadedFrame.FrameIndex))
+		{
+			// Mark as not loaded
+			LoadedFrameIndices.Remove(LoadedFrame.FrameIndex);
+
+			// Remove this frame
+			FreeLoadedFrame(LoadedFrame);
+
+			FScopeLock LoadedFramesLock(&LoadedFramesCriticalSection);
+			LoadedFrames.RemoveAtSwap(LoadedFrameIdx, 1, false);
+		}
+	}
+
+	LoadFailedFrameIndices.Reset();
+
+	bool bHasPendingRequestInFlight = false;
+	TArray<int32> IndicesToLoad;
+	TArray<int32> IndicesToFree;
+
+	if (HasPendingRequests(IndicesToLoad, IndicesToFree))
+	{
+		for (FLoadedSparseVolumeTextureFrame& LoadedFrame : LoadedFrames)
+		{
+			if (LoadedFrame.IORequest)
+			{
+				const bool bRequestFinished = LoadedFrame.IORequest->PollCompletion();
+				bHasPendingRequestInFlight |= !bRequestFinished;
+				if (bRequestFinished)
+				{
+					LoadedFrame.CleanUpIORequest();
+				}
+			}
+		}
+
+		LoadedFrameIndices = RequestedFrameIndices;
+
+		BeginPendingRequests(IndicesToLoad, IndicesToFree);
+	}
+
+	ResetRequestedFrames();
+
+	// Print out the currently resident frames to the screen
+#if 0
+	if (GEngine)
+	{
+#if WITH_EDITORONLY_DATA
+		FString Message = TEXT("Referenced DDC SVT Frames: ");
+#else
+		FString Message = TEXT("Streaming SVT Frames in Memory: ");
+#endif
+		TArray<int32> FramesInMemoryIndices;
+		for (FLoadedSparseVolumeTextureFrame& LoadedFrame : LoadedFrames)
+		{
+			if (LoadedFrame.bOwnsProxy)
+			{
+				FramesInMemoryIndices.Add(LoadedFrame.FrameIndex);
+			}
+		}
+		
+		if (!FramesInMemoryIndices.IsEmpty())
+		{
+			FramesInMemoryIndices.Sort();
+			for (int32 FrameIndex : FramesInMemoryIndices)
+			{
+				Message += FString::Format(TEXT("{0}, "), { FrameIndex });
+			}
+			GEngine->AddOnScreenDebugMessage(-1, 0.1f, FColor::Yellow, Message);
+		}
+	}
+#endif
+
+	return bHasPendingRequestInFlight;
+}
+
+bool FStreamingSparseVolumeTextureData::HasPendingRequests(TArray<int32>& IndicesToLoad, TArray<int32>& IndicesToFree) const
+{
+	IndicesToLoad.Reset();
+	IndicesToFree.Reset();
+
+	// Find indices that aren't loaded
+	for (int32 RequestedIndex : RequestedFrameIndices)
+	{
+		if (!LoadedFrameIndices.Contains(RequestedIndex))
+		{
+			IndicesToLoad.AddUnique(RequestedIndex);
+		}
+	}
+
+	// Find indices that aren't needed anymore
+	for (int32 LoadedIndex : LoadedFrameIndices)
+	{
+		if (!RequestedFrameIndices.Contains(LoadedIndex))
+		{
+			IndicesToFree.AddUnique(LoadedIndex);
+		}
+	}
+
+	return IndicesToLoad.Num() > 0 || IndicesToFree.Num() > 0;
+}
+
+void FStreamingSparseVolumeTextureData::BeginPendingRequests(const TArray<int32>& IndicesToLoad, const TArray<int32>& IndicesToFree)
+{
+	// Mark frames for removal in case they can be reused
+	{
+		for (int32 IndexToFree : IndicesToFree)
+		{
+			for (int32 LoadedFrameIdx = 0; LoadedFrameIdx < LoadedFrames.Num(); ++LoadedFrameIdx)
+			{
+				check(IndexToFree != 0);
+				if (LoadedFrames[LoadedFrameIdx].FrameIndex == IndexToFree)
+				{
+					FreeLoadedFrame(LoadedFrames[LoadedFrameIdx]);
+
+					FScopeLock LoadedFramesLock(&LoadedFramesCriticalSection);
+					LoadedFrames.RemoveAtSwap(LoadedFrameIdx, 1, false);
+					break;
+				}
+			}
+		}
+	}
+
+	// Set off all IO Requests
+
+	const EAsyncIOPriorityAndFlags AsyncIOPriority = AIOP_CriticalPath; //Set to Crit temporarily as emergency speculative fix for streaming issue
+
+	for (int32 IndexToLoad : IndicesToLoad)
+	{
+		const FSparseVolumeTextureFrame& Frame = SparseVolumeTexture->GetFrames()[IndexToLoad];
+		FSparseVolumeTextureSceneProxy* ExistingProxy = Frame.SparseVolumeTextureSceneProxy;
+		FLoadedSparseVolumeTextureFrame& FrameStorage = AddNewLoadedFrame(IndexToLoad, ExistingProxy);
+
+		if (!ExistingProxy)
+		{
+			UE_CLOG(FrameStorage.Proxy.load() != nullptr, LogSparseVolumeTextureStreaming, Fatal, TEXT("Existing render proxy for streaming SparseVolumeTexture frame."));
+			UE_CLOG(FrameStorage.IORequest, LogSparseVolumeTextureStreaming, Fatal, TEXT("Streaming SparseVolumeTexture frame already has IORequest."));
+
+			int64 FrameSize = Frame.RuntimeStreamedInData.GetBulkDataSize();
+			FrameStorage.RequestStart = FPlatformTime::Seconds();
+			UE_LOG(LogSparseVolumeTextureStreaming, Warning, TEXT("SparseVolumeTexture streaming request started %s Frame:%i At:%.3f\n"), *SparseVolumeTexture->GetName(), IndexToLoad, FrameStorage.RequestStart);
+			FBulkDataIORequestCallBack AsyncFileCallBack = [this, IndexToLoad, FrameSize](bool bWasCancelled, IBulkDataIORequest* Req)
+			{
+				StreamingManager->OnAsyncFileCallback(this, IndexToLoad, FrameSize, Req, bWasCancelled);
+			};
+
+			UE_LOG(LogSparseVolumeTextureStreaming, Warning, TEXT("Loading streaming SparseVolumeTexture %s Frame:%i Offset:%i Size:%i File:%s\n"),
+				*SparseVolumeTexture->GetName(), IndexToLoad, Frame.RuntimeStreamedInData.GetBulkDataOffsetInFile(), Frame.RuntimeStreamedInData.GetBulkDataSize(), *Frame.RuntimeStreamedInData.GetDebugName());
+			FrameStorage.IORequest = Frame.RuntimeStreamedInData.CreateStreamingRequest(AsyncIOPriority, &AsyncFileCallBack, nullptr);
+			if (!FrameStorage.IORequest)
+			{
+				UE_LOG(LogSparseVolumeTextureStreaming, Error, TEXT("SparseVolumeTexture streaming read request failed."));
+			}
+		}
+	}
+}
+
+bool FStreamingSparseVolumeTextureData::BlockTillAllRequestsFinished(float TimeLimit)
+{
+	QUICK_SCOPE_CYCLE_COUNTER(FStreamingSparseVolumeTextureData_BlockTillAllRequestsFinished);
+	if (TimeLimit == 0.0f)
+	{
+		for (FLoadedSparseVolumeTextureFrame& LoadedFrame : LoadedFrames)
+		{
+			LoadedFrame.CleanUpIORequest();
+		}
+	}
+	else
+	{
+		double EndTime = FPlatformTime::Seconds() + TimeLimit;
+		for (FLoadedSparseVolumeTextureFrame& LoadedFrame : LoadedFrames)
+		{
+			if (LoadedFrame.IORequest)
+			{
+				float ThisTimeLimit = EndTime - FPlatformTime::Seconds();
+				if (ThisTimeLimit < 0.001f || // one ms is the granularity of the platform event system
+					!LoadedFrame.IORequest->WaitCompletion(ThisTimeLimit))
+				{
+					return false;
+				}
+
+				LoadedFrame.CleanUpIORequest();
+			}
+		}
+	}
+	return true;
+}
+
+void FStreamingSparseVolumeTextureData::GetMemorySize(SIZE_T* SizeCPU, SIZE_T* SizeGPU) const
+{
+	FScopeLock LoadedFramesLock(&LoadedFramesCriticalSection);
+	for (const FLoadedSparseVolumeTextureFrame& LoadedFrame : LoadedFrames)
+	{
+		const FSparseVolumeTextureSceneProxy* Proxy = LoadedFrame.Proxy.load();
+		if (LoadedFrame.bOwnsProxy && Proxy)
+		{
+			Proxy->GetMemorySize(SizeCPU, SizeGPU);
+		}
+	}
+}
+
+FLoadedSparseVolumeTextureFrame& FStreamingSparseVolumeTextureData::AddNewLoadedFrame(int32 FrameIndex, FSparseVolumeTextureSceneProxy* ExistingProxy)
+{
+	FScopeLock LoadedFramesLock(&LoadedFramesCriticalSection);
+	const int32 NewIndex = LoadedFrames.AddDefaulted();
+
+	LoadedFrames[NewIndex].Proxy = ExistingProxy;
+	LoadedFrames[NewIndex].FrameIndex = FrameIndex;
+	LoadedFrames[NewIndex].bOwnsProxy = false;
+	return LoadedFrames[NewIndex];
+}
+
+void FStreamingSparseVolumeTextureData::FreeLoadedFrame(FLoadedSparseVolumeTextureFrame& LoadedFrame)
+{
+	if (LoadedFrame.IORequest)
+	{
+		LoadedFrame.IORequest->Cancel();
+		LoadedFrame.IORequest->WaitCompletion();
+		delete LoadedFrame.IORequest;
+		LoadedFrame.IORequest = nullptr;
+	}
+
+	if (LoadedFrame.bOwnsProxy)
+	{
+		FSparseVolumeTextureSceneProxy* Proxy = LoadedFrame.Proxy.load();
+		ENQUEUE_RENDER_COMMAND(DeleteSVTProxy)(
+			[Proxy](FRHICommandListImmediate& RHICmdList)
+			{
+				Proxy->ReleaseResource();
+				delete Proxy; 
+			});
+	}
+
+	LoadedFrame.Proxy = nullptr;
+	LoadedFrame.bOwnsProxy = false;
+	LoadedFrame.FrameIndex = INDEX_NONE;
+}
+
+void FStreamingSparseVolumeTextureData::ResetRequestedFrames()
+{
+	RequestedFrameIndices.Reset();
+	RequestedFrameIndices.Add(0); //Always want first frame
+}
+
+////////////////////////////////////////////////////////////////////////////////////////////////
+
+FSparseVolumeTextureStreamingManager::FSparseVolumeTextureStreamingManager()
+{
+}
+
+FSparseVolumeTextureStreamingManager::~FSparseVolumeTextureStreamingManager()
+{
+	checkf(StreamingSparseVolumeTextures.IsEmpty(), TEXT("FSparseVolumeTextureStreamingManager still has %i streaming SparseVolumeTextures registered!"), StreamingSparseVolumeTextures.Num());
+}
+
+void FSparseVolumeTextureStreamingManager::UpdateResourceStreaming(float DeltaTime, bool bProcessEverything)
+{
+	FScopeLock Lock(&CriticalSection);
+	for (auto& StreamingSVTPair : StreamingSparseVolumeTextures)
+	{
+		StreamingSVTPair.Value->UpdateStreamingStatus();
+	}
+}
+
+int32 FSparseVolumeTextureStreamingManager::BlockTillAllRequestsFinished(float TimeLimit, bool bLogResults)
+{
+	FScopeLock Lock(&CriticalSection);
+
+	QUICK_SCOPE_CYCLE_COUNTER(FSparseVolumeTextureStreamingManager_BlockTillAllRequestsFinished);
+	int32 Result = 0;
+
+	if (TimeLimit == 0.0f)
+	{
+		for (auto& StreamingSVTPair : StreamingSparseVolumeTextures)
+		{
+			StreamingSVTPair.Value->BlockTillAllRequestsFinished();
+		}
+	}
+	else
+	{
+		double EndTime = FPlatformTime::Seconds() + TimeLimit;
+		for (auto& StreamingSVTPair : StreamingSparseVolumeTextures)
+		{
+			float ThisTimeLimit = EndTime - FPlatformTime::Seconds();
+			if (ThisTimeLimit < .001f || // one ms is the granularity of the platform event system
+				!StreamingSVTPair.Value->BlockTillAllRequestsFinished(ThisTimeLimit))
+			{
+				Result = 1; // we don't report the actual number, just 1 for any number of outstanding requests
+				break;
+			}
+		}
+	}
+
+	return Result;
+}
+
+void FSparseVolumeTextureStreamingManager::AddSparseVolumeTexture(UStreamableSparseVolumeTexture* SparseVolumeTexture)
+{
+	FScopeLock Lock(&CriticalSection);
+	if (StreamingSparseVolumeTextures.FindRef(SparseVolumeTexture) == nullptr)
+	{
+		FStreamingSparseVolumeTextureData* NewStreamingData = new FStreamingSparseVolumeTextureData();
+		if (NewStreamingData->Initialize(SparseVolumeTexture, this))
+		{
+			StreamingSparseVolumeTextures.Add(SparseVolumeTexture, NewStreamingData);
+		}
+		else
+		{
+			delete NewStreamingData;
+		}
+	}
+}
+
+bool FSparseVolumeTextureStreamingManager::RemoveSparseVolumeTexture(UStreamableSparseVolumeTexture* SparseVolumeTexture)
+{
+	FScopeLock Lock(&CriticalSection);
+	FStreamingSparseVolumeTextureData* StreamingData = StreamingSparseVolumeTextures.FindRef(SparseVolumeTexture);
+	if (StreamingData)
+	{
+		StreamingSparseVolumeTextures.Remove(SparseVolumeTexture);
+
+		// Free resources. This blocks pending IO requests
+		StreamingData->FreeResources();
+		delete StreamingData;
+		return true;
+	}
+	return false;
+}
+
+void FSparseVolumeTextureStreamingManager::GetMemorySizeForSparseVolumeTexture(const UStreamableSparseVolumeTexture* SparseVolumeTexture, SIZE_T* SizeCPU, SIZE_T* SizeGPU) const
+{
+	FScopeLock Lock(&CriticalSection);
+	FStreamingSparseVolumeTextureData* StreamingData = StreamingSparseVolumeTextures.FindRef(SparseVolumeTexture);
+	if (StreamingData)
+	{
+		StreamingData->GetMemorySize(SizeCPU, SizeGPU);
+	}
+}
+
+const FSparseVolumeTextureSceneProxy* FSparseVolumeTextureStreamingManager::GetSparseVolumeTextureSceneProxy(const UStreamableSparseVolumeTexture* SparseVolumeTexture, int32 FrameIndex, bool bTrackAsRequested)
+{
+	FScopeLock Lock(&CriticalSection);
+
+	FStreamingSparseVolumeTextureData* StreamingData = StreamingSparseVolumeTextures.FindRef(SparseVolumeTexture);
+	if (StreamingData)
+	{
+		if (bTrackAsRequested)
+		{
+			StreamingData->RequestedFrameIndices.AddUnique(FrameIndex);
+
+			// Prefetch next frames
+			const int32 NumFrames = SparseVolumeTexture->GetFrames().Num();
+			const int32 NumPrefetchFrames = FMath::Clamp(GSVTNumPrefetchFrames, 0, NumFrames - 1);
+			for (int32 i = 0; i < NumPrefetchFrames; ++i)
+			{
+				StreamingData->RequestedFrameIndices.AddUnique((FrameIndex + 1 + i) % NumFrames);
+			}
+		}
+
+		if (StreamingData->LoadedFrameIndices.Contains(FrameIndex))
+		{
+			if (const FLoadedSparseVolumeTextureFrame* Frame = Algo::FindBy(StreamingData->LoadedFrames, FrameIndex, &FLoadedSparseVolumeTextureFrame::FrameIndex))
+			{
+				if (Frame->Proxy.load() == nullptr)
+				{
+					const double RequestTimer = Frame->RequestStart < 0.0 ? Frame->RequestStart : FPlatformTime::Seconds() - Frame->RequestStart;
+					UE_LOG(LogSparseVolumeTextureStreaming, Warning, TEXT("No render proxy for loaded frame: %i, SVT: %s Request timer : %.3f"), FrameIndex, *SparseVolumeTexture->GetFullName(), RequestTimer);
+				}
+				return Frame->Proxy.load();
+			}
+			else
+			{
+				UE_LOG(LogSparseVolumeTextureStreaming, Warning, TEXT("Unable to find requested frame: %i, SVT: %s - Is in LoadedFrameIndices however"), FrameIndex, *SparseVolumeTexture->GetFullName());
+			}
+		}
+		else
+		{
+			UE_LOG(LogSparseVolumeTextureStreaming, Warning, TEXT("Requested previously unknown frame: %i, SVT: %s"), FrameIndex, *SparseVolumeTexture->GetFullName());
+		}
+	}
+	else
+	{
+		UE_LOG(LogSparseVolumeTextureStreaming, Warning, TEXT("Tried to get frame for SVT that is not registered with the streaming manager SVT: %s"), *SparseVolumeTexture->GetFullName());
+	}
+
+	return nullptr;
+}
+
+void FSparseVolumeTextureStreamingManager::OnAsyncFileCallback(FStreamingSparseVolumeTextureData* StreamingSVTData, int32 FrameIndex, int64 ReadSize, IBulkDataIORequest* ReadRequest, bool bWasCancelled)
+{
+	// Check to see if we successfully managed to load anything
+	uint8* Mem = ReadRequest->GetReadResults();
+
+	FScopeLock Lock(&StreamingSVTData->LoadedFramesCriticalSection);
+
+	const int32 LoadedFrameIdx = StreamingSVTData->LoadedFrames.IndexOfByPredicate([FrameIndex](const FLoadedSparseVolumeTextureFrame& Frame) { return Frame.FrameIndex == FrameIndex; });
+	check(LoadedFrameIdx != INDEX_NONE);
+	FLoadedSparseVolumeTextureFrame& FrameStorage = StreamingSVTData->LoadedFrames[LoadedFrameIdx];
+
+	const double CurrentTime = FPlatformTime::Seconds();
+	const double RequestDuration = CurrentTime - FrameStorage.RequestStart;
+
+	if (Mem)
+	{
+		checkf(FrameStorage.Proxy.load() == nullptr, TEXT("Frame storage already has data. (0x%p) FrameIndex:%i LoadedFrameIdx:%i"), FrameStorage.Proxy.load(), FrameIndex, LoadedFrameIdx);
+
+		TArrayView<const uint8> MemView(Mem, ReadSize);
+		FMemoryReaderView Reader(MemView);
+
+		FSparseVolumeTextureSceneProxy* NewProxy = new FSparseVolumeTextureSceneProxy();
+		NewProxy->GetRuntimeData().Serialize(Reader);
+		BeginInitResource(NewProxy);
+
+		FrameStorage.Proxy = NewProxy;
+		FrameStorage.bOwnsProxy = true;
+		FrameStorage.RequestStart = -2.0; //Signify we have finished loading
+
+		UE_LOG(LogSparseVolumeTextureStreaming, Log, TEXT("Request Finished %.2f\n SparseVolumeTexture Frame Streamed %.4f\n"), CurrentTime, RequestDuration);
+
+		FMemory::Free(Mem);
+	}
+	else
+	{
+		const TCHAR* WasCancelledText = bWasCancelled ? TEXT("Yes") : TEXT("No");
+		UE_LOG(LogSparseVolumeTextureStreaming, Warning, TEXT("Streaming SparseVolumeTexture failed to load frame: %i Load Duration:%.3f, SVT:%s WasCancelled: %s\n"), FrameIndex, RequestDuration, *StreamingSVTData->SparseVolumeTexture->GetName(), WasCancelledText);
+		
+		StreamingSVTData->LoadFailedFrameIndices.AddUnique(FrameIndex);
+	}
+}
diff --git a/Engine/Source/Runtime/Engine/Public/ContentStreaming.h b/Engine/Source/Runtime/Engine/Public/ContentStreaming.h
--- a/Engine/Source/Runtime/Engine/Public/ContentStreaming.h
+++ b/Engine/Source/Runtime/Engine/Public/ContentStreaming.h
@@ -31,6 +31,8 @@ struct FWaveInstance;
 class UAnimStreamable;
 enum class EStreamableRenderAssetType : uint8;
 struct FCompressedAnimSequence;
+class UStreamableSparseVolumeTexture;
+class FSparseVolumeTextureSceneProxy;
 
 using FSoundWaveProxyPtr = TSharedPtr<FSoundWaveProxy, ESPMode::ThreadSafe>;
 
@@ -627,6 +629,31 @@ struct IAnimationStreamingManager : public IStreamingManager
 	virtual const FCompressedAnimSequence* GetLoadedChunk(const UAnimStreamable* Anim, uint32 ChunkIndex, bool bRequestNextChunk) const = 0;
 };
 
+/**
+ * Interface to add functions specifically related to SparseVolumeTexture streaming
+ */
+struct ISparseVolumeTextureStreamingManager : public IStreamingManager
+{
+	/** Adds a new StreamableSparseVolumeTexture to the streaming manager. */
+	virtual void AddSparseVolumeTexture(UStreamableSparseVolumeTexture* SparseVolumeTexture) = 0;
+
+	/** Removes a StreamableSparseVolumeTexture from the streaming manager. */
+	virtual bool RemoveSparseVolumeTexture(UStreamableSparseVolumeTexture* SparseVolumeTexture) = 0;
+
+	/** Returns the memory a StreamableSparseVolumeTexture is currently using */
+	virtual void GetMemorySizeForSparseVolumeTexture(const UStreamableSparseVolumeTexture* SparseVolumeTexture, SIZE_T* SizeCPU, SIZE_T* SizeGPU) const = 0;
+
+	/**
+	 * Gets a pointer to a proxy of a frame of a SparseVolumeTexture
+	 *
+	 * @param SparseVolumeTexture	SparseVolumeTexture we want a frame from
+	 * @param FrameIndex			Index of the frame we want
+	 * @return Either the desired frame proxy or NULL if it's not loaded
+	 */
+	virtual const FSparseVolumeTextureSceneProxy* GetSparseVolumeTextureSceneProxy(const UStreamableSparseVolumeTexture* SparseVolumeTexture, int32 FrameIndex, bool bTrackAsRequested) = 0;
+};
+
+
 /**
  * Streaming manager collection, routing function calls to streaming managers that have been added
  * via AddStreamingManager.
@@ -736,6 +763,11 @@ struct FStreamingManagerCollection : public IStreamingManager
 	*/
 	ENGINE_API Nanite::FCoarseMeshStreamingManager* GetNaniteCoarseMeshStreamingManager() const;
 
+	/**
+	 * Gets a reference to the SparseVolumeTexture Streaming Manager interface
+	 */
+	ENGINE_API ISparseVolumeTextureStreamingManager& GetSparseVolumeTextureStreamingManager() const;
+
 	/**
 	 * Adds a streaming manager to the array of managers to route function calls to.
 	 *
@@ -840,6 +872,9 @@ protected:
 	/** The nanite coarse mesh streaming manager, should always exist */
 	Nanite::FCoarseMeshStreamingManager* NaniteCoarseMeshStreamingManager;
 
+	/** The SparseVolumeTexture streaming manager, should always exist */
+	ISparseVolumeTextureStreamingManager* SparseVolumeTextureStreamingManager;
+
 #if WITH_EDITOR
 	// Locks out any audio streaming manager call when we are re-initializing the audio streaming manager.
 	mutable FCriticalSection AudioStreamingManagerCriticalSection;
diff --git a/Engine/Source/Runtime/Engine/Public/SparseVolumeTexture/SparseVolumeTextureStreaming.h b/Engine/Source/Runtime/Engine/Public/SparseVolumeTexture/SparseVolumeTextureStreaming.h
new file mode 100644
--- /dev/null
+++ b/Engine/Source/Runtime/Engine/Public/SparseVolumeTexture/SparseVolumeTextureStreaming.h
@@ -0,0 +1,97 @@
+// Copyright Epic Games, Inc. All Rights Reserved.
+
+#pragma once
+
+#include "CoreMinimal.h"
+#include "Engine/EngineTypes.h"
+#include "ContentStreaming.h"
+#include "Containers/Map.h"
+
+class UStreamableSparseVolumeTexture;
+class FSparseVolumeTextureSceneProxy;
+
+struct FLoadedSparseVolumeTextureFrame
+{
+	std::atomic<FSparseVolumeTextureSceneProxy*> Proxy;
+	class IBulkDataIORequest* IORequest;
+	double RequestStart;
+	int32 FrameIndex;
+	bool bOwnsProxy;
+
+	explicit FLoadedSparseVolumeTextureFrame()
+		:Proxy(),
+		IORequest(),
+		RequestStart(-1.0),
+		FrameIndex(INDEX_NONE)
+	{
+	}
+	~FLoadedSparseVolumeTextureFrame()
+	{
+		checkf(Proxy.load() == nullptr, TEXT("Render proxy ptr not null (%p), FrameIndex: %u"), Proxy.load(), FrameIndex);
+	}
+
+	void CleanUpIORequest();
+};
+
+struct FStreamingSparseVolumeTextureData
+{
+	UStreamableSparseVolumeTexture* SparseVolumeTexture;
+	TArray<FLoadedSparseVolumeTextureFrame> LoadedFrames;
+	TArray<int32> LoadedFrameIndices;
+	TArray<int32> RequestedFrameIndices;
+	TArray<int32> LoadFailedFrameIndices;
+	mutable FCriticalSection LoadedFramesCriticalSection;
+
+	explicit FStreamingSparseVolumeTextureData();
+	~FStreamingSparseVolumeTextureData();
+
+	// Don't allow copy construction as it could free shared memory
+	FStreamingSparseVolumeTextureData(const FStreamingSparseVolumeTextureData&) = delete;
+	FStreamingSparseVolumeTextureData& operator=(FStreamingSparseVolumeTextureData const&) = delete;
+
+	bool Initialize(UStreamableSparseVolumeTexture* InSparseVolumeTexture, class FSparseVolumeTextureStreamingManager* InStreamingManager);
+	void FreeResources();
+	bool UpdateStreamingStatus();
+	bool HasPendingRequests(TArray<int32>& IndicesToLoad, TArray<int32>& IndicesToFree) const;
+	void BeginPendingRequests(const TArray<int32>& IndicesToLoad, const TArray<int32>& IndicesToFree);
+	bool BlockTillAllRequestsFinished(float TimeLimit = 0.0f);
+	void GetMemorySize(SIZE_T* SizeCPU, SIZE_T* SizeGPU) const;
+
+private:
+	class FSparseVolumeTextureStreamingManager* StreamingManager;
+
+	FLoadedSparseVolumeTextureFrame& AddNewLoadedFrame(int32 FrameIndex, FSparseVolumeTextureSceneProxy* ExistingProxy);
+	void FreeLoadedFrame(FLoadedSparseVolumeTextureFrame& LoadedFrame);
+	void ResetRequestedFrames();
+};
+
+class FSparseVolumeTextureStreamingManager : public ISparseVolumeTextureStreamingManager
+{
+public:
+	explicit FSparseVolumeTextureStreamingManager();
+	virtual ~FSparseVolumeTextureStreamingManager();
+
+	//~ Begin IStreamingManager Interface
+	virtual void UpdateResourceStreaming(float DeltaTime, bool bProcessEverything = false) override;
+	virtual int32 BlockTillAllRequestsFinished(float TimeLimit = 0.0f, bool bLogResults = false) override;
+	virtual void CancelForcedResources() override { /*empty*/ };
+	virtual void NotifyLevelChange() override { /*empty*/ };
+	virtual void SetDisregardWorldResourcesForFrames(int32 NumFrames) override { /*empty*/ };
+	virtual void AddLevel(class ULevel* Level) override { /*empty*/ };
+	virtual void RemoveLevel(class ULevel* Level) override { /*empty*/ };
+	virtual void NotifyLevelOffset(class ULevel* Level, const FVector& Offset) override { /*empty*/ };
+	//~ End IStreamingManager Interface
+
+	//~ Begin ISparseVolumeTextureStreamingManager Interface
+	virtual void AddSparseVolumeTexture(UStreamableSparseVolumeTexture* SparseVolumeTexture) override;
+	virtual bool RemoveSparseVolumeTexture(UStreamableSparseVolumeTexture* SparseVolumeTexture) override;
+	virtual void GetMemorySizeForSparseVolumeTexture(const UStreamableSparseVolumeTexture* SparseVolumeTexture, SIZE_T* SizeCPU, SIZE_T* SizeGPU) const override;
+	virtual const FSparseVolumeTextureSceneProxy* GetSparseVolumeTextureSceneProxy(const UStreamableSparseVolumeTexture* SparseVolumeTexture, int32 FrameIndex, bool bTrackAsRequested) override;
+	//~ End ISparseVolumeTextureStreamingManager Interface
+
+	void OnAsyncFileCallback(FStreamingSparseVolumeTextureData* StreamingSVTData, int32 FrameIndex, int64 ReadSize, IBulkDataIORequest* ReadRequest, bool bWasCancelled);
+
+private:
+	TMap<UStreamableSparseVolumeTexture*, FStreamingSparseVolumeTextureData*> StreamingSparseVolumeTextures;
+	mutable FCriticalSection CriticalSection;
+};
