Remove duplicate call to BeginReplication when a component gets added to an actor

diff --git a/Engine/Source/Runtime/Engine/Private/Actor.cpp b/Engine/Source/Runtime/Engine/Private/Actor.cpp
--- a/Engine/Source/Runtime/Engine/Private/Actor.cpp
+++ b/Engine/Source/Runtime/Engine/Private/Actor.cpp
@@ -5505,6 +5505,8 @@ void AActor::HandleRegisterComponentWithWorld(UActorComponent* Component)
 	{
 		Component->InitializeComponent();
 
+		// The component was finally initialized, it can now be replicated
+		// Note that if this component does not ask to be initialized, it would have started to be replicated inside AddOwnedComponent.
 		if (bOwnerBeginPlayStarted)
 		{
 			AddComponentForReplication(Component);
@@ -5517,13 +5519,6 @@ void AActor::HandleRegisterComponentWithWorld(UActorComponent* Component)
 
 		if (!Component->HasBegunPlay())
 		{
-#if UE_WITH_IRIS
-		    if (Component->GetIsReplicated())
-		    {
-		    	UpdateReplicatedComponent(Component);
-		    }
-#endif // UE_WITH_IRIS
-
 			Component->BeginPlay();
 			ensureMsgf(Component->HasBegunPlay(), TEXT("Failed to route BeginPlay (%s)"), *Component->GetFullName());
 		}
