//UE5/Main - Incremental FastCook - Assertion failed: IsCachedCookedPlatformDataLoaded in AnimSequence.cpp

#fix Remove ClearCachedCookedPlatformData implementation from AnimSequence, as it could clear out compressed anim data shared between platforms. Which causes invalid behaviour when trying to save out a packages for each cook target platform.

diff --git a/Engine/Source/Runtime/Engine/Classes/Animation/AnimSequence.h b/Engine/Source/Runtime/Engine/Classes/Animation/AnimSequence.h
--- a/Engine/Source/Runtime/Engine/Classes/Animation/AnimSequence.h
+++ b/Engine/Source/Runtime/Engine/Classes/Animation/AnimSequence.h
@@ -389,7 +389,6 @@ public:
 	virtual void BeginCacheForCookedPlatformData(const ITargetPlatform* TargetPlatform) override;
 	virtual bool IsCachedCookedPlatformDataLoaded(const ITargetPlatform* TargetPlatform) override;
 	virtual void WillNeverCacheCookedPlatformDataAgain() override;
-	virtual void ClearCachedCookedPlatformData(const ITargetPlatform* TargetPlatform) override;
 	virtual void ClearAllCachedCookedPlatformData() override;
 #endif // WITH_EDITOR
 	virtual void BeginDestroy() override;
diff --git a/Engine/Source/Runtime/Engine/Private/Animation/AnimSequence.cpp b/Engine/Source/Runtime/Engine/Private/Animation/AnimSequence.cpp
--- a/Engine/Source/Runtime/Engine/Private/Animation/AnimSequence.cpp
+++ b/Engine/Source/Runtime/Engine/Private/Animation/AnimSequence.cpp
@@ -464,26 +464,6 @@ void UAnimSequence::WillNeverCacheCookedPlatformDataAgain()
 	PRAGMA_ENABLE_DEPRECATION_WARNINGS
 }
 
-void UAnimSequence::ClearCachedCookedPlatformData(const ITargetPlatform* TargetPlatform)
-{
-	Super::ClearCachedCookedPlatformData(TargetPlatform);
-
-	// Remove platform data if any
-	const FIoHash PlatformHash = CreateDerivedDataKeyHash(TargetPlatform);
-
-	CacheTasksByKeyHash.Remove(PlatformHash);
-	DataByPlatformKeyHash.Remove(PlatformHash);
-
-	if (TargetPlatform->IsRunningPlatform() && PlatformHash == DataKeyHash)
-	{
-		CompressedData.Reset();
-		DataKeyHash = FIoHash::Zero;
-		PRAGMA_DISABLE_DEPRECATION_WARNINGS
-		bUseRawDataOnly = true;
-		PRAGMA_ENABLE_DEPRECATION_WARNINGS
-	}
-}
-
 void UAnimSequence::ClearAllCachedCookedPlatformData()
 {
 	Super::ClearAllCachedCookedPlatformData();
@@ -5076,7 +5056,7 @@ FIoHash UAnimSequence::BeginCacheDerivedData(const ITargetPlatform* TargetPlatfo
 
 	// Wait for any in-flight requests
 	UE::Anim::FAnimSequenceCompilingManager::Get().FinishCompilation({this});
-	
+
 	// We should wait here for any previous compilations?
 	FCompressedAnimSequence* TargetData = nullptr;
 	if (TargetPlatform->IsRunningPlatform())
@@ -5160,7 +5140,7 @@ void UAnimSequence::EndCacheDerivedData(const FIoHash& KeyHash)
 	{
 		return;
 	}
-
+	
 	UE::Anim::FAnimSequenceCompilingManager::Get().FinishCompilation({this});
 }
 	
