Fix removal of Slovak keyboard support for console key by an incorrect merge

diff --git a/Engine/Source/Runtime/Engine/Private/UserInterface/InputSettings.cpp b/Engine/Source/Runtime/Engine/Private/UserInterface/InputSettings.cpp
--- a/Engine/Source/Runtime/Engine/Private/UserInterface/InputSettings.cpp
+++ b/Engine/Source/Runtime/Engine/Private/UserInterface/InputSettings.cpp
@@ -161,6 +161,7 @@ void UInputSettings::AddInternationalConsoleKey()
 			DefaultConsoleKey = FInputKeyManager::Get().GetKeyFromCodes(VK_OEM_5, 0);
 			break;
 
+		case LANG_SLOVAK:
 		case LANG_SWEDISH:
 			DefaultConsoleKey = EKeys::Section;
 			break;
