[shaders] add a cvar 'r.ShaderCompiler.ParallelInProcess' which when enabled will go wide on inprocess compilation when -noshaderworker is specified on the commandline

diff --git a/Engine/Source/Runtime/Engine/Private/ShaderCompiler/ShaderCompiler.cpp b/Engine/Source/Runtime/Engine/Private/ShaderCompiler/ShaderCompiler.cpp
--- a/Engine/Source/Runtime/Engine/Private/ShaderCompiler/ShaderCompiler.cpp
+++ b/Engine/Source/Runtime/Engine/Private/ShaderCompiler/ShaderCompiler.cpp
@@ -117,6 +117,12 @@ static TAutoConsoleVariable<bool> CVarDebugDumpWorkerInputs(
 	TEXT("If true, worker input files will be saved for each individual compile job alongside other debug data (note that r.DumpShaderDebugInfo must also be enabled for this to function)"),
 	ECVF_ReadOnly);
 
+static TAutoConsoleVariable<bool> CVarCompileParallelInProcess(
+	TEXT("r.ShaderCompiler.ParallelInProcess"),
+	false,
+	TEXT("EXPERIMENTAL- If true, shader compilation will be executed in-process in parallel. Note that this will serialize if the legacy preprocessor is enabled."),
+	ECVF_ReadOnly);
+
 static bool IsShaderJobCacheDDCRemotePolicyEnabled()
 {
 	return CVarJobCacheDDCPolicy.GetValueOnAnyThread();
@@ -2988,81 +2994,148 @@ void FShaderCompileUtilities::DeleteFileHelper(const FString& Filename)
 
 int32 FShaderCompileThreadRunnable::CompilingLoop()
 {
-	// push completed jobs to Manager->ShaderMapJobs before asking for new ones, so we can free the workers now and avoid them waiting a cycle
-	PushCompletedJobsToManager();
+	if (!Manager->bAllowCompilingThroughWorkers && CVarCompileParallelInProcess.GetValueOnAnyThread())
+	{
+		int32 NumJobs = Manager->GetNumPendingJobs();
+		if (NumJobs == 0)
+		{
+			return 0;
+		}
+
+		// if -noshaderworker is specified and the experimental in-process parallel compile is enabled, submit tasks for a batch 
+		// of pending jobs to run wide in-process (and a tailing task to mark those jobs as complete in the manager)
+		TUniquePtr<TArray<FShaderCommonCompileJobPtr>> Jobs = TUniquePtr<TArray<FShaderCommonCompileJobPtr>>(new TArray<FShaderCommonCompileJobPtr>());
+		TUniquePtr<TArray<float>> JobTimes = TUniquePtr<TArray<float>>(new TArray<float>());
+		{
+			FScopeLock Lock(&Manager->CompileQueueSection);
+			for (int32 PriorityIndex = MaxPriorityIndex; PriorityIndex >= MinPriorityIndex; --PriorityIndex)
+			{
+				// Throttle how many jobs we kick per tick so we get more frequent progress updates in the UI;
+				// this doesn't seem to have much effect on overall throughput.
+				const int32 MaxNumJobs = 64;
+				Manager->AllJobs.GetPendingJobs(EShaderCompilerWorkerType::LocalThread, (EShaderCompileJobPriority)PriorityIndex, 1, MaxNumJobs, *Jobs);
+			}
+		}
 
-	// Grab more shader compile jobs from the input queue
-	const int32 NumActiveThreads = PullTasksFromQueue();
+		JobTimes->SetNum(Jobs->Num());
 
-	if (NumActiveThreads == 0 && Manager->bAllowAsynchronousShaderCompiling)
-	{
-		// Yield while there's nothing to do
-		// Note: sleep-looping is bad threading practice, wait on an event instead!
-		// The shader worker thread does it because it needs to communicate with other processes through the file system
-		FPlatformProcess::Sleep(.010f);
-	}
+		TArray<UE::Tasks::FTask> CompileTasks;
+		CompileTasks.Reserve(Jobs->Num());
+		for (TArray<FShaderCommonCompileJobPtr>::SizeType JobIndex = 0; JobIndex < Jobs->Num(); ++JobIndex)
+		{
+			FShaderCommonCompileJob* Job = (*Jobs)[JobIndex];
+			float* Time = &(*JobTimes)[JobIndex];
+			CompileTasks.Add(UE::Tasks::Launch(UE_SOURCE_LOCATION, [Job, Time]()
+				{
+					const float StartTime = FPlatformTime::Seconds();
+					FShaderCompileUtilities::ExecuteShaderCompileJob(*Job);
+					*Time = FPlatformTime::Seconds() - StartTime;
+				}));
+		}
+
+		if (!Jobs->IsEmpty())
+		{
+			UE::Tasks::Launch(UE_SOURCE_LOCATION, [Jobs = MoveTemp(Jobs), JobTimes = MoveTemp(JobTimes), this]()
+				{
+					FScopeLock Lock(&Manager->CompileQueueSection);
+					for (FShaderCommonCompileJobPtr Job : *Jobs)
+					{
+						Manager->ProcessFinishedJob(Job.GetReference());
+					}
+
+					float ElapsedTime = 0.0f;
+					for (const float& JobTime : *JobTimes)
+					{
+						ElapsedTime += JobTime;
+					}
 
-	if (Manager->bAllowCompilingThroughWorkers)
+					Manager->WorkersBusyTime += ElapsedTime;
+					COOK_STAT(ShaderCompilerCookStats::AsyncCompileTimeSec += ElapsedTime);
+				}, CompileTasks);
+			// num active threads is up to the task system; as long as we return non-zero
+			// this will indicate to the caller that pending jobs remain
+			return 1;
+		}
+		return 0;
+	}
+	else // compile either through worker processes or single-threaded in-process (depending on Manager->bAllowCompilingThroughWorkers)
 	{
-		// Write out the files which are input to the shader compile workers
-		WriteNewTasks();
+		// push completed jobs to Manager->ShaderMapJobs before asking for new ones, so we can free the workers now and avoid them waiting a cycle
+		PushCompletedJobsToManager();
 
-		// Launch shader compile workers if they are not already running
-		// Workers can time out when idle so they may need to be relaunched
-		bool bAbandonWorkers = LaunchWorkersIfNeeded();
+		// Grab more shader compile jobs from the input queue
+		const int32 NumActiveThreads = PullTasksFromQueue();
 
-		if (bAbandonWorkers)
+		if (NumActiveThreads == 0 && Manager->bAllowAsynchronousShaderCompiling)
 		{
-			// Fall back to local compiles if the SCW crashed.
-			// This is nasty but needed to work around issues where message passing through files to SCW is unreliable on random PCs
-			Manager->bAllowCompilingThroughWorkers = false;
+			// Yield while there's nothing to do
+			// Note: sleep-looping is bad threading practice, wait on an event instead!
+			// The shader worker thread does it because it needs to communicate with other processes through the file system
+			FPlatformProcess::Sleep(.010f);
+		}
 
-			// Try to recover from abandoned workers after a certain amount of single-threaded compilations
-			if (Manager->NumSingleThreadedRunsBeforeRetry == GSingleThreadedRunsIdle)
-			{
-				// First try to recover, only run single-threaded approach once
-				Manager->NumSingleThreadedRunsBeforeRetry = 1;
-			}
-			else if (Manager->NumSingleThreadedRunsBeforeRetry > GSingleThreadedRunsMaxCount)
+		if (Manager->bAllowCompilingThroughWorkers)
+		{
+			// Write out the files which are input to the shader compile workers
+			WriteNewTasks();
+
+			// Launch shader compile workers if they are not already running
+			// Workers can time out when idle so they may need to be relaunched
+			bool bAbandonWorkers = LaunchWorkersIfNeeded();
+
+			if (bAbandonWorkers)
 			{
-				// Stop retry approach after too many retries have failed
-				Manager->NumSingleThreadedRunsBeforeRetry = GSingleThreadedRunsDisabled;
+				// Fall back to local compiles if the SCW crashed.
+				// This is nasty but needed to work around issues where message passing through files to SCW is unreliable on random PCs
+				Manager->bAllowCompilingThroughWorkers = false;
+
+				// Try to recover from abandoned workers after a certain amount of single-threaded compilations
+				if (Manager->NumSingleThreadedRunsBeforeRetry == GSingleThreadedRunsIdle)
+				{
+					// First try to recover, only run single-threaded approach once
+					Manager->NumSingleThreadedRunsBeforeRetry = 1;
+				}
+				else if (Manager->NumSingleThreadedRunsBeforeRetry > GSingleThreadedRunsMaxCount)
+				{
+					// Stop retry approach after too many retries have failed
+					Manager->NumSingleThreadedRunsBeforeRetry = GSingleThreadedRunsDisabled;
+				}
+				else
+				{
+					// Next time increase runs by factor X
+					Manager->NumSingleThreadedRunsBeforeRetry *= GSingleThreadedRunsIncreaseFactor;
+				}
 			}
 			else
 			{
-				// Next time increase runs by factor X
-				Manager->NumSingleThreadedRunsBeforeRetry *= GSingleThreadedRunsIncreaseFactor;
+				// Read files which are outputs from the shader compile workers
+				int32 NumProcessedResults = ReadAvailableResults();
+				if (NumProcessedResults == 0)
+				{
+					// Reduce filesystem query rate while actively waiting for results.
+					FPlatformProcess::Sleep(0.1f);
+				}
 			}
 		}
 		else
 		{
-			// Read files which are outputs from the shader compile workers
-			int32 NumProcessedResults = ReadAvailableResults();
-			if (NumProcessedResults == 0)
-			{
-				// Reduce filesystem query rate while actively waiting for results.
-				FPlatformProcess::Sleep(0.1f);
-			}
-		}
-	}
-	else
-	{
-		// Execute all pending worker tasks single-threaded
-		CompileDirectlyThroughDll();
+			// Execute all pending worker tasks single-threaded
+			CompileDirectlyThroughDll();
 
-		// If single-threaded mode was enabled by an abandoned worker, try to recover after the given amount of runs
-		if (Manager->NumSingleThreadedRunsBeforeRetry > 0)
-		{
-			Manager->NumSingleThreadedRunsBeforeRetry--;
-			if (Manager->NumSingleThreadedRunsBeforeRetry == 0)
+			// If single-threaded mode was enabled by an abandoned worker, try to recover after the given amount of runs
+			if (Manager->NumSingleThreadedRunsBeforeRetry > 0)
 			{
-				UE_LOG(LogShaderCompilers, Display, TEXT("Retry shader compiling through workers."));
-				Manager->bAllowCompilingThroughWorkers = true;
+				Manager->NumSingleThreadedRunsBeforeRetry--;
+				if (Manager->NumSingleThreadedRunsBeforeRetry == 0)
+				{
+					UE_LOG(LogShaderCompilers, Display, TEXT("Retry shader compiling through workers."));
+					Manager->bAllowCompilingThroughWorkers = true;
+				}
 			}
 		}
-	}
 
-	return NumActiveThreads;
+		return NumActiveThreads;
+	}
 }
 
 FShaderCompilerStats* GShaderCompilerStats = nullptr;
