Added include for ENQUEUE_RENDER_COMMAND to SparseVolumeTexture to fix compile errors in some configs.

#rnx

diff --git a/Engine/Source/Runtime/Engine/Private/SparseVolumeTexture/SparseVolumeTexture.cpp b/Engine/Source/Runtime/Engine/Private/SparseVolumeTexture/SparseVolumeTexture.cpp
--- a/Engine/Source/Runtime/Engine/Private/SparseVolumeTexture/SparseVolumeTexture.cpp
+++ b/Engine/Source/Runtime/Engine/Private/SparseVolumeTexture/SparseVolumeTexture.cpp
@@ -9,6 +9,7 @@
 #include "Misc/SecureHash.h"
 #include "EngineUtils.h"
 #include "Shader/ShaderTypes.h"
+#include "RenderingThread.h"
 
 #if WITH_EDITORONLY_DATA
 #include "DerivedDataCache.h"
@@ -436,7 +437,7 @@ void UStreamableSparseVolumeTexture::BeginDestroy()
 	{
 		if (Frame.SparseVolumeTextureSceneProxy)
 		{
-			ENQUEUE_RENDER_COMMAND(DeleteSVTProxy)(
+			ENQUEUE_RENDER_COMMAND(UStreamableSparseVolumeTexture_DeleteSVTProxy)(
 				[Proxy = Frame.SparseVolumeTextureSceneProxy](FRHICommandListImmediate& RHICmdList)
 				{
 					Proxy->ReleaseResource();
diff --git a/Engine/Source/Runtime/Engine/Private/SparseVolumeTexture/SparseVolumeTextureStreaming.cpp b/Engine/Source/Runtime/Engine/Private/SparseVolumeTexture/SparseVolumeTextureStreaming.cpp
--- a/Engine/Source/Runtime/Engine/Private/SparseVolumeTexture/SparseVolumeTextureStreaming.cpp
+++ b/Engine/Source/Runtime/Engine/Private/SparseVolumeTexture/SparseVolumeTextureStreaming.cpp
@@ -2,6 +2,7 @@
 
 #include "SparseVolumeTexture/SparseVolumeTextureStreaming.h"
 #include "SparseVolumeTexture/SparseVolumeTexture.h"
+#include "RenderingThread.h"
 
 DEFINE_LOG_CATEGORY_STATIC(LogSparseVolumeTextureStreaming, Log, All);
 
@@ -312,7 +313,7 @@ void FStreamingSparseVolumeTextureData::FreeLoadedFrame(FLoadedSparseVolumeTextu
 	if (LoadedFrame.bOwnsProxy)
 	{
 		FSparseVolumeTextureSceneProxy* Proxy = LoadedFrame.Proxy.load();
-		ENQUEUE_RENDER_COMMAND(DeleteSVTProxy)(
+		ENQUEUE_RENDER_COMMAND(FStreamingSparseVolumeTextureData_DeleteSVTProxy)(
 			[Proxy](FRHICommandListImmediate& RHICmdList)
 			{
 				Proxy->ReleaseResource();
