Added more missing headers to SparseVolumeTextureStreaming.cpp.

#rnx

diff --git a/Engine/Source/Runtime/Engine/Private/SparseVolumeTexture/SparseVolumeTextureStreaming.cpp b/Engine/Source/Runtime/Engine/Private/SparseVolumeTexture/SparseVolumeTextureStreaming.cpp
--- a/Engine/Source/Runtime/Engine/Private/SparseVolumeTexture/SparseVolumeTextureStreaming.cpp
+++ b/Engine/Source/Runtime/Engine/Private/SparseVolumeTexture/SparseVolumeTextureStreaming.cpp
@@ -4,6 +4,10 @@
 #include "SparseVolumeTexture/SparseVolumeTexture.h"
 #include "RenderingThread.h"
 
+#include "Serialization/LargeMemoryReader.h"
+#include "Serialization/MemoryReader.h"
+#include "Serialization/BulkDataReader.h"
+
 DEFINE_LOG_CATEGORY_STATIC(LogSparseVolumeTextureStreaming, Log, All);
 
 int32 GSVTNumPrefetchFrames = 1;
