Refactor FSkinnedMeshComponentRecreateRenderStateContext class to take USkinnedAsset instead of USkeletalMesh as a constructor argument.
This enables its use in the ClothAsset::Build method.

diff --git a/Engine/Source/Runtime/Engine/Private/SkeletalMesh.cpp b/Engine/Source/Runtime/Engine/Private/SkeletalMesh.cpp
--- a/Engine/Source/Runtime/Engine/Private/SkeletalMesh.cpp
+++ b/Engine/Source/Runtime/Engine/Private/SkeletalMesh.cpp
@@ -62,7 +62,6 @@
 #include "MeshUtilities.h"
 #include "Engine/SkeletalMeshEditorData.h"
 #include "DerivedDataCacheInterface.h"
-#include "IMeshReductionManagerModule.h"
 #include "Engine/RendererSettings.h"
 #include "SkeletalDebugRendering.h"
 #else
@@ -7100,54 +7099,6 @@ void FSkeletalMeshSceneProxy::OnTransformChanged()
 	UpdateDefaultInstanceSceneData();
 }
 
-FSkinnedMeshComponentRecreateRenderStateContext::FSkinnedMeshComponentRecreateRenderStateContext(USkeletalMesh* InSkeletalMesh, bool InRefreshBounds /*= false*/)
-	: bRefreshBounds(InRefreshBounds)
-{
-	bool bRequireFlushRenderingCommands = false;
-	for (TObjectIterator<USkeletalMeshComponent> It; It; ++It)
-	{
-		if (It->GetSkeletalMeshAsset() == InSkeletalMesh)
-		{
-			checkf(!It->IsUnreachable(), TEXT("%s"), *It->GetFullName());
-
-			if (It->IsRenderStateCreated())
-			{
-				check(It->IsRegistered());
-				It->DestroyRenderState_Concurrent();
-				bRequireFlushRenderingCommands = true;
-			}
-				MeshComponents.Add(*It);
-			}
-		}
-
-	// Flush the rendering commands generated by the detachments.
-	// The static mesh scene proxies reference the UStaticMesh, and this ensures that they are cleaned up before the UStaticMesh changes.
-	if (bRequireFlushRenderingCommands)
-	{
-		FlushRenderingCommands();
-	}
-}
-
-FSkinnedMeshComponentRecreateRenderStateContext::~FSkinnedMeshComponentRecreateRenderStateContext()
-{
-	const int32 ComponentCount = MeshComponents.Num();
-	for (int32 ComponentIndex = 0; ComponentIndex < ComponentCount; ++ComponentIndex)
-	{
-		if(USkinnedMeshComponent* Component = MeshComponents[ComponentIndex].Get())
-		{
-			if (bRefreshBounds)
-			{
-				Component->UpdateBounds();
-			}
-
-			if (Component->IsRegistered() && !Component->IsRenderStateCreated() && Component->ShouldCreateRenderState())
-			{
-				Component->CreateRenderState_Concurrent(nullptr);
-			}
-		}
-	}
-}
-
 //////////////////////////////////////////////////////////////////////////
 
 FVector GetRefVertexLocationTyped(
diff --git a/Engine/Source/Runtime/Engine/Private/SkinnedAsset.cpp b/Engine/Source/Runtime/Engine/Private/SkinnedAsset.cpp
--- a/Engine/Source/Runtime/Engine/Private/SkinnedAsset.cpp
+++ b/Engine/Source/Runtime/Engine/Private/SkinnedAsset.cpp
@@ -389,4 +389,53 @@ FString USkinnedAsset::GetLODPathName(const USkinnedAsset* Mesh, int32 LODIndex)
 #else
 	return TEXT("");
 #endif
-}
\ No newline at end of file
+}
+
+FSkinnedMeshComponentRecreateRenderStateContext::FSkinnedMeshComponentRecreateRenderStateContext(USkinnedAsset* InSkinnedAsset, bool InRefreshBounds /*= false*/)
+	: bRefreshBounds(InRefreshBounds)
+{
+	bool bRequireFlushRenderingCommands = false;
+	for (TObjectIterator<USkinnedMeshComponent> It; It; ++It)
+	{
+		if (It->GetSkinnedAsset() == InSkinnedAsset)
+		{
+			checkf(!It->IsUnreachable(), TEXT("%s"), *It->GetFullName());
+
+			if (It->IsRenderStateCreated())
+			{
+				check(It->IsRegistered());
+				It->DestroyRenderState_Concurrent();
+				bRequireFlushRenderingCommands = true;
+			}
+			MeshComponents.Add(*It);
+		}
+	}
+
+	// Flush the rendering commands generated by the detachments.
+	// The static mesh scene proxies reference the UStaticMesh, and this ensures that they are cleaned up before the UStaticMesh changes.
+	if (bRequireFlushRenderingCommands)
+	{
+		FlushRenderingCommands();
+	}
+}
+
+FSkinnedMeshComponentRecreateRenderStateContext::~FSkinnedMeshComponentRecreateRenderStateContext()
+{
+	const int32 ComponentCount = MeshComponents.Num();
+	for (int32 ComponentIndex = 0; ComponentIndex < ComponentCount; ++ComponentIndex)
+	{
+		if(USkinnedMeshComponent* Component = MeshComponents[ComponentIndex].Get())
+		{
+			if (bRefreshBounds)
+			{
+				Component->UpdateBounds();
+			}
+
+			if (Component->IsRegistered() && !Component->IsRenderStateCreated() && Component->ShouldCreateRenderState())
+			{
+				Component->CreateRenderState_Concurrent(nullptr);
+			}
+		}
+	}
+}
+
diff --git a/Engine/Source/Runtime/Engine/Public/SkeletalMeshTypes.h b/Engine/Source/Runtime/Engine/Public/SkeletalMeshTypes.h
--- a/Engine/Source/Runtime/Engine/Public/SkeletalMeshTypes.h
+++ b/Engine/Source/Runtime/Engine/Public/SkeletalMeshTypes.h
@@ -22,6 +22,7 @@ class USkinnedMeshComponent;
 class USkeletalMesh;
 class FSkeletalMeshRenderData;
 class FSkeletalMeshLODRenderData;
+class USkinnedAsset;
 
 /** Flags used when building vertex buffers. */
 enum class ESkeletalMeshVertexFlags : uint8
@@ -133,13 +134,13 @@ struct FClothingSectionData
 };
 
 
-/** Used to recreate all skinned mesh components for a given skeletal mesh */
+/** Used to recreate all skinned mesh components for a given skinned asset. */
 class ENGINE_API FSkinnedMeshComponentRecreateRenderStateContext
 {
 public:
 
 	/** Initialization constructor. */
-	FSkinnedMeshComponentRecreateRenderStateContext(USkeletalMesh* InSkeletalMesh, bool InRefreshBounds = false);
+	FSkinnedMeshComponentRecreateRenderStateContext(USkinnedAsset* InSkinnedAsset, bool InRefreshBounds = false);
 
 	/** Destructor: recreates render state for all components that had their render states destroyed in the constructor. */
 	~FSkinnedMeshComponentRecreateRenderStateContext();
@@ -147,7 +148,7 @@ public:
 private:
 
 	/** List of components to reset */
-	TArray< TWeakObjectPtr<USkinnedMeshComponent>> MeshComponents;
+	TArray<TWeakObjectPtr<USkinnedMeshComponent>> MeshComponents;
 
 	/** Whether we'll refresh the component bounds as we reset */
 	bool bRefreshBounds;
