NonUnity compile error in MaterialCachedData .h/.cpp fix.

diff --git a/Engine/Source/Runtime/Engine/Private/Materials/MaterialCachedData.cpp b/Engine/Source/Runtime/Engine/Private/Materials/MaterialCachedData.cpp
--- a/Engine/Source/Runtime/Engine/Private/Materials/MaterialCachedData.cpp
+++ b/Engine/Source/Runtime/Engine/Private/Materials/MaterialCachedData.cpp
@@ -4,6 +4,7 @@
 #include "MaterialCachedHLSLTree.h"
 #include "MaterialExpressionIO.h"
 #include "Materials/MaterialAttributeDefinitionMap.h"
+#include "Materials/Material.h"
 #include "Materials/MaterialExpression.h"
 #include "Materials/MaterialExpressionCollectionParameter.h"
 #include "Materials/MaterialExpressionDynamicParameter.h"
diff --git a/Engine/Source/Runtime/Engine/Public/MaterialCachedData.h b/Engine/Source/Runtime/Engine/Public/MaterialCachedData.h
--- a/Engine/Source/Runtime/Engine/Public/MaterialCachedData.h
+++ b/Engine/Source/Runtime/Engine/Public/MaterialCachedData.h
@@ -16,6 +16,7 @@ class UTexture;
 class UCurveLinearColor;
 class UCurveLinearColorAtlas;
 class UFont;
+class UMaterial;
 class UMaterialExpression;
 class URuntimeVirtualTexture;
 class ULandscapeGrassType;
