Added Trace when asyncloading is triggered by WorldSetting's High Priority flag

diff --git a/Engine/Source/Runtime/Engine/Private/LevelTick.cpp b/Engine/Source/Runtime/Engine/Private/LevelTick.cpp
--- a/Engine/Source/Runtime/Engine/Private/LevelTick.cpp
+++ b/Engine/Source/Runtime/Engine/Private/LevelTick.cpp
@@ -1398,6 +1398,7 @@ void UWorld::Tick( ELevelTick TickType, float DeltaSeconds )
 	if (Info->bHighPriorityLoading || Info->bHighPriorityLoadingLocal || IsInSeamlessTravel())
 	{
 		CSV_SCOPED_SET_WAIT_STAT(AsyncLoading);
+		TRACE_CPUPROFILER_EVENT_SCOPE(HighPriorityAsyncLoading)
 		// Force it to use the entire time slice, even if blocked on I/O
 		ProcessAsyncLoading(true, true, GPriorityAsyncLoadingExtraTime / 1000.0f);
 	}
