Adding SceneInterface to FXSystem.

diff --git a/Engine/Source/Runtime/Engine/Private/Particles/FXSystem.cpp b/Engine/Source/Runtime/Engine/Private/Particles/FXSystem.cpp
--- a/Engine/Source/Runtime/Engine/Private/Particles/FXSystem.cpp
+++ b/Engine/Source/Runtime/Engine/Private/Particles/FXSystem.cpp
@@ -51,6 +51,7 @@ FFXSystemInterface* FFXSystemInterface::Create(ERHIFeatureLevel::Type InFeatureL
 	{
 		FFXSystemSet* Set = new FFXSystemSet(GPUSortManager);
 		Scene->SetFXSystem(Set);
+		Set->SetSceneInterface(Scene);
 		Set->SetScene((FScene*)Scene);
 
 		Set->FXSystems.Add(new FFXSystem(InFeatureLevel, InShaderPlatform, GPUSortManager));
@@ -60,6 +61,7 @@ FFXSystemInterface* FFXSystemInterface::Create(ERHIFeatureLevel::Type InFeatureL
 			FFXSystemInterface* CustomFX = Ite.Value().Execute(InFeatureLevel, InShaderPlatform, GPUSortManager);
 			if (CustomFX)
 			{
+				CustomFX->SetSceneInterface(Scene);
 				CustomFX->SetScene((FScene*)Scene);
 				Set->FXSystems.Add(CustomFX);
 			}
@@ -70,6 +72,7 @@ FFXSystemInterface* FFXSystemInterface::Create(ERHIFeatureLevel::Type InFeatureL
 	{
 		FFXSystemInterface* Ret = new FFXSystem(InFeatureLevel, InShaderPlatform, GPUSortManager);
 		Scene->SetFXSystem(Ret);
+		Ret->SetSceneInterface(Scene);
 		Ret->SetScene((FScene*)Scene);
 		return Ret;
 	}
diff --git a/Engine/Source/Runtime/Engine/Public/FXSystem.h b/Engine/Source/Runtime/Engine/Public/FXSystem.h
--- a/Engine/Source/Runtime/Engine/Public/FXSystem.h
+++ b/Engine/Source/Runtime/Engine/Public/FXSystem.h
@@ -266,6 +266,9 @@ public:
 
 	virtual void SetSceneTexturesUniformBuffer(const TUniformBufferRef<FSceneTextureUniformParameters>& InSceneTexturesUniformParams) {}
 
+	FORCEINLINE FSceneInterface* GetSceneInterface() const { return SceneInterface; }
+	FORCEINLINE void SetSceneInterface(FSceneInterface* InSceneInterface) { SceneInterface = InSceneInterface; }
+
 	FORCEINLINE FScene* GetScene()const { return Scene; }
 	FORCEINLINE void SetScene(FScene* InScene) { Scene = InScene; }
 
@@ -285,7 +288,8 @@ private:
 
 	bool bIsPendingKill = false;
 
-	class FScene* Scene = nullptr;
+	FSceneInterface* SceneInterface = nullptr;
+	FScene* Scene = nullptr;
 
 	static TMap<FName, FCreateCustomFXSystemDelegate> CreateCustomFXDelegates;
 };
