Fixed bone index handling in UE::Anim::Attributes::MirrorAttributes
#misc handle deprecation path

diff --git a/Engine/Source/Runtime/AnimGraphRuntime/Private/AnimNodes/AnimNode_Mirror.cpp b/Engine/Source/Runtime/AnimGraphRuntime/Private/AnimNodes/AnimNode_Mirror.cpp
--- a/Engine/Source/Runtime/AnimGraphRuntime/Private/AnimNodes/AnimNode_Mirror.cpp
+++ b/Engine/Source/Runtime/AnimGraphRuntime/Private/AnimNodes/AnimNode_Mirror.cpp
@@ -208,7 +208,7 @@ void FAnimNode_MirrorBase::Evaluate_AnyThread(FPoseContext& Output)
 
 		if (GetAttributeMirroring())
 		{
-			UE::Anim::Attributes::MirrorAttributes(Output.CustomAttributes, *MirrorDataTable);
+			UE::Anim::Attributes::MirrorAttributes(Output.CustomAttributes, *MirrorDataTable, CompactPoseMirrorBones);
 		}
 	}
 }
diff --git a/Engine/Source/Runtime/Engine/Private/Animation/AttributesRuntime.cpp b/Engine/Source/Runtime/Engine/Private/Animation/AttributesRuntime.cpp
--- a/Engine/Source/Runtime/Engine/Private/Animation/AttributesRuntime.cpp
+++ b/Engine/Source/Runtime/Engine/Private/Animation/AttributesRuntime.cpp
@@ -439,6 +439,12 @@ void Attributes::BlendAttributesPerBoneFilter(const FStackAttributeContainer& Ba
 }
 
 void Attributes::MirrorAttributes(FStackAttributeContainer& CustomAttributes, const UMirrorDataTable& MirrorTable)
+{
+	static const TArray<FCompactPoseBoneIndex> CompactPoseMirrorBones;
+	Attributes::MirrorAttributes(CustomAttributes, MirrorTable, CompactPoseMirrorBones);
+}
+
+void Attributes::MirrorAttributes(FStackAttributeContainer& CustomAttributes, const UMirrorDataTable& MirrorTable, const TArray<FCompactPoseBoneIndex>& CompactPoseMirrorBones)
 {
 	struct FMirrorSet
 	{
@@ -481,19 +487,41 @@ void Attributes::MirrorAttributes(FStackAttributeContainer& CustomAttributes, co
 			TArray<int32> SortedBoneIndices(UniqueBoneIndices); 
 			SortedBoneIndices.Sort();
 
-			for (int32 CurBoneIndex : UniqueBoneIndices)
+			if(CompactPoseMirrorBones.Num())
 			{
-				// @FIXME: this needs addressing!
-				// Mirror tables are stored using skeleton indices, but are assumed to be used interchangeably with compact pose indices here.
-				int32 MirrorBoneIndex = MirrorTable.BoneToMirrorBoneIndex[CurBoneIndex].GetInt();
-				int32 SwapCheckBoneIndex = MirrorTable.BoneToMirrorBoneIndex[MirrorBoneIndex].GetInt();
-
-				// only handle simple swaps - both entries exist and are different
-				if (MirrorBoneIndex != INDEX_NONE && CurBoneIndex != INDEX_NONE && CurBoneIndex != MirrorBoneIndex &&
-					SwapCheckBoneIndex == CurBoneIndex && Algo::BinarySearch(SortedBoneIndices, MirrorBoneIndex))
-				{
-					AddUnique(CurBoneIndex, MirrorBoneIndex);
-				}
+			    for (int32 CurBoneIndex : UniqueBoneIndices)
+			    {
+				    // only handle simple swaps - both entries exist and are different
+				    if (CompactPoseMirrorBones.IsValidIndex(CurBoneIndex))
+				    {
+					    int32 MirrorBoneIndex = CompactPoseMirrorBones[CurBoneIndex].GetInt();
+					    if (CurBoneIndex != MirrorBoneIndex && CompactPoseMirrorBones.IsValidIndex(MirrorBoneIndex))
+					    {
+						    int32 SwapCheckBoneIndex = CompactPoseMirrorBones[MirrorBoneIndex].GetInt();
+						    if (SwapCheckBoneIndex == CurBoneIndex && Algo::BinarySearch(SortedBoneIndices, MirrorBoneIndex))
+						    {
+							    AddUnique(CurBoneIndex, MirrorBoneIndex);
+						    }
+					    }
+				    }				
+			    }
+			}
+            else
+            {
+				// Deprecated behaviour support, remove when deleting other signature
+			    for (int32 CurBoneIndex : UniqueBoneIndices)
+			    {
+				    // Mirror tables are stored using skeleton indices, but are assumed to be used interchangeably with compact pose indices here.
+				    int32 MirrorBoneIndex = MirrorTable.BoneToMirrorBoneIndex[CurBoneIndex].GetInt();
+				    int32 SwapCheckBoneIndex = MirrorTable.BoneToMirrorBoneIndex[MirrorBoneIndex].GetInt();
+    
+				    // only handle simple swaps - both entries exist and are different
+				    if (MirrorBoneIndex != INDEX_NONE && CurBoneIndex != INDEX_NONE && CurBoneIndex != MirrorBoneIndex &&
+					    SwapCheckBoneIndex == CurBoneIndex && Algo::BinarySearch(SortedBoneIndices, MirrorBoneIndex))
+				    {
+					    AddUnique(CurBoneIndex, MirrorBoneIndex);
+				    }
+			    }
 			}
 
 			TArray<TWrappedAttribute<FAnimStackAllocator>>& ValuesArray = FStackAttributeContainerAccessor::GetValues(CustomAttributes, TypeIndex);
diff --git a/Engine/Source/Runtime/Engine/Public/Animation/AttributesRuntime.h b/Engine/Source/Runtime/Engine/Public/Animation/AttributesRuntime.h
--- a/Engine/Source/Runtime/Engine/Public/Animation/AttributesRuntime.h
+++ b/Engine/Source/Runtime/Engine/Public/Animation/AttributesRuntime.h
@@ -99,6 +99,9 @@ namespace UE
 			static void InterpolateAttributes(FMeshAttributeContainer& FromAttributes, const FMeshAttributeContainer& ToAttributes, float Alpha);
 
 			/** Mirror (swap) attributes with the specified MirrorDataTable. Attributes are swapped using the bone mapping such that bones which are mirrored swap attributes */
+			static void MirrorAttributes(FStackAttributeContainer& Attributes, const UMirrorDataTable& MirrorDataTable, const TArray<FCompactPoseBoneIndex>& CompactPoseMirrorBones);
+
+			UE_DEPRECATED(5.2, "MirrorAttributes has been deprecated, use other signature instead")
 			static void MirrorAttributes(FStackAttributeContainer& Attributes, const UMirrorDataTable& MirrorDataTable);
 			
 			/** Helper functionality to retrieve the correct blend type (from UAnimationSettings) for the provided attribute name */
