Nothing occurs when compressing an animation on a skeleton with no bones

#fix handle fetched from DDC vs build locally case for animation compression

diff --git a/Engine/Source/Runtime/Engine/Private/Animation/AnimCompress.cpp b/Engine/Source/Runtime/Engine/Private/Animation/AnimCompress.cpp
--- a/Engine/Source/Runtime/Engine/Private/Animation/AnimCompress.cpp
+++ b/Engine/Source/Runtime/Engine/Private/Animation/AnimCompress.cpp
@@ -77,7 +77,7 @@ FCompressionMemorySummary::FCompressionMemorySummary(bool bInEnabled)
 
 FCompressionMemorySummary::~FCompressionMemorySummary()
 {
-	if (bEnabled && bUsed && WorstBoneError.IsValid())
+	if (bEnabled && bUsed)
 	{
 		const int32 TotalBeforeSaving = TotalRaw - TotalBeforeCompressed;
 		const int32 TotalAfterSaving = TotalRaw - TotalAfterCompressed;
@@ -101,18 +101,24 @@ FCompressionMemorySummary::~FCompressionMemorySummary()
 		Args.Add(TEXT("TotalTimeSpentCompressingPretty"), FText::FromString(FPlatformTime::PrettyTime(TotalCompressionExecutionTime)));
 		Args.Add(TEXT("TotalTimeSpentCompressingRawSeconds"), FText::AsNumber((float)TotalCompressionExecutionTime, &Options));
 
-		const FErrorTrackerWorstBone WorstBone = WorstBoneError.GetMaxErrorItem();
-		const FErrorTrackerWorstAnimation WorstAnimation = WorstAnimationError.GetMaxErrorItem();
-
 		Args.Add(TEXT("AverageError"), FText::AsNumber(AverageError, &Options));
-
-		Args.Add(TEXT("WorstBoneError"), WorstBone.ToText());
-		Args.Add(TEXT("WorstAnimationError"), WorstAnimation.ToText());
+		if (WorstBoneError.IsValid())
+		{
+			const FErrorTrackerWorstBone WorstBone = WorstBoneError.GetMaxErrorItem();
+			const FErrorTrackerWorstAnimation WorstAnimation = WorstAnimationError.GetMaxErrorItem();
+			Args.Add(TEXT("WorstBoneError"), WorstBone.ToText());
+			Args.Add(TEXT("WorstAnimationError"), WorstAnimation.ToText());
+		}
+		else
+		{
+			Args.Add(TEXT("WorstBoneError"), FText::FromName(NAME_None));
+			Args.Add(TEXT("WorstAnimationError"), FText::FromName(NAME_None));
+		}
 
 		FText Message;
 		if (!bPerformedCompression)
 		{
-			Message = FText::Format(NSLOCTEXT("Engine", "CompressionMemorySummaryNoChange", "Compressed {NumberOfAnimations} Animation(s)\n\nRaw: {TotalRaw} - Compressed: {TotalAfterCompressed}\nSaving: {TotalAfterSaving} ({NewCompressionRatio})\n\nEnd Effector Translation Added By Compression:\n Average: {AverageError} Max:\n{WorstBoneError}\n\nMax Average Animation Error:\n{WorstAnimationError}"), Args);
+			Message = FText::Format(NSLOCTEXT("Engine", "CompressionMemorySummaryNoChange", "Fetched compressed data for {NumberOfAnimations} Animation(s)\n\nRaw: {TotalRaw} - Compressed: {TotalAfterCompressed}\nSaving: {TotalAfterSaving} ({NewCompressionRatio})\n\nEnd Effector Translation Added By Compression:\n Average: {AverageError} Max:\n{WorstBoneError}\n\nMax Average Animation Error:\n{WorstAnimationError}"), Args);
 		}
 		else
 		{
