Regression in UAnimSequence::ExtractRootMotionFromRange
#fix round to frame-number when calculating # of sampled frames using GetPlayLength (float)

diff --git a/Engine/Source/Runtime/Engine/Private/Animation/AnimSequence.cpp b/Engine/Source/Runtime/Engine/Private/Animation/AnimSequence.cpp
--- a/Engine/Source/Runtime/Engine/Private/Animation/AnimSequence.cpp
+++ b/Engine/Source/Runtime/Engine/Private/Animation/AnimSequence.cpp
@@ -1134,7 +1134,7 @@ void UAnimSequence::GetBoneTransform(FTransform& OutAtom, FSkeletonPoseBoneIndex
 
 	if ( bEvaluateCompressedData && IsCompressedDataValid())
 	{
-		FAnimSequenceDecompressionContext DecompContext(PlatformTargetFrameRate.Default, PlatformTargetFrameRate.Default.AsFrameNumber(GetPlayLength()).Value, Interpolation, GetRetargetTransformsSourceName(), *CompressedData.CompressedDataStructure, GetSkeleton()->GetRefLocalPoses(), CompressedData.CompressedTrackToSkeletonMapTable, GetSkeleton(), IsValidAdditive());
+		FAnimSequenceDecompressionContext DecompContext(PlatformTargetFrameRate.Default, PlatformTargetFrameRate.Default.AsFrameTime(GetPlayLength()).RoundToFrame().Value, Interpolation, GetRetargetTransformsSourceName(), *CompressedData.CompressedDataStructure, GetSkeleton()->GetRefLocalPoses(), CompressedData.CompressedTrackToSkeletonMapTable, GetSkeleton(), IsValidAdditive());
 		DecompContext.Seek(Time);
 		if (CompressedData.BoneCompressionCodec)
 		{
@@ -3318,7 +3318,7 @@ int32 UAnimSequence::GetNumberOfSampledKeys() const
 #if WITH_EDITOR
 	return NumberOfSampledKeys;
 #else
-	return CompressedData.CompressedDataStructure ? CompressedData.CompressedDataStructure->CompressedNumberOfKeys :  PlatformTargetFrameRate.Default.AsFrameNumber(GetPlayLength()).Value + 1;
+	return CompressedData.CompressedDataStructure ? CompressedData.CompressedDataStructure->CompressedNumberOfKeys :  PlatformTargetFrameRate.Default.AsFrameTime(GetPlayLength()).RoundToFrame().Value + 1;
 #endif
 }
 
