Fixed warning about keeping a reference to an element of a temporary TArrayView element in SparseVolumeTextureStreaming.

#rnx

diff --git a/Engine/Source/Runtime/Engine/Private/SparseVolumeTexture/SparseVolumeTextureStreaming.cpp b/Engine/Source/Runtime/Engine/Private/SparseVolumeTexture/SparseVolumeTextureStreaming.cpp
--- a/Engine/Source/Runtime/Engine/Private/SparseVolumeTexture/SparseVolumeTextureStreaming.cpp
+++ b/Engine/Source/Runtime/Engine/Private/SparseVolumeTexture/SparseVolumeTextureStreaming.cpp
@@ -218,10 +218,10 @@ void FStreamingSparseVolumeTextureData::BeginPendingRequests(const TArray<int32>
 	// Set off all IO Requests
 
 	const EAsyncIOPriorityAndFlags AsyncIOPriority = AIOP_CriticalPath; //Set to Crit temporarily as emergency speculative fix for streaming issue
-
+	TArrayView<const FSparseVolumeTextureFrame> SVTFrames = SparseVolumeTexture->GetFrames();
 	for (int32 IndexToLoad : IndicesToLoad)
 	{
-		const FSparseVolumeTextureFrame& Frame = SparseVolumeTexture->GetFrames()[IndexToLoad];
+		const FSparseVolumeTextureFrame& Frame = SVTFrames[IndexToLoad];
 		FSparseVolumeTextureSceneProxy* ExistingProxy = Frame.SparseVolumeTextureSceneProxy;
 		FLoadedSparseVolumeTextureFrame& FrameStorage = AddNewLoadedFrame(IndexToLoad, ExistingProxy);
 
