Add CVar to dump shader output cache hits to compare cache hit ratio between HLSLcc and DXC.

This adds CVar "r.DumpShaderOutputCacheHits" to dump shader code output and cache hits.
For FMemsetBufferCS compiled with HLSLcc in a Lyra cook for instance, this produces ByteBuffer.bytecode, ByteBuffer.1.cachehit, ..., ByteBuffer.3.cachehit
whereas the .cachehit files are text files only containing the path to the debug info folder where the cache hit was requested.
This makes it easy to search for cache hits in a file manager and compare outputs between HLSLcc and DXC for example (or other backends).

#rnx

diff --git a/Engine/Source/Runtime/Engine/Private/ShaderCompiler/ShaderCompiler.cpp b/Engine/Source/Runtime/Engine/Private/ShaderCompiler/ShaderCompiler.cpp
--- a/Engine/Source/Runtime/Engine/Private/ShaderCompiler/ShaderCompiler.cpp
+++ b/Engine/Source/Runtime/Engine/Private/ShaderCompiler/ShaderCompiler.cpp
@@ -752,7 +752,7 @@ void FShaderCompileJobCollection::AddToCacheAndProcessPending(FShaderCommonCompi
 	{
 		const bool bAddToDDC = !(FinishedJob->bIsDefaultMaterial || FinishedJob->bIsGlobalShader);
 		// we only cache jobs that succeded
-		CompletedJobsCache.Add(InputHash, Buffer, NumOutstandingJobsWithSameHash, bAddToDDC);
+		CompletedJobsCache.AddJobOutput(FinishedJob, InputHash, Buffer, NumOutstandingJobsWithSameHash, bAddToDDC);
 	}
 
 	// remove ourselves from the jobs in flight, if we were there (if this job is a cloned job it might not have been)
@@ -928,6 +928,13 @@ static FAutoConsoleVariableRef CVarDumpShaderDebugInfo(
 	TEXT("On iOS, if the PowerVR graphics SDK is installed to the default path, the PowerVR shader compiler will be called and errors will be reported during the cook.")
 	);
 
+static TAutoConsoleVariable<bool> CVarDumpShaderOutputCacheHits(
+	TEXT("r.DumpShaderOutputCacheHits"),
+	false,
+	TEXT("Dumps shader output bytecode and cache hits with reference to original output.\n")
+	TEXT("Dumping shader output bytecode for all compile shaders also requires CVar r.DumpShaderDebugInfo=1."),
+	ECVF_ReadOnly);
+
 static int32 GDumpShaderDebugInfoShort = 0;
 static FAutoConsoleVariableRef CVarDumpShaderDebugShortNames(
 	TEXT("r.DumpShaderDebugShortNames"),
@@ -8502,6 +8509,9 @@ public:
 	/** Canned output */
 	FSharedBuffer JobOutput;
 
+	/** Path to where the cached debug info is stored. */
+	FString CachedDebugInfoPath;
+
 	/** Similar to FRefCountBase AddRef, but not atomic */
 	int32 AddRef()
 	{
@@ -8650,7 +8660,7 @@ FShaderJobCache::~FShaderJobCache()
 	}
 }
 
-void FShaderJobCache::Add(const FJobInputHash& Hash, const FJobCachedOutput& Contents, int32 InitialHitCount, const bool bAddToDDC)
+void FShaderJobCache::AddJobOutput(const FShaderCommonCompileJob* FinishedJob, const FJobInputHash& Hash, const FJobCachedOutput& Contents, int32 InitialHitCount, const bool bAddToDDC)
 {
 	TRACE_CPUPROFILER_EVENT_SCOPE(FShaderJobCache::Add);
 
@@ -8667,6 +8677,23 @@ void FShaderJobCache::Add(const FJobInputHash& Hash, const FJobCachedOutput& Con
 
 	FJobOutputHash OutputHash = FBlake3::HashBuffer(Contents.GetData(), Contents.GetSize());
 
+	const bool bDumpCachedDebugInfo = CVarDumpShaderOutputCacheHits.GetValueOnAnyThread();
+
+	// Get dump shader debug output path
+	FString InputDebugInfoPath, InputSourceFilename;
+	if (bDumpCachedDebugInfo)
+	{
+		if (const FShaderCompileJob* SingleJob = FinishedJob->GetSingleShaderJob())
+		{
+			const FShaderCompilerInput& Input = SingleJob->Input;
+			if (!Input.DumpDebugInfoPath.IsEmpty())
+			{
+				InputDebugInfoPath = Input.DumpDebugInfoPath;
+				InputSourceFilename = FPaths::GetBaseFilename(Input.GetSourceFilename());
+			}
+		}
+	}
+
 	// add the record
 	const uint64 InputHashToOutputOriginalSize = InputHashToOutput.GetAllocatedSize();
 
@@ -8679,7 +8706,19 @@ void FShaderJobCache::Add(const FJobInputHash& Hash, const FJobCachedOutput& Con
 	if (CannedOutput && (bCachePerShaderDDC == false))
 	{
 		// update the output hit count
-		(*CannedOutput)->AddRef();
+		const int32 NumRef = (*CannedOutput)->AddRef();
+
+		if (bDumpCachedDebugInfo)
+		{
+			// Write cache hit debug file
+			const FString& CachedDebugInfoPath = (*CannedOutput)->CachedDebugInfoPath;
+			if (!CachedDebugInfoPath.IsEmpty())
+			{
+				const int32 CacheHit = NumRef - 1;
+				const FString CacheHitFilename = FString::Printf(TEXT("%s/%s.%d.cachehit"), *CachedDebugInfoPath, *InputSourceFilename, CacheHit);
+				FFileHelper::SaveStringToFile(InputDebugInfoPath, *CacheHitFilename);
+			}
+		}
 	}
 	else
 	{
@@ -8688,12 +8727,22 @@ void FShaderJobCache::Add(const FJobInputHash& Hash, const FJobCachedOutput& Con
 		FStoredOutput* NewStoredOutput = new FStoredOutput();
 		NewStoredOutput->NumHits = InitialHitCount;
 		NewStoredOutput->JobOutput = Contents;
-
+		NewStoredOutput->CachedDebugInfoPath = InputDebugInfoPath;
 		NewStoredOutput->AddRef();
 		Outputs.Add(OutputHash, NewStoredOutput);
 
 		CurrentlyAllocatedMemory += NewStoredOutput->GetAllocatedSize();
 
+		if (bDumpCachedDebugInfo)
+		{
+			// Write new allocated cache file
+			if (!InputDebugInfoPath.IsEmpty())
+			{
+				const FString CacheFilename = FString::Printf(TEXT("%s/%s.bytecode"), *InputDebugInfoPath, *InputSourceFilename);
+				FFileHelper::SaveArrayToFile(NewStoredOutput->JobOutput, *CacheFilename);
+			}
+		}
+
 #if WITH_EDITOR
 		if (bCachePerShaderDDC)
 		{
diff --git a/Engine/Source/Runtime/Engine/Public/ShaderCompiler.h b/Engine/Source/Runtime/Engine/Public/ShaderCompiler.h
--- a/Engine/Source/Runtime/Engine/Public/ShaderCompiler.h
+++ b/Engine/Source/Runtime/Engine/Public/ShaderCompiler.h
@@ -59,7 +59,7 @@ public:
 	FJobCachedOutput* Find(const FJobInputHash& Hash, const bool bCheckDDC);
 
 	/** Adds a job output to the cache */
-	void Add(const FJobInputHash& Hash, const FJobCachedOutput& Contents, int InitialHitCount, const bool bAddToDDC);
+	void AddJobOutput(const FShaderCommonCompileJob* FinishedJob, const FJobInputHash& Hash, const FJobCachedOutput& Contents, int InitialHitCount, const bool bAddToDDC);
 
 	/** Calculates memory used by the cache*/
 	uint64 GetAllocatedMemory();
