Fix Horde build failure after JobOutput has been changed from TArray<uint8> to FSharedBuffer.

#rnx

diff --git a/Engine/Source/Runtime/Engine/Private/ShaderCompiler/ShaderCompiler.cpp b/Engine/Source/Runtime/Engine/Private/ShaderCompiler/ShaderCompiler.cpp
--- a/Engine/Source/Runtime/Engine/Private/ShaderCompiler/ShaderCompiler.cpp
+++ b/Engine/Source/Runtime/Engine/Private/ShaderCompiler/ShaderCompiler.cpp
@@ -8739,7 +8739,7 @@ void FShaderJobCache::AddJobOutput(const FShaderCommonCompileJob* FinishedJob, c
 			if (!InputDebugInfoPath.IsEmpty())
 			{
 				const FString CacheFilename = FString::Printf(TEXT("%s/%s.bytecode"), *InputDebugInfoPath, *InputSourceFilename);
-				FFileHelper::SaveArrayToFile(NewStoredOutput->JobOutput, *CacheFilename);
+				FFileHelper::SaveArrayToFile(TArrayView<const uint8>((const uint8*)NewStoredOutput->JobOutput.GetData(), NewStoredOutput->JobOutput.GetSize()), *CacheFilename);
 			}
 		}
 
