Strata - added Strata UI node to be able to author UI element material.
For project Opal

Work done by Sebastien.hillaire

diff --git a/Engine/Shaders/Private/SlateElementPixelShader.usf b/Engine/Shaders/Private/SlateElementPixelShader.usf
--- a/Engine/Shaders/Private/SlateElementPixelShader.usf
+++ b/Engine/Shaders/Private/SlateElementPixelShader.usf
@@ -129,17 +129,40 @@ half4 GetMaterialColor( VertexToPixelInterpolants InVertex, float2 MeshUV )
 
 	// Clip if masked
 	GetMaterialCoverageAndClipping( Parameters, PixelMaterialInputs );
+
+	half4 OutColor;
 	
+#if STRATA_ENABLED
+
+	// UI element are always rendered with a unlit BSDF behind the scene
+	FStrataBSDF UnlitBSDF = PixelMaterialInputs.FrontMaterial.InlinedBSDF;
+	SanitizeStrataUnlit(UnlitBSDF);
+
+#if ((STRATA_OPAQUE_DEFERRED || (STRATA_OPAQUE_MATERIAL && STRATA_FORWARD_SHADING)))
+	OutColor = half4(BSDF_GETEMISSIVE(UnlitBSDF), 1.0f);
+#else // STRATA_TRANSLUCENT_FORWARD || STRATA_FORWARD_SHADING || STRATA_MATERIAL_EXPORT_FROM_TRANSLUCENT
+	const float Coverage = saturate(1.0f - UNLIT_TRANSMITTANCE(UnlitBSDF).x);	// Unlit node transmittance is only generated from opacity
+#if MATERIALBLENDING_ADDITIVE
+	OutColor = half4(BSDF_GETEMISSIVE(UnlitBSDF) * Coverage, 0.0f);
+#else
+	OutColor = half4(BSDF_GETEMISSIVE(UnlitBSDF), Coverage);
+#endif
+#endif
+
+	// STRATA_TODO fix premultiplied alpha properly
+
+#else // STRATA_ENABLED
+
 	half Opacity = GetMaterialOpacity(PixelMaterialInputs);
 	half3 Color = GetMaterialEmissive(PixelMaterialInputs);
 
-	half4 OutColor;
+#if MATERIALBLENDING_ADDITIVE
+	OutColor = half4(Color * Opacity, 0.0f);
+#else
+	OutColor = half4(Color, Opacity);
+#endif
 
-	#if MATERIALBLENDING_ADDITIVE
-		OutColor = half4( Color * Opacity, 0.0f );
-	#else
-		OutColor = half4( Color, Opacity );
-	#endif
+#endif // STRATA_ENABLED
 
 	return OutColor;
 }
diff --git a/Engine/Source/Editor/UnrealEd/Classes/MaterialGraph/MaterialGraph.h b/Engine/Source/Editor/UnrealEd/Classes/MaterialGraph/MaterialGraph.h
--- a/Engine/Source/Editor/UnrealEd/Classes/MaterialGraph/MaterialGraph.h
+++ b/Engine/Source/Editor/UnrealEd/Classes/MaterialGraph/MaterialGraph.h
@@ -64,9 +64,19 @@ struct FMaterialInputInfo
 		}
 		else if( Material->IsUIMaterial() )
 		{
-			if( Property == MP_EmissiveColor || Property == MP_Opacity || Property == MP_OpacityMask || Property == MP_WorldPositionOffset )
+			if (Strata::IsStrataEnabled())
 			{
-				return true;
+				if (Property == MP_FrontMaterial || Property == MP_OpacityMask || Property == MP_WorldPositionOffset)
+				{
+					return true;
+				}
+			}
+			else
+			{
+				if (Property == MP_EmissiveColor || Property == MP_Opacity || Property == MP_OpacityMask || Property == MP_WorldPositionOffset)
+				{
+					return true;
+				}
 			}
 
 			if (Property >= MP_CustomizedUVs0 && Property <= MP_CustomizedUVs7)
diff --git a/Engine/Source/Runtime/Engine/Classes/Engine/EngineTypes.h b/Engine/Source/Runtime/Engine/Classes/Engine/EngineTypes.h
--- a/Engine/Source/Runtime/Engine/Classes/Engine/EngineTypes.h
+++ b/Engine/Source/Runtime/Engine/Classes/Engine/EngineTypes.h
@@ -686,6 +686,7 @@ enum EStrataShadingModel : int
 	SSM_LightFunction			UMETA(DisplayName = "LightFunction"),
 	SSM_PostProcess				UMETA(DisplayName = "PostProcess"),
 	SSM_Decal					UMETA(DisplayName = "Decal"),
+	SSM_UI						UMETA(DisplayName = "UI"),
 	/** Number of unique shading models. */
 	SSM_NUM						UMETA(Hidden),
 };
diff --git a/Engine/Source/Runtime/Engine/Classes/Materials/MaterialExpressionStrata.h b/Engine/Source/Runtime/Engine/Classes/Materials/MaterialExpressionStrata.h
--- a/Engine/Source/Runtime/Engine/Classes/Materials/MaterialExpressionStrata.h
+++ b/Engine/Source/Runtime/Engine/Classes/Materials/MaterialExpressionStrata.h
@@ -731,6 +731,36 @@ class UMaterialExpressionStrataPostProcess : public UMaterialExpressionStrataBSD
 	//~ End UMaterialExpression Interface
 };
 
+UCLASS(MinimalAPI, collapsecategories, hidecategories = Object, DisplayName = "Strata UI")
+class UMaterialExpressionStrataUI : public UMaterialExpressionStrataBSDF
+{
+	GENERATED_UCLASS_BODY()
+
+	/**
+	 * The output color of the UI element.
+	 */
+	UPROPERTY()
+	FExpressionInput Color;
+
+	/**
+	 * The coverage of the UI element: the more the value is high, the less the back buffer will be visible.
+	 */
+	UPROPERTY()
+	FExpressionInput Opacity;
+
+	//~ Begin UMaterialExpression Interface
+#if WITH_EDITOR
+	virtual int32 Compile(class FMaterialCompiler* Compiler, int32 OutputIndex) override;
+	virtual void GetCaption(TArray<FString>& OutCaptions) const override;
+	virtual uint32 GetOutputType(int32 OutputIndex) override;
+	virtual uint32 GetInputType(int32 InputIndex) override;
+	virtual bool IsResultStrataMaterial(int32 OutputIndex) override;
+	virtual void GatherStrataMaterialInfo(FStrataMaterialInfo& StrataMaterialInfo, int32 OutputIndex) override;
+	virtual FStrataOperator* StrataGenerateMaterialTopologyTree(class FMaterialCompiler* Compiler, class UMaterialExpression* Parent, int32 OutputIndex) override;
+#endif
+	//~ End UMaterialExpression Interface
+};
+
 UCLASS(MinimalAPI, collapsecategories, hidecategories = Object, DisplayName = "Strata Convert To Decal")
 class UMaterialExpressionStrataConvertToDecal : public UMaterialExpressionStrataBSDF
 {
diff --git a/Engine/Source/Runtime/Engine/Private/Materials/Material.cpp b/Engine/Source/Runtime/Engine/Private/Materials/Material.cpp
--- a/Engine/Source/Runtime/Engine/Private/Materials/Material.cpp
+++ b/Engine/Source/Runtime/Engine/Private/Materials/Material.cpp
@@ -3533,6 +3533,21 @@ void UMaterial::ConvertMaterialToStrataMaterial()
 			EditorOnly->FrontMaterial.Connect(0, ConvertToDecalNode);
 			bInvalidateShader = true;
 		}
+		else if (MaterialDomain == MD_UI)
+		{
+			// Some materials don't have their shading mode set correctly to Unlit. Since only Unlit is supported, forcing it here.
+			ShadingModel = MSM_Unlit;
+			ShadingModels.ClearShadingModels();
+			ShadingModels.AddShadingModel(MSM_Unlit);
+
+			UMaterialExpressionStrataUI* UINode = NewObject<UMaterialExpressionStrataUI>(this);
+			SetPosXAndMoveReferenceToTheRight(UINode);
+			MoveConnectionTo(EditorOnly->EmissiveColor, UINode, 0);
+			MoveConnectionTo(EditorOnly->Opacity, UINode, 1);
+
+			EditorOnly->FrontMaterial.Connect(0, UINode);
+			bInvalidateShader = true;
+		}
 
 		BlendMode = ConvertLegacyBlendMode(BlendMode, ShadingModels);
 	}
@@ -4266,6 +4281,11 @@ bool UMaterial::CanEditChange(const FProperty* InProperty) const
 		FString PropertyName = InProperty->GetName();
 		const bool bStrataEnabled = Strata::IsStrataEnabled();
 
+		if (PropertyName == GET_MEMBER_NAME_STRING_CHECKED(UMaterial, MaterialDomain))
+		{
+			return !bStrataEnabled; // Material domain is no longer tweakable with Strata. It is instead derived from the graph.
+		}
+
 		if (PropertyName == GET_MEMBER_NAME_STRING_CHECKED(UMaterial, PhysMaterial) || PropertyName == GET_MEMBER_NAME_STRING_CHECKED(UMaterial, PhysMaterialMask))
 		{
 			return MaterialDomain == MD_Surface;
@@ -4874,6 +4894,11 @@ void UMaterial::RebuildShadingModelField()
 					BlendMode = BLEND_Opaque;
 				}
 			}
+			else if (StrataMaterialInfo.HasOnlyShadingModel(SSM_UI))
+			{
+				MaterialDomain = EMaterialDomain::MD_UI;
+				ShadingModel = MSM_Unlit;
+			}
 			else if (StrataMaterialInfo.HasShadingModel(SSM_Decal))
 			{
 				// Decal can have multiple shading model
@@ -6449,15 +6474,6 @@ bool UMaterial::IsPropertySupported(EMaterialProperty InProperty) const
 		case MP_FrontMaterial:
 			bSupported = true;
 			break;
-		// UI domain uses the following remapping:
-		// * MP_EmissiveColor       -> Final Color
-		// * MP_WorldPositionOffset -> Screen Position
-		// * MP_OpacityMask         -> Opacity Mask
-		// * MP_Opacity             -> Opacity
-		case MP_EmissiveColor:
-		case MP_Opacity:
-			bSupported = MaterialDomain == EMaterialDomain::MD_UI;
-			break;
 		}
 	}
 	return bSupported;
@@ -6585,11 +6601,21 @@ static bool IsPropertyActive_Internal(EMaterialProperty InProperty,
 	}
 	else if (Domain == MD_UI)
 	{
-		return InProperty == MP_EmissiveColor
-			|| (InProperty == MP_WorldPositionOffset)
-			|| (InProperty == MP_OpacityMask && IsMaskedBlendMode(BlendMode))
-			|| (InProperty == MP_Opacity && IsTranslucentBlendMode(BlendMode) && BlendMode != BLEND_Modulate)
-			|| (InProperty >= MP_CustomizedUVs0 && InProperty <= MP_CustomizedUVs7);
+		if (bStrataEnabled)
+		{
+			return InProperty == MP_FrontMaterial
+				|| (InProperty == MP_WorldPositionOffset)
+				|| (InProperty == MP_OpacityMask && IsMaskedBlendMode(BlendMode))
+				|| (InProperty >= MP_CustomizedUVs0 && InProperty <= MP_CustomizedUVs7);
+		}
+		else
+		{
+			return InProperty == MP_EmissiveColor
+				|| (InProperty == MP_WorldPositionOffset)
+				|| (InProperty == MP_OpacityMask && IsMaskedBlendMode(BlendMode))
+				|| (InProperty == MP_Opacity && IsTranslucentBlendMode(BlendMode) && BlendMode != BLEND_Modulate)
+				|| (InProperty >= MP_CustomizedUVs0 && InProperty <= MP_CustomizedUVs7);
+		}
 	}
 
 	// Now processing MD_Surface
diff --git a/Engine/Source/Runtime/Engine/Private/Materials/MaterialExpressions.cpp b/Engine/Source/Runtime/Engine/Private/Materials/MaterialExpressions.cpp
--- a/Engine/Source/Runtime/Engine/Private/Materials/MaterialExpressions.cpp
+++ b/Engine/Source/Runtime/Engine/Private/Materials/MaterialExpressions.cpp
@@ -23824,6 +23824,80 @@ FStrataOperator* UMaterialExpressionStrataPostProcess::StrataGenerateMaterialTop
 
 
 
+UMaterialExpressionStrataUI::UMaterialExpressionStrataUI(const FObjectInitializer& ObjectInitializer)
+	: Super(ObjectInitializer)
+{
+	struct FConstructorStatics
+	{
+		FText NAME_Strata;
+		FConstructorStatics() : NAME_Strata(LOCTEXT("Strata Extras", "Strata Extras")) { }
+	};
+	static FConstructorStatics ConstructorStatics;
+#if WITH_EDITORONLY_DATA
+	MenuCategories.Add(ConstructorStatics.NAME_Strata);
+#endif
+}
+
+#if WITH_EDITOR
+int32 UMaterialExpressionStrataUI::Compile(class FMaterialCompiler* Compiler, int32 OutputIndex)
+{
+	int OpacityCodeChunk = CompileWithDefaultFloat1(Compiler, Opacity, 0.0f);
+	int TransmittanceCodeChunk = Compiler->Saturate(Compiler->Sub(Compiler->Constant(1.0f), OpacityCodeChunk));
+
+	int32 OutputCodeChunk = Compiler->StrataUnlitBSDF(
+		CompileWithDefaultFloat3(Compiler, Color, 0.0f, 0.0f, 0.0f),
+		TransmittanceCodeChunk);
+
+	return OutputCodeChunk;
+}
+
+void UMaterialExpressionStrataUI::GetCaption(TArray<FString>& OutCaptions) const
+{
+	OutCaptions.Add(TEXT("Strata UI"));
+}
+
+uint32 UMaterialExpressionStrataUI::GetOutputType(int32 OutputIndex)
+{
+	return MCT_Strata;
+}
+
+uint32 UMaterialExpressionStrataUI::GetInputType(int32 InputIndex)
+{
+	switch (InputIndex)
+	{
+	case 0:
+		return MCT_Float3;
+		break;
+	case 1:
+		return MCT_Float;
+		break;
+	}
+
+	check(false);
+	return MCT_Float1;
+}
+
+bool UMaterialExpressionStrataUI::IsResultStrataMaterial(int32 OutputIndex)
+{
+	return true;
+}
+
+void UMaterialExpressionStrataUI::GatherStrataMaterialInfo(FStrataMaterialInfo& StrataMaterialInfo, int32 OutputIndex)
+{
+	StrataMaterialInfo.AddShadingModel(SSM_UI);
+}
+
+FStrataOperator* UMaterialExpressionStrataUI::StrataGenerateMaterialTopologyTree(class FMaterialCompiler* Compiler, class UMaterialExpression* Parent, int32 OutputIndex)
+{
+	FStrataOperator& StrataOperator = Compiler->StrataCompilationRegisterOperator(STRATA_OPERATOR_BSDF, Compiler->StrataTreeStackGetPathUniqueId(), Parent, Compiler->StrataTreeStackGetParentPathUniqueId());
+	StrataOperator.BSDFType = STRATA_BSDF_TYPE_UNLIT;
+	StrataOperator.ThicknessIndex = Compiler->StrataThicknessStackGetThicknessIndex();
+	return &StrataOperator;
+}
+#endif // WITH_EDITOR
+
+
+
 UMaterialExpressionStrataConvertToDecal::UMaterialExpressionStrataConvertToDecal(const FObjectInitializer& ObjectInitializer)
 	: Super(ObjectInitializer)
 {
diff --git a/Engine/Source/Runtime/Engine/Private/Materials/MaterialShared.cpp b/Engine/Source/Runtime/Engine/Private/Materials/MaterialShared.cpp
--- a/Engine/Source/Runtime/Engine/Private/Materials/MaterialShared.cpp
+++ b/Engine/Source/Runtime/Engine/Private/Materials/MaterialShared.cpp
@@ -1885,8 +1885,9 @@ bool FMaterialResource::HasMaterialPropertyConnected(EMaterialProperty In) const
 		Material->MaterialDomain == MD_PostProcess || 
 		Material->MaterialDomain == MD_LightFunction ||
 		Material->MaterialDomain == MD_DeferredDecal || 
-		Material->MaterialDomain == MD_Surface || 
-		Material->MaterialDomain == MD_Volume;
+		Material->MaterialDomain == MD_Surface ||
+		Material->MaterialDomain == MD_Volume ||
+		Material->MaterialDomain == MD_UI;
 
 	if (Strata::IsStrataEnabled() && bIsStrataSupportedDomain)
 	{
