USkinnedAsset: allow a potentially running AsyncTask to finish before deleting it

diff --git a/Engine/Source/Runtime/Engine/Private/SkinnedAsset.cpp b/Engine/Source/Runtime/Engine/Private/SkinnedAsset.cpp
--- a/Engine/Source/Runtime/Engine/Private/SkinnedAsset.cpp
+++ b/Engine/Source/Runtime/Engine/Private/SkinnedAsset.cpp
@@ -35,7 +35,16 @@ USkinnedAsset::USkinnedAsset(const FObjectInitializer& ObjectInitializer)
 	: Super(ObjectInitializer)
 {}
 
-USkinnedAsset::~USkinnedAsset() {}
+USkinnedAsset::~USkinnedAsset() 
+{
+#if WITH_EDITOR
+	if (AsyncTask)
+	{
+		// Allow AsyncTask to finish if it hasn't yet, otherwise CheckIdle() could fail on deletion
+		AsyncTask->EnsureCompletion();
+	}
+#endif
+}
 
 bool USkinnedAsset::IsValidMaterialIndex(int32 Index) const
 {
