Added logic to respected r.ForceLod

diff --git a/Engine/Source/Runtime/Engine/Private/Components/LODSyncComponent.cpp b/Engine/Source/Runtime/Engine/Private/Components/LODSyncComponent.cpp
--- a/Engine/Source/Runtime/Engine/Private/Components/LODSyncComponent.cpp
+++ b/Engine/Source/Runtime/Engine/Private/Components/LODSyncComponent.cpp
@@ -195,7 +195,17 @@ void ULODSyncComponent::UpdateLOD()
 		// it seems we have a situation where components becomes nullptr between registration but has the array entry
 		// so we ensure this is set before we set the data. 
 		bool bHaveValidSetting = false;
-		if (ForcedLOD >= 0 && ForcedLOD < CurrentNumLODs)
+		
+		//Make sure we respect r.ForceLod before starting to apply lodsync logic
+		static IConsoleVariable* ForceLODCvar = IConsoleManager::Get().FindConsoleVariable(TEXT("r.ForceLOD"));
+		if (ForceLODCvar && ForceLODCvar->GetInt() >= 0)
+		{
+			HighestPriLOD = ForceLODCvar->GetInt();
+			bHaveValidSetting = true;
+			UE_LOG(LogLODSync, Verbose, TEXT("LOD Sync : Using Cvar r.ForceLod Value to set lodsync [%d]"), HighestPriLOD);
+		}
+
+		else if (ForcedLOD >= 0 && ForcedLOD < CurrentNumLODs)
 		{
 			HighestPriLOD = ForcedLOD;
 			bHaveValidSetting = true;
