Add logging to see what component and actor have been broken in between deferred kinematic update

diff --git a/Engine/Source/Runtime/Engine/Private/PhysicsEngine/Experimental/PhysScene_Chaos.cpp b/Engine/Source/Runtime/Engine/Private/PhysicsEngine/Experimental/PhysScene_Chaos.cpp
--- a/Engine/Source/Runtime/Engine/Private/PhysicsEngine/Experimental/PhysScene_Chaos.cpp
+++ b/Engine/Source/Runtime/Engine/Private/PhysicsEngine/Experimental/PhysScene_Chaos.cpp
@@ -48,6 +48,9 @@ FAutoConsoleVariableRef CVar_EnableKinematicDeferralStartPhysicsCondition(TEXT("
 bool GKinematicDeferralCheckValidBodies = true;
 FAutoConsoleVariableRef CVar_KinematicDeferralCheckValidBodies(TEXT("p.KinematicDeferralCheckValidBodies"), GKinematicDeferralCheckValidBodies, TEXT("If true, don't attempt to update deferred kinematic skeletal mesh bodies which are pending delete."));
 
+bool GKinematicDeferralLogInvalidBodies = true;
+FAutoConsoleVariableRef CVar_KinematicDeferralLogInvalidBodies(TEXT("p.KinematicDeferralLogInvalidBodies"), GKinematicDeferralLogInvalidBodies, TEXT("If true and p.KinematicDeferralCheckValidBodies is true, log when an invalid body is found on kinematic update."));
+
 DECLARE_CYCLE_STAT(TEXT("Update Kinematics On Deferred SkelMeshes"), STAT_UpdateKinematicsOnDeferredSkelMeshesChaos, STATGROUP_Physics);
 
 struct FPendingAsyncPhysicsCommand
@@ -1472,12 +1475,30 @@ void FPhysScene_Chaos::UpdateKinematicsOnDeferredSkelMeshes()
 					FBodyInstance* BodyInst = SkelComp->Bodies[i];
 					if (GKinematicDeferralCheckValidBodies && (BodyInst == nullptr || !BodyInst->IsValidBodyInstance()))
 					{
+						if (GKinematicDeferralLogInvalidBodies)
+						{
+							UE_LOG(LogChaos, Warning, TEXT("\n"
+								"Invalid FBodyInstance in FPhysScene_Chaos::UpdateKinematicsOnDeferredSkelMesh - Gathering Proxies\n"
+								"SkeletalMesh: %s\n"
+								"OwningActor: %s"),
+								*SkelComp->GetName(),
+								*(SkelComp->GetOwner() ? SkelComp->GetOwner()->GetName() : FString("None")));
+						}
 						continue;
 					}
 
 					FPhysicsActorHandle& ActorHandle = BodyInst->ActorHandle;
 					if (GKinematicDeferralCheckValidBodies && (ActorHandle == nullptr || ActorHandle->GetMarkedDeleted()))
 					{
+						if (GKinematicDeferralLogInvalidBodies)
+						{
+							UE_LOG(LogChaos, Warning, TEXT("\n"
+								"Invalid FPhysicsActorHandle in FPhysScene_Chaos::UpdateKinematicsOnDeferredSkelMesh - Gathering Proxies\n"
+								"SkeletalMesh: %s\n"
+								"OwningActor: %s"),
+								*SkelComp->GetName(),
+								*(SkelComp->GetOwner() ? SkelComp->GetOwner()->GetName() : FString("None")));
+						}
 						continue;
 					}
 
@@ -1534,12 +1555,30 @@ void FPhysScene_Chaos::UpdateKinematicsOnDeferredSkelMeshes()
 					FBodyInstance* BodyInst = SkelComp->Bodies[i];
 					if (GKinematicDeferralCheckValidBodies && (BodyInst == nullptr || !BodyInst->IsValidBodyInstance()))
 					{
+						if (GKinematicDeferralLogInvalidBodies)
+						{
+							UE_LOG(LogChaos, Warning, TEXT("\n"
+								"Invalid FBodyInstance in FPhysScene_Chaos::UpdateKinematicsOnDeferredSkelMesh - Add to Teleport Pool\n"
+								"SkeletalMesh: %s\n"
+								"OwningActor: %s"),
+								*SkelComp->GetName(),
+								*(SkelComp->GetOwner() ? SkelComp->GetOwner()->GetName() : FString("None")));
+						}
 						continue;
 					}
 
 					FPhysicsActorHandle& ActorHandle = BodyInst->ActorHandle;
 					if (GKinematicDeferralCheckValidBodies && (ActorHandle == nullptr || ActorHandle->GetMarkedDeleted()))
 					{
+						if (GKinematicDeferralLogInvalidBodies)
+						{
+							UE_LOG(LogChaos, Warning, TEXT("\n"
+								"Invalid FPhysicsActorHandle in FPhysScene_Chaos::UpdateKinematicsOnDeferredSkelMesh - Add to Teleport Pool\n"
+								"SkeletalMesh: %s\n"
+								"OwningActor: %s"),
+								*SkelComp->GetName(),
+								*(SkelComp->GetOwner() ? SkelComp->GetOwner()->GetName() : FString("None")));
+						}
 						continue;
 					}
 
