Adding channel description to package map guid warning

diff --git a/Engine/Source/Runtime/Engine/Private/DataChannel.cpp b/Engine/Source/Runtime/Engine/Private/DataChannel.cpp
--- a/Engine/Source/Runtime/Engine/Private/DataChannel.cpp
+++ b/Engine/Source/Runtime/Engine/Private/DataChannel.cpp
@@ -3347,7 +3347,7 @@ int64 UActorChannel::ReplicateActor()
 	UPackageMapClient* PackageMapClient = static_cast<UPackageMapClient*>(Connection->PackageMap);
 	if (PackageMapClient->GetMustBeMappedGuidsInLastBunch().Num() != 0)
 	{
-		UE_LOG(LogNet, Warning, TEXT("ReplicateActor: PackageMap->GetMustBeMappedGuidsInLastBunch().Num() != 0: %i"), PackageMapClient->GetMustBeMappedGuidsInLastBunch().Num());
+		UE_LOG(LogNet, Warning, TEXT("ReplicateActor: PackageMap->GetMustBeMappedGuidsInLastBunch().Num() != 0: %i: Channel: %s"), PackageMapClient->GetMustBeMappedGuidsInLastBunch().Num(), *Describe());
 	}
 
 	bool bWroteSomethingImportant = bIsNewlyReplicationUnpaused || bIsNewlyReplicationPaused;
