Adding a const version of GetRenderScene.

diff --git a/Engine/Source/Runtime/Engine/Public/SceneInterface.h b/Engine/Source/Runtime/Engine/Public/SceneInterface.h
--- a/Engine/Source/Runtime/Engine/Public/SceneInterface.h
+++ b/Engine/Source/Runtime/Engine/Public/SceneInterface.h
@@ -35,6 +35,7 @@ class FRenderResource;
 class FRenderTarget;
 class FRHICommandListImmediate;
 class FRHIUniformBuffer;
+class FScene;
 class FSceneRenderer;
 class FSceneViewStateInterface;
 class FSkyAtmosphereRenderSceneInfo;
@@ -521,10 +522,8 @@ public:
 	 * Return the scene to be used for rendering. Note that this can return null if rendering has
 	 * been disabled!
 	 */
-	virtual class FScene* GetRenderScene()
-	{
-		return nullptr;
-	}
+	virtual FScene* GetRenderScene() = 0;
+	virtual const FScene* GetRenderScene() const = 0;
 
 	virtual void OnWorldCleanup()
 	{
diff --git a/Engine/Source/Runtime/Renderer/Private/RendererScene.cpp b/Engine/Source/Runtime/Renderer/Private/RendererScene.cpp
--- a/Engine/Source/Runtime/Renderer/Private/RendererScene.cpp
+++ b/Engine/Source/Runtime/Renderer/Private/RendererScene.cpp
@@ -6553,12 +6553,14 @@ public:
 		return World;
 	}
 
-	/**
-	* Return the scene to be used for rendering
-	*/
-	virtual class FScene* GetRenderScene() override
+	virtual FScene* GetRenderScene() override
+	{
+		return nullptr;
+	}
+
+	virtual const FScene* GetRenderScene() const override
 	{
-		return NULL;
+		return nullptr;
 	}
 
 	/**
diff --git a/Engine/Source/Runtime/Renderer/Private/ScenePrivate.h b/Engine/Source/Runtime/Renderer/Private/ScenePrivate.h
--- a/Engine/Source/Runtime/Renderer/Private/ScenePrivate.h
+++ b/Engine/Source/Runtime/Renderer/Private/ScenePrivate.h
@@ -3274,10 +3274,9 @@ public:
 	/**
 	* Return the scene to be used for rendering
 	*/
-	virtual class FScene* GetRenderScene() override
-	{
-		return this;
-	}
+	virtual FScene* GetRenderScene() final { return this; }
+	virtual const FScene* GetRenderScene() const final { return this; }
+
 	virtual void OnWorldCleanup() override;
 
 
