Persistent error state on Soundwave

diff --git a/Engine/Source/Runtime/AudioMixer/Private/AudioMixerBuffer.cpp b/Engine/Source/Runtime/AudioMixer/Private/AudioMixerBuffer.cpp
--- a/Engine/Source/Runtime/AudioMixer/Private/AudioMixerBuffer.cpp
+++ b/Engine/Source/Runtime/AudioMixer/Private/AudioMixerBuffer.cpp
@@ -330,6 +330,13 @@ namespace Audio
 		{
 			return nullptr;
 		}
+		
+		// Ignore attempts to create if this wave has been flagged as containing errors
+		if (InWave->HasError())
+		{
+			UE_LOG(LogAudioMixer, VeryVerbose, TEXT("FMixerBuffer::CreateStreamingBuffer, ignoring '%s' as it contains previously seen errors"), *InWave->GetName() );
+			return nullptr;
+		}
 
 		AUDIO_MIXER_TRACE_CPUPROFILER_EVENT_SCOPE(FMixerBuffer::CreateStreamingBuffer);
 
@@ -356,6 +363,12 @@ namespace Audio
 		}
 		else
 		{
+			// Failed to stream in compressed info, so mark the wave as having an error.
+			if (Buffer->DecompressionState)
+			{
+				InWave->SetError(TEXT("ICompressedAudioInfo::StreamCompressedInfo failed"));
+			}
+
 			// When set to seekable streaming, missing the first chunk is possible and
 			// does not signify any issue with the asset itself, so don't mark it as invalid.
 			if (InWave && !InWave->IsSeekable())
diff --git a/Engine/Source/Runtime/AudioMixer/Private/AudioMixerSourceBuffer.cpp b/Engine/Source/Runtime/AudioMixer/Private/AudioMixerSourceBuffer.cpp
--- a/Engine/Source/Runtime/AudioMixer/Private/AudioMixerSourceBuffer.cpp
+++ b/Engine/Source/Runtime/AudioMixer/Private/AudioMixerSourceBuffer.cpp
@@ -60,6 +60,13 @@ namespace Audio
 	{
 		LLM_SCOPE(ELLMTag::AudioMixer);
 
+		// Fail if the Wave has been flagged to contain an error
+		if (InArgs.SoundWave && InArgs.SoundWave->HasError())
+		{
+			UE_LOG(LogAudioMixer, VeryVerbose, TEXT("FMixerSourceBuffer::Create failed as '%s' is flagged as containing errors"), *InArgs.SoundWave->GetName());
+			return {};
+		}
+
 		TSharedPtr<FMixerSourceBuffer, ESPMode::ThreadSafe> NewSourceBuffer = MakeShareable(new FMixerSourceBuffer(InArgs, MoveTemp(InDefaultParams)));
 
 		return NewSourceBuffer;
@@ -370,9 +377,13 @@ namespace Audio
 
 		// Handle the case that the decoder has an error and can't continue.
 		if (InDecoder && InDecoder->HasError())
-		{			
-			UE_LOG(LogAudioMixer, Warning, TEXT("Decoder Error, stopping source [%s]"), 
-				*InDecoder->GetStreamingSoundWave()->GetFName().ToString());
+		{
+			FScopeTryLock Lock(&SoundWaveCritSec);
+			if (Lock.IsLocked() && SoundWave)
+			{
+				SoundWave->SetError(TEXT("ICompressedAudioInfo::HasError() flagged on the Decoder"));
+			}
+
 			bHasError = true;
 			bBufferFinished = true;
 			return false;	
diff --git a/Engine/Source/Runtime/Engine/Classes/Sound/SoundWave.h b/Engine/Source/Runtime/Engine/Classes/Sound/SoundWave.h
--- a/Engine/Source/Runtime/Engine/Classes/Sound/SoundWave.h
+++ b/Engine/Source/Runtime/Engine/Classes/Sound/SoundWave.h
@@ -1153,6 +1153,10 @@ public:
 	*/
 	float GetSampleRateForCompressionOverrides(const FPlatformAudioCookOverrides* CompressionOverrides);
 
+	void SetError(const TCHAR* InErrorMsg=nullptr);
+	void ResetError();
+	bool HasError() const;
+
 #if WITH_EDITORONLY_DATA
 
 #if WITH_EDITOR
@@ -1247,6 +1251,7 @@ public:
 		, bIsSeekable(0)
 		, bShouldUseStreamCaching(0)
 		, bLoadingBehaviorOverridden(0)
+		, bHasError(0)
 #if WITH_EDITOR
 		, bLoadedFromCookedData(0)
 #endif //WITH_EDITOR
@@ -1277,6 +1282,10 @@ public:
 
 	void ReleaseCompressedAudio();
 
+	void SetError(const TCHAR* InErrorMsg=nullptr);  
+	bool HasError() const;
+	bool ResetError();
+
 	bool IsStreaming() const { return bIsStreaming; }
 	ESoundAssetCompressionType GetSoundCompressionType() const { return SoundAssetCompressionType; }
 	bool IsLooping() const { return bIsLooping; }
@@ -1352,7 +1361,7 @@ private:
 	FObjectKey SoundWaveKeyCached;
 	TArray<FSoundWaveCuePoint> CuePoints;
 	ESoundAssetCompressionType SoundAssetCompressionType;
-
+	
 	float SampleRate = 0;
 	float Duration = 0;
 
@@ -1366,6 +1375,7 @@ private:
 	uint8 bIsSeekable : 1;
 	uint8 bShouldUseStreamCaching : 1;
 	uint8 bLoadingBehaviorOverridden : 1;
+	uint8 bHasError : 1;
 
 #if WITH_EDITOR
 	uint8 bLoadedFromCookedData : 1;
diff --git a/Engine/Source/Runtime/Engine/Private/SoundWave.cpp b/Engine/Source/Runtime/Engine/Private/SoundWave.cpp
--- a/Engine/Source/Runtime/Engine/Private/SoundWave.cpp
+++ b/Engine/Source/Runtime/Engine/Private/SoundWave.cpp
@@ -327,6 +327,27 @@ bool FSoundWaveData::HasCompressedData(FName Format, ITargetPlatform* TargetPlat
 	}
 }
 
+void FSoundWaveData::SetError(const TCHAR* InErrorMsg)
+{
+	UE_CLOG(!bHasError, LogAudio, Warning, TEXT("Flagging SoundWave '%s' to have errors, preventing playback. ErrorMsg='%s'"), 
+		*GetPackageName().ToString(),
+		InErrorMsg ? InErrorMsg : TEXT("")
+	);
+	bHasError = true;
+}
+
+bool FSoundWaveData::HasError() const
+{
+	return bHasError;
+}
+
+bool FSoundWaveData::ResetError()
+{
+	bool bHadError = bHasError;
+	bHasError = false;
+	return bHadError;
+}
+
 FByteBulkData* FSoundWaveData::GetCompressedData(FName Format, const FPlatformAudioCookOverrides* CompressionOverrides)
 {
 	if (IsTemplate() || IsRunningDedicatedServer())
@@ -3847,6 +3868,10 @@ void USoundWave::UpdateAsset(bool bMarkDirty)
 	InvalidateCompressedData();
 	FreeResources();
 	UpdatePlatformData();
+	
+	// Reset any error state we've encountered.
+	ResetError();
+
 	if(bMarkDirty)
 	{
 		MarkPackageDirty();
@@ -4267,3 +4292,21 @@ TArrayView<const uint8> FSoundWaveProxy::GetZerothChunk(const FSoundWaveProxyPtr
 
 	return {};
 }
+
+void USoundWave::SetError(const TCHAR* InErrorMsg)
+{
+	check(SoundWaveDataPtr);
+	SoundWaveDataPtr->SetError(InErrorMsg);
+}
+
+void USoundWave::ResetError()
+{
+	check(SoundWaveDataPtr);
+	SoundWaveDataPtr->ResetError();
+}
+
+bool USoundWave::HasError() const
+{
+	check(SoundWaveDataPtr);
+	return SoundWaveDataPtr->HasError();
+}
\ No newline at end of file
diff --git a/Engine/Source/Runtime/Engine/Private/VorbisAudioInfo.cpp b/Engine/Source/Runtime/Engine/Private/VorbisAudioInfo.cpp
--- a/Engine/Source/Runtime/Engine/Private/VorbisAudioInfo.cpp
+++ b/Engine/Source/Runtime/Engine/Private/VorbisAudioInfo.cpp
@@ -313,6 +313,7 @@ bool FVorbisAudioInfo::GetCompressedInfoCommon(void* Callbacks, FSoundQualityInf
 	if (!bDllLoaded)
 	{
 		UE_LOG(LogAudio, Error, TEXT("FVorbisAudioInfo::GetCompressedInfoCommon failed due to vorbis DLL not being loaded."));
+		bHasError = true;
 		return false;
 	}
 
@@ -320,7 +321,9 @@ bool FVorbisAudioInfo::GetCompressedInfoCommon(void* Callbacks, FSoundQualityInf
 	int Result = ov_open_callbacks(this, &VFWrapper->vf, NULL, 0, (*(ov_callbacks*)Callbacks));
 	if (Result < 0)
 	{
-		UE_LOG(LogAudio, Error, TEXT("FVorbisAudioInfo::ReadCompressedInfo, ov_open_callbacks error code: %d : %s"), Result,GetVorbisError(Result));
+		UE_LOG(LogAudio, Warning, TEXT("FVorbisAudioInfo::ReadCompressedInfo, ov_open_callbacks error code: %d : %s"), Result,GetVorbisError(Result));
+		bHasError = true;
+		LastErrorCode = Result;
 		return false;
 	}
 
@@ -603,7 +606,8 @@ bool FVorbisAudioInfo::StreamCompressedInfoInternal(const FSoundWaveProxyPtr& In
 		bHeaderParsed = GetCompressedInfoCommon(&Callbacks, QualityInfo);
 		if (!bHeaderParsed)
 		{
-			UE_LOG(LogAudio, Error, TEXT("FVorbisAudioInfo::StreamCompressedInfoInternal failed to parse header for '%s'."), *InWaveProxy->GetFName().ToString());
+			UE_LOG(LogAudio, Error, TEXT("FVorbisAudioInfo::StreamCompressedInfoInternal failed to parse header for '%s'. LastError=%s"), 
+				*InWaveProxy->GetFName().ToString(), GetVorbisError(LastErrorCode));
 		}
 	
 
@@ -734,7 +738,9 @@ bool FVorbisAudioInfo::StreamCompressedData(uint8* InDestination, bool bLooping,
 			int Result = ov_open_callbacks(this, &VFWrapper->vf, NULL, 0, Callbacks);
 			if (Result < 0)
 			{
-				UE_LOG(LogAudio, Error, TEXT("FVorbisAudioInfo::StreamCompressedData, ov_open_callbacks error code: %d"), Result);
+				LastErrorCode = Result;
+				bHasError = true;
+				UE_LOG(LogAudio, Error, TEXT("FVorbisAudioInfo::StreamCompressedData, ov_open_callbacks error code: %d : %s"), Result, GetVorbisError(Result));
 				break;
 			}
 
diff --git a/Engine/Source/Runtime/Engine/Public/VorbisAudioInfo.h b/Engine/Source/Runtime/Engine/Public/VorbisAudioInfo.h
--- a/Engine/Source/Runtime/Engine/Public/VorbisAudioInfo.h
+++ b/Engine/Source/Runtime/Engine/Public/VorbisAudioInfo.h
@@ -102,6 +102,7 @@ public:
 	virtual bool StreamCompressedData(uint8* InDestination, bool bLooping, uint32 BufferSize, int32& OutNumBytesStreamed) override;
 	virtual int32 GetCurrentChunkIndex() const override {return CurrentStreamingChunkIndex;}
 	virtual int32 GetCurrentChunkOffset() const override {return BufferOffset % CurrentStreamingChunksSize;}
+	virtual bool HasError() const override { return bHasError; }
 	// End of ICompressedAudioInfo Interface
 
 protected:
@@ -131,6 +132,8 @@ private:
 	int32 NextStreamingChunkIndex;
 	uint32 CurrentStreamingChunksSize;
 	bool bHeaderParsed;
+	bool bHasError = false;
+	int32 LastErrorCode = 0;
 
 	// This handle is used to ensure that a chunk currently being decoded isn't evicted until we are done with it.
 	FAudioChunkHandle CurCompressedChunkHandle;
