Adding the -DisableOptionalMips cooker option to disable optional mips when cooking, overriding any platform & texture settings.

 - useful when cooking some content plugins.

#rnx

diff --git a/Engine/Source/Runtime/Engine/Private/TextureDerivedData.cpp b/Engine/Source/Runtime/Engine/Private/TextureDerivedData.cpp
--- a/Engine/Source/Runtime/Engine/Private/TextureDerivedData.cpp
+++ b/Engine/Source/Runtime/Engine/Private/TextureDerivedData.cpp
@@ -2532,15 +2532,20 @@ static void SerializePlatformData(
 			FirstInlineMip = NumMips - NumNonStreamingMips;
 
 		#if WITH_EDITORONLY_DATA
-			const int32 LODGroup = Texture->LODGroup;
+			static bool bDisableOptionalMips = FParse::Param(FCommandLine::Get(), TEXT("DisableOptionalMips"));
+			if (!bDisableOptionalMips)
+			{
+				const int32 LODGroup = Texture->LODGroup;
+				const int32 FirstMipWidth  = PlatformData->Mips[FirstMipToSerialize].SizeX;
+				const int32 FirstMipHeight = PlatformData->Mips[FirstMipToSerialize].SizeY;
+
+				OptionalMips = Ar.CookingTarget()->GetTextureLODSettings().CalculateNumOptionalMips(LODGroup, FirstMipWidth, FirstMipHeight, NumMips, FirstInlineMip, Texture->MipGenSettings);
+				bDuplicateNonOptionalMips = Ar.CookingTarget()->GetTextureLODSettings().TextureLODGroups[LODGroup].DuplicateNonOptionalMips;
 
-			const int32 FirstMipWidth  = PlatformData->Mips[FirstMipToSerialize].SizeX;
-			const int32 FirstMipHeight = PlatformData->Mips[FirstMipToSerialize].SizeY;
-			OptionalMips = Ar.CookingTarget()->GetTextureLODSettings().CalculateNumOptionalMips(LODGroup, FirstMipWidth, FirstMipHeight, NumMips, FirstInlineMip, Texture->MipGenSettings);
-			bDuplicateNonOptionalMips = Ar.CookingTarget()->GetTextureLODSettings().TextureLODGroups[LODGroup].DuplicateNonOptionalMips;
+				// OptionalMips must be streaming mips.
+				check(OptionalMips <= FirstInlineMip);
+			}
 
-			// OptionalMips must be streaming mips.
-			check(OptionalMips <= FirstInlineMip);
 		#endif
 
 		#if WITH_EDITOR
