Added new functions to UAnimatedSparseVolumeTextureController to make controlling SVT playback from blueprint easier.

#rnx

diff --git a/Engine/Source/Runtime/Engine/Classes/SparseVolumeTexture/SparseVolumeTexture.h b/Engine/Source/Runtime/Engine/Classes/SparseVolumeTexture/SparseVolumeTexture.h
--- a/Engine/Source/Runtime/Engine/Classes/SparseVolumeTexture/SparseVolumeTexture.h
+++ b/Engine/Source/Runtime/Engine/Classes/SparseVolumeTexture/SparseVolumeTexture.h
@@ -184,6 +184,7 @@ public:
 	USparseVolumeTexture();
 	virtual ~USparseVolumeTexture() = default;
 
+	UFUNCTION(BlueprintCallable, Category = "Animation")
 	virtual int32 GetFrameCount() const { return 0; }
 	virtual FBox GetVolumeBounds() const { return FBox(); }
 	virtual const FSparseVolumeTextureSceneProxy* GetSparseVolumeTextureSceneProxy() const { return nullptr; }
@@ -329,9 +330,6 @@ public:
 	UAnimatedSparseVolumeTextureController();
 	virtual ~UAnimatedSparseVolumeTextureController() = default;
 
-	UFUNCTION(BlueprintCallable, Category = "Animation")
-	void SetSparseVolumeTexture(USparseVolumeTexture* Texture);
-
 	UFUNCTION(BlueprintCallable, Category = "Animation")
 	void Play();
 
@@ -345,7 +343,34 @@ public:
 	bool IsPlaying();
 
 	UFUNCTION(BlueprintCallable, Category = "Animation")
-	USparseVolumeTextureFrame* Update(float DeltaTime);
+	void Update(float DeltaTime);
+
+	UFUNCTION(BlueprintCallable, Category = "Animation")
+	void SetSparseVolumeTexture(USparseVolumeTexture* Texture);
+
+	UFUNCTION(BlueprintCallable, Category = "Animation")
+	void SetTime(float Time);
+
+	UFUNCTION(BlueprintCallable, Category = "Animation")
+	void SetFractionalFrameIndex(float Frame);
+
+	UFUNCTION(BlueprintCallable, Category = "Animation")
+	USparseVolumeTexture* GetSparseVolumeTexture();
+
+	UFUNCTION(BlueprintCallable, Category = "Animation")
+	float GetTime();
+
+	UFUNCTION(BlueprintCallable, Category = "Animation")
+	float GetFractionalFrameIndex();
+
+	UFUNCTION(BlueprintCallable, Category = "Animation")
+	USparseVolumeTextureFrame* GetCurrentFrame();
+
+	UFUNCTION(BlueprintCallable, Category = "Animation")
+	void GetLerpFrames(USparseVolumeTextureFrame*& Frame0, USparseVolumeTextureFrame*& Frame1, float& LerpAlpha);
+
+	UFUNCTION(BlueprintCallable, Category = "Animation")
+	float GetDuration();
 
 private:
 	TObjectPtr<USparseVolumeTexture> SparseVolumeTexture;
diff --git a/Engine/Source/Runtime/Engine/Private/SparseVolumeTexture/SparseVolumeTexture.cpp b/Engine/Source/Runtime/Engine/Private/SparseVolumeTexture/SparseVolumeTexture.cpp
--- a/Engine/Source/Runtime/Engine/Private/SparseVolumeTexture/SparseVolumeTexture.cpp
+++ b/Engine/Source/Runtime/Engine/Private/SparseVolumeTexture/SparseVolumeTexture.cpp
@@ -667,7 +667,10 @@ USparseVolumeTextureFrame::USparseVolumeTextureFrame(const FObjectInitializer& O
 
 USparseVolumeTextureFrame* USparseVolumeTextureFrame::CreateFrame(USparseVolumeTexture* Texture, int32 FrameIndex)
 {
-	check(Texture);
+	if (!Texture)
+	{
+		return nullptr;
+	}
 	
 	const FSparseVolumeTextureSceneProxy* Proxy = nullptr;
 	if (Texture->IsA<UStreamableSparseVolumeTexture>())
@@ -704,18 +707,6 @@ UAnimatedSparseVolumeTextureController::UAnimatedSparseVolumeTextureController(c
 {
 }
 
-void UAnimatedSparseVolumeTextureController::SetSparseVolumeTexture(USparseVolumeTexture* Texture)
-{
-	if (Texture == SparseVolumeTexture)
-	{
-		return;
-	}
-
-	SparseVolumeTexture = Texture;
-	bIsPlaying = bIsPlaying && (SparseVolumeTexture != nullptr);
-	Time = 0.0f;
-}
-
 void UAnimatedSparseVolumeTextureController::Play()
 {
 	bIsPlaying = true;
@@ -740,24 +731,79 @@ bool UAnimatedSparseVolumeTextureController::IsPlaying()
 	return bIsPlaying;
 }
 
-USparseVolumeTextureFrame* UAnimatedSparseVolumeTextureController::Update(float DeltaTime)
+void UAnimatedSparseVolumeTextureController::Update(float DeltaTime)
+{
+	if (!SparseVolumeTexture || !bIsPlaying)
+	{
+		return;
+	}
+
+	// Update animation time
+	const float AnimationDuration = GetDuration();
+	Time = FMath::Fmod(Time + DeltaTime, AnimationDuration + UE_SMALL_NUMBER);
+}
+
+void UAnimatedSparseVolumeTextureController::SetSparseVolumeTexture(USparseVolumeTexture* Texture)
+{
+	if (Texture == SparseVolumeTexture)
+	{
+		return;
+	}
+
+	SparseVolumeTexture = Texture;
+	bIsPlaying = bIsPlaying && (SparseVolumeTexture != nullptr);
+	Time = 0.0f;
+}
+
+void UAnimatedSparseVolumeTextureController::SetTime(float InTime)
+{
+	const float AnimationDuration = GetDuration();
+	Time = FMath::Fmod(InTime, AnimationDuration + UE_SMALL_NUMBER);
+}
+
+void UAnimatedSparseVolumeTextureController::SetFractionalFrameIndex(float Frame)
 {
 	if (!SparseVolumeTexture)
 	{
-		return nullptr;
+		return;
 	}
 
 	const int32 FrameCount = SparseVolumeTexture->GetFrameCount();
-	const float AnimationDuration = FrameCount / (FrameRate + UE_SMALL_NUMBER);
+	Frame = FMath::Fmod(Frame, (float)FrameCount);
+	Time = Frame / (FrameRate + UE_SMALL_NUMBER);
+}
 
-	// Update animation time
-	if (bIsPlaying)
+USparseVolumeTexture* UAnimatedSparseVolumeTextureController::GetSparseVolumeTexture()
+{
+	return SparseVolumeTexture;
+}
+
+float UAnimatedSparseVolumeTextureController::GetTime()
+{
+	return Time;
+}
+
+float UAnimatedSparseVolumeTextureController::GetFractionalFrameIndex()
+{
+	if (!SparseVolumeTexture)
 	{
-		Time = FMath::Fmod(Time + DeltaTime, AnimationDuration + UE_SMALL_NUMBER);
+		return 0.0f;
 	}
 
-	// Compute (fractional) index of frame to sample
+	const int32 FrameCount = SparseVolumeTexture->GetFrameCount();
 	const float FrameIndexF = FMath::Fmod(Time * FrameRate, (float)FrameCount);
+	return FrameIndexF;
+}
+
+USparseVolumeTextureFrame* UAnimatedSparseVolumeTextureController::GetCurrentFrame()
+{
+	if (!SparseVolumeTexture)
+	{
+		return nullptr;
+	}
+
+	// Compute (fractional) index of frame to sample
+	const float FrameIndexF = GetFractionalFrameIndex();
 	const int32 FrameIndex = (int32)FrameIndexF;
 
 	// Create and initialize a USparseVolumeTextureFrame which holds the frame to sample and can be bound to shaders
@@ -766,6 +812,35 @@ USparseVolumeTextureFrame* UAnimatedSparseVolumeTextureController::Update(float
 	return Frame;
 }
 
+void UAnimatedSparseVolumeTextureController::GetLerpFrames(USparseVolumeTextureFrame*& Frame0, USparseVolumeTextureFrame*& Frame1, float& LerpAlpha)
+{
+	if (!SparseVolumeTexture)
+	{
+		return;
+	}
+
+	// Compute (fractional) index of frame to sample
+	const float FrameIndexF = GetFractionalFrameIndex();
+	const int32 FrameIndex = (int32)FrameIndexF;
+	LerpAlpha = FMath::Frac(FrameIndexF);
+
+	// Create and initialize a USparseVolumeTextureFrame which holds the frame to sample and can be bound to shaders
+	Frame0 = USparseVolumeTextureFrame::CreateFrame(SparseVolumeTexture, FrameIndex);
+	Frame1 = USparseVolumeTextureFrame::CreateFrame(SparseVolumeTexture, FrameIndex + 1);
+}
+
+float UAnimatedSparseVolumeTextureController::GetDuration()
+{
+	if (!SparseVolumeTexture)
+	{
+		return 0.0f;
+	}
+
+	const int32 FrameCount = SparseVolumeTexture->GetFrameCount();
+	const float AnimationDuration = FrameCount / (FrameRate + UE_SMALL_NUMBER);
+	return AnimationDuration;
+}
+
 ////////////////////////////////////////////////////////////////////////////////////////////////
 
 #undef LOCTEXT_NAMESPACE
