More cases covered by the new MaterialCachedData graph analyzer. Added support for Reroute nodes and MaterialFunctions opened in the editor.

diff --git a/Engine/Source/Runtime/Engine/Private/Materials/MaterialCachedData.cpp b/Engine/Source/Runtime/Engine/Private/Materials/MaterialCachedData.cpp
--- a/Engine/Source/Runtime/Engine/Private/Materials/MaterialCachedData.cpp
+++ b/Engine/Source/Runtime/Engine/Private/Materials/MaterialCachedData.cpp
@@ -625,6 +625,13 @@ struct FMaterialConnectedPropertiesAnalyzer
 					bVisitAllInputExpressions = false;
 				}
 			}
+			else if (UMaterialExpressionNamedRerouteUsage* ExpressionRerouteUsage = Cast<UMaterialExpressionNamedRerouteUsage>(Expression))
+			{
+				if (ExpressionRerouteUsage->Declaration)
+				{
+					PushUnexploredExpression(ExpressionRerouteUsage->Declaration->Input.Expression);
+				}
+			}
 
 			// If otherwise specified, explore all input expressions to this node.
 			if (bVisitAllInputExpressions)
@@ -760,6 +767,16 @@ void FMaterialCachedExpressionData::AnalyzeMaterial(UMaterial& Material)
 			}
 		}
 	}
+	
+	// If there are any connected function output expressions, mark them as unexplored expressions.
+	// This occurs when the user opens a MaterialFunction in the editor.
+	for (UMaterialExpression* Expression : Material.GetExpressions())
+	{
+		if (Expression->IsA<UMaterialExpressionFunctionOutput>())
+		{
+			Analyzer.PushUnexploredExpression(Expression);
+		}
+	}
 
 	Analyzer.VisitUnexploredExpressions(GlobalParameter, -1);
 }
