Fix issues with stationary light when r.AllowStaticLighting is disabled.
- Renderer will treat stationary lights as movable in projects with static lighting disabled.
    - Don't render "Preview" indicator in stationary lights shadows
    - Stationary lights are added to TranslucencyLightingVolume

diff --git a/Engine/Source/Runtime/Engine/Private/LightSceneProxy.cpp b/Engine/Source/Runtime/Engine/Private/LightSceneProxy.cpp
--- a/Engine/Source/Runtime/Engine/Private/LightSceneProxy.cpp
+++ b/Engine/Source/Runtime/Engine/Private/LightSceneProxy.cpp
@@ -39,7 +39,7 @@ FLightSceneProxy::FLightSceneProxy(const ULightComponent* InLightComponent)
 	, bCastDynamicShadow(InLightComponent->CastShadows&& InLightComponent->CastDynamicShadows)
 	, bCastStaticShadow(InLightComponent->CastShadows&& InLightComponent->CastStaticShadows)
 	, bCastTranslucentShadows(InLightComponent->CastTranslucentShadows)
-	, bTransmission(InLightComponent->bTransmission&& bCastDynamicShadow && !bStaticShadowing)
+	, bTransmission(InLightComponent->bTransmission && bCastDynamicShadow && !bStaticShadowing)
 	, bCastVolumetricShadow(InLightComponent->bCastVolumetricShadow)
 	, bCastHairStrandsDeepShadow(InLightComponent->bCastDeepShadow)
 	, bCastShadowsFromCinematicObjectsOnly(InLightComponent->bCastShadowsFromCinematicObjectsOnly)
@@ -83,6 +83,14 @@ FLightSceneProxy::FLightSceneProxy(const ULightComponent* InLightComponent)
 		bStaticShadowing = bStaticLighting;
 	}
 
+	static const auto AllowStaticLightingVar = IConsoleManager::Get().FindTConsoleVariableDataInt(TEXT("r.AllowStaticLighting"));
+	const bool bAllowStaticLighting = (!AllowStaticLightingVar || AllowStaticLightingVar->GetValueOnGameThread() != 0);
+
+	if (!bAllowStaticLighting)
+	{
+		bStaticShadowing = bStaticLighting;
+	}
+
 	const FLightComponentMapBuildData* MapBuildData = InLightComponent->GetLightComponentMapBuildData();
 
 	if (MapBuildData && bStaticShadowing && !bStaticLighting)
