Allow object reference collector to visit CachedExpressionData, otherwise some textures might be missed

#rnx

diff --git a/Engine/Source/Runtime/Engine/Private/Materials/MaterialInterface.cpp b/Engine/Source/Runtime/Engine/Private/Materials/MaterialInterface.cpp
--- a/Engine/Source/Runtime/Engine/Private/Materials/MaterialInterface.cpp
+++ b/Engine/Source/Runtime/Engine/Private/Materials/MaterialInterface.cpp
@@ -192,6 +192,13 @@ void UMaterialInterface::Serialize(FArchive& Ar)
 		FObjectCacheEventSink::NotifyReferencedTextureChanged_Concurrent(this);
 #endif
 	}
+    // we don't consider bLoadedCachedExpressionData here because in editor we call UpdateCachedExpressionData which can
+    // create CachedExpressionData if it wasn't serialized
+	else if (Ar.IsObjectReferenceCollector() && CachedExpressionData)
+	{
+ 		UScriptStruct* Struct = FMaterialCachedExpressionData::StaticStruct();
+ 		Struct->SerializeTaggedProperties(Ar, (uint8*)CachedExpressionData.Get(), Struct, nullptr);
+	}
 }
 
 void UMaterialInterface::PostLoad()
