[shaders] non-unity build fix in ShaderCompiler.cpp

diff --git a/Engine/Source/Runtime/Engine/Private/ShaderCompiler/ShaderCompiler.cpp b/Engine/Source/Runtime/Engine/Private/ShaderCompiler/ShaderCompiler.cpp
--- a/Engine/Source/Runtime/Engine/Private/ShaderCompiler/ShaderCompiler.cpp
+++ b/Engine/Source/Runtime/Engine/Private/ShaderCompiler/ShaderCompiler.cpp
@@ -46,6 +46,7 @@
 #include "ShaderPlatformCachedIniValue.h"
 #include "StaticBoundShaderState.h"
 #include "StereoRenderUtils.h"
+#include "Tasks/Task.h"
 #include "UObject/DevObjectVersion.h"
 #include "UObject/UObjectIterator.h"
 #include "Math/UnitConversion.h"
