[NO NON-PII COMMIT MESSAGE]

diff --git a/Engine/Source/Editor/ComponentVisualizers/Private/SplineComponentVisualizer.cpp b/Engine/Source/Editor/ComponentVisualizers/Private/SplineComponentVisualizer.cpp
--- a/Engine/Source/Editor/ComponentVisualizers/Private/SplineComponentVisualizer.cpp
+++ b/Engine/Source/Editor/ComponentVisualizers/Private/SplineComponentVisualizer.cpp
@@ -3070,6 +3070,8 @@ void FSplineComponentVisualizer::OnResetToDefault()
 	check(SelectionState);
 
 	SplineComp->Modify();
+	SplineComp->ResetToDefault();
+
 	if (AActor* Owner = SplineComp->GetOwner())
 	{
 		Owner->Modify();
diff --git a/Engine/Source/Runtime/Engine/Classes/Components/SplineComponent.h b/Engine/Source/Runtime/Engine/Classes/Components/SplineComponent.h
--- a/Engine/Source/Runtime/Engine/Classes/Components/SplineComponent.h
+++ b/Engine/Source/Runtime/Engine/Classes/Components/SplineComponent.h
@@ -356,6 +356,9 @@ public:
 
 	void ApplyComponentInstanceData(struct FSplineInstanceData* ComponentInstanceData, const bool bPostUCS);
 
+	/** Reset the spline to its default shape (a spline of two points) */
+	void ResetToDefault();
+
 	/** Update the spline tangents and SplineReparamTable */
 	UFUNCTION(BlueprintCallable, Category = Spline)
 	virtual void UpdateSpline();
@@ -810,6 +813,10 @@ private:
 	static const FInterpCurvePointVector DummyPointScale;
 
 private:
+
+	/** Set the SplineCurves with the default shape i.e. a spline of two points (Used by default constructor) */
+	void SetDefaultSpline();
+
 	/** Returns the length of the specified spline segment up to the parametric value given */
 	float GetSegmentLength(const int32 Index, const float Param = 1.0f) const;
 
diff --git a/Engine/Source/Runtime/Engine/Private/Components/SplineComponent.cpp b/Engine/Source/Runtime/Engine/Private/Components/SplineComponent.cpp
--- a/Engine/Source/Runtime/Engine/Private/Components/SplineComponent.cpp
+++ b/Engine/Source/Runtime/Engine/Private/Components/SplineComponent.cpp
@@ -55,17 +55,7 @@ USplineComponent::USplineComponent(const FObjectInitializer& ObjectInitializer)
 	, ScaleVisualizationWidth(30.0f)
 #endif
 {
-	SplineCurves.Position.Points.Reset(10);
-	SplineCurves.Rotation.Points.Reset(10);
-	SplineCurves.Scale.Points.Reset(10);
-
-	SplineCurves.Position.Points.Emplace(0.0f, FVector(0, 0, 0), FVector::ZeroVector, FVector::ZeroVector, CIM_CurveAuto);
-	SplineCurves.Rotation.Points.Emplace(0.0f, FQuat::Identity, FQuat::Identity, FQuat::Identity, CIM_CurveAuto);
-	SplineCurves.Scale.Points.Emplace(0.0f, FVector(1.0f), FVector::ZeroVector, FVector::ZeroVector, CIM_CurveAuto);
-	
-	SplineCurves.Position.Points.Emplace(1.0f, FVector(100, 0, 0), FVector::ZeroVector, FVector::ZeroVector, CIM_CurveAuto);
-	SplineCurves.Rotation.Points.Emplace(1.0f, FQuat::Identity, FQuat::Identity, FQuat::Identity, CIM_CurveAuto);
-	SplineCurves.Scale.Points.Emplace(1.0f, FVector(1.0f), FVector::ZeroVector, FVector::ZeroVector, CIM_CurveAuto);
+	SetDefaultSpline();
 
 #if WITH_EDITORONLY_DATA
 	if (GEngine)
@@ -83,6 +73,45 @@ USplineComponent::USplineComponent(const FObjectInitializer& ObjectInitializer)
 	SplineReparamTable_DEPRECATED = SplineCurves.ReparamTable;
 }
 
+void USplineComponent::ResetToDefault()
+{
+	SetDefaultSpline();
+
+	bAllowSplineEditingPerInstance_DEPRECATED = true;
+	ReparamStepsPerSegment = 10;
+	Duration = 1.0f;
+	bStationaryEndpoints = false;
+	bSplineHasBeenEdited = false;
+	bModifiedByConstructionScript = false;
+	bInputSplinePointsToConstructionScript = false;
+	bDrawDebug  = true ;
+	bClosedLoop = false;
+	DefaultUpVector = FVector::UpVector;
+#if WITH_EDITORONLY_DATA
+	EditorUnselectedSplineSegmentColor = FStyleColors::White.GetSpecifiedColor();
+	EditorSelectedSplineSegmentColor = FStyleColors::AccentOrange.GetSpecifiedColor();
+	EditorTangentColor = FLinearColor(0.718f, 0.589f, 0.921f);
+	bAllowDiscontinuousSpline = false;
+	bShouldVisualizeScale = false;
+	ScaleVisualizationWidth = 30.0f;
+#endif
+}
+
+void USplineComponent::SetDefaultSpline()
+{
+	SplineCurves.Position.Points.Reset(10);
+	SplineCurves.Rotation.Points.Reset(10);
+	SplineCurves.Scale.Points.Reset(10);
+
+	SplineCurves.Position.Points.Emplace(0.0f, FVector(0, 0, 0), FVector::ZeroVector, FVector::ZeroVector, CIM_CurveAuto);
+	SplineCurves.Rotation.Points.Emplace(0.0f, FQuat::Identity, FQuat::Identity, FQuat::Identity, CIM_CurveAuto);
+	SplineCurves.Scale.Points.Emplace(0.0f, FVector(1.0f), FVector::ZeroVector, FVector::ZeroVector, CIM_CurveAuto);
+
+	SplineCurves.Position.Points.Emplace(1.0f, FVector(100, 0, 0), FVector::ZeroVector, FVector::ZeroVector, CIM_CurveAuto);
+	SplineCurves.Rotation.Points.Emplace(1.0f, FQuat::Identity, FQuat::Identity, FQuat::Identity, CIM_CurveAuto);
+	SplineCurves.Scale.Points.Emplace(1.0f, FVector(1.0f), FVector::ZeroVector, FVector::ZeroVector, CIM_CurveAuto);
+}
+
 void USplineComponent::GetLifetimeReplicatedProps(TArray<FLifetimeProperty>& OutLifetimeProps) const
 {
 	Super::GetLifetimeReplicatedProps(OutLifetimeProps);
