fixed null ptr dereferencing in UPrimitiveComponent::BatchSendRenderDebugPhysics

diff --git a/Engine/Source/Runtime/Engine/Private/Components/PrimitiveComponent.cpp b/Engine/Source/Runtime/Engine/Private/Components/PrimitiveComponent.cpp
--- a/Engine/Source/Runtime/Engine/Private/Components/PrimitiveComponent.cpp
+++ b/Engine/Source/Runtime/Engine/Private/Components/PrimitiveComponent.cpp
@@ -989,21 +989,26 @@ void UPrimitiveComponent::BatchSendRenderDebugPhysics(TArrayView<UPrimitiveCompo
 	TArray<uint32> DebugMassCounts;
 	TArray<FPrimitiveSceneProxy::FDebugMassData> DebugMassData;
 
-	SceneProxies.AddUninitialized(InPrimitives.Num());
-	DebugMassCounts.AddUninitialized(InPrimitives.Num());
+	SceneProxies.Reserve(InPrimitives.Num());
+	DebugMassCounts.Reserve(InPrimitives.Num());
 	DebugMassData.Reserve(InPrimitives.Num());
 
 	for (int32 PrimitiveIndex = 0; PrimitiveIndex < InPrimitives.Num(); PrimitiveIndex++)
 	{
-		SceneProxies[PrimitiveIndex] = InPrimitives[PrimitiveIndex]->SceneProxy;
-		
-		uint32 NumDebugMassDataBefore = DebugMassData.Num();
-		AppendDebugMassData(InPrimitives[PrimitiveIndex], DebugMassData);
-		DebugMassCounts[PrimitiveIndex] = DebugMassData.Num() - NumDebugMassDataBefore;
+		if (InPrimitives[PrimitiveIndex]->SceneProxy)
+		{
+			SceneProxies.Add(InPrimitives[PrimitiveIndex]->SceneProxy);
+
+			uint32 NumDebugMassDataBefore = DebugMassData.Num();
+			AppendDebugMassData(InPrimitives[PrimitiveIndex], DebugMassData);
+			DebugMassCounts.Add(DebugMassData.Num() - NumDebugMassDataBefore);
+		}
 	}
 
-	ENQUEUE_RENDER_COMMAND(PrimitiveComponent_BatchSendRenderDebugPhysics)(
-		[SceneProxies = MoveTemp(SceneProxies), DebugMassCounts = MoveTemp(DebugMassCounts), DebugMassData = MoveTemp(DebugMassData)](FRHICommandList& RHICmdList)
+	if (SceneProxies.Num())
+	{
+		ENQUEUE_RENDER_COMMAND(PrimitiveComponent_BatchSendRenderDebugPhysics)(
+			[SceneProxies = MoveTemp(SceneProxies), DebugMassCounts = MoveTemp(DebugMassCounts), DebugMassData = MoveTemp(DebugMassData)](FRHICommandList& RHICmdList)
 		{
 			TArray<FPrimitiveSceneProxy::FDebugMassData> SingleDebugMassData;
 			uint32 DebugMassOffset = 0;
@@ -1018,10 +1023,11 @@ void UPrimitiveComponent::BatchSendRenderDebugPhysics(TArrayView<UPrimitiveCompo
 				}
 
 				SceneProxies[ProxyIndex]->SetDebugMassData(SingleDebugMassData);
-				
+
 				DebugMassOffset += DebugMassCount;
 			}
 		});
+	}
 }
 #endif  // UE_ENABLE_DEBUG_DRAWING
 
