World Partition Handle Tests: remove access do destroyed package.

diff --git a/Engine/Source/Runtime/Engine/Private/Tests/WorldPartition/WorldPartitionHandleTests.cpp b/Engine/Source/Runtime/Engine/Private/Tests/WorldPartition/WorldPartitionHandleTests.cpp
--- a/Engine/Source/Runtime/Engine/Private/Tests/WorldPartition/WorldPartitionHandleTests.cpp
+++ b/Engine/Source/Runtime/Engine/Private/Tests/WorldPartition/WorldPartitionHandleTests.cpp
@@ -297,10 +297,6 @@ namespace WorldPartitionTests
 			ScopedEditorWorld.Reset();
 			CollectGarbage(GARBAGE_COLLECTION_KEEPFLAGS);
 
-			// Unload Package so test can run twice without issues.
-			UPackage* Package = FindPackage(NULL, *ActorDescContainer->GetContainerPackage().ToString());
-			UPackageTools::UnloadPackages({ Package });
-
 			TestFalse(TEXT("Invalid container test"), Handle.IsValid());
 			TestFalse(TEXT("Invalid container test"), Reference.IsValid());
 		}
