Fix distorted meshes (and potential crashes) when assigning assets via skeleton compatibility

Incorrect skeleton was being used for the uncompressed retargeting path

diff --git a/Engine/Plugins/Animation/AnimationData/Source/AnimationData/Private/AnimSequencerDataModel.cpp b/Engine/Plugins/Animation/AnimationData/Source/AnimationData/Private/AnimSequencerDataModel.cpp
--- a/Engine/Plugins/Animation/AnimationData/Source/AnimationData/Private/AnimSequencerDataModel.cpp
+++ b/Engine/Plugins/Animation/AnimationData/Source/AnimationData/Private/AnimSequencerDataModel.cpp
@@ -1130,7 +1130,7 @@ void UAnimationSequencerDataModel::GeneratePoseData(UControlRig* ControlRig, FAn
 			RigPose.ResetToRefPose();
 			const FBoneContainer& RequiredBones = RigPose.GetBoneContainer();
 			FBlendedCurve& Curve = InOutPoseData.GetCurve();
-			UE::Anim::Retargeting::FRetargetingScope RetargetingScope(RigPose, EvaluationContext);
+			UE::Anim::Retargeting::FRetargetingScope RetargetingScope(GetSkeleton(), RigPose, EvaluationContext);
 			
 			// Populate bone/curve elements to Pose/Curve indices
 			{
diff --git a/Engine/Source/Runtime/Engine/Classes/Animation/AnimData/AnimDataModel.h b/Engine/Source/Runtime/Engine/Classes/Animation/AnimData/AnimDataModel.h
--- a/Engine/Source/Runtime/Engine/Classes/Animation/AnimData/AnimDataModel.h
+++ b/Engine/Source/Runtime/Engine/Classes/Animation/AnimData/AnimDataModel.h
@@ -94,6 +94,8 @@ private:
 	FAnimCurveBase* FindMutableCurveById(const FAnimationCurveIdentifier& CurveIdentifier);	   
 	FRichCurve* GetMutableRichCurve(const FAnimationCurveIdentifier& CurveIdentifier);
 
+	USkeleton* GetSkeleton() const;
+
 private:
 	UPROPERTY(Transient)
 	int32 BracketCounter = 0;
diff --git a/Engine/Source/Runtime/Engine/Classes/Animation/AnimSequenceHelpers.h b/Engine/Source/Runtime/Engine/Classes/Animation/AnimSequenceHelpers.h
--- a/Engine/Source/Runtime/Engine/Classes/Animation/AnimSequenceHelpers.h
+++ b/Engine/Source/Runtime/Engine/Classes/Animation/AnimSequenceHelpers.h
@@ -245,11 +245,12 @@ namespace Anim {
 
 		struct ENGINE_API FRetargetingScope
 		{
-			FRetargetingScope(FCompactPose& ToRetargetPose, const DataModel::FEvaluationContext& InEvaluationContext);
+			FRetargetingScope(const USkeleton* InSourceSkeleton, FCompactPose& ToRetargetPose, const DataModel::FEvaluationContext& InEvaluationContext);
 			~FRetargetingScope();
 			
 			void AddTrackedBone(FCompactPoseBoneIndex CompactBoneIndex, int32 SkeletonBoneIndex) const;
 		private:
+			const USkeleton* SourceSkeleton;
 			FCompactPose& RetargetPose;
 			const DataModel::FEvaluationContext& EvaluationContext;
 			TArray<FRetargetTracking>& RetargetTracking;
diff --git a/Engine/Source/Runtime/Engine/Private/Animation/AnimDataModel.cpp b/Engine/Source/Runtime/Engine/Private/Animation/AnimDataModel.cpp
--- a/Engine/Source/Runtime/Engine/Private/Animation/AnimDataModel.cpp
+++ b/Engine/Source/Runtime/Engine/Private/Animation/AnimDataModel.cpp
@@ -794,7 +794,7 @@ void UAnimDataModel::Evaluate(FAnimationPoseData& InOutPoseData, const UE::Anim:
 	}
 
 	{		
-		UE::Anim::Retargeting::FRetargetingScope RetargetingScope(OutPose, EvaluationContext);
+		UE::Anim::Retargeting::FRetargetingScope RetargetingScope(GetSkeleton(), OutPose, EvaluationContext);
 		if (bShouldInterpolate)
 		{		
 			ExtractPose<true>(BoneAnimationTracks, ActiveCurves, OutPose, KeyIndex1, KeyIndex2, Alpha, TimePerFrame, RetargetingScope);
@@ -855,6 +855,20 @@ FRichCurve* UAnimDataModel::GetMutableRichCurve(const FAnimationCurveIdentifier&
 	return RichCurve;
 }
 
+USkeleton* UAnimDataModel::GetSkeleton() const
+{
+	const UAnimationAsset* AnimationAsset = CastChecked<UAnimationAsset>(GetOuter());	
+	checkf(AnimationAsset, TEXT("Unable to retrieve owning AnimationAsset"));
+
+	USkeleton* Skeleton = AnimationAsset->GetSkeleton();
+	if (Skeleton == nullptr)
+	{
+		IAnimationDataController::ReportObjectErrorf(this, LOCTEXT("UnableToFindSkeleton", "Unable to retrieve target USkeleton for Animation Asset ({0})"), FText::FromString(*AnimationAsset->GetPathName()));
+	} 
+
+	return Skeleton;
+}
+
 FBoneAnimationTrack* UAnimDataModel::FindMutableBoneTrackByName(FName Name)
 {
 	return BoneAnimationTracks.FindByPredicate([Name](const FBoneAnimationTrack& Track)
diff --git a/Engine/Source/Runtime/Engine/Private/Animation/AnimSequenceHelpers.cpp b/Engine/Source/Runtime/Engine/Private/Animation/AnimSequenceHelpers.cpp
--- a/Engine/Source/Runtime/Engine/Private/Animation/AnimSequenceHelpers.cpp
+++ b/Engine/Source/Runtime/Engine/Private/Animation/AnimSequenceHelpers.cpp
@@ -189,12 +189,14 @@ void CopyCurveDataToModel(const FRawCurveTracks& CurveData, const USkeleton* Ske
 	}	
 }
 	
-Retargeting::FRetargetingScope::FRetargetingScope(FCompactPose& ToRetargetPose, const DataModel::FEvaluationContext& InEvaluationContext)
-	: RetargetPose(ToRetargetPose),
+Retargeting::FRetargetingScope::FRetargetingScope(const USkeleton* InSourceSkeleton, FCompactPose& ToRetargetPose, const DataModel::FEvaluationContext& InEvaluationContext)
+	: SourceSkeleton(InSourceSkeleton),
+	RetargetPose(ToRetargetPose),
 	EvaluationContext(InEvaluationContext),
 	RetargetTracking(FBuildRawPoseScratchArea::Get().RetargetTracking),
 	bShouldRetarget(!ToRetargetPose.GetBoneContainer().GetDisableRetargeting() && EvaluationContext.RetargetTransforms.Num())
 {
+	check(SourceSkeleton);
 	RetargetTracking.Reset();
 }
 
@@ -211,10 +213,9 @@ Retargeting::FRetargetingScope::~FRetargetingScope()
 	if (bShouldRetarget)
 	{
 		const FBoneContainer& RequiredBones = RetargetPose.GetBoneContainer();
-		const USkeleton* Skeleton = RequiredBones.GetSkeletonAsset();
 		for (const FRetargetTracking& RT : RetargetTracking)
 		{
-			FAnimationRuntime::RetargetBoneTransform(Skeleton, EvaluationContext.RetargetSource, EvaluationContext.RetargetTransforms, RetargetPose[RT.PoseBoneIndex], RT.SkeletonBoneIndex, RT.PoseBoneIndex, RequiredBones, false);
+			FAnimationRuntime::RetargetBoneTransform(SourceSkeleton, EvaluationContext.RetargetSource, EvaluationContext.RetargetTransforms, RetargetPose[RT.PoseBoneIndex], RT.SkeletonBoneIndex, RT.PoseBoneIndex, RequiredBones, false);
 		}
 	}
 }
