World Partition: unify streaming object creation code between different hash implementations.

diff --git a/Engine/Source/Runtime/Engine/Private/WorldPartition/WorldPartitionRuntimeHash.cpp b/Engine/Source/Runtime/Engine/Private/WorldPartition/WorldPartitionRuntimeHash.cpp
--- a/Engine/Source/Runtime/Engine/Private/WorldPartition/WorldPartitionRuntimeHash.cpp
+++ b/Engine/Source/Runtime/Engine/Private/WorldPartition/WorldPartitionRuntimeHash.cpp
@@ -68,6 +68,15 @@ UWorldPartitionRuntimeHash::UWorldPartitionRuntimeHash(const FObjectInitializer&
 {}
 
 #if WITH_EDITOR
+URuntimeHashExternalStreamingObjectBase* UWorldPartitionRuntimeHash::CreateExternalStreamingObject(TSubclassOf<URuntimeHashExternalStreamingObjectBase> InClass, UObject* InOuter, FName InName, UWorld* InOwningWorld, UWorld* InOuterWorld)
+{
+	check(!InOwningWorld->IsGameWorld());
+	URuntimeHashExternalStreamingObjectBase* StreamingObject = NewObject<URuntimeHashExternalStreamingObjectBase>(InOuter, InClass, InName, RF_Public);
+	StreamingObject->OwningWorld = InOwningWorld;
+	StreamingObject->OuterWorld = InOuterWorld;
+	return StreamingObject;
+}
+
 void UWorldPartitionRuntimeHash::OnBeginPlay()
 {
 	// Mark always loaded actors so that the Level will force reference to these actors for PIE.
diff --git a/Engine/Source/Runtime/Engine/Private/WorldPartition/WorldPartitionRuntimeSpatialHash.cpp b/Engine/Source/Runtime/Engine/Private/WorldPartition/WorldPartitionRuntimeSpatialHash.cpp
--- a/Engine/Source/Runtime/Engine/Private/WorldPartition/WorldPartitionRuntimeSpatialHash.cpp
+++ b/Engine/Source/Runtime/Engine/Private/WorldPartition/WorldPartitionRuntimeSpatialHash.cpp
@@ -1009,9 +1009,8 @@ void UWorldPartitionRuntimeSpatialHash::DrawPreview() const
 
 URuntimeHashExternalStreamingObjectBase* UWorldPartitionRuntimeSpatialHash::StoreToExternalStreamingObject(UObject* StreamingObjectOuter, FName StreamingObjectName)
 {
-	check(!GetWorld()->IsGameWorld());
-	URuntimeSpatialHashExternalStreamingObject* StreamingObject = NewObject<URuntimeSpatialHashExternalStreamingObject>(StreamingObjectOuter, StreamingObjectName, RF_Public);
-	StreamingObject->Initialize(GetWorld(), GetTypedOuter<UWorld>());
+	URuntimeSpatialHashExternalStreamingObject* StreamingObject = CreateExternalStreamingObject<URuntimeSpatialHashExternalStreamingObject>(StreamingObjectOuter, StreamingObjectName);
+
 	StreamingObject->StreamingGrids = MoveTemp(StreamingGrids);
 
 	for (FSpatialHashStreamingGrid& StreamingGrid : StreamingObject->StreamingGrids)
diff --git a/Engine/Source/Runtime/Engine/Public/WorldPartition/WorldPartitionRuntimeHash.h b/Engine/Source/Runtime/Engine/Public/WorldPartition/WorldPartitionRuntimeHash.h
--- a/Engine/Source/Runtime/Engine/Public/WorldPartition/WorldPartitionRuntimeHash.h
+++ b/Engine/Source/Runtime/Engine/Public/WorldPartition/WorldPartitionRuntimeHash.h
@@ -4,13 +4,12 @@
 #include "CoreMinimal.h"
 #include "UObject/ObjectMacros.h"
 #include "UObject/Object.h"
+#include "Templates/SubclassOf.h"
 #include "Engine/World.h"
 #include "WorldPartition.h"
-#include "WorldPartitionActorDescView.h"
 #include "WorldPartition/ActorDescList.h"
 #include "WorldPartition/WorldPartitionHandle.h"
 #include "WorldPartition/WorldPartitionRuntimeCell.h"
-#include "WorldPartition/WorldPartitionActorDescViewProxy.h"
 #include "WorldPartition/WorldPartitionStreamingGeneration.h"
 #include "WorldPartition/WorldPartitionStreamingGenerationContext.h"
 #if WITH_EDITOR
@@ -19,7 +18,6 @@
 #endif
 #include "WorldPartitionRuntimeHash.generated.h"
 
-class FActorDescViewMap;
 struct FHierarchicalLogArchive;
 
 UENUM()
@@ -35,18 +33,16 @@ class ENGINE_API URuntimeHashExternalStreamingObjectBase : public UObject
 {
 	GENERATED_BODY()
 
+	friend class UWorldPartitionRuntimeHash;
+
 public:
-	virtual void Initialize(UWorld* InOwningWorld, UWorld* InOuterWorld) 
-	{ 
-		OwningWorld = InOwningWorld;  
-		OuterWorld = InOuterWorld;
-	}
+	//~ Begin UObject Interface
+	virtual class UWorld* GetWorld() const override final { return GetOwningWorld(); }
+	//~ End UObject Interface
 
 	UWorld* GetOwningWorld() const { return OwningWorld.Get(); }
 	UWorld* GetOuterWorld() const { return OuterWorld.Get(); }
 
-	virtual class UWorld* GetWorld() const override final { return GetOwningWorld(); }
-
 	void ForEachStreamingCells(TFunctionRef<void(UWorldPartitionRuntimeCell&)> Func);
 	
 	void OnStreamingObjectLoaded(UWorld* InjectedWorld);
@@ -158,12 +154,16 @@ public:
 protected:
 	virtual EWorldPartitionStreamingPerformance GetStreamingPerformanceForCell(const UWorldPartitionRuntimeCell* Cell) const { return EWorldPartitionStreamingPerformance::Good; }
 
-private:
 #if WITH_EDITOR
-	void ForceExternalActorLevelReference(bool bForceExternalActorLevelReferenceForPIE);
+	static URuntimeHashExternalStreamingObjectBase* CreateExternalStreamingObject(TSubclassOf<URuntimeHashExternalStreamingObjectBase> InClass, UObject* InOuter, FName InName, UWorld* InOwningWorld, UWorld* InOuterWorld);
+
+	template <class T>
+	T* CreateExternalStreamingObject(UObject* InOuter, FName InName)
+	{
+		return static_cast<T*>(CreateExternalStreamingObject(T::StaticClass(), InOuter, InName, GetWorld(), GetTypedOuter<UWorld>()));
+	}
 #endif
 
-protected:
 #if WITH_EDITORONLY_DATA
 	struct FAlwaysLoadedActorForPIE
 	{
@@ -182,4 +182,9 @@ protected:
 public:
 	mutable FActorDescList ModifiedActorDescListForPIE;
 #endif
+
+#if WITH_EDITOR
+private:
+	void ForceExternalActorLevelReference(bool bForceExternalActorLevelReferenceForPIE);
+#endif
 };
\ No newline at end of file
