Added FStringDev, which is a string that is present in memory only in editor and development builds.

diff --git a/Engine/Source/Runtime/Engine/Private/StringDev.cpp b/Engine/Source/Runtime/Engine/Private/StringDev.cpp
new file mode 100644
--- /dev/null
+++ b/Engine/Source/Runtime/Engine/Private/StringDev.cpp
@@ -0,0 +1,29 @@
+// Copyright Epic Games, Inc. All Rights Reserved.
+
+#include "StringDev.h"
+
+const FString& FStringDev::GetString() const
+{
+#if !WITH_EDITOR && !UE_BUILD_DEVELOPMENT
+	static FString String;
+#endif
+	return String;
+}
+
+void FStringDev::SetString(const FString& InString)
+{
+#if WITH_EDITOR || UE_BUILD_DEVELOPMENT
+	String = InString;
+#endif
+}
+
+FArchive& operator<<(FArchive& Ar, FStringDev& StringDev)
+{
+#if !WITH_EDITOR && !UE_BUILD_DEVELOPMENT
+	FString String;
+	Ar << String;
+#else
+	Ar << StringDev.String;
+#endif
+	return Ar;
+}
\ No newline at end of file
diff --git a/Engine/Source/Runtime/Engine/Private/WorldPartition/WorldPartitionRuntimeCellDataSpatialHash.cpp b/Engine/Source/Runtime/Engine/Private/WorldPartition/WorldPartitionRuntimeCellDataSpatialHash.cpp
--- a/Engine/Source/Runtime/Engine/Private/WorldPartition/WorldPartitionRuntimeCellDataSpatialHash.cpp
+++ b/Engine/Source/Runtime/Engine/Private/WorldPartition/WorldPartitionRuntimeCellDataSpatialHash.cpp
@@ -21,6 +21,12 @@ UWorldPartitionRuntimeCellDataSpatialHash::UWorldPartitionRuntimeCellDataSpatial
 	, CachedMinSquareDistanceToSource(MAX_dbl)
 {}
 
+void UWorldPartitionRuntimeCellDataSpatialHash::Serialize(FArchive& Ar)
+{
+	Super::Serialize(Ar);
+	Ar << DebugName;
+}
+
 void UWorldPartitionRuntimeCellDataSpatialHash::ResetStreamingSourceInfo() const
 {
 	Super::ResetStreamingSourceInfo();
@@ -174,5 +180,5 @@ bool UWorldPartitionRuntimeCellDataSpatialHash::IsDebugShown() const
 
 FString UWorldPartitionRuntimeCellDataSpatialHash::GetDebugName() const
 {
-	return DebugName;
+	return DebugName.GetString();
 }
\ No newline at end of file
diff --git a/Engine/Source/Runtime/Engine/Public/StringDev.h b/Engine/Source/Runtime/Engine/Public/StringDev.h
new file mode 100644
--- /dev/null
+++ b/Engine/Source/Runtime/Engine/Public/StringDev.h
@@ -0,0 +1,29 @@
+// Copyright Epic Games, Inc. All Rights Reserved.
+#pragma once
+
+#include "CoreTypes.h"
+
+/**
+ * A string that is present in memory only in editor and development builds.
+ */
+class ENGINE_API FStringDev
+{
+public:
+	FStringDev()
+	{}
+
+	FStringDev(const FString& InString)
+	{
+		SetString(InString);
+	}
+
+	const FString& GetString() const;
+	void SetString(const FString& InString);
+
+	ENGINE_API friend FArchive& operator<<(FArchive& Ar, FStringDev& StringDev);
+
+private:
+#if WITH_EDITOR || UE_BUILD_DEVELOPMENT
+	FString String;
+#endif
+};
\ No newline at end of file
diff --git a/Engine/Source/Runtime/Engine/Public/WorldPartition/WorldPartitionRuntimeCellDataSpatialHash.h b/Engine/Source/Runtime/Engine/Public/WorldPartition/WorldPartitionRuntimeCellDataSpatialHash.h
--- a/Engine/Source/Runtime/Engine/Public/WorldPartition/WorldPartitionRuntimeCellDataSpatialHash.h
+++ b/Engine/Source/Runtime/Engine/Public/WorldPartition/WorldPartitionRuntimeCellDataSpatialHash.h
@@ -2,6 +2,7 @@
 #pragma once
 
 #include "CoreMinimal.h"
+#include "StringDev.h"
 #include "WorldPartition/WorldPartitionRuntimeCellData.h"
 #include "WorldPartitionRuntimeCellDataSpatialHash.generated.h"
 
@@ -13,6 +14,10 @@ class ENGINE_API UWorldPartitionRuntimeCellDataSpatialHash : public UWorldPartit
 {
 	GENERATED_UCLASS_BODY()
 
+	//~Begin UObject Interface
+	void Serialize(FArchive& Ar);
+	//~End UObject Interface
+
 	//~Begin UWorldPartitionRuntimeCellData
 	virtual void ResetStreamingSourceInfo() const override;
 	virtual void AppendStreamingSourceInfo(const FWorldPartitionStreamingSource& Source, const FSphericalSector& SourceShape) const override;
@@ -38,8 +43,7 @@ class ENGINE_API UWorldPartitionRuntimeCellDataSpatialHash : public UWorldPartit
 	UPROPERTY()
 	FName GridName;
 
-	UPROPERTY()
-	FString DebugName;
+	FStringDev DebugName;
 
 private:
 	float ComputeSourceToCellAngleFactor(const FSphericalSector& SourceShape) const;
