FAnimNode_LayeredBoneBlend not compatible with Compatible Skeletons

#fix make sure skeleton bone indices from BlendMasks are (re)mapped to target skeleton for LayeredBoneBlend node instance

diff --git a/Engine/Source/Runtime/Engine/Private/Animation/AnimationRuntime.cpp b/Engine/Source/Runtime/Engine/Private/Animation/AnimationRuntime.cpp
--- a/Engine/Source/Runtime/Engine/Private/Animation/AnimationRuntime.cpp
+++ b/Engine/Source/Runtime/Engine/Private/Animation/AnimationRuntime.cpp
@@ -2178,9 +2178,11 @@ void FAnimationRuntime::CreateMaskWeights(TArray<FPerBoneBlendWeight>& BoneBlend
 				continue;
 			}
 
+			const USkeleton* MaskSkeleton = BlendMask->OwningSkeleton;
+			const FSkeletonRemapping& SkeletonRemapping = UE::Anim::FSkeletonRemappingRegistry::Get().GetRemapping(MaskSkeleton, Skeleton);
 			for (int32 EntryIndex = 0; EntryIndex < BlendMask->GetNumBlendEntries(); EntryIndex++)
 			{
-				int32 BoneIndex = BlendMask->ProfileEntries[EntryIndex].BoneReference.BoneIndex;
+				const int32 BoneIndex = SkeletonRemapping.IsValid() ? SkeletonRemapping.GetTargetSkeletonBoneIndex(BlendMask->ProfileEntries[EntryIndex].BoneReference.BoneIndex) : BlendMask->ProfileEntries[EntryIndex].BoneReference.BoneIndex;
 				if (BoneBlendWeights.IsValidIndex(BoneIndex))
 				{
 					// Match the BoneBlendWeight's input pose with BlendMasks's MaskIndex and use the blend mask's weight
