Preview scene directional light now supports transmission for SSS profile.

diff --git a/Engine/Source/Runtime/Engine/Private/PreviewScene.cpp b/Engine/Source/Runtime/Engine/Private/PreviewScene.cpp
--- a/Engine/Source/Runtime/Engine/Private/PreviewScene.cpp
+++ b/Engine/Source/Runtime/Engine/Private/PreviewScene.cpp
@@ -75,6 +75,7 @@ FPreviewScene::FPreviewScene(FPreviewScene::ConstructionValues CVS)
 		DirectionalLight = NewObject<UDirectionalLightComponent>(GetTransientPackage(), NAME_None, RF_Transient);
 		DirectionalLight->Intensity = CVS.LightBrightness;
 		DirectionalLight->LightColor = FColor::White;
+		DirectionalLight->bTransmission = true;
 		AddComponent(DirectionalLight, FTransform(CVS.LightRotation));
 
 		SkyLight = NewObject<USkyLightComponent>(GetTransientPackage(), NAME_None, RF_Transient);
