Add r.MaterialLogErrorOnFailure which will emit an `Error` when we fail to compile a material instead of a warning.

- This defaults to `false` (the original behavior)
- Can be enabled from the commandline via: -dpcvars="r.MaterialLogErrorOnFailure=1"

diff --git a/Engine/Source/Runtime/Engine/Private/Materials/Material.cpp b/Engine/Source/Runtime/Engine/Private/Materials/Material.cpp
--- a/Engine/Source/Runtime/Engine/Private/Materials/Material.cpp
+++ b/Engine/Source/Runtime/Engine/Private/Materials/Material.cpp
@@ -123,6 +123,13 @@ static TAutoConsoleVariable<int32> CVarMaterialParameterLegacyChecks(
 	TEXT("Note that this can be slow"),
 	ECVF_RenderThreadSafe);
 
+static TAutoConsoleVariable<bool> CVarMaterialLogErrorOnFailure(
+	TEXT("r.MaterialLogErrorOnFailure"),
+	false,
+	TEXT("When enabled, when a material fails to compile it will issue an Error instead of a Warning.\n")
+	TEXT("Default: false"),
+	ECVF_RenderThreadSafe);
+
 namespace MaterialImpl
 {
 	// Filtered list of properties that we follow on mobile platforms.
@@ -2458,6 +2465,10 @@ namespace MaterialImpl
 			{
 				UE_ASSET_LOG(LogMaterial, Fatal, This, TEXT("%s"), *ErrorString);
 			}
+			else if (CVarMaterialLogErrorOnFailure.GetValueOnAnyThread())
+			{
+				UE_ASSET_LOG(LogMaterial, Error, This, TEXT("%s"), *ErrorString);
+			}
 			else
 			{
 				UE_ASSET_LOG(LogMaterial, Warning, This, TEXT("%s"), *ErrorString);
@@ -3058,7 +3069,7 @@ void UMaterial::ConvertMaterialToStrataMaterial()
 	{
 		return;
 	}
-	
+
 	// Store current node post from the root node.
 	int32 CurrentNodePosX = EditorX;
 	const int32 TranslationOffsetX = 350.0f;
@@ -3090,7 +3101,7 @@ void UMaterial::ConvertMaterialToStrataMaterial()
 			NewNode->GetInput(NewInputIndex)->Connect(OldNodeInput.OutputIndex, OldNodeInput.Expression);
 		}
 	};
-	
+
 	UMaterialEditorOnlyData* EditorOnly = GetEditorOnlyData();
 
 	bool bCustomNodesGathered = false;
@@ -3206,7 +3217,7 @@ void UMaterial::ConvertMaterialToStrataMaterial()
 		ConvertNode->Roughness.Connect(3, BreakMatAtt);
 		if (bIsAnisotropyConnected)
 		{
-			ConvertNode->Anisotropy.Connect(4, BreakMatAtt);
+		ConvertNode->Anisotropy.Connect(4, BreakMatAtt);
 		}
 		ConvertNode->EmissiveColor.Connect(5, BreakMatAtt);
 		ConvertNode->Normal.Connect(8, BreakMatAtt);
@@ -3526,7 +3537,7 @@ void UMaterial::ConvertMaterialToStrataMaterial()
 			check(ConvertNode->ConvertedStrataMaterialInfo.CountShadingModels() == 1);
 
 			// Now pass through the convert to decal node, which flag the material as SSM_Decal, which will set the domain to Decal.
-			UMaterialExpressionStrataConvertToDecal* ConvertToDecalNode = NewObject<UMaterialExpressionStrataConvertToDecal>(this);
+			UMaterialExpressionStrataConvertToDecal* ConvertToDecalNode= NewObject<UMaterialExpressionStrataConvertToDecal>(this);
 			ReplaceNodeAndMoveToTheRight(ConvertNode, ConvertToDecalNode);
 			ConvertToDecalNode->DecalMaterial.Connect(0, ConvertNode);
 
@@ -4356,9 +4367,9 @@ bool UMaterial::CanEditChange(const FProperty* InProperty) const
 			if (bStrataEnabled)
 			{
 				return ((MaterialDomain != MD_PostProcess && MaterialDomain != MD_LightFunction && MaterialDomain != MD_Volume) || (MaterialDomain == MD_PostProcess && BlendableOutputAlpha));
-			}
+		}
 			else
-			{
+		{
 				return (MaterialDomain == MD_DeferredDecal || MaterialDomain == MD_Surface || MaterialDomain == MD_Volume || MaterialDomain == MD_UI || (MaterialDomain == MD_PostProcess && BlendableOutputAlpha));
 			}
 		}
@@ -4504,8 +4515,8 @@ void UMaterial::PostEditChangePropertyInternal(FPropertyChangedEvent& PropertyCh
 
 	// If BLEND_TranslucentColoredTransmittance is selected while Strata is not enabled, force BLEND_Translucent blend mode
 	if (!Strata::IsStrataEnabled() && BlendMode == BLEND_TranslucentColoredTransmittance)
-	{
-		BlendMode = BLEND_Translucent;
+		{
+			BlendMode = BLEND_Translucent;
 	}
 
 	bool bRequiresCompilation = true;
@@ -4893,7 +4904,7 @@ void UMaterial::RebuildShadingModelField()
 				{
 					BlendMode = BLEND_Opaque;
 				}
-			}
+				}
 			else if (StrataMaterialInfo.HasOnlyShadingModel(SSM_UI))
 			{
 				MaterialDomain = EMaterialDomain::MD_UI;
@@ -6610,12 +6621,12 @@ static bool IsPropertyActive_Internal(EMaterialProperty InProperty,
 		}
 		else
 		{
-			return InProperty == MP_EmissiveColor
-				|| (InProperty == MP_WorldPositionOffset)
+		return InProperty == MP_EmissiveColor
+			|| (InProperty == MP_WorldPositionOffset)
 				|| (InProperty == MP_OpacityMask && IsMaskedBlendMode(BlendMode))
-				|| (InProperty == MP_Opacity && IsTranslucentBlendMode(BlendMode) && BlendMode != BLEND_Modulate)
-				|| (InProperty >= MP_CustomizedUVs0 && InProperty <= MP_CustomizedUVs7);
-		}
+			|| (InProperty == MP_Opacity && IsTranslucentBlendMode(BlendMode) && BlendMode != BLEND_Modulate)
+			|| (InProperty >= MP_CustomizedUVs0 && InProperty <= MP_CustomizedUVs7);
+	}
 	}
 
 	// Now processing MD_Surface
