Fix non-unity build.

diff --git a/Engine/Source/Runtime/Engine/Public/StringDev.h b/Engine/Source/Runtime/Engine/Public/StringDev.h
--- a/Engine/Source/Runtime/Engine/Public/StringDev.h
+++ b/Engine/Source/Runtime/Engine/Public/StringDev.h
@@ -2,6 +2,7 @@
 #pragma once
 
 #include "CoreTypes.h"
+#include "Containers/UnrealString.h"
 
 /**
  * A string that is present in memory only in editor and development builds.
