Stop the component instance cache from reattaching garbage components during reconstruction
This was causing a crash during seamless travel when transferring a pawn with child actors that had niagara components attached

diff --git a/Engine/Source/Runtime/Engine/Private/Components/SceneComponent.cpp b/Engine/Source/Runtime/Engine/Private/Components/SceneComponent.cpp
--- a/Engine/Source/Runtime/Engine/Private/Components/SceneComponent.cpp
+++ b/Engine/Source/Runtime/Engine/Private/Components/SceneComponent.cpp
@@ -2418,7 +2418,8 @@ void FSceneComponentInstanceData::ApplyToComponent(UActorComponent* Component, c
 		USceneComponent* ChildComponent = ChildComponentPair.Key;
 		// If the ChildComponent now has a "good" attach parent it was set by the transaction and it means we are undoing/redoing attachment
 		// and so the rebuilt component should not take back attachment ownership
-		if (ChildComponent && !IsValid(ChildComponent->GetAttachParent()))
+		// We don't want to do this for garbage components
+		if (IsValid(ChildComponent) && !IsValid(ChildComponent->GetAttachParent()))
 		{
 			ChildComponent->SetRelativeTransform_Direct(ChildComponentPair.Value);
 			ChildComponent->AttachToComponent(SceneComponent, FAttachmentTransformRules::KeepRelativeTransform);
