Fix for SetDisableRetargeting not working correctly for packaged projects

diff --git a/Engine/Source/Runtime/Engine/Classes/Animation/Skeleton.h b/Engine/Source/Runtime/Engine/Classes/Animation/Skeleton.h
--- a/Engine/Source/Runtime/Engine/Classes/Animation/Skeleton.h
+++ b/Engine/Source/Runtime/Engine/Classes/Animation/Skeleton.h
@@ -837,9 +837,9 @@ public:
 	 */
 	ENGINE_API int32 GetMeshBoneIndexFromSkeletonBoneIndex(const USkinnedAsset* InSkinnedAsset, const int32 SkeletonBoneIndex);
 
-	EBoneTranslationRetargetingMode::Type GetBoneTranslationRetargetingMode(const int32 BoneTreeIdx) const
+	EBoneTranslationRetargetingMode::Type GetBoneTranslationRetargetingMode(const int32 BoneTreeIdx, bool bDisableRetargeting = false) const
 	{
-		if (BoneTree.IsValidIndex(BoneTreeIdx))
+		if (!bDisableRetargeting && BoneTree.IsValidIndex(BoneTreeIdx))
 		{
 			return BoneTree[BoneTreeIdx].TranslationRetargetingMode;
 		}
diff --git a/Engine/Source/Runtime/Engine/Private/Animation/AnimationDecompression.cpp b/Engine/Source/Runtime/Engine/Private/Animation/AnimationDecompression.cpp
--- a/Engine/Source/Runtime/Engine/Private/Animation/AnimationDecompression.cpp
+++ b/Engine/Source/Runtime/Engine/Private/Animation/AnimationDecompression.cpp
@@ -86,7 +86,7 @@ void DecompressPose(FCompactPose& OutPose,
 					RotationScalePairs.Add(BoneTrackPair(CompactPoseBoneIndex, TrackIndex));
 
 					// Skip extracting translation component for EBoneTranslationRetargetingMode::Skeleton.
-					switch (TargetSkeleton->GetBoneTranslationRetargetingMode(TargetSkeletonBoneIndex))
+					switch (TargetSkeleton->GetBoneTranslationRetargetingMode(TargetSkeletonBoneIndex, RequiredBones.GetDisableRetargeting()))
 					{
 					case EBoneTranslationRetargetingMode::Animation:
 						TranslationPairs.Add(BoneTrackPair(CompactPoseBoneIndex, TrackIndex));
@@ -144,7 +144,7 @@ void DecompressPose(FCompactPose& OutPose,
 				if (DecompressionContext.IsAdditiveAnimation())
 				{
 					RootAtom.SetRotation(SkeletonRemapping.RetargetAdditiveRotationToTargetSkeleton(TargetSkeletonBoneIndex, RootAtom.GetRotation()));
-					if (TargetSkeleton->GetBoneTranslationRetargetingMode(TargetSkeletonBoneIndex) != EBoneTranslationRetargetingMode::Skeleton)
+					if (TargetSkeleton->GetBoneTranslationRetargetingMode(TargetSkeletonBoneIndex, RequiredBones.GetDisableRetargeting()) != EBoneTranslationRetargetingMode::Skeleton)
 					{
 						RootAtom.SetTranslation(SkeletonRemapping.RetargetAdditiveTranslationToTargetSkeleton(TargetSkeletonBoneIndex, RootAtom.GetTranslation()));
 					}
@@ -152,7 +152,7 @@ void DecompressPose(FCompactPose& OutPose,
 				else
 				{
 					RootAtom.SetRotation(SkeletonRemapping.RetargetBoneRotationToTargetSkeleton(TargetSkeletonBoneIndex, RootAtom.GetRotation()));
-					if (TargetSkeleton->GetBoneTranslationRetargetingMode(TargetSkeletonBoneIndex) != EBoneTranslationRetargetingMode::Skeleton)
+					if (TargetSkeleton->GetBoneTranslationRetargetingMode(TargetSkeletonBoneIndex, RequiredBones.GetDisableRetargeting()) != EBoneTranslationRetargetingMode::Skeleton)
 					{
 						RootAtom.SetTranslation(SkeletonRemapping.RetargetBoneTranslationToTargetSkeleton(TargetSkeletonBoneIndex, RootAtom.GetTranslation()));
 					}
@@ -180,7 +180,7 @@ void DecompressPose(FCompactPose& OutPose,
 			{
 				const int32 TargetSkeletonBoneIndex = RequiredBones.GetSkeletonIndex(BoneIndex);
 				OutPose[BoneIndex].SetRotation(SkeletonRemapping.RetargetAdditiveRotationToTargetSkeleton(TargetSkeletonBoneIndex, OutPose[BoneIndex].GetRotation()));
-				if (TargetSkeleton->GetBoneTranslationRetargetingMode(TargetSkeletonBoneIndex) != EBoneTranslationRetargetingMode::Skeleton)
+				if (TargetSkeleton->GetBoneTranslationRetargetingMode(TargetSkeletonBoneIndex, RequiredBones.GetDisableRetargeting()) != EBoneTranslationRetargetingMode::Skeleton)
 				{
 					OutPose[BoneIndex].SetTranslation(SkeletonRemapping.RetargetAdditiveTranslationToTargetSkeleton(TargetSkeletonBoneIndex, OutPose[BoneIndex].GetTranslation()));
 				}
@@ -192,7 +192,7 @@ void DecompressPose(FCompactPose& OutPose,
 			{
 				const int32 TargetSkeletonBoneIndex = RequiredBones.GetSkeletonIndex(BoneIndex);
 				OutPose[BoneIndex].SetRotation(SkeletonRemapping.RetargetBoneRotationToTargetSkeleton(TargetSkeletonBoneIndex, OutPose[BoneIndex].GetRotation()));
-				if (TargetSkeleton->GetBoneTranslationRetargetingMode(TargetSkeletonBoneIndex) != EBoneTranslationRetargetingMode::Skeleton)
+				if (TargetSkeleton->GetBoneTranslationRetargetingMode(TargetSkeletonBoneIndex, RequiredBones.GetDisableRetargeting()) != EBoneTranslationRetargetingMode::Skeleton)
 				{
 					OutPose[BoneIndex].SetTranslation(SkeletonRemapping.RetargetBoneTranslationToTargetSkeleton(TargetSkeletonBoneIndex, OutPose[BoneIndex].GetTranslation()));
 				}
diff --git a/Engine/Source/Runtime/Engine/Private/Animation/AnimationRuntime.cpp b/Engine/Source/Runtime/Engine/Private/Animation/AnimationRuntime.cpp
--- a/Engine/Source/Runtime/Engine/Private/Animation/AnimationRuntime.cpp
+++ b/Engine/Source/Runtime/Engine/Private/Animation/AnimationRuntime.cpp
@@ -2627,7 +2627,7 @@ void FAnimationRuntime::RetargetBoneTransform(const USkeleton* SourceSkeleton, c
 		const int32 TargetSkeletonBoneIndex = RequiredBones.GetSkeletonIndex(BoneIndex);
 		const int32 SourceSkeletonBoneIndex = SkeletonRemapping.IsValid() ? SkeletonRemapping.GetSourceSkeletonBoneIndex(TargetSkeletonBoneIndex) : SkeletonBoneIndex;
 
-		switch (TargetSkeleton->GetBoneTranslationRetargetingMode(TargetSkeletonBoneIndex))
+		switch (TargetSkeleton->GetBoneTranslationRetargetingMode(TargetSkeletonBoneIndex, RequiredBones.GetDisableRetargeting()))
 		{
 			case EBoneTranslationRetargetingMode::AnimationScaled:
 			{
diff --git a/Engine/Source/Runtime/Engine/Private/BoneContainer.cpp b/Engine/Source/Runtime/Engine/Private/BoneContainer.cpp
--- a/Engine/Source/Runtime/Engine/Private/BoneContainer.cpp
+++ b/Engine/Source/Runtime/Engine/Private/BoneContainer.cpp
@@ -463,7 +463,7 @@ const FRetargetSourceCachedData& FBoneContainer::GetRetargetSourceCachedData(con
 			const int32 TargetSkeletonBoneIndex = CompactPoseToSkeletonIndex[CompactBoneIndex];
 			const int32 SourceSkeletonBoneIndex = InRemapping.IsValid() ? InRemapping.GetSourceSkeletonBoneIndex(TargetSkeletonBoneIndex) : TargetSkeletonBoneIndex;
 
-			if (SourceSkeletonBoneIndex != INDEX_NONE && AssetSkeleton->GetBoneTranslationRetargetingMode(SourceSkeletonBoneIndex) == EBoneTranslationRetargetingMode::OrientAndScale)
+			if (SourceSkeletonBoneIndex != INDEX_NONE && AssetSkeleton->GetBoneTranslationRetargetingMode(SourceSkeletonBoneIndex, bDisableRetargeting) == EBoneTranslationRetargetingMode::OrientAndScale)
 			{
 				if(AuthoredOnRefSkeleton.IsValidIndex(SourceSkeletonBoneIndex))
 				{
