s23.50: Improved boom sniper performace by disabling wobble on stick behavior when more than 15m from camera

This replaces the previous implementation by using a culling range instead of disabling the effect completely.

diff --git a/Engine/Source/Runtime/Engine/Classes/Kismet/GameplayStatics.h b/Engine/Source/Runtime/Engine/Classes/Kismet/GameplayStatics.h
--- a/Engine/Source/Runtime/Engine/Classes/Kismet/GameplayStatics.h
+++ b/Engine/Source/Runtime/Engine/Classes/Kismet/GameplayStatics.h
@@ -217,6 +217,15 @@ class ENGINE_API UGameplayStatics : public UBlueprintFunctionLibrary
 	UFUNCTION(BlueprintPure, Category="Game", meta=(WorldContext="WorldContextObject", UnsafeDuringActorConstruction="true"))
 	static class APlayerCameraManager* GetPlayerCameraManager(const UObject* WorldContextObject, int32 PlayerIndex);
 
+	/** 
+	 * Determines if any local player controller's camera is within range of the specified location.
+	 * @param Location		The location from which test range
+	 * @param MaximumRange	The distance away from Location to test range
+	 * @note This will always return false on dedicated servers.
+	 */
+	UFUNCTION(BlueprintPure, BlueprintCosmetic, Category="Game", meta=(WorldContext="WorldContextObject"))
+	static bool IsAnyLocalPlayerCameraWithinRange(const UObject* WorldContextObject, const FVector& Location, float MaximumRange);
+
 	/** 
 	 * Returns the player controller with the specified physical controller ID. This only works for local players.
 	 *
diff --git a/Engine/Source/Runtime/Engine/Private/GameplayStatics.cpp b/Engine/Source/Runtime/Engine/Private/GameplayStatics.cpp
--- a/Engine/Source/Runtime/Engine/Private/GameplayStatics.cpp
+++ b/Engine/Source/Runtime/Engine/Private/GameplayStatics.cpp
@@ -420,6 +420,43 @@ APlayerCameraManager* UGameplayStatics::GetPlayerCameraManager(const UObject* Wo
 	return PC ? PC->PlayerCameraManager : nullptr;
 }
 
+bool UGameplayStatics::IsAnyLocalPlayerCameraWithinRange(const UObject* WorldContextObject, const FVector& Location, float MaximumRange)
+{
+	if (!GEngine || IsRunningDedicatedServer())
+	{
+		return false;
+	}
+	
+	UWorld* World = GEngine->GetWorldFromContextObject(WorldContextObject, EGetWorldErrorMode::LogAndReturnNull);
+	if (!World)
+	{
+		return false;
+	}
+
+	const float MaximumRangeSq = MaximumRange * MaximumRange; 
+
+	for (FConstPlayerControllerIterator PlayerControllerIterator = World->GetPlayerControllerIterator();
+		 PlayerControllerIterator; 
+		 ++PlayerControllerIterator)
+	{
+		const APlayerController* PlayerController = PlayerControllerIterator->Get();
+		if (!PlayerController || !PlayerController->IsLocalController())
+		{
+			continue;
+		}
+
+		if (const APlayerCameraManager* PlayerCameraManager = PlayerController->PlayerCameraManager)
+		{
+			if (FVector::DistSquared(PlayerCameraManager->GetCameraLocation(), Location) <= MaximumRangeSq)
+			{
+				return true;
+			}
+		}
+	}
+
+	return false;
+}
+
 APlayerController* UGameplayStatics::CreatePlayer(const UObject* WorldContextObject, int32 ControllerId, bool bSpawnPlayerController)
 {
 	return CreatePlayerFromPlatformUser(WorldContextObject, FGenericPlatformMisc::GetPlatformUserForUserIndex(ControllerId), bSpawnPlayerController);
