Added logging of Persistent Level content in WorldPartition streaming generation log.

diff --git a/Engine/Source/Runtime/Engine/Private/WorldPartition/WorldPartitionRuntimeHash.cpp b/Engine/Source/Runtime/Engine/Private/WorldPartition/WorldPartitionRuntimeHash.cpp
--- a/Engine/Source/Runtime/Engine/Private/WorldPartition/WorldPartitionRuntimeHash.cpp
+++ b/Engine/Source/Runtime/Engine/Private/WorldPartition/WorldPartitionRuntimeHash.cpp
@@ -276,7 +276,18 @@ void UWorldPartitionRuntimeHash::DumpStateLog(FHierarchicalLogArchive& Ar)
 	Ar.Printf(TEXT("----------------------------------------------------------------------------------------------------------------"));
 	Ar.Printf(TEXT("%s - Persistent Level"), *GetWorld()->GetName());
 	Ar.Printf(TEXT("----------------------------------------------------------------------------------------------------------------"));
-	Ar.Printf(TEXT("Always loaded Actor Count: %d "), GetWorld()->PersistentLevel->Actors.Num());
+	{
+		FHierarchicalLogArchive::FIndentScope CellIndentScope = Ar.PrintfIndent(TEXT("Content of %s Persistent Level"), *GetWorld()->GetName());
+		Ar.Printf(TEXT("Always loaded Actor Count: %d "), GetWorld()->PersistentLevel->Actors.Num());
+		for (const AActor* Actor : GetWorld()->PersistentLevel->Actors)
+		{
+			if (Actor)
+			{
+				Ar.Printf(TEXT("Actor Path: %s"), *Actor->GetPathName());
+				Ar.Printf(TEXT("Actor Package: %s"), *Actor->GetPackage()->GetName());
+			}
+		}
+	}
 	Ar.Printf(TEXT(""));
 }
 
