FWorldPartitionActorDesc: Keep bFailedToLoad transient flag to display in outliner

diff --git a/Engine/Source/Runtime/Engine/Private/WorldPartition/LoaderAdapter/LoaderAdapterPinnedActors.cpp b/Engine/Source/Runtime/Engine/Private/WorldPartition/LoaderAdapter/LoaderAdapterPinnedActors.cpp
--- a/Engine/Source/Runtime/Engine/Private/WorldPartition/LoaderAdapter/LoaderAdapterPinnedActors.cpp
+++ b/Engine/Source/Runtime/Engine/Private/WorldPartition/LoaderAdapter/LoaderAdapterPinnedActors.cpp
@@ -20,6 +20,11 @@ FText FLoaderAdapterPinnedActors::GetUnloadedReason(FWorldPartitionActorDesc* In
 {
 	if (InActorDesc)
 	{
+		if (InActorDesc->FailedToLoad())
+		{
+			return LOCTEXT("FailedToLoadReason", "Failed to load");
+		}
+
 		UActorDescContainer* ActorDescContainer = InActorDesc->GetContainer();
 		UWorldPartition* WorldPartition = ActorDescContainer != nullptr ? ActorDescContainer->GetWorldPartition() : nullptr;
 		bool bShouldBeLoaded = !InActorDesc->GetIsSpatiallyLoaded() && !InActorDesc->GetActorIsRuntimeOnly();
diff --git a/Engine/Source/Runtime/Engine/Private/WorldPartition/WorldPartitionActorDesc.cpp b/Engine/Source/Runtime/Engine/Private/WorldPartition/WorldPartitionActorDesc.cpp
--- a/Engine/Source/Runtime/Engine/Private/WorldPartition/WorldPartitionActorDesc.cpp
+++ b/Engine/Source/Runtime/Engine/Private/WorldPartition/WorldPartitionActorDesc.cpp
@@ -34,6 +34,7 @@ FWorldPartitionActorDesc::FWorldPartitionActorDesc()
 	, HardRefCount(0)
 	, Container(nullptr)
 	, bIsForcedNonSpatiallyLoaded(false)
+	, bFailedToLoad(false)
 {}
 
 void FWorldPartitionActorDesc::Init(const AActor* InActor)
@@ -611,6 +612,8 @@ AActor* FWorldPartitionActorDesc::GetActor(bool bEvenIfPendingKill, bool bEvenIf
 
 AActor* FWorldPartitionActorDesc::Load() const
 {
+	bFailedToLoad = false;
+
 	if (ActorPtr.IsExplicitlyNull() || ActorPtr.IsStale())
 	{
 		// First, try to find the existing actor which could have been loaded by another actor (through standard serialization)
@@ -640,6 +643,7 @@ AActor* FWorldPartitionActorDesc::Load() const
 			if (!ActorPtr.IsValid())
 			{
 				UE_LOG(LogWorldPartition, Warning, TEXT("Can't load actor guid `%s` ('%s') from package '%s'"), *Guid.ToString(), *GetActorName().ToString(), *ActorPackage.ToString());
+				bFailedToLoad = true;
 			}
 		}
 	}
diff --git a/Engine/Source/Runtime/Engine/Public/WorldPartition/WorldPartitionActorDesc.h b/Engine/Source/Runtime/Engine/Public/WorldPartition/WorldPartitionActorDesc.h
--- a/Engine/Source/Runtime/Engine/Public/WorldPartition/WorldPartitionActorDesc.h
+++ b/Engine/Source/Runtime/Engine/Public/WorldPartition/WorldPartitionActorDesc.h
@@ -159,6 +159,11 @@ PRAGMA_ENABLE_DEPRECATION_WARNINGS
 		return GetTypeHash(Key.Guid);
 	}
 
+	bool FailedToLoad() const
+	{
+		return bFailedToLoad;
+	}
+
 protected:
 	inline uint32 IncSoftRefCount() const
 	{
@@ -284,6 +289,7 @@ protected:
 	UActorDescContainer*			Container;
 	TOptional<TArray<FName>>		ResolvedDataLayerInstanceNames; // Can only resolve in ActorDesc if Container is not used as a template
 	bool							bIsForcedNonSpatiallyLoaded;
+	mutable bool					bFailedToLoad;
 
 	static TMap<TSubclassOf<AActor>, FActorDescDeprecator> Deprecators;
 };
