FExec: Added a Exec_Runtime function to go with Exec_Editor and Exec_Dev. Exec_Runtime is routed on all targets where exec commands are enabled.

diff --git a/Engine/Plugins/Experimental/GameplayTargetingSystem/Source/GameplayTargetingSystem/TargetingSystem/TargetingSubsystem.cpp b/Engine/Plugins/Experimental/GameplayTargetingSystem/Source/GameplayTargetingSystem/TargetingSystem/TargetingSubsystem.cpp
--- a/Engine/Plugins/Experimental/GameplayTargetingSystem/Source/GameplayTargetingSystem/TargetingSystem/TargetingSubsystem.cpp
+++ b/Engine/Plugins/Experimental/GameplayTargetingSystem/Source/GameplayTargetingSystem/TargetingSystem/TargetingSubsystem.cpp
@@ -139,7 +139,7 @@ void UTargetingSubsystem::Deinitialize()
 	AsyncTargetingRequests.Empty();
 }
 
-bool UTargetingSubsystem::Exec(UWorld* Inworld, const TCHAR* Cmd, FOutputDevice& Ar)
+bool UTargetingSubsystem::Exec_Runtime(UWorld* Inworld, const TCHAR* Cmd, FOutputDevice& Ar)
 {
 	if (HasAnyFlags(RF_ClassDefaultObject) || IsDefaultSubobject())
 	{
diff --git a/Engine/Plugins/Experimental/GameplayTargetingSystem/Source/GameplayTargetingSystem/TargetingSystem/TargetingSubsystem.h b/Engine/Plugins/Experimental/GameplayTargetingSystem/Source/GameplayTargetingSystem/TargetingSystem/TargetingSubsystem.h
--- a/Engine/Plugins/Experimental/GameplayTargetingSystem/Source/GameplayTargetingSystem/TargetingSystem/TargetingSubsystem.h
+++ b/Engine/Plugins/Experimental/GameplayTargetingSystem/Source/GameplayTargetingSystem/TargetingSystem/TargetingSubsystem.h
@@ -96,7 +96,7 @@ private:
 	// End USubsystem
 
 	/** FSelfRegisteringExec implementation */
-	virtual bool Exec(UWorld* Inworld, const TCHAR* Cmd, FOutputDevice& Ar) override;
+	virtual bool Exec_Runtime(UWorld* Inworld, const TCHAR* Cmd, FOutputDevice& Ar) override;
 	/** ~FSelfRegisteringExec implementation */
 
 	// FTickableGameObject interface
diff --git a/Engine/Plugins/Experimental/PythonScriptPlugin/Source/PythonScriptPlugin/Private/PythonScriptPlugin.cpp b/Engine/Plugins/Experimental/PythonScriptPlugin/Source/PythonScriptPlugin/Private/PythonScriptPlugin.cpp
--- a/Engine/Plugins/Experimental/PythonScriptPlugin/Source/PythonScriptPlugin/Private/PythonScriptPlugin.cpp
+++ b/Engine/Plugins/Experimental/PythonScriptPlugin/Source/PythonScriptPlugin/Private/PythonScriptPlugin.cpp
@@ -705,7 +705,7 @@ void FPythonScriptPlugin::ShutdownModule()
 #endif	// WITH_PYTHON
 }
 
-bool FPythonScriptPlugin::Exec(UWorld* InWorld, const TCHAR* Cmd, FOutputDevice& Ar)
+bool FPythonScriptPlugin::Exec_Runtime(UWorld* InWorld, const TCHAR* Cmd, FOutputDevice& Ar)
 {
 #if WITH_PYTHON
 	if (FParse::Command(&Cmd, TEXT("PY")))
diff --git a/Engine/Plugins/Experimental/PythonScriptPlugin/Source/PythonScriptPlugin/Private/PythonScriptPlugin.h b/Engine/Plugins/Experimental/PythonScriptPlugin/Source/PythonScriptPlugin/Private/PythonScriptPlugin.h
--- a/Engine/Plugins/Experimental/PythonScriptPlugin/Source/PythonScriptPlugin/Private/PythonScriptPlugin.h
+++ b/Engine/Plugins/Experimental/PythonScriptPlugin/Source/PythonScriptPlugin/Private/PythonScriptPlugin.h
@@ -116,9 +116,6 @@ public:
 	virtual void StartupModule() override;
 	virtual void ShutdownModule() override;
 
-	//~ FSelfRegisteringExec interface
-	virtual bool Exec(UWorld* InWorld, const TCHAR* Cmd, FOutputDevice& Ar) override;
-
 	//~ FEnumEditorUtils::INotifyOnEnumChanged interface
 	virtual void PreChange(const UUserDefinedEnum* Enum, FEnumEditorUtils::EEnumEditorChangeInfo Info) override;
 	virtual void PostChange(const UUserDefinedEnum* Enum, FEnumEditorUtils::EEnumEditorChangeInfo Info) override;
@@ -155,6 +152,10 @@ public:
 	PyObject* GetConsoleLocalDict()  { return PyConsoleLocalDict.Get();  }
 #endif	// WITH_PYTHON
 
+protected:
+	//~ FSelfRegisteringExec interface
+	virtual bool Exec_Runtime(UWorld* InWorld, const TCHAR* Cmd, FOutputDevice& Ar) override;
+
 private:
 #if WITH_PYTHON
 	void InitializePython();
diff --git a/Engine/Source/Runtime/Core/Private/Misc/Exec.cpp b/Engine/Source/Runtime/Core/Private/Misc/Exec.cpp
--- a/Engine/Source/Runtime/Core/Private/Misc/Exec.cpp
+++ b/Engine/Source/Runtime/Core/Private/Misc/Exec.cpp
@@ -19,6 +19,8 @@ bool FExec::Exec( UWorld* InWorld, const TCHAR* Cmd, FOutputDevice& Ar )
 	bExecSuccess = bExecSuccess || Exec_Dev( InWorld, Cmd, Ar );
 #endif
 
+	bExecSuccess = bExecSuccess || Exec_Runtime( InWorld, Cmd, Ar );
+
 	return bExecSuccess;
 }
 #endif // UE_ALLOW_EXEC_COMMANDS
\ No newline at end of file
diff --git a/Engine/Source/Runtime/Core/Public/Misc/Exec.h b/Engine/Source/Runtime/Core/Public/Misc/Exec.h
--- a/Engine/Source/Runtime/Core/Public/Misc/Exec.h
+++ b/Engine/Source/Runtime/Core/Public/Misc/Exec.h
@@ -58,6 +58,9 @@ public:
 #endif // !UE_ALLOW_EXEC_COMMANDS
 
 protected:
+	/** Implementation of Exec that is called on all targets where UE_ALLOW_EXEC_COMMANDS is true */
+	virtual bool Exec_Runtime(UWorld* InWorld, const TCHAR* Cmd, FOutputDevice& Ar) { return false; }
+
 	/** Implementation of Exec that is only called in non-shipping targets */
 	virtual bool Exec_Dev( UWorld* InWorld, const TCHAR* Cmd, FOutputDevice& Ar ) { return false; }
 
diff --git a/Engine/Source/Runtime/Engine/Classes/Engine/GameEngine.h b/Engine/Source/Runtime/Engine/Classes/Engine/GameEngine.h
--- a/Engine/Source/Runtime/Engine/Classes/Engine/GameEngine.h
+++ b/Engine/Source/Runtime/Engine/Classes/Engine/GameEngine.h
@@ -140,11 +140,15 @@ public:
 	virtual bool NetworkRemapPath(UNetConnection* Connection, FString& Str, bool bReading = true) override;
 	virtual bool ShouldDoAsyncEndOfFrameTasks() const override;
 
-public:
-
 	// FExec interface
-
+public:
+#if UE_ALLOW_EXEC_COMMANDS
 	virtual bool Exec( UWorld* InWorld, const TCHAR* Cmd, FOutputDevice& Ar=*GLog ) override;
+#endif
+protected:
+	virtual bool Exec_Runtime( UWorld* InWorld, const TCHAR* Cmd, FOutputDevice& Ar ) override;
+	virtual bool Exec_Dev( UWorld* InWorld, const TCHAR* Cmd, FOutputDevice& Ar ) override;
+	virtual bool Exec_Editor( UWorld* InWorld, const TCHAR* Cmd, FOutputDevice& Ar ) override;
 
 public:
 
diff --git a/Engine/Source/Runtime/Engine/Classes/Engine/GameViewportClient.h b/Engine/Source/Runtime/Engine/Classes/Engine/GameViewportClient.h
--- a/Engine/Source/Runtime/Engine/Classes/Engine/GameViewportClient.h
+++ b/Engine/Source/Runtime/Engine/Classes/Engine/GameViewportClient.h
@@ -208,9 +208,14 @@ public:
 	{}
 
 	//~ Begin FExec Interface.
+#if UE_ALLOW_EXEC_COMMANDS
 	virtual bool Exec( UWorld* InWorld, const TCHAR* Cmd,FOutputDevice& Ar) override;
+#endif
+protected:
+	virtual bool Exec_Runtime( UWorld* InWorld, const TCHAR* Cmd,FOutputDevice& Ar) override;
 	//~ End FExec Interface.
 
+public:
 	/**
 	 * Gives the game's custom viewport client a way to handle F11 or Alt+Enter before processing the input
 	 */
diff --git a/Engine/Source/Runtime/Engine/Private/GameEngine.cpp b/Engine/Source/Runtime/Engine/Private/GameEngine.cpp
--- a/Engine/Source/Runtime/Engine/Private/GameEngine.cpp
+++ b/Engine/Source/Runtime/Engine/Private/GameEngine.cpp
@@ -1016,7 +1016,7 @@ public:
 		return GEngineIni;
 	}
 
-	virtual bool Exec(UWorld* Inworld, const TCHAR* Cmd, FOutputDevice& Ar) override
+	virtual bool Exec_Runtime(UWorld* Inworld, const TCHAR* Cmd, FOutputDevice& Ar) override
 	{
 		if (FParse::Command(&Cmd, TEXT("exitembedded")))
 		{
@@ -1378,7 +1378,51 @@ bool UGameEngine::ShouldDoAsyncEndOfFrameTasks() const
 	Command line executor.
 -----------------------------------------------------------------------------*/
 
+#if UE_ALLOW_EXEC_COMMANDS
 bool UGameEngine::Exec( UWorld* InWorld, const TCHAR* Cmd, FOutputDevice& Ar )
+{
+	if (FExec::Exec(InWorld, Cmd, Ar))
+	{
+		return true;
+	}
+	else if (InWorld && InWorld->Exec(InWorld, Cmd, Ar))
+	{
+		return true;
+	}
+	else if (InWorld && InWorld->GetAuthGameMode() && InWorld->GetAuthGameMode()->ProcessConsoleExec(Cmd, Ar, nullptr))
+	{
+		return true;
+	}
+	else
+	{
+#if UE_BUILD_SHIPPING
+		// disallow set of actor properties if network game
+		if ((FParse::Command(&Cmd, TEXT("SET")) || FParse::Command(&Cmd, TEXT("SETNOPEC"))))
+		{
+			FWorldContext &Context = GetWorldContextFromWorldChecked(InWorld);
+			if (Context.PendingNetGame != NULL || InWorld->GetNetMode() != NM_Standalone)
+			{
+				return true;
+			}
+			// the effects of this cannot be easily reversed, so prevent the user from playing network games without restarting to avoid potential exploits
+			GDisallowNetworkTravel = true;
+		}
+#endif // UE_BUILD_SHIPPING
+		if (UEngine::Exec(InWorld, Cmd, Ar))
+		{
+			return true;
+		}
+		else if (UPlatformInterfaceBase::StaticExec(Cmd, Ar))
+		{
+			return true;
+		}
+
+		return false;
+	}
+}
+#endif // UE_ALLOW_EXEC_COMMANDS
+
+bool UGameEngine::Exec_Runtime( UWorld* InWorld, const TCHAR* Cmd, FOutputDevice& Ar )
 {
 	if( FParse::Command( &Cmd,TEXT("REATTACHCOMPONENTS")) || FParse::Command( &Cmd,TEXT("REREGISTERCOMPONENTS")))
 	{
@@ -1472,14 +1516,26 @@ bool UGameEngine::Exec( UWorld* InWorld, const TCHAR* Cmd, FOutputDevice& Ar )
 
 		return true;
 	}
+
+	return false;
+}
+
+bool UGameEngine::Exec_Dev( UWorld* InWorld, const TCHAR* Cmd, FOutputDevice& Ar )
+{
 #if !UE_BUILD_SHIPPING
-	else if( FParse::Command( &Cmd, TEXT("ApplyUserSettings") ) )
+	if( FParse::Command( &Cmd, TEXT("ApplyUserSettings") ) )
 	{
-		return HandleApplyUserSettingsCommand( Cmd, Ar );
+		return HandleApplyUserSettingsCommand(Cmd, Ar);
 	}
 #endif // !UE_BUILD_SHIPPING
+
+	return false;
+}
+
+bool UGameEngine::Exec_Editor( UWorld* InWorld, const TCHAR* Cmd, FOutputDevice& Ar )
+{
 #if WITH_EDITOR
-	else if( FParse::Command(&Cmd,TEXT("STARTMOVIECAPTURE")) && GIsEditor )
+	if( FParse::Command(&Cmd,TEXT("STARTMOVIECAPTURE")) && GIsEditor )
 	{
 		IMovieSceneCaptureInterface* CaptureInterface = IMovieSceneCaptureModule::Get().GetFirstActiveMovieSceneCapture();
 		if (CaptureInterface)
@@ -1496,41 +1552,9 @@ bool UGameEngine::Exec( UWorld* InWorld, const TCHAR* Cmd, FOutputDevice& Ar )
 		}
 		return false;
 	}
-#endif
-	else if( InWorld && InWorld->Exec( InWorld, Cmd, Ar ) )
-	{
-		return true;
-	}
-	else if( InWorld && InWorld->GetAuthGameMode() && InWorld->GetAuthGameMode()->ProcessConsoleExec(Cmd,Ar,NULL) )
-	{
-		return true;
-	}
-	else
-	{
-#if UE_BUILD_SHIPPING
-		// disallow set of actor properties if network game
-		if ((FParse::Command( &Cmd, TEXT("SET")) || FParse::Command( &Cmd, TEXT("SETNOPEC"))))
-		{
-			FWorldContext &Context = GetWorldContextFromWorldChecked(InWorld);
-			if( Context.PendingNetGame != NULL || InWorld->GetNetMode() != NM_Standalone)
-			{
-				return true;
-			}
-			// the effects of this cannot be easily reversed, so prevent the user from playing network games without restarting to avoid potential exploits
-			GDisallowNetworkTravel = true;
-		}
-#endif // UE_BUILD_SHIPPING
-		if (UEngine::Exec(InWorld, Cmd, Ar))
-		{
-			return true;
-		}
-		else if (UPlatformInterfaceBase::StaticExec(Cmd, Ar))
-		{
-			return true;
-		}
-	
-		return false;
-	}
+#endif // WITH_EDITOR
+
+	return false;
 }
 
 bool UGameEngine::HandleExitCommand( const TCHAR* Cmd, FOutputDevice& Ar )
diff --git a/Engine/Source/Runtime/Engine/Private/GameViewportClient.cpp b/Engine/Source/Runtime/Engine/Private/GameViewportClient.cpp
--- a/Engine/Source/Runtime/Engine/Private/GameViewportClient.cpp
+++ b/Engine/Source/Runtime/Engine/Private/GameViewportClient.cpp
@@ -2954,7 +2954,33 @@ void UGameViewportClient::SetHideCursorDuringCapture(bool InHideCursorDuringCapt
 	}
 }
 
-bool UGameViewportClient::Exec( UWorld* InWorld, const TCHAR* Cmd,FOutputDevice& Ar)
+#if UE_ALLOW_EXEC_COMMANDS
+bool UGameViewportClient::Exec(UWorld* InWorld, const TCHAR* Cmd, FOutputDevice& Ar)
+{
+	if (FExec::Exec(InWorld, Cmd, Ar))
+	{
+		return true;
+	}
+	else if (ProcessConsoleExec(Cmd, Ar, NULL))
+	{
+		return true;
+	}
+	else if (GameInstance && (GameInstance->Exec(InWorld, Cmd, Ar) || GameInstance->ProcessConsoleExec(Cmd, Ar, nullptr)))
+	{
+		return true;
+	}
+	else if (GEngine->Exec(InWorld, Cmd, Ar))
+	{
+		return true;
+	}
+	else
+	{
+		return false;
+	}
+}
+#endif // UE_ALLOW_EXEC_COMMANDS
+
+bool UGameViewportClient::Exec_Runtime( UWorld* InWorld, const TCHAR* Cmd,FOutputDevice& Ar)
 {
 	if ( FParse::Command(&Cmd,TEXT("FORCEFULLSCREEN")) )
 	{
@@ -3060,19 +3086,6 @@ bool UGameViewportClient::Exec( UWorld* InWorld, const TCHAR* Cmd,FOutputDevice&
 	{
 		return HandlePauseRenderClockCommand( Cmd, Ar );
 	}
-
-	if(ProcessConsoleExec(Cmd,Ar,NULL))
-	{
-		return true;
-	}
-	else if ( GameInstance && (GameInstance->Exec(InWorld, Cmd, Ar) || GameInstance->ProcessConsoleExec(Cmd, Ar, nullptr)) )
-	{
-		return true;
-	}
-	else if( GEngine->Exec( InWorld, Cmd,Ar) )
-	{
-		return true;
-	}
 	else
 	{
 		return false;
