Fixed crash when unloading animation assets in a game feature plugin that are referenced by editor thumbnails

Unloading a game feature plugin attempts to unload all of the assets therein. It attempts to GC unreferenced RF_Standalone assets in editor, but if the assets are still reachable it will leave them alone. Hoever the game feature plugin system then marks the 'leaking' assets with MarkAsGarbage. This causes weak references to those assets to return null, even though hard references in higher level code think they are valid. In this case, the animation decompression code was trying to use a skeleton that was marked as garbage.

In this case, assets were sill being referenced by the tick record tracking in the Sync member of the preview instance's proxy. The fix is to clear out the sync on setting animations and uninitializing the proxy so that the animations dont get picked up as reachable for GC.

diff --git a/Engine/Source/Runtime/Engine/Private/Animation/AnimInstanceProxy.cpp b/Engine/Source/Runtime/Engine/Private/Animation/AnimInstanceProxy.cpp
--- a/Engine/Source/Runtime/Engine/Private/Animation/AnimInstanceProxy.cpp
+++ b/Engine/Source/Runtime/Engine/Private/Animation/AnimInstanceProxy.cpp
@@ -433,6 +433,7 @@ void FAnimInstanceProxy::Uninitialize(UAnimInstance* InAnimInstance)
 	ResetAnimationCurves();
 	MaterialParametersToClear.Reset();
 	ActiveAnimNotifiesSinceLastTick.Reset();
+	Sync.ResetAll();
 }
 
 void FAnimInstanceProxy::UpdateActiveAnimNotifiesSinceLastTick(const FAnimNotifyQueue& AnimInstanceQueue)
diff --git a/Engine/Source/Runtime/Engine/Private/Animation/AnimSingleNodeInstanceProxy.cpp b/Engine/Source/Runtime/Engine/Private/Animation/AnimSingleNodeInstanceProxy.cpp
--- a/Engine/Source/Runtime/Engine/Private/Animation/AnimSingleNodeInstanceProxy.cpp
+++ b/Engine/Source/Runtime/Engine/Private/Animation/AnimSingleNodeInstanceProxy.cpp
@@ -218,6 +218,7 @@ void FAnimSingleNodeInstanceProxy::SetAnimationAsset(class UAnimationAsset* NewA
 	BlendSpacePosition = FVector::ZeroVector;
 	BlendSampleData.Reset();
 	MarkerTickRecord.Reset();
+	ResetSync();
 	UpdateBlendspaceSamples(BlendSpacePosition);
 
 #if WITH_EDITORONLY_DATA
diff --git a/Engine/Source/Runtime/Engine/Private/Animation/AnimSync.cpp b/Engine/Source/Runtime/Engine/Private/Animation/AnimSync.cpp
--- a/Engine/Source/Runtime/Engine/Private/Animation/AnimSync.cpp
+++ b/Engine/Source/Runtime/Engine/Private/Animation/AnimSync.cpp
@@ -25,6 +25,25 @@ void FAnimSync::Reset()
 	}
 }
 
+void FAnimSync::ResetAll()
+{
+	for(TArray<FAnimTickRecord>& UngroupedActivePlayers : UngroupedActivePlayerArrays)
+	{
+		UngroupedActivePlayers.Reset();
+	}
+
+	for(FSyncGroupMap& SyncGroupMap : SyncGroupMaps)
+	{
+		for (auto& SyncGroupPair : SyncGroupMap)
+		{
+			SyncGroupPair.Value.Reset();
+		}
+	}
+
+	MirrorDataTable = nullptr;
+	SyncGroupWriteIndex = 0;
+}
+
 void FAnimSync::AddTickRecord(const FAnimTickRecord& InTickRecord, const FAnimSyncParams& InSyncParams)
 {
 	if (InSyncParams.GroupName != NAME_None)
diff --git a/Engine/Source/Runtime/Engine/Public/Animation/AnimInstanceProxy.h b/Engine/Source/Runtime/Engine/Public/Animation/AnimInstanceProxy.h
--- a/Engine/Source/Runtime/Engine/Public/Animation/AnimInstanceProxy.h
+++ b/Engine/Source/Runtime/Engine/Public/Animation/AnimInstanceProxy.h
@@ -946,6 +946,12 @@ protected:
 	static void EvaluateInputProxy(FAnimInstanceProxy* InputProxy, FPoseContext& Output);
 	static void ResetCounterInputProxy(FAnimInstanceProxy* InputProxy);
 
+	/** Reset the sync system to its starting state */
+	void ResetSync()
+	{
+		Sync.ResetAll();
+	}
+
 private:
 
 	/** Executes the provided functor on each valid state machine on this anim instance */
diff --git a/Engine/Source/Runtime/Engine/Public/Animation/AnimSync.h b/Engine/Source/Runtime/Engine/Public/Animation/AnimSync.h
--- a/Engine/Source/Runtime/Engine/Public/Animation/AnimSync.h
+++ b/Engine/Source/Runtime/Engine/Public/Animation/AnimSync.h
@@ -37,6 +37,9 @@ struct ENGINE_API FAnimSync
 	// Reset internal buffers ready for writing
 	void Reset();
 
+	// Completely reset the sync system back to its starting state
+	void ResetAll();
+	
 	// Adds a tick record in the list for the correct group or the ungrouped array.
 	void AddTickRecord(const FAnimTickRecord& InTickRecord, const FAnimSyncParams& InSyncParams = FAnimSyncParams());
 
