Fixed Lyra crash in MaterialCacheData.

diff --git a/Engine/Source/Runtime/Engine/Private/Materials/MaterialCachedData.cpp b/Engine/Source/Runtime/Engine/Private/Materials/MaterialCachedData.cpp
--- a/Engine/Source/Runtime/Engine/Private/Materials/MaterialCachedData.cpp
+++ b/Engine/Source/Runtime/Engine/Private/Materials/MaterialCachedData.cpp
@@ -772,7 +772,7 @@ void FMaterialCachedExpressionData::AnalyzeMaterial(UMaterial& Material)
 	// This occurs when the user opens a MaterialFunction in the editor.
 	for (UMaterialExpression* Expression : Material.GetExpressions())
 	{
-		if (Expression->IsA<UMaterialExpressionFunctionOutput>())
+		if (Expression && Expression->IsA<UMaterialExpressionFunctionOutput>())
 		{
 			Analyzer.PushUnexploredExpression(Expression);
 		}
