Turned improvements to the MaterialCachedData behind a CVAR and made them experimental until all corner cases are solved.

diff --git a/Engine/Source/Runtime/Engine/Private/Materials/MaterialCachedData.cpp b/Engine/Source/Runtime/Engine/Private/Materials/MaterialCachedData.cpp
--- a/Engine/Source/Runtime/Engine/Private/Materials/MaterialCachedData.cpp
+++ b/Engine/Source/Runtime/Engine/Private/Materials/MaterialCachedData.cpp
@@ -45,6 +45,9 @@ const FMaterialCachedExpressionEditorOnlyData FMaterialCachedExpressionEditorOnl
 
 static_assert((uint32)(EMaterialProperty::MP_MaterialAttributes)-1 <= (8 * sizeof(FMaterialCachedExpressionData::PropertyConnectedBitmask)), "PropertyConnectedBitmask cannot contain entire EMaterialProperty enumeration.");
 
+static bool GExperimentalMaterialCachedDataAnalysisEnabled = false;
+static FAutoConsoleVariableRef CVarExperimentalMaterialCachedDataAnalysisEnabled(TEXT("r.Material.ExperimentalMaterialCachedDataAnalysisEnabled"), GExperimentalMaterialCachedDataAnalysisEnabled, TEXT("Enables material cached data experimental graph based analysis"));
+
 FMaterialCachedExpressionData::FMaterialCachedExpressionData()
 	: FunctionInfosStateCRC(0xffffffff)
 	, bHasMaterialLayers(false)
@@ -365,7 +368,76 @@ void FMaterialCachedExpressionData::UpdateForExpressions(const FMaterialCachedEx
 
 		Expression->GetLandscapeLayerNames(EditorOnlyData->LandscapeLayerNames);
 
-		if (UMaterialExpressionMaterialAttributeLayers* LayersExpression = Cast<UMaterialExpressionMaterialAttributeLayers>(Expression))
+		if (UMaterialExpressionCollectionParameter* ExpressionCollectionParameter = Cast<UMaterialExpressionCollectionParameter>(Expression))
+		{
+			UMaterialParameterCollection* Collection = ExpressionCollectionParameter->Collection;
+			if (Collection)
+			{
+				FMaterialParameterCollectionInfo NewInfo;
+				NewInfo.ParameterCollection = Collection;
+				NewInfo.StateId = Collection->StateId;
+				ParameterCollectionInfos.AddUnique(NewInfo);
+			}
+		}
+		else if (UMaterialExpressionDynamicParameter* ExpressionDynamicParameter = Cast< UMaterialExpressionDynamicParameter>(Expression))
+		{
+			DynamicParameterNames.Empty(ExpressionDynamicParameter->ParamNames.Num());
+			for (const FString& Name : ExpressionDynamicParameter->ParamNames)
+			{
+				DynamicParameterNames.Add(*Name);
+			}
+		}
+		else if (UMaterialExpressionLandscapeGrassOutput* ExpressionGrassOutput = Cast<UMaterialExpressionLandscapeGrassOutput>(Expression))
+		{
+			for (const auto& Type : ExpressionGrassOutput->GrassTypes)
+			{
+				GrassTypes.AddUnique(Type.GrassType);
+			}
+		}
+		else if (UMaterialExpressionQualitySwitch* QualitySwitchNode = Cast<UMaterialExpressionQualitySwitch>(Expression))
+		{
+			const FExpressionInput DefaultInput = QualitySwitchNode->Default.GetTracedInput();
+
+			for (int32 InputIndex = 0; InputIndex < EMaterialQualityLevel::Num; InputIndex++)
+			{
+				if (QualitySwitchNode->Inputs[InputIndex].IsConnected())
+				{
+					// We can ignore quality levels that are defined the same way as 'Default'
+					// This avoids compiling a separate explicit quality level resource, that will end up exactly the same as the default resource
+					const FExpressionInput Input = QualitySwitchNode->Inputs[InputIndex].GetTracedInput();
+					if (Input.Expression != DefaultInput.Expression ||
+						Input.OutputIndex != DefaultInput.OutputIndex)
+					{
+						QualityLevelsUsed[InputIndex] = true;
+					}
+				}
+			}
+		}
+		else if (Expression->IsA(UMaterialExpressionRuntimeVirtualTextureOutput::StaticClass()))
+		{
+			bHasRuntimeVirtualTextureOutput = true;
+		}
+		else if (Expression->IsA(UMaterialExpressionSceneColor::StaticClass()))
+		{
+			bHasSceneColor = true;
+		}
+		else if (Expression->IsA(UMaterialExpressionPerInstanceRandom::StaticClass()))
+		{
+			bHasPerInstanceRandom = true;
+		}
+		else if (Expression->IsA(UMaterialExpressionPerInstanceCustomData::StaticClass()))
+		{
+			bHasPerInstanceCustomData = true;
+		}
+		else if (Expression->IsA(UMaterialExpressionPerInstanceCustomData3Vector::StaticClass()))
+		{
+			bHasPerInstanceCustomData = true;
+		}
+		else if (Expression->IsA(UMaterialExpressionVertexInterpolator::StaticClass()))
+		{
+			bHasVertexInterpolator = true;
+		}
+		else if (UMaterialExpressionMaterialAttributeLayers* LayersExpression = Cast<UMaterialExpressionMaterialAttributeLayers>(Expression))
 		{
 			checkf(Association == GlobalParameter, TEXT("UMaterialExpressionMaterialAttributeLayers can't be nested"));
 			// Only a single layers expression is allowed/expected...creating additional layer expression will cause a compile error
@@ -392,7 +464,58 @@ void FMaterialCachedExpressionData::UpdateForExpressions(const FMaterialCachedEx
 				// Update even if MaterialFunctionNode->MaterialFunction is NULL, because we need to remove the invalid inputs in that case
 				FunctionCall->UpdateFromFunctionResource();
 			}
-		} 
+		}
+		else if (UMaterialExpressionSetMaterialAttributes* SetMatAttributes = Cast<UMaterialExpressionSetMaterialAttributes>(Expression))
+		{
+			for (int32 PinIndex = 0; PinIndex < SetMatAttributes->AttributeSetTypes.Num(); ++PinIndex)
+			{
+				// For this material attribute pin do we have something connected?
+				const FGuid& Guid = SetMatAttributes->AttributeSetTypes[PinIndex];
+				const FExpressionInput& AttributeInput = SetMatAttributes->Inputs[PinIndex + 1];
+				const EMaterialProperty MaterialProperty = FMaterialAttributeDefinitionMap::GetProperty(Guid);
+				if (AttributeInput.Expression)
+				{
+					SetPropertyConnected(MaterialProperty);
+				}
+			}
+		}
+		else if (UMaterialExpressionMakeMaterialAttributes* MakeMatAttributes = Cast<UMaterialExpressionMakeMaterialAttributes>(Expression))
+		{
+			auto SetMatAttributeConditionally = [&](EMaterialProperty InMaterialProperty, bool InIsConnected)
+			{
+				if (InIsConnected)
+				{
+					SetPropertyConnected(InMaterialProperty);
+				}
+			};
+
+			SetMatAttributeConditionally(EMaterialProperty::MP_BaseColor, MakeMatAttributes->BaseColor.IsConnected());
+			SetMatAttributeConditionally(EMaterialProperty::MP_Metallic, MakeMatAttributes->Metallic.IsConnected());
+			SetMatAttributeConditionally(EMaterialProperty::MP_Specular, MakeMatAttributes->Specular.IsConnected());
+			SetMatAttributeConditionally(EMaterialProperty::MP_Roughness, MakeMatAttributes->Roughness.IsConnected());
+			SetMatAttributeConditionally(EMaterialProperty::MP_Anisotropy, MakeMatAttributes->Anisotropy.IsConnected());
+			SetMatAttributeConditionally(EMaterialProperty::MP_EmissiveColor, MakeMatAttributes->EmissiveColor.IsConnected());
+			SetMatAttributeConditionally(EMaterialProperty::MP_Opacity, MakeMatAttributes->Opacity.IsConnected());
+			SetMatAttributeConditionally(EMaterialProperty::MP_OpacityMask, MakeMatAttributes->OpacityMask.IsConnected());
+			SetMatAttributeConditionally(EMaterialProperty::MP_Normal, MakeMatAttributes->Normal.IsConnected());
+			SetMatAttributeConditionally(EMaterialProperty::MP_Tangent, MakeMatAttributes->Tangent.IsConnected());
+			SetMatAttributeConditionally(EMaterialProperty::MP_WorldPositionOffset, MakeMatAttributes->WorldPositionOffset.IsConnected());
+			SetMatAttributeConditionally(EMaterialProperty::MP_SubsurfaceColor, MakeMatAttributes->SubsurfaceColor.IsConnected());
+			SetMatAttributeConditionally(EMaterialProperty::MP_CustomData0, MakeMatAttributes->ClearCoat.IsConnected());
+			SetMatAttributeConditionally(EMaterialProperty::MP_CustomData1, MakeMatAttributes->ClearCoatRoughness.IsConnected());
+			SetMatAttributeConditionally(EMaterialProperty::MP_AmbientOcclusion, MakeMatAttributes->AmbientOcclusion.IsConnected());
+			SetMatAttributeConditionally(EMaterialProperty::MP_Refraction, MakeMatAttributes->Refraction.IsConnected());
+			SetMatAttributeConditionally(EMaterialProperty::MP_CustomizedUVs0, MakeMatAttributes->CustomizedUVs[0].IsConnected());
+			SetMatAttributeConditionally(EMaterialProperty::MP_CustomizedUVs1, MakeMatAttributes->CustomizedUVs[1].IsConnected());
+			SetMatAttributeConditionally(EMaterialProperty::MP_CustomizedUVs2, MakeMatAttributes->CustomizedUVs[2].IsConnected());
+			SetMatAttributeConditionally(EMaterialProperty::MP_CustomizedUVs3, MakeMatAttributes->CustomizedUVs[3].IsConnected());
+			SetMatAttributeConditionally(EMaterialProperty::MP_CustomizedUVs4, MakeMatAttributes->CustomizedUVs[4].IsConnected());
+			SetMatAttributeConditionally(EMaterialProperty::MP_CustomizedUVs5, MakeMatAttributes->CustomizedUVs[5].IsConnected());
+			SetMatAttributeConditionally(EMaterialProperty::MP_CustomizedUVs6, MakeMatAttributes->CustomizedUVs[6].IsConnected());
+			SetMatAttributeConditionally(EMaterialProperty::MP_CustomizedUVs7, MakeMatAttributes->CustomizedUVs[7].IsConnected());
+			SetMatAttributeConditionally(EMaterialProperty::MP_PixelDepthOffset, MakeMatAttributes->PixelDepthOffset.IsConnected());
+			SetMatAttributeConditionally(EMaterialProperty::MP_ShadingModel, MakeMatAttributes->ShadingModel.IsConnected());
+		}
 	}
 }
 
@@ -735,50 +858,64 @@ struct FMaterialConnectedPropertiesAnalyzer
 
 void FMaterialCachedExpressionData::AnalyzeMaterial(UMaterial& Material)
 {
-	// First run a pre-pass to determine referenced textures and fix-up material layers and functions, add parameters, etc.
-	// This is executed for all expressions in the material.
+	if (!GExperimentalMaterialCachedDataAnalysisEnabled)
 	{
-		FMaterialCachedExpressionContext Context;
-		UpdateForExpressions(Context, Material.GetExpressions(), EMaterialParameterAssociation::GlobalParameter, -1);
-	}
-
-	// Now crawl the graph from outputs to inputs skipping unused subgraphs ("compiled out" by static switches).
-	FMaterialConnectedPropertiesAnalyzer Analyzer{ *this };
-
-	// Add material outputs depending on whether material uses material attributes.
-	if (Material.bUseMaterialAttributes)
-	{
-		const FExpressionInput* Input = Material.GetExpressionInputForProperty(MP_MaterialAttributes);
-		if (Input && Input->IsConnected())
+		if (!Material.bUseMaterialAttributes)
 		{
-			Analyzer.PushUnexploredExpression(Input->Expression);
+			for (int32 PropertyIndex = 0; PropertyIndex < MP_MAX; ++PropertyIndex)
+			{
+				const EMaterialProperty Property = (EMaterialProperty)PropertyIndex;
+				const FExpressionInput* Input = Material.GetExpressionInputForProperty(Property);
+				if (Input && Input->IsConnected())
+				{
+					SetPropertyConnected(Property);
+				}
+			}
 		}
+
+		FMaterialCachedExpressionContext Context;
+		UpdateForExpressions(Context, Material.GetExpressions(), EMaterialParameterAssociation::GlobalParameter, -1);
 	}
 	else
 	{
-		for (int32 PropertyIndex = 0; PropertyIndex < MP_MAX; ++PropertyIndex)
+		// Now crawl the graph from outputs to inputs skipping unused subgraphs ("compiled out" by static switches).
+		FMaterialConnectedPropertiesAnalyzer Analyzer{ *this };
+
+		// Add material outputs depending on whether material uses material attributes.
+		if (Material.bUseMaterialAttributes)
 		{
-			const EMaterialProperty Property = (EMaterialProperty)PropertyIndex;
-			const FExpressionInput* Input = Material.GetExpressionInputForProperty(Property);
+			const FExpressionInput* Input = Material.GetExpressionInputForProperty(MP_MaterialAttributes);
 			if (Input && Input->IsConnected())
 			{
 				Analyzer.PushUnexploredExpression(Input->Expression);
-				SetPropertyConnected(Property);
 			}
 		}
-	}
+		else
+		{
+			for (int32 PropertyIndex = 0; PropertyIndex < MP_MAX; ++PropertyIndex)
+			{
+				const EMaterialProperty Property = (EMaterialProperty)PropertyIndex;
+				const FExpressionInput* Input = Material.GetExpressionInputForProperty(Property);
+				if (Input && Input->IsConnected())
+				{
+					Analyzer.PushUnexploredExpression(Input->Expression);
+					SetPropertyConnected(Property);
+				}
+			}
+		}
 	
-	// If there are any connected function output expressions, mark them as unexplored expressions.
-	// This occurs when the user opens a MaterialFunction in the editor.
-	for (UMaterialExpression* Expression : Material.GetExpressions())
-	{
-		if (Expression && Expression->IsA<UMaterialExpressionFunctionOutput>())
+		// If there are any connected function output expressions, mark them as unexplored expressions.
+		// This occurs when the user opens a MaterialFunction in the editor.
+		for (UMaterialExpression* Expression : Material.GetExpressions())
 		{
-			Analyzer.PushUnexploredExpression(Expression);
+			if (Expression && Expression->IsA<UMaterialExpressionFunctionOutput>())
+			{
+				Analyzer.PushUnexploredExpression(Expression);
+			}
 		}
-	}
 
-	Analyzer.VisitUnexploredExpressions(GlobalParameter, -1);
+		Analyzer.VisitUnexploredExpressions(GlobalParameter, -1);
+	}
 }
 
 
