World Partition: misc code cleanup.

diff --git a/Engine/Source/Runtime/Engine/Private/WorldPartition/ActorDescContainer.cpp b/Engine/Source/Runtime/Engine/Private/WorldPartition/ActorDescContainer.cpp
--- a/Engine/Source/Runtime/Engine/Private/WorldPartition/ActorDescContainer.cpp
+++ b/Engine/Source/Runtime/Engine/Private/WorldPartition/ActorDescContainer.cpp
@@ -25,6 +25,7 @@ UActorDescContainer::UActorDescContainer(const FObjectInitializer& ObjectInitial
 #endif
 {}
 
+#if WITH_EDITOR
 void UActorDescContainer::Initialize(UWorld* InWorld, FName InPackageName)
 {
 	Initialize({ InWorld, InPackageName });
@@ -39,7 +40,6 @@ void UActorDescContainer::Initialize(const FInitializeParams& InitParams)
 	// With this, we could always find the owning WorldPartition and of course the owning world (GetOwningWorldPartition()->GetWorld()).
 	UWorld* OwningWorld = InitParams.World;
 
-#if WITH_EDITOR
 	check(!bContainerInitialized);
 	ContainerPackageName = InitParams.PackageName;
 	TArray<FAssetData> Assets;
@@ -89,12 +89,10 @@ void UActorDescContainer::Initialize(const FInitializeParams& InitParams)
 	RegisterEditorDelegates();
 
 	bContainerInitialized = true;
-#endif
 }
 
 void UActorDescContainer::Update()
 {
-#if WITH_EDITOR
 	check(bContainerInitialized);
 	TArray<FAssetData> Assets;
 
@@ -157,12 +155,10 @@ void UActorDescContainer::Update()
 			RemoveActor(ActorDescGuidToRemove);
 		}
 	}
-#endif
 }
 
 void UActorDescContainer::Uninitialize()
 {
-#if WITH_EDITOR
 	if (bContainerInitialized)
 	{
 		UnregisterEditorDelegates();
@@ -177,7 +173,6 @@ void UActorDescContainer::Uninitialize()
 		}
 		ActorDescPtr.Reset();
 	}
-#endif
 }
 
 void UActorDescContainer::BeginDestroy()
@@ -187,7 +182,6 @@ void UActorDescContainer::BeginDestroy()
 	Uninitialize();
 }
 
-#if WITH_EDITOR
 FString UActorDescContainer::GetExternalActorPath() const
 {
 	return ULevel::GetExternalActorsPath(ContainerPackageName.ToString());
diff --git a/Engine/Source/Runtime/Engine/Private/WorldPartition/WorldPartitionEditorHash.cpp b/Engine/Source/Runtime/Engine/Private/WorldPartition/WorldPartitionEditorHash.cpp
--- a/Engine/Source/Runtime/Engine/Private/WorldPartition/WorldPartitionEditorHash.cpp
+++ b/Engine/Source/Runtime/Engine/Private/WorldPartition/WorldPartitionEditorHash.cpp
@@ -1,7 +1,6 @@
 // Copyright Epic Games, Inc. All Rights Reserved.
 
 #include "WorldPartition/WorldPartitionEditorHash.h"
-#include "WorldPartition/IWorldPartitionEditorModule.h"
 
 #if WITH_EDITOR
 #include "Modules/ModuleManager.h"
@@ -11,14 +10,7 @@
 
 UWorldPartitionEditorHash::UWorldPartitionEditorHash(const FObjectInitializer& ObjectInitializer)
 	: Super(ObjectInitializer)
-{
-#if WITH_EDITOR
-	if (!HasAnyFlags(RF_ClassDefaultObject))
-	{
-		WorldPartitionEditorModule = &FModuleManager::LoadModuleChecked<IWorldPartitionEditorModule>("WorldPartitionEditor");
-	}
-#endif
-}
+{}
 
 #if WITH_EDITOR
 UWorldPartitionEditorHash::FForEachIntersectingActorParams::FForEachIntersectingActorParams()
diff --git a/Engine/Source/Runtime/Engine/Public/WorldPartition/ActorDescContainer.h b/Engine/Source/Runtime/Engine/Public/WorldPartition/ActorDescContainer.h
--- a/Engine/Source/Runtime/Engine/Public/WorldPartition/ActorDescContainer.h
+++ b/Engine/Source/Runtime/Engine/Public/WorldPartition/ActorDescContainer.h
@@ -18,6 +18,7 @@ class ENGINE_API UActorDescContainer : public UObject, public FActorDescList
 {
 	GENERATED_UCLASS_BODY()
 
+#if WITH_EDITOR
 	friend struct FWorldPartitionHandleUtils;
 	friend class FWorldPartitionActorDesc;
 
@@ -36,10 +37,8 @@ public:
 		/* The long package name of the container package on disk. */
 		FName PackageName;
 
-#if WITH_EDITOR
 		/* Custom filter function used to filter actors descriptors. */
 		TUniqueFunction<bool(const FWorldPartitionActorDesc*)> FilterActorDesc;
-#endif
 	};
 
 	UE_DEPRECATED(5.1, "UActorDescContainer::Initialize is deprecated, UActorDescContainer::Initialize with UActorDescContainer::FInitializeParams should be used instead.")
@@ -48,7 +47,6 @@ public:
 	void Update();
 	void Uninitialize();
 
-#if WITH_EDITOR
 	bool IsInitialized() const { return bContainerInitialized; }
 
 	void OnObjectPreSave(UObject* Object, FObjectPreSaveContext SaveContext);
@@ -108,7 +106,6 @@ public:
 	FGuid ContentBundleGuid;
 
 	TArray<TUniquePtr<FWorldPartitionActorDesc>> InvalidActors;
-#endif
 
 protected:
 	//~ Begin UObject Interface
@@ -119,7 +116,6 @@ private:
 	// GetWorld() should never be called on an ActorDescContainer to avoid any confusion as it can be used as a template
 	virtual UWorld* GetWorld() const override { return nullptr; }
 
-#if WITH_EDITOR
 	bool ShouldRegisterDelegates();
 	void RegisterEditorDelegates();
 	void UnregisterEditorDelegates();
diff --git a/Engine/Source/Runtime/Engine/Public/WorldPartition/WorldPartitionEditorHash.h b/Engine/Source/Runtime/Engine/Public/WorldPartition/WorldPartitionEditorHash.h
--- a/Engine/Source/Runtime/Engine/Public/WorldPartition/WorldPartitionEditorHash.h
+++ b/Engine/Source/Runtime/Engine/Public/WorldPartition/WorldPartitionEditorHash.h
@@ -8,8 +8,6 @@
 #include "WorldPartition/WorldPartitionHandle.h"
 #include "WorldPartitionEditorHash.generated.h"
 
-class IWorldPartitionEditorModule;
-
 UCLASS(Abstract, Config=Engine, Within = WorldPartition)
 class ENGINE_API UWorldPartitionEditorHash : public UObject
 {
@@ -47,8 +45,5 @@ public:
 	};
 
 	virtual int32 ForEachIntersectingActor(const FBox& Box, TFunctionRef<void(FWorldPartitionActorDesc*)> InOperation, const FForEachIntersectingActorParams& Params = FForEachIntersectingActorParams()) PURE_VIRTUAL(UWorldPartitionEditorHash::ForEachIntersectingActor, return 0;);
-
-protected:
-	IWorldPartitionEditorModule* WorldPartitionEditorModule;
 #endif
 };
