added LLM_SCOPE for the ShaderCompiler to tag the allocations

diff --git a/Engine/Source/Runtime/Engine/Private/ShaderCompiler/ShaderCompiler.cpp b/Engine/Source/Runtime/Engine/Private/ShaderCompiler/ShaderCompiler.cpp
--- a/Engine/Source/Runtime/Engine/Private/ShaderCompiler/ShaderCompiler.cpp
+++ b/Engine/Source/Runtime/Engine/Private/ShaderCompiler/ShaderCompiler.cpp
@@ -76,6 +76,8 @@ DEFINE_LOG_CATEGORY(LogShaderCompilers);
 // whether to parallelize writing/reading task files
 #define UE_SHADERCOMPILER_FIFO_JOB_EXECUTION  1
 
+LLM_DEFINE_TAG(ShaderCompiler);
+
 int32 GShaderCompilerJobCache = 1;
 static FAutoConsoleVariableRef CVarShaderCompilerJobCache(
 	TEXT("r.ShaderCompiler.JobCache"),
@@ -2351,6 +2353,7 @@ void FShaderCompileThreadRunnable::OnMachineResourcesChanged()
 /** Entry point for the shader compiling thread. */
 uint32 FShaderCompileThreadRunnableBase::Run()
 {
+	LLM_SCOPE_BYTAG(ShaderCompiler);
 	check(Manager->bAllowAsynchronousShaderCompiling);
 	while (!bForceFinish)
 	{
