Bug fix for hiding objects in the editor, where the objects still appear in secondary forms of rendering like Lumen and Ray Tracing.  Batch render state update logic was incorrectly recreating scene proxies for the hidden components.  Fix is to check if components are hidden by calling "ShouldComponentAddToScene" before adding them to the batch update list.

#rnx

diff --git a/Engine/Source/Runtime/Engine/Private/StaticMeshResources.cpp b/Engine/Source/Runtime/Engine/Private/StaticMeshResources.cpp
--- a/Engine/Source/Runtime/Engine/Private/StaticMeshResources.cpp
+++ b/Engine/Source/Runtime/Engine/Private/StaticMeshResources.cpp
@@ -120,8 +120,9 @@ FStaticMeshComponentBulkReregisterContext::FStaticMeshComponentBulkReregisterCon
 	{
 		UStaticMeshComponent* StaticMeshComponent = Cast<UStaticMeshComponent>(Component);
 
-		// When running a render state update, only process components that are dirty
-		if (StaticMeshComponent && StaticMeshComponent->SceneProxy && (ReregisterType == EBulkReregister::Component || StaticMeshComponent->IsRenderStateDirty()))
+		// When running a render state update, only batch process components that are dirty and not hidden
+		if (StaticMeshComponent && StaticMeshComponent->SceneProxy &&
+			(ReregisterType == EBulkReregister::Component || (StaticMeshComponent->IsRenderStateDirty() && StaticMeshComponent->ShouldComponentAddToScene())))
 		{
 			StaticMeshComponents.Add(StaticMeshComponent);
 
