World Partition: unify streaming object creation code between different hash implementations.

diff --git a/Engine/Source/Runtime/Engine/Private/WorldPartition/WorldPartitionRuntimeHash.cpp b/Engine/Source/Runtime/Engine/Private/WorldPartition/WorldPartitionRuntimeHash.cpp
--- a/Engine/Source/Runtime/Engine/Private/WorldPartition/WorldPartitionRuntimeHash.cpp
+++ b/Engine/Source/Runtime/Engine/Private/WorldPartition/WorldPartitionRuntimeHash.cpp
@@ -10,6 +10,7 @@
 #include "WorldPartition/WorldPartitionRuntimeLevelStreamingCell.h"
 #include "WorldPartition/WorldPartitionLevelStreamingDynamic.h"
 #if WITH_EDITOR
+#include "Misc/ArchiveMD5.h"
 #include "WorldPartition/Cook/WorldPartitionCookPackage.h"
 #endif
 
@@ -95,6 +96,32 @@ void UWorldPartitionRuntimeHash::OnEndPlay()
 	ModifiedActorDescListForPIE.Empty();
 }
 
+UWorldPartitionRuntimeCell* UWorldPartitionRuntimeHash::CreateRuntimeCell(UClass* CellClass, UClass* CellDataClass, const FString& CellName)
+{
+	//@todo_ow: to reduce file paths on windows, we compute a MD5 hash from the unique cell name and use that hash as the cell filename. We
+	//			should create the runtime cell using the name but make sure the package that gets associated with it when cooking gets the
+	//			hash indtead.
+	auto GetCellObjectName = [](FString CellName) -> FString
+	{
+		FArchiveMD5 ArMD5;
+		ArMD5 << CellName;
+
+		FMD5Hash MD5Hash;
+		ArMD5.GetHash(MD5Hash);
+
+		FGuid CellNameGuid = MD5HashToGuid(MD5Hash);
+		check(CellNameGuid.IsValid());
+
+		return CellNameGuid.ToString(EGuidFormats::Base36Encoded);
+	};
+
+	const FString CellObjectName = GetCellObjectName(CellName);
+	check(!FindObject<UWorldPartitionRuntimeCell>(this, *CellObjectName));
+	UWorldPartitionRuntimeCell* RuntimeCell = NewObject<UWorldPartitionRuntimeCell>(this, CellClass, *CellObjectName);
+	RuntimeCell->RuntimeCellData = NewObject<UWorldPartitionRuntimeCellData>(RuntimeCell, CellDataClass);
+	return RuntimeCell;
+}
+
 bool UWorldPartitionRuntimeHash::GenerateStreaming(class UWorldPartitionStreamingPolicy* StreamingPolicy, const IStreamingGenerationContext* StreamingGenerationContext, TArray<FString>* OutPackagesToGenerate)
 {
 	return PackagesToGenerateForCook.IsEmpty();
@@ -143,6 +170,36 @@ bool UWorldPartitionRuntimeHash::ConditionalRegisterAlwaysLoadedActorsForPIE(con
 	return false;
 }
 
+bool UWorldPartitionRuntimeHash::PopulateCellActorInstances(const TArray<const IStreamingGenerationContext::FActorSetInstance*> ActorSetInstances, bool bIsMainWorldPartition, bool bIsCellAlwaysLoaded, TArray<IStreamingGenerationContext::FActorInstance>& OutCellActorInstances)
+{
+	for (const IStreamingGenerationContext::FActorSetInstance* ActorSetInstance : ActorSetInstances)
+	{
+		// Instanced world partition, ContainerID is the main container, but it's not the main world partition,
+		// so the always loaded actors don't be part of the process of ForceExternalActorLevelReference/AlwaysLoadedActorsForPIE.
+		// In PIE, always loaded actors of an instanced world partition will go in the always loaded cell.
+		if (!ConditionalRegisterAlwaysLoadedActorsForPIE(ActorSetInstance, bIsMainWorldPartition, ActorSetInstance->ContainerID.IsMainContainer(), bIsCellAlwaysLoaded))
+		{
+			for (const FGuid& ActorGuid : ActorSetInstance->ActorSet->Actors)
+			{
+				// Actors that return true to ShouldLevelKeepRefIfExternal will always be part of the partitioned persistent level of a
+				// world partition. In PIE, for an instanced world partition, we don't want this actor to be both in the persistent level
+				// and also part of the always loaded cell level.
+				//
+				// @todo_ow: We need to implement PIE always loaded actors of instanced world partitions to be part
+				//			 of the persistent level and get rid of the always loaded cell (to have the same behavior
+				//			 as non-instanced world partition and as cooked world partition).
+				IStreamingGenerationContext::FActorInstance ActorInstance(ActorGuid, ActorSetInstance);
+				if (!CastChecked<AActor>(ActorInstance.GetActorDescView().GetActorNativeClass()->GetDefaultObject())->ShouldLevelKeepRefIfExternal())
+				{
+					OutCellActorInstances.Emplace(ActorInstance);
+				}
+			}
+		}
+	}
+
+	return OutCellActorInstances.Num() > 0;
+}
+
 void UWorldPartitionRuntimeHash::PopulateRuntimeCell(UWorldPartitionRuntimeCell* RuntimeCell, const TArray<IStreamingGenerationContext::FActorInstance>& ActorInstances, TArray<FString>* OutPackagesToGenerate)
 {
 	for (const IStreamingGenerationContext::FActorInstance& ActorInstance : ActorInstances)
diff --git a/Engine/Source/Runtime/Engine/Private/WorldPartition/WorldPartitionRuntimeSpatialHash.cpp b/Engine/Source/Runtime/Engine/Private/WorldPartition/WorldPartitionRuntimeSpatialHash.cpp
--- a/Engine/Source/Runtime/Engine/Private/WorldPartition/WorldPartitionRuntimeSpatialHash.cpp
+++ b/Engine/Source/Runtime/Engine/Private/WorldPartition/WorldPartitionRuntimeSpatialHash.cpp
@@ -1325,101 +1325,53 @@ bool UWorldPartitionRuntimeSpatialHash::CreateStreamingGrid(const FSpatialHashRu
 				// Cell cannot be treated as always loaded if it has data layers
 				const bool bIsCellAlwaysLoaded = (&TempCell == &PartionedActors.GetAlwaysLoadedCell()) && !GridCellDataChunk.HasDataLayers() && !GridCellDataChunk.GetContentBundleID().IsValid();
 				
-				const TArray<const IStreamingGenerationContext::FActorSetInstance*> ActorSetInstances = GridCellDataChunk.GetActorSetInstances();
-
 				TArray<IStreamingGenerationContext::FActorInstance> FilteredActors;
-				for (const IStreamingGenerationContext::FActorSetInstance* ActorSetInstance : ActorSetInstances)
+				if (PopulateCellActorInstances(GridCellDataChunk.GetActorSetInstances(), bIsMainWorldPartition, bIsCellAlwaysLoaded, FilteredActors))
 				{
-					if (!ConditionalRegisterAlwaysLoadedActorsForPIE(ActorSetInstance, bIsMainWorldPartition, ActorSetInstance->ContainerID.IsMainContainer(), bIsCellAlwaysLoaded))
+					FGridCellCoord CellGlobalCoords;
+					verify(PartionedActors.GetCellGlobalCoords(FGridCellCoord(CellCoordX, CellCoordY, Level), CellGlobalCoords));
+					const FString CellName = GetCellNameString(OuterWorld->GetPackage()->GetFName(), CurrentStreamingGrid.GridName, CellGlobalCoords, GridCellDataChunk.GetDataLayersID(), GridCellDataChunk.GetContentBundleID());
+					const FGuid CellGuid = GetCellGuid(CurrentStreamingGrid.GridName, CellGlobalCoords, GridCellDataChunk.GetDataLayersID(), GridCellDataChunk.GetContentBundleID());
+
+					UWorldPartitionRuntimeCell* StreamingCell = CreateRuntimeCell(StreamingPolicy->GetRuntimeCellClass(), UWorldPartitionRuntimeCellDataSpatialHash::StaticClass(), CellName);
+					UWorldPartitionRuntimeCellDataSpatialHash* CellDataSpatialHash = CastChecked<UWorldPartitionRuntimeCellDataSpatialHash>(StreamingCell->RuntimeCellData);
+
+					StreamingCell->SetIsAlwaysLoaded(bIsCellAlwaysLoaded);
+					StreamingCell->SetDataLayers(GridCellDataChunk.GetDataLayers());
+					StreamingCell->SetContentBundleUID(GridCellDataChunk.GetContentBundleID());
+					StreamingCell->SetPriority(RuntimeGrid.Priority);
+					StreamingCell->SetClientOnlyVisible(CurrentStreamingGrid.bClientOnlyVisible);
+					StreamingCell->SetBlockOnSlowLoading(CurrentStreamingGrid.bBlockOnSlowStreaming);
+					StreamingCell->SetIsHLOD(RuntimeGrid.HLODLayer ? true : false);
+					StreamingCell->SetGuid(CellGuid);
+
+					FBox2D Bounds;
+					verify(TempLevel.GetCellBounds(FGridCellCoord2(CellCoordX, CellCoordY), Bounds));				
+					const double CellExtent = Bounds.GetExtent().X;
+					check(CellExtent < MAX_flt);
+
+					CellDataSpatialHash->Level = Level;
+					CellDataSpatialHash->Position = FVector(Bounds.GetCenter(), 0.f);
+					CellDataSpatialHash->Extent = (float)CellExtent;
+					CellDataSpatialHash->GridName = RuntimeGrid.GridName;			
+					CellDataSpatialHash->DebugName = CellName;
+
+					PopulateRuntimeCell(StreamingCell, FilteredActors, OutPackagesToGenerate);
+
+					int32 LayerCellIndex;
+					int32* LayerCellIndexPtr = GridLevel.LayerCellsMapping.Find(CellIndex);
+					if (LayerCellIndexPtr)
 					{
-						for (const FGuid& ActorGuid : ActorSetInstance->ActorSet->Actors)
-						{
-							// Instanced world partition, ContainerID is the main container, but it's not the main world partition,
-							// so the always loaded actors don't be part of the process of ForceExternalActorLevelReference/AlwaysLoadedActorsForPIE.
-							// In PIE, always loaded actors of an instanced world partition will go in the always loaded cell.
-							// 
-							// WorldDataLayers::ShouldLevelKeepRefIfExternal return true, so it will always be part of the partitioned 
-							// persistent level of a world partition (ideally we would need this information to be part of the ActorDesc).
-							// 
-							// In PIE, for an instanced world partition, we don't want this WorldDataLayers actor to be both in the persistent level 
-							// and also part of the always loaded cell level.
-							//
-							// @todo_ow: We need to implement PIE always loaded actors of instanced world partitions to be part
-							//			 of the persistent level and get rid of the always loaded cell (to have the same behavior
-							//			 as non-instanced world partition and as cooked world partition).
-							//
-							// For now, we force skip AWorldDataLayers actors.
-							IStreamingGenerationContext::FActorInstance ActorInstance(ActorGuid, ActorSetInstance);
-							if (!ActorInstance.GetActorDescView().GetActorNativeClass()->IsChildOf<AWorldDataLayers>())
-							{
-								FilteredActors.Emplace(ActorGuid, ActorSetInstance);
-							}
-						}
+						LayerCellIndex = *LayerCellIndexPtr;
+					}
+					else
+					{
+						LayerCellIndex = GridLevel.LayerCells.AddDefaulted();
+						GridLevel.LayerCellsMapping.Add(CellIndex, LayerCellIndex);
 					}
-				}
-				
-				if (!FilteredActors.Num())
-				{
-					continue;
-				}
-				
-				auto GetCellObjectName = [](FString CellName) -> FString
-				{
-					FArchiveMD5 ArMD5;
-					ArMD5 << CellName;
-
-					FMD5Hash MD5Hash;
-					ArMD5.GetHash(MD5Hash);
-
-					FGuid CellNameGuid = MD5HashToGuid(MD5Hash);
-					check(CellNameGuid.IsValid());
-
-					return CellNameGuid.ToString(EGuidFormats::Base36Encoded);
-				};
 
-				FGridCellCoord CellGlobalCoords;
-				verify(PartionedActors.GetCellGlobalCoords(FGridCellCoord(CellCoordX, CellCoordY, Level), CellGlobalCoords));
-				const FString CellName = GetCellNameString(OuterWorld->GetPackage()->GetFName(), CurrentStreamingGrid.GridName, CellGlobalCoords, GridCellDataChunk.GetDataLayersID(), GridCellDataChunk.GetContentBundleID());
-				const FGuid CellGuid = GetCellGuid(CurrentStreamingGrid.GridName, CellGlobalCoords, GridCellDataChunk.GetDataLayersID(), GridCellDataChunk.GetContentBundleID());
-
-				UWorldPartitionRuntimeCell* StreamingCell = CreateRuntimeCell(StreamingPolicy->GetRuntimeCellClass(), UWorldPartitionRuntimeCellDataSpatialHash::StaticClass(), FName(GetCellObjectName(CellName.ToLower())));
-				UWorldPartitionRuntimeCellDataSpatialHash* CellDataSpatialHash = CastChecked<UWorldPartitionRuntimeCellDataSpatialHash>(StreamingCell->RuntimeCellData);
-
-				StreamingCell->SetIsAlwaysLoaded(bIsCellAlwaysLoaded);
-				StreamingCell->SetDataLayers(GridCellDataChunk.GetDataLayers());
-				StreamingCell->SetContentBundleUID(GridCellDataChunk.GetContentBundleID());
-				StreamingCell->SetPriority(RuntimeGrid.Priority);
-				StreamingCell->SetClientOnlyVisible(CurrentStreamingGrid.bClientOnlyVisible);
-				StreamingCell->SetBlockOnSlowLoading(CurrentStreamingGrid.bBlockOnSlowStreaming);
-				StreamingCell->SetIsHLOD(RuntimeGrid.HLODLayer ? true : false);
-				StreamingCell->SetGuid(CellGuid);
-
-				FBox2D Bounds;
-				verify(TempLevel.GetCellBounds(FGridCellCoord2(CellCoordX, CellCoordY), Bounds));				
-				const double CellExtent = Bounds.GetExtent().X;
-				check(CellExtent < MAX_flt);
-
-				CellDataSpatialHash->Level = Level;
-				CellDataSpatialHash->Position = FVector(Bounds.GetCenter(), 0.f);
-				CellDataSpatialHash->Extent = (float)CellExtent;
-				CellDataSpatialHash->GridName = RuntimeGrid.GridName;			
-				CellDataSpatialHash->DebugName = CellName;
-
-				PopulateRuntimeCell(StreamingCell, FilteredActors, OutPackagesToGenerate);
-
-				int32 LayerCellIndex;
-				int32* LayerCellIndexPtr = GridLevel.LayerCellsMapping.Find(CellIndex);
-				if (LayerCellIndexPtr)
-				{
-					LayerCellIndex = *LayerCellIndexPtr;
+					GridLevel.LayerCells[LayerCellIndex].GridCells.Add(StreamingCell);
 				}
-				else
-				{
-					LayerCellIndex = GridLevel.LayerCells.AddDefaulted();
-					GridLevel.LayerCellsMapping.Add(CellIndex, LayerCellIndex);
-				}
-
-				GridLevel.LayerCells[LayerCellIndex].GridCells.Add(StreamingCell);
 			}
 		}
 	}
diff --git a/Engine/Source/Runtime/Engine/Public/WorldPartition/WorldPartitionRuntimeHash.h b/Engine/Source/Runtime/Engine/Public/WorldPartition/WorldPartitionRuntimeHash.h
--- a/Engine/Source/Runtime/Engine/Public/WorldPartition/WorldPartitionRuntimeHash.h
+++ b/Engine/Source/Runtime/Engine/Public/WorldPartition/WorldPartitionRuntimeHash.h
@@ -99,16 +99,11 @@ class ENGINE_API UWorldPartitionRuntimeHash : public UObject
 	void OnBeginPlay();
 	void OnEndPlay();
 
-	inline UWorldPartitionRuntimeCell* CreateRuntimeCell(UClass* CellClass, UClass* CellDataClass, FName CellName)
-	{
-		check(!FindObject<UWorldPartitionRuntimeCell>(this, *CellName.ToString()));
-		UWorldPartitionRuntimeCell* RuntimeCell = NewObject<UWorldPartitionRuntimeCell>(this, CellClass, CellName);
-		RuntimeCell->RuntimeCellData = NewObject<UWorldPartitionRuntimeCellData>(RuntimeCell, CellDataClass);
-		return RuntimeCell;
-	}
+	UWorldPartitionRuntimeCell* CreateRuntimeCell(UClass* CellClass, UClass* CellDataClass, const FString& CellName);
 
 protected:
 	bool ConditionalRegisterAlwaysLoadedActorsForPIE(const IStreamingGenerationContext::FActorSetInstance* ActorSetInstance, bool bIsMainWorldPartition, bool bIsMainContainer, bool bIsCellAlwaysLoaded);
+	bool PopulateCellActorInstances(const TArray<const IStreamingGenerationContext::FActorSetInstance*> ActorSetInstances, bool bIsMainWorldPartition, bool bIsCellAlwaysLoaded, TArray<IStreamingGenerationContext::FActorInstance>& OutCellActorInstances);
 	void PopulateRuntimeCell(UWorldPartitionRuntimeCell* RuntimeCell, const TArray<IStreamingGenerationContext::FActorInstance>& ActorInstances, TArray<FString>* OutPackagesToGenerate);
 #endif
 
