Editor asserts when deleting a skeleton that's used by the current level

#fix introduce similar Skeleton validity checks for post-process aninm instance already performed for the main anim instance

diff --git a/Engine/Source/Runtime/Engine/Private/Components/SkeletalMeshComponent.cpp b/Engine/Source/Runtime/Engine/Private/Components/SkeletalMeshComponent.cpp
--- a/Engine/Source/Runtime/Engine/Private/Components/SkeletalMeshComponent.cpp
+++ b/Engine/Source/Runtime/Engine/Private/Components/SkeletalMeshComponent.cpp
@@ -589,8 +589,17 @@ bool USkeletalMeshComponent::NeedToSpawnPostPhysicsInstance(bool bForceReinit) c
 		const UClass* ClassToUse = *GetSkeletalMeshAsset()->GetPostProcessAnimBlueprint();
 		const UClass* CurrentClass = PostProcessAnimInstance ? PostProcessAnimInstance->GetClass() : nullptr;
 
+		const IAnimClassInterface* AnimClassInterface = IAnimClassInterface::GetFromClass(ClassToUse);
+		const USkeleton* AnimSkeleton = (AnimClassInterface) ? AnimClassInterface->GetTargetSkeleton() : nullptr;
+		if(AnimSkeleton == nullptr)
+		{
+			const USkeletalMesh* MeshAsset = GetSkeletalMeshAsset();
+			// Fall back to mesh skeleton (post-process anim BP could have been a template)
+			AnimSkeleton = MeshAsset ? MeshAsset->GetSkeleton() : nullptr;
+		}
+
 		// We need to have an instance, and we have the wrong class (different or null)
-		if(ClassToUse && (ClassToUse != CurrentClass || bForceReinit ) && MainInstanceClass != ClassToUse)
+		if(ClassToUse && (ClassToUse != CurrentClass || bForceReinit ) && MainInstanceClass != ClassToUse && AnimSkeleton)
 		{
 			return true;
 		}
@@ -949,7 +958,7 @@ bool USkeletalMeshComponent::InitializeAnimScriptInstance(bool bForceReinit, boo
 
 		// May need to clear out the post physics instance
 		UClass* NewMeshInstanceClass = *SkelMesh->GetPostProcessAnimBlueprint();
-		if(!NewMeshInstanceClass || NewMeshInstanceClass == *AnimClass)
+		if(!NewMeshInstanceClass || NewMeshInstanceClass == *AnimClass || (PostProcessAnimInstance && PostProcessAnimInstance->CurrentSkeleton != GetSkeletalMeshAsset()->GetSkeleton()) || !GetSkeletalMeshAsset()->GetSkeleton())
 		{
 			PostProcessAnimInstance = nullptr;
 		}
