Crash when creating a USkeletalMeshComponent inside a UMotionControllerComponent

#fix make ::CleanupUpdateRateParametersRef use Find rather than FindChecked. In this case a DefaultSubObject was being Unregistered after being Renamed to the transient-package. (and no previous URO entry was made to begin with)

diff --git a/Engine/Source/Runtime/Engine/Private/Components/SkinnedMeshComponent.cpp b/Engine/Source/Runtime/Engine/Private/Components/SkinnedMeshComponent.cpp
--- a/Engine/Source/Runtime/Engine/Private/Components/SkinnedMeshComponent.cpp
+++ b/Engine/Source/Runtime/Engine/Private/Components/SkinnedMeshComponent.cpp
@@ -176,14 +176,17 @@ namespace FAnimUpdateRateManager
 
 	void CleanupUpdateRateParametersRef(USkinnedMeshComponent* SkinnedComponent)
 	{
-		UObject* TrackerIndex = GetMapIndexForComponent(SkinnedComponent);
-
-		FAnimUpdateRateParametersTracker* Tracker = ActorToUpdateRateParams.FindChecked(TrackerIndex);
-		Tracker->RegisteredComponents.Remove(SkinnedComponent);
-		if (Tracker->RegisteredComponents.Num() == 0)
+		const UObject* TrackerIndex = GetMapIndexForComponent(SkinnedComponent);
+		FAnimUpdateRateParametersTracker** TrackerPtr = ActorToUpdateRateParams.Find(TrackerIndex);
+		if (TrackerPtr)
 		{
-			ActorToUpdateRateParams.Remove(TrackerIndex);
-			delete Tracker;
+			FAnimUpdateRateParametersTracker* Tracker = *TrackerPtr;
+			Tracker->RegisteredComponents.Remove(SkinnedComponent);
+			if (Tracker->RegisteredComponents.Num() == 0)
+			{
+				ActorToUpdateRateParams.Remove(TrackerIndex);
+				delete Tracker;
+			}
 		}
 	}
 
