Always call InitDefaultMaterials early when running with cooked data, not only when GEventDrivenLoaderEnabled is true

diff --git a/Engine/Source/Runtime/Engine/Private/Materials/Material.cpp b/Engine/Source/Runtime/Engine/Private/Materials/Material.cpp
--- a/Engine/Source/Runtime/Engine/Private/Materials/Material.cpp
+++ b/Engine/Source/Runtime/Engine/Private/Materials/Material.cpp
@@ -598,7 +598,7 @@ void UMaterialInterface::InitDefaultMaterials()
 
 void UMaterialInterface::PostCDOContruct()
 {
-	if (GEventDrivenLoaderEnabled && EVENT_DRIVEN_ASYNC_LOAD_ACTIVE_AT_RUNTIME)
+	if (FPlatformProperties::RequiresCookedData())
 	{
 		UMaterial::StaticClass()->GetDefaultObject();
 		UMaterialInterface::InitDefaultMaterials();
