Remove deprecated LevelPartition code (unrelated to WorldPartition)

diff --git a/Engine/Plugins/Enterprise/DatasmithImporter/Source/DatasmithImporter/Private/DatasmithConsumer.cpp b/Engine/Plugins/Enterprise/DatasmithImporter/Source/DatasmithImporter/Private/DatasmithConsumer.cpp
--- a/Engine/Plugins/Enterprise/DatasmithImporter/Source/DatasmithImporter/Private/DatasmithConsumer.cpp
+++ b/Engine/Plugins/Enterprise/DatasmithImporter/Source/DatasmithImporter/Private/DatasmithConsumer.cpp
@@ -1846,21 +1846,17 @@ namespace DatasmithConsumerUtils
 
 		UEditorLevelUtils::SetLevelVisibility(DestLevel, true, false, ELevelVisibilityDirtyMode::DontModify);
 
-		// Scope this so that Actors that have been pasted will have their final levels set before doing the actor mapping
-		{
-			// Set the new level and force it visible while we do the paste
-			FLevelPartitionOperationScope LevelPartitionScope(DestLevel);
-			OwningWorld->SetCurrentLevel(LevelPartitionScope.GetLevel());
-
-			//const bool bDuplicate = false;
-			const bool bOffsetLocations = false;
-			const bool bWarnIfHidden = false;
-			GEditor->edactPasteSelected(OwningWorld, bDuplicate, bOffsetLocations, bWarnIfHidden);
+		// Set the new level and force it visible while we do the paste
+		OwningWorld->SetCurrentLevel(DestLevel);
 
-			// Restore the original current level
-			OwningWorld->SetCurrentLevel(OldCurrentLevel);
-		}
+		//const bool bDuplicate = false;
+		const bool bOffsetLocations = false;
+		const bool bWarnIfHidden = false;
+		GEditor->edactPasteSelected(OwningWorld, bDuplicate, bOffsetLocations, bWarnIfHidden);
 
+		// Restore the original current level
+		OwningWorld->SetCurrentLevel(OldCurrentLevel);
+		
 		TArray<AActor*> NewActors;
 		NewActors.Reserve(GEditor->GetSelectedActorCount());
 
diff --git a/Engine/Source/Editor/UnrealEd/Private/AssetSelection.cpp b/Engine/Source/Editor/UnrealEd/Private/AssetSelection.cpp
--- a/Engine/Source/Editor/UnrealEd/Private/AssetSelection.cpp
+++ b/Engine/Source/Editor/UnrealEd/Private/AssetSelection.cpp
@@ -587,15 +587,6 @@ static AActor* PrivateAddActor( UObject* Asset, UActorFactory* Factory, bool Sel
 
 	ULevel* DesiredLevel = GWorld->GetCurrentLevel();
 
-	// If DesireLevel is part of a LevelPartition find the proper DesiredLevel by asking the Partition
-	if (const ILevelPartitionInterface* LevelPartition = DesiredLevel->GetLevelPartition())
-	{
-		if (ULevel* SubLevel = LevelPartition->GetSubLevel(ActorTransform.GetLocation()))
-		{
-			DesiredLevel = SubLevel;
-		}
-	}
-
 	bool bSpawnActor = true;
 
 	if ((ObjectFlags & RF_Transactional) != 0)
diff --git a/Engine/Source/Editor/UnrealEd/Private/EditorActor.cpp b/Engine/Source/Editor/UnrealEd/Private/EditorActor.cpp
--- a/Engine/Source/Editor/UnrealEd/Private/EditorActor.cpp
+++ b/Engine/Source/Editor/UnrealEd/Private/EditorActor.cpp
@@ -484,9 +484,7 @@ public:
 		// Paste to the dest level
 		{
 			ULevel* CurrentDestLevel = DestLevel->OwningWorld->GetCurrentLevel();
-
-			FLevelPartitionOperationScope LevelPartitionScope(DestLevel);
-			DestLevel->OwningWorld->SetCurrentLevel(LevelPartitionScope.GetLevel());
+			DestLevel->OwningWorld->SetCurrentLevel(DestLevel);
 
 			TArray<AActor*> PastedActors;
 			GUnrealEd->PasteActors(PastedActors, DestLevel->OwningWorld, LocationOffset, /*bDuplicate*/true, /*bWarnIfHidden*/true, &ScratchData);
diff --git a/Engine/Source/Editor/UnrealEd/Private/EditorLevelUtils.cpp b/Engine/Source/Editor/UnrealEd/Private/EditorLevelUtils.cpp
--- a/Engine/Source/Editor/UnrealEd/Private/EditorLevelUtils.cpp
+++ b/Engine/Source/Editor/UnrealEd/Private/EditorLevelUtils.cpp
@@ -181,26 +181,22 @@ int32 UEditorLevelUtils::CopyOrMoveActorsToLevel(const TArray<AActor*>& ActorsTo
 				FString DestinationData;
 				GEditor->CopySelectedActorsToClipboard(OwningWorld, bMoveActors, bMoveActors, bWarnAboutReferences, &DestinationData);
 
+				// Set the new level and force it visible while we do the paste
 				const bool bLevelVisible = DestLevel->bIsVisible;
 				if (!bLevelVisible)
 				{
 					UEditorLevelUtils::SetLevelVisibility(DestLevel, true, false);
 				}
+				
+				OwningWorld->SetCurrentLevel(DestLevel);
 
-				// Scope this so that Actors that have been pasted will have their final levels set before doing the actor mapping
-				{
-					// Set the new level and force it visible while we do the paste
-					FLevelPartitionOperationScope LevelPartitionScope(DestLevel);
-					OwningWorld->SetCurrentLevel(LevelPartitionScope.GetLevel());
-
-					const bool bDuplicate = false;
-					const bool bOffsetLocations = false;
-					const bool bWarnIfHidden = false;
-					GEditor->edactPasteSelected(OwningWorld, bDuplicate, bOffsetLocations, bWarnIfHidden, &DestinationData);
+				const bool bDuplicate = false;
+				const bool bOffsetLocations = false;
+				const bool bWarnIfHidden = false;
+				GEditor->edactPasteSelected(OwningWorld, bDuplicate, bOffsetLocations, bWarnIfHidden, &DestinationData);
 
-					// Restore the original current level
-					OwningWorld->SetCurrentLevel(OldCurrentLevel);
-				}
+				// Restore the original current level
+				OwningWorld->SetCurrentLevel(OldCurrentLevel);
 
 				// Build a remapping of old to new names so we can do a fixup
 				for (FSelectionIterator It(GEditor->GetSelectedActorIterator()); It; ++It)
diff --git a/Engine/Source/Editor/WorldBrowser/Private/SWorldDetails.cpp b/Engine/Source/Editor/WorldBrowser/Private/SWorldDetails.cpp
--- a/Engine/Source/Editor/WorldBrowser/Private/SWorldDetails.cpp
+++ b/Engine/Source/Editor/WorldBrowser/Private/SWorldDetails.cpp
@@ -73,7 +73,6 @@ void SWorldDetails::OnBrowseWorld(UWorld* InWorld)
 		Args.bHideSelectionTip = true;
 	
 		DetailsView = PropertyModule.CreateDetailView(Args);
-		WorldDetailsView = PropertyModule.CreateDetailView(Args);
 		ChildSlot
 		[
 			SNew(SBorder)
@@ -186,29 +185,10 @@ void SWorldDetails::OnSelectionChanged()
 	{
 		// Clear ComboBox selection in case we have multiple selection
 		SubLevelsComboBox->ClearSelection();
-		WorldDetailsView->SetObject(nullptr);
 	}
 	else
 	{
 		SubLevelsComboBox->SetSelectedItem(SelectedLevels[0]);
-		ULevel* LevelObject = SelectedLevels[0]->GetLevelObject();
-		UObject* LevelPartition = Cast<UObject>(LevelObject ? LevelObject->GetLevelPartition() : nullptr);
-
-		if (LevelPartition)
-		{
-			VerticalBox->AddSlot()
-				.FillHeight(1.f)
-				.Padding(0,4,0,0)
-				[
-					SAssignNew(VerticalBoxBorder, SBorder)
-					.BorderImage(FAppStyle::GetBrush(TEXT("ToolPanel.GroupBorder")))
-					[
-						WorldDetailsView.ToSharedRef()
-					]
-				];
-		}
-
-		WorldDetailsView->SetObject(LevelPartition);
 	}
 
 	bUpdatingSelection = false;
diff --git a/Engine/Source/Editor/WorldBrowser/Private/SWorldDetails.h b/Engine/Source/Editor/WorldBrowser/Private/SWorldDetails.h
--- a/Engine/Source/Editor/WorldBrowser/Private/SWorldDetails.h
+++ b/Engine/Source/Editor/WorldBrowser/Private/SWorldDetails.h
@@ -56,7 +56,6 @@ private:
 private:
 	TSharedPtr<FLevelCollectionModel>				WorldModel;
 	TSharedPtr<IDetailsView>						DetailsView;
-	TSharedPtr<IDetailsView>						WorldDetailsView;
 	TSharedPtr<SVerticalBox>						VerticalBox;
 	TSharedPtr<SBorder>								VerticalBoxBorder;
 	TSharedPtr<SComboBox<TSharedPtr<FLevelModel>>>	SubLevelsComboBox;
diff --git a/Engine/Source/Editor/WorldBrowser/Private/StreamingLevels/StreamingLevelCollectionModel.cpp b/Engine/Source/Editor/WorldBrowser/Private/StreamingLevels/StreamingLevelCollectionModel.cpp
--- a/Engine/Source/Editor/WorldBrowser/Private/StreamingLevels/StreamingLevelCollectionModel.cpp
+++ b/Engine/Source/Editor/WorldBrowser/Private/StreamingLevels/StreamingLevelCollectionModel.cpp
@@ -126,16 +126,6 @@ void FStreamingLevelCollectionModel::UnloadLevels(const FLevelModelList& InLevel
 			// this level is dirty and can be removed from the world
 			bHaveDirtyLevels = true;
 		}
-
-		if (const ULevel* Level = LevelModel->GetLevelObject())
-		{
-			if (Level->IsPartitionSubLevel())
-			{
-				// this level is a partition sublevel and cannot be removed from the world
-				FMessageDialog::Open(EAppMsgType::Ok, LOCTEXT("RemoveLevel_PartitionSubLevel", "Cannot remove a level which is part of a partition. Delete the top level instead."));
-				return;
-			}
-		}
 	}
 
 	// Depending on the state of the level, create a warning message
diff --git a/Engine/Source/Runtime/Engine/Classes/Engine/Level.h b/Engine/Source/Runtime/Engine/Classes/Engine/Level.h
--- a/Engine/Source/Runtime/Engine/Classes/Engine/Level.h
+++ b/Engine/Source/Runtime/Engine/Classes/Engine/Level.h
@@ -34,7 +34,6 @@ class UNavigationDataChunk;
 class UTexture2D;
 struct FLevelCollection;
 class ULevelActorContainer;
-class FLevelPartitionOperationScope;
 class FRegisterComponentContext;
 class SNotificationItem;
 class UActorFolder;
@@ -42,30 +41,6 @@ class IWorldPartitionCell;
 class UWorldPartitionRuntimeCell;
 struct FFolder;
 
-UINTERFACE()
-class ULevelPartitionInterface : public UInterface
-{
-	GENERATED_BODY()
-};
-
-class ILevelPartitionInterface
-{
-	GENERATED_BODY()
-
-#if WITH_EDITOR
-	friend class FLevelPartitionOperationScope;
-#endif
-
-public:
-	virtual ULevel* GetSubLevel(const FVector& Coords) const = 0;
-
-private:
-#if WITH_EDITOR
-	virtual void BeginOperation(class FLevelPartitionOperationScope* Scope) = 0;
-	virtual void EndOperation() = 0;
-#endif
-};
-
 #if WITH_EDITOR
 struct ENGINE_API FLevelActorFoldersHelper
 {
@@ -84,22 +59,6 @@ private:
 	friend class UWorldPartitionLevelStreamingDynamic;
 };
 
-class ENGINE_API FLevelPartitionOperationScope
-{
-public:
-	FLevelPartitionOperationScope(ULevel* Level);
-	~FLevelPartitionOperationScope();
-		
-	TArray<AActor*> GetActors() const;
-	ULevel* GetLevel() const;
-
-private:
-	static ULevel* CreateTransientLevel(UWorld* InWorld);
-	static void DestroyTransientLevel(ULevel* InLevel);
-
-	ILevelPartitionInterface* InterfacePtr = nullptr;
-	ULevel* Level = nullptr;
-};
 #endif
 
 // Actor container class used to duplicate actors during cells streaming in PIE
@@ -725,22 +684,6 @@ public:
 	UPROPERTY()
 	EActorPackagingScheme ActorPackagingScheme;
 
-	/** Returns true if the current level is a partitioned level */
-	ENGINE_API bool IsPartitionedLevel() const;
-
-	/** Returns true if the current level is a sublevel (managed by a parent partitioned level) */
-	ENGINE_API bool IsPartitionSubLevel() const;
-
-	/** Assign a level partition to this level */
-	ENGINE_API void SetLevelPartition(ILevelPartitionInterface* LevelPartition);
-
-	/** Get the level partition assigned to this level, if any */
-	ENGINE_API ILevelPartitionInterface* GetLevelPartition();
-	ENGINE_API const ILevelPartitionInterface* GetLevelPartition() const;
-	
-	/** Setup the provided sublevel so that it is handled by this level's partition */
-	ENGINE_API void SetPartitionSubLevel(ULevel* SubLevel);
-
 #endif //WITH_EDITORONLY_DATA
 
 	/** Actor which defines level logical bounding box				*/
@@ -820,14 +763,6 @@ private:
 	TArray<FReplicatedStaticActorDestructionInfo> DestroyedReplicatedStaticActors;
 
 #if WITH_EDITORONLY_DATA
-	/** Level partition, if any */
-	UPROPERTY(EditInstanceOnly, Category = World)
-	TScriptInterface<ILevelPartitionInterface> LevelPartition;
-
-	/** When the level is partitioned, this will point to the owner partition (will be the same as this->LevelPartition in case that is the top partition level */
-	UPROPERTY()
-	TSoftObjectPtr<UObject> OwnerLevelPartition;
-
 	/** Use actor folder objects, actor folders of this level will be persistent in their own object. */
 	UPROPERTY(EditInstanceOnly, Category = World)
 	bool bUseActorFolders;
@@ -892,7 +827,6 @@ public:
 	virtual void PostEditUndo() override;	
 	virtual void PostEditChangeProperty(FPropertyChangedEvent& PropertyChangedEvent) override;
 	virtual void BeginCacheForCookedPlatformData(const ITargetPlatform *TargetPlatform) override;
-	virtual bool CanEditChange(const FProperty* PropertyThatWillChange) const override;
 #endif // WITH_EDITOR
 	virtual bool ResolveSubobject(const TCHAR* SubObjectPath, UObject*& OutObject, bool bLoadIfExists) override;
 	virtual void PostLoad() override;
diff --git a/Engine/Source/Runtime/Engine/Private/ActorPartition/ActorPartitionSubsystem.cpp b/Engine/Source/Runtime/Engine/Private/ActorPartition/ActorPartitionSubsystem.cpp
--- a/Engine/Source/Runtime/Engine/Private/ActorPartition/ActorPartitionSubsystem.cpp
+++ b/Engine/Source/Runtime/Engine/Private/ActorPartition/ActorPartitionSubsystem.cpp
@@ -163,13 +163,6 @@ private:
 	{
 		check(InLevelHint);
 		ULevel* SpawnLevel = InLevelHint;
-		if (ILevelPartitionInterface* LevelPartition = InLevelHint->GetLevelPartition())
-		{
-			if (ULevel* SubLevel = LevelPartition->GetSubLevel(InLocationHint))
-			{
-				SpawnLevel = SubLevel;
-			}
-		}
 		return SpawnLevel;
 	}
 
diff --git a/Engine/Source/Runtime/Engine/Private/Level.cpp b/Engine/Source/Runtime/Engine/Private/Level.cpp
--- a/Engine/Source/Runtime/Engine/Private/Level.cpp
+++ b/Engine/Source/Runtime/Engine/Private/Level.cpp
@@ -146,72 +146,6 @@ void FLevelActorFoldersHelper::DeleteFolder(ULevel* InLevel, const FFolder& InFo
 	}
 };
 
-FLevelPartitionOperationScope::FLevelPartitionOperationScope(ULevel* InLevel)
-{
-	InterfacePtr = InLevel->GetLevelPartition();
-	Level = InLevel;
-	if (InterfacePtr)
-	{
-		InterfacePtr->BeginOperation(this);
-		Level = CreateTransientLevel(InLevel->GetWorld());
-	}
-}
-
-FLevelPartitionOperationScope::~FLevelPartitionOperationScope()
-{
-	if (InterfacePtr)
-	{
-		InterfacePtr->EndOperation();
-		DestroyTransientLevel(Level);
-	}
-	Level = nullptr;
-}
-
-TArray<AActor*> FLevelPartitionOperationScope::GetActors() const
-{
-	if (InterfacePtr)
-	{
-		return Level->Actors;
-	}
-
-	return {};
-}
-
-ULevel* FLevelPartitionOperationScope::GetLevel() const
-{
-	check(Level);
-	return Level;
-}
-
-ULevel* FLevelPartitionOperationScope::CreateTransientLevel(UWorld* InWorld)
-{
-	ULevel* Level = NewObject<ULevel>(GetTransientPackage(), TEXT("TempLevelPartitionOperationScopeLevel"));
-	check(Level);
-	Level->Initialize(FURL(nullptr));
-	Level->AddToRoot();
-	Level->OwningWorld = InWorld;
-	Level->Model = NewObject<UModel>(Level);
-	Level->Model->Initialize(nullptr, true);
-	Level->bIsVisible = true;
-
-	Level->SetFlags(RF_Transactional);
-	Level->Model->SetFlags(RF_Transactional);
-
-	return Level;
-}
-
-void FLevelPartitionOperationScope::DestroyTransientLevel(ULevel* Level)
-{
-	check(Level->GetOutermost() == GetTransientPackage());
-	// Make sure Level doesn't contain any Actors before destroying. That would mean the operation failed.
-	check(!Algo::AnyOf(Level->Actors, [](AActor* Actor) { return Actor != nullptr; }));
-	// Delete the temporary level
-	Level->ClearLevelComponents();
-	Level->GetWorld()->RemoveLevel(Level);
-	Level->OwningWorld = nullptr;
-	Level->RemoveFromRoot();
-	Level = nullptr;
-}
 #endif
 
 /*-----------------------------------------------------------------------------
@@ -4163,32 +4097,6 @@ void ULevel::BeginCacheForCookedPlatformData(const ITargetPlatform *TargetPlatfo
 	}
 }
 
-bool ULevel::CanEditChange(const FProperty* PropertyThatWillChange) const
-{
-	static FName NAME_LevelPartition = GET_MEMBER_NAME_CHECKED(ULevel, LevelPartition);
-	if (PropertyThatWillChange->GetFName() == NAME_LevelPartition)
-	{
-		// Can't set a partition on the persistent level
-		if (IsPersistentLevel())
-		{
-			return false;
-		}
-
-		// Can't set a partition on partition sublevels
-		if (IsPartitionSubLevel())
-		{
-			return false;
-		}
-
-		// Can't set a partition if using world composition or partition
-		if (WorldSettings && (WorldSettings->bEnableWorldComposition || WorldSettings->IsPartitionedWorld()))
-		{
-			return false;
-		}
-	}
-	return Super::CanEditChange(PropertyThatWillChange);
-}
-
 void ULevel::FixupForPIE(int32 InPIEInstanceID, TFunctionRef<void(int32, FSoftObjectPath&)> InCustomFixupFunction)
 {
 	FPIEFixupSerializer FixupSerializer(this, InPIEInstanceID, InCustomFixupFunction);
@@ -4430,42 +4338,6 @@ bool ULevel::HasVisibilityChangeRequestPending() const
 	return (OwningWorld && ( this == OwningWorld->GetCurrentLevelPendingVisibility() || this == OwningWorld->GetCurrentLevelPendingInvisibility() ) );
 }
 
-#if WITH_EDITORONLY_DATA
-
-bool ULevel::IsPartitionedLevel() const
-{
-	return LevelPartition != nullptr;
-}
-
-bool ULevel::IsPartitionSubLevel() const
-{
-	return OwnerLevelPartition.IsValid() && LevelPartition == nullptr;
-}
-
-void ULevel::SetLevelPartition(ILevelPartitionInterface* InLevelPartition)
-{
-	UObject* PartitionObject = Cast<UObject>(InLevelPartition);
-	LevelPartition = PartitionObject;
-	OwnerLevelPartition = PartitionObject;
-}
-
-ILevelPartitionInterface* ULevel::GetLevelPartition()
-{
-	return Cast<ILevelPartitionInterface>(OwnerLevelPartition.Get());
-}
-
-const ILevelPartitionInterface* ULevel::GetLevelPartition() const
-{
-	return Cast<ILevelPartitionInterface>(OwnerLevelPartition.Get());
-}
-
-void ULevel::SetPartitionSubLevel(ULevel* SubLevel)
-{
-	check(LevelPartition);
-	SubLevel->OwnerLevelPartition = Cast<UObject>(&*LevelPartition);
-}
-
-#endif // #if WITH_EDITORONLY_DATA
 #if WITH_EDITOR
 void ULevel::RepairLevelScript()
 {
