SubObjects replication
* Added cvar net.SubObjects.DetectDeprecatedReplicatedSubObjects
* When enabled it will ensure if it detects that ReplicateSubObjects() is implemented in a class that switched over to the new registration method
* Only available in development builds

diff --git a/Engine/Plugins/Runtime/GameplayAbilities/Source/GameplayAbilities/Private/AbilitySystemComponent.cpp b/Engine/Plugins/Runtime/GameplayAbilities/Source/GameplayAbilities/Private/AbilitySystemComponent.cpp
--- a/Engine/Plugins/Runtime/GameplayAbilities/Source/GameplayAbilities/Private/AbilitySystemComponent.cpp
+++ b/Engine/Plugins/Runtime/GameplayAbilities/Source/GameplayAbilities/Private/AbilitySystemComponent.cpp
@@ -1653,8 +1653,17 @@ void UAbilitySystemComponent::ForceAvatarReplication()
 	}
 }
 
-bool UAbilitySystemComponent::ReplicateSubobjects(class UActorChannel *Channel, class FOutBunch *Bunch, FReplicationFlags *RepFlags)
+bool UAbilitySystemComponent::ReplicateSubobjects(UActorChannel* Channel, FOutBunch* Bunch, FReplicationFlags *RepFlags)
 {
+#if SUBOBJECT_TRANSITION_VALIDATION
+	// When true it means we are calling this function to find any leftover replicated subobjects in classes that transitioned to the new registry list.
+	// This shared class needs to keep supporting the old ways until we fully deprecate the API, so by only returning false we prevent the ensures to trigger
+	if (UActorChannel::CanIgnoreDeprecatedReplicateSubObjects())
+	{
+		return false;
+	}
+#endif
+
 	bool WroteSomething = Super::ReplicateSubobjects(Channel, Bunch, RepFlags);
 
 	for (const UAttributeSet* Set : GetSpawnedAttributes())
diff --git a/Engine/Source/Runtime/Engine/Classes/Engine/ActorChannel.h b/Engine/Source/Runtime/Engine/Classes/Engine/ActorChannel.h
--- a/Engine/Source/Runtime/Engine/Classes/Engine/ActorChannel.h
+++ b/Engine/Source/Runtime/Engine/Classes/Engine/ActorChannel.h
@@ -15,6 +15,11 @@
 #define NET_ENABLE_SUBOBJECT_REPKEYS 1
 #endif // NET_ENABLE_SUBOBJECT_REPKEYS
 
+// Enable the code allowing to validate replicate subobjects when converting from the legacy method to the explicit registration list
+#ifndef SUBOBJECT_TRANSITION_VALIDATION
+	#define SUBOBJECT_TRANSITION_VALIDATION !(UE_BUILD_SHIPPING || UE_BUILD_TEST)
+#endif
+
 class UActorComponent;
 class AActor;
 class FInBunch;
@@ -301,6 +306,11 @@ public:
 	static void SetCurrentSubObjectOwner(AActor* SubObjectOwner);
 	static void SetCurrentSubObjectOwner(UActorComponent* SubObjectOwner);
 
+#if SUBOBJECT_TRANSITION_VALIDATION
+	/** Only returns true when calling ReplicateSubobjects() while we are trying to detect bad usage of the legacy function. */
+	static bool CanIgnoreDeprecatedReplicateSubObjects();
+#endif
+
 	/** Replicates given subobject on this actor channel */
 	bool ReplicateSubobject(UObject* Obj, FOutBunch& Bunch, FReplicationFlags RepFlags);
 
diff --git a/Engine/Source/Runtime/Engine/Private/ActorReplication.cpp b/Engine/Source/Runtime/Engine/Private/ActorReplication.cpp
--- a/Engine/Source/Runtime/Engine/Private/ActorReplication.cpp
+++ b/Engine/Source/Runtime/Engine/Private/ActorReplication.cpp
@@ -499,6 +499,14 @@ bool AActor::ReplicateSubobjects(UActorChannel *Channel, FOutBunch *Bunch, FRepl
 		{
 			UActorChannel::SetCurrentSubObjectOwner(ActorComp);
 			WroteSomething |= ActorComp->ReplicateSubobjects(Channel, Bunch, RepFlags);		// Lets the component add subobjects before replicating its own properties.
+
+#if SUBOBJECT_TRANSITION_VALIDATION
+			if (UActorChannel::CanIgnoreDeprecatedReplicateSubObjects())
+			{
+				// Don't replicate the component itself when testing since we know this function is not voluntarily implemented
+				continue;
+			}
+#endif
 			UActorChannel::SetCurrentSubObjectOwner(this);
 			WroteSomething |= Channel->ReplicateSubobject(ActorComp, *Bunch, *RepFlags);	// (this makes those subobjects 'supported', and from here on those objects may have reference replicated)		
 		}
diff --git a/Engine/Source/Runtime/Engine/Private/DataChannel.cpp b/Engine/Source/Runtime/Engine/Private/DataChannel.cpp
--- a/Engine/Source/Runtime/Engine/Private/DataChannel.cpp
+++ b/Engine/Source/Runtime/Engine/Private/DataChannel.cpp
@@ -56,11 +56,6 @@ DECLARE_CYCLE_STAT(TEXT("ActorChan_FindOrCreateRep"), Stat_ActorChanFindOrCreate
 DECLARE_LLM_MEMORY_STAT(TEXT("NetChannel"), STAT_NetChannelLLM, STATGROUP_LLMFULL);
 LLM_DEFINE_TAG(NetChannel, NAME_None, TEXT("Networking"), GET_STATFNAME(STAT_NetChannelLLM), GET_STATFNAME(STAT_NetworkingSummaryLLM));
 
-// Enable the code allowing to validate replicate subobjects when converting from the legacy method to the explicit registration list
-#ifndef SUBOBJECT_TRANSITION_VALIDATION
-	#define SUBOBJECT_TRANSITION_VALIDATION !UE_BUILD_SHIPPING
-#endif
-
 extern int32 GDoReplicationContextString;
 extern int32 GNetDormancyValidate;
 extern bool GbNetReuseReplicatorsForDormantObjects;
@@ -130,7 +125,7 @@ namespace UE::Net
 		TEXT("Only allow property replication of actors that had BeginPlay called on them."), ECVF_Default);
 
 #if SUBOBJECT_TRANSITION_VALIDATION
-	static int32 GCVarCompareSubObjectsReplicated = 0;
+	static bool GCVarCompareSubObjectsReplicated = false;
 	static FAutoConsoleVariableRef CVarCompareSubObjectsReplicated(
 		TEXT("net.SubObjects.CompareWithLegacy"),
 		GCVarCompareSubObjectsReplicated,
@@ -140,6 +135,12 @@ namespace UE::Net
 		TEXT("net.SubObjects.LogAllComparisonErrors"),
 		0,
 		TEXT("If enabled log all the errors detected by the CompareWithLegacy cheat. Otherwise only the first ensure triggered gets logged."));
+
+	static bool GCVarDetectDeprecatedReplicateSubObjects = false;
+	static FAutoConsoleVariableRef CVarDetectDeprecatedReplicateSubObjects(
+		TEXT("net.SubObjects.DetectDeprecatedReplicatedSubObjects"),
+		GCVarDetectDeprecatedReplicateSubObjects,
+		TEXT("When turned on, we trigger an ensure if we detect a ReplicateSubObjects() function is still implemented in a class that is using the new SubObject list."));
 #endif
 
 	static bool bEnableNetInitialSubObjects = true;
@@ -3756,12 +3757,12 @@ bool UActorChannel::ReplicateRegisteredSubObjects(FOutBunch& Bunch, FReplication
 	check(Actor->IsUsingRegisteredSubObjectList());
 
 #if SUBOBJECT_TRANSITION_VALIDATION
-	if (GCVarCompareSubObjectsReplicated)
+	if (GCVarCompareSubObjectsReplicated || GCVarDetectDeprecatedReplicateSubObjects)
 	{
 		TestLegacyReplicateSubObjects(Bunch, RepFlags);
 
-		// From here track all subobjects that write into the bunch.
-		DataChannelInternal::bTrackReplicatedSubObjects = true;
+		// From here track all subobjects from the real registry that write into the bunch.
+		DataChannelInternal::bTrackReplicatedSubObjects = GCVarCompareSubObjectsReplicated;
 	}
 #endif
 
@@ -3818,6 +3819,13 @@ void UActorChannel::SetCurrentSubObjectOwner(UActorComponent* SubObjectOwner)
 	DataChannelInternal::bIgnoreLegacyReplicateSubObject = SubObjectOwner && SubObjectOwner->IsUsingRegisteredSubObjectList();
 }
 
+#if SUBOBJECT_TRANSITION_VALIDATION
+bool UActorChannel::CanIgnoreDeprecatedReplicateSubObjects()
+{
+	return UE::Net::GCVarDetectDeprecatedReplicateSubObjects && DataChannelInternal::bTestingLegacyMethodForComparison;
+}
+#endif
+
 bool UActorChannel::ReplicateSubobject(UObject* SubObj, FOutBunch& Bunch, FReplicationFlags RepFlags)
 {
 	if (!IsValid(SubObj))
@@ -3828,7 +3836,25 @@ bool UActorChannel::ReplicateSubobject(UObject* SubObj, FOutBunch& Bunch, FRepli
 #if SUBOBJECT_TRANSITION_VALIDATION
 	if (DataChannelInternal::bTestingLegacyMethodForComparison)
 	{
-		DataChannelInternal::LegacySubObjectsCollected.Emplace(DataChannelInternal::FSubObjectReplicatedInfo(SubObj));
+		if (UE::Net::GCVarDetectDeprecatedReplicateSubObjects)
+		{
+			if (DataChannelInternal::CurrentSubObjectOwner == Actor)
+			{
+				ensureMsgf(false, TEXT("ReplicateSubObjects() should not be implemented in class %s since bReplicateUsingRegisteredSubObjectList is true. Delete the function and use AddReplicatedSubObject instead."),
+					*Actor->GetName(), *Actor->GetClass()->GetName());
+			}
+			else
+			{
+				ensureMsgf(false, TEXT("ReplicateSubObjects() should not be implemented in class %s owned by %s (class %s) since bReplicateUsingRegisteredSubObjectList is true. Delete the function and use AddReplicatedSubObject instead."),
+					*GetNameSafe(DataChannelInternal::CurrentSubObjectOwner?DataChannelInternal::CurrentSubObjectOwner->GetClass():nullptr),
+					*Actor->GetName(), *Actor->GetClass()->GetName());
+			}
+		}
+		else
+		{
+			DataChannelInternal::LegacySubObjectsCollected.Emplace(DataChannelInternal::FSubObjectReplicatedInfo(SubObj));
+		}
+
 		return false;
 	}
 #endif
@@ -3889,16 +3915,18 @@ bool UActorChannel::ReplicateSubobject(UActorComponent* ReplicatedComponent, FOu
 		return false;
 	}
 
+	using namespace UE::Net;
+
 	// This traps actor components using the registration list even if their owning Actor isn't using the list itself.
 	if (ReplicatedComponent->IsUsingRegisteredSubObjectList() && !DataChannelInternal::bTestingLegacyMethodForComparison)
 	{
 #if SUBOBJECT_TRANSITION_VALIDATION
-		if (UE::Net::GCVarCompareSubObjectsReplicated)
+		if (GCVarCompareSubObjectsReplicated || GCVarDetectDeprecatedReplicateSubObjects)
 		{
 			TestLegacyReplicateSubObjects(ReplicatedComponent, Bunch, RepFlags);
 
-			// From here track all subobjects that write into the bunch.
-			DataChannelInternal::bTrackReplicatedSubObjects = true;
+			// From here track all subobjects from the real registry that write into the bunch.
+			DataChannelInternal::bTrackReplicatedSubObjects = GCVarCompareSubObjectsReplicated;
 		}
 #endif
 
@@ -3913,7 +3941,7 @@ bool UActorChannel::ReplicateSubobject(UActorComponent* ReplicatedComponent, FOu
 		bWroteSomethingImportant |= WriteComponentSubObjects(ReplicatedComponent, Bunch, RepFlags, ConditionMap);
 
 #if SUBOBJECT_TRANSITION_VALIDATION
-		if (UE::Net::GCVarCompareSubObjectsReplicated)
+		if (GCVarCompareSubObjectsReplicated)
 		{
 			ValidateReplicatedSubObjects();
 
diff --git a/Engine/Source/Runtime/GameplayTasks/Private/GameplayTasksComponent.cpp b/Engine/Source/Runtime/GameplayTasks/Private/GameplayTasksComponent.cpp
--- a/Engine/Source/Runtime/GameplayTasks/Private/GameplayTasksComponent.cpp
+++ b/Engine/Source/Runtime/GameplayTasks/Private/GameplayTasksComponent.cpp
@@ -179,6 +179,15 @@ void UGameplayTasksComponent::GetLifetimeReplicatedProps(TArray< FLifetimeProper
 
 bool UGameplayTasksComponent::ReplicateSubobjects(UActorChannel* Channel, class FOutBunch *Bunch, FReplicationFlags *RepFlags)
 {
+#if SUBOBJECT_TRANSITION_VALIDATION
+	// When true it means we are calling this function to find any leftover replicated subobjects in classes that transitioned to the new registry list.
+	// This shared class needs to keep supporting the old ways until we fully deprecate the API, so by only returning false we prevent the ensures to trigger
+	if (UActorChannel::CanIgnoreDeprecatedReplicateSubObjects())
+	{
+		return false;
+	}
+#endif
+
 	bool WroteSomething = Super::ReplicateSubobjects(Channel, Bunch, RepFlags);
 	
 	if (!RepFlags->bNetOwner)
