Fixed check condition on GEventDrivenLoaderEnabled

diff --git a/Engine/Source/Runtime/Engine/Private/Materials/MaterialRenderProxy.cpp b/Engine/Source/Runtime/Engine/Private/Materials/MaterialRenderProxy.cpp
--- a/Engine/Source/Runtime/Engine/Private/Materials/MaterialRenderProxy.cpp
+++ b/Engine/Source/Runtime/Engine/Private/Materials/MaterialRenderProxy.cpp
@@ -444,7 +444,7 @@ void FMaterialRenderProxy::CacheUniformExpressions(bool bRecreateUniformBuffer)
 	// Register the render proxy's as a render resource so it can receive notifications to free the uniform buffer.
 	InitResource();
 
-	bool bUsingNewLoader = EVENT_DRIVEN_ASYNC_LOAD_ACTIVE_AT_RUNTIME && GEventDrivenLoaderEnabled;
+	bool bUsingNewLoader = FPlatformProperties::RequiresCookedData();
 
 	check((bUsingNewLoader && GIsInitialLoad) || // The EDL at boot time maybe not load the default materials first; we need to intialize materials before the default materials are done
 		UMaterial::GetDefaultMaterial(MD_Surface));
