Export buffer Serialize() functions.
Want this for storing vertex buffers in GeometryCollection object.

diff --git a/Engine/Source/Runtime/Engine/Public/RawIndexBuffer.h b/Engine/Source/Runtime/Engine/Public/RawIndexBuffer.h
--- a/Engine/Source/Runtime/Engine/Public/RawIndexBuffer.h
+++ b/Engine/Source/Runtime/Engine/Public/RawIndexBuffer.h
@@ -293,7 +293,7 @@ public:
 	 * @param	Ar				Archive to serialize with
 	 * @param	bNeedsCPUAccess	Whether the elements need to be accessed by the CPU
 	 */
-	void Serialize(FArchive& Ar, bool bNeedsCPUAccess);
+	ENGINE_API void Serialize(FArchive& Ar, bool bNeedsCPUAccess);
 
 	/** Serialize only meta data (e.g. number of indices) but not the actual index data */
 	void SerializeMetaData(FArchive& Ar);
diff --git a/Engine/Source/Runtime/Engine/Public/Rendering/ColorVertexBuffer.h b/Engine/Source/Runtime/Engine/Public/Rendering/ColorVertexBuffer.h
--- a/Engine/Source/Runtime/Engine/Public/Rendering/ColorVertexBuffer.h
+++ b/Engine/Source/Runtime/Engine/Public/Rendering/ColorVertexBuffer.h
@@ -58,7 +58,7 @@ public:
 	* @param	Ar					Archive to serialize with
 	* @param	bInNeedsCPUAccess	Whether the elements need to be accessed by the CPU
 	*/
-	void Serialize(FArchive& Ar, bool bNeedsCPUAccess);
+	ENGINE_API void Serialize(FArchive& Ar, bool bNeedsCPUAccess);
 
 	void SerializeMetaData(FArchive& Ar);
 
diff --git a/Engine/Source/Runtime/Engine/Public/Rendering/PositionVertexBuffer.h b/Engine/Source/Runtime/Engine/Public/Rendering/PositionVertexBuffer.h
--- a/Engine/Source/Runtime/Engine/Public/Rendering/PositionVertexBuffer.h
+++ b/Engine/Source/Runtime/Engine/Public/Rendering/PositionVertexBuffer.h
@@ -64,7 +64,7 @@ public:
 	* @param	Ar					Archive to serialize with
 	* @param	bInNeedsCPUAccess	Whether the elements need to be accessed by the CPU
 	*/
-	void Serialize(FArchive& Ar, bool bInNeedsCPUAccess);
+	ENGINE_API void Serialize(FArchive& Ar, bool bInNeedsCPUAccess);
 
 	void SerializeMetaData(FArchive& Ar);
 
diff --git a/Engine/Source/Runtime/Engine/Public/Rendering/StaticMeshVertexBuffer.h b/Engine/Source/Runtime/Engine/Public/Rendering/StaticMeshVertexBuffer.h
--- a/Engine/Source/Runtime/Engine/Public/Rendering/StaticMeshVertexBuffer.h
+++ b/Engine/Source/Runtime/Engine/Public/Rendering/StaticMeshVertexBuffer.h
@@ -206,7 +206,7 @@ public:
 	* @param	Ar				Archive to serialize with
 	* @param	bNeedsCPUAccess	Whether the elements need to be accessed by the CPU
 	*/
-	void Serialize(FArchive& Ar, bool bNeedsCPUAccess);
+	ENGINE_API void Serialize(FArchive& Ar, bool bNeedsCPUAccess);
 
 	void SerializeMetaData(FArchive& Ar);
 
