Fix for crash when changing codecs on playing soundwave set to ForceInline

diff --git a/Engine/Source/Runtime/Engine/Classes/Sound/SoundWave.h b/Engine/Source/Runtime/Engine/Classes/Sound/SoundWave.h
--- a/Engine/Source/Runtime/Engine/Classes/Sound/SoundWave.h
+++ b/Engine/Source/Runtime/Engine/Classes/Sound/SoundWave.h
@@ -1046,6 +1046,9 @@ private:
 	// Removes any in-progress async loading data formats. 
 	void FlushAsyncLoadingDataFormats();
 
+	// Waits for audio rendering commands to execute
+	void FlushAudioRenderingCommands() const;
+
 	void BakeFFTAnalysis();
 	void BakeEnvelopeAnalysis();
 
diff --git a/Engine/Source/Runtime/Engine/Private/SoundWave.cpp b/Engine/Source/Runtime/Engine/Private/SoundWave.cpp
--- a/Engine/Source/Runtime/Engine/Private/SoundWave.cpp
+++ b/Engine/Source/Runtime/Engine/Private/SoundWave.cpp
@@ -1638,6 +1638,27 @@ void USoundWave::FlushAsyncLoadingDataFormats()
 		AsyncLoadingDataFormats.Empty();
 	}
 }
+
+void USoundWave::FlushAudioRenderingCommands() const
+{
+	check(IsInAudioThread());
+
+	// Waits for audio rendering commands to execute
+	if (GEngine && !GExitPurge)
+	{
+		FAudioDeviceManager* DeviceManager = GEngine->GetAudioDeviceManager();
+		if (DeviceManager)
+		{
+			DeviceManager->IterateOverAllDevices([](Audio::FDeviceId InDeviceID, FAudioDevice* InDevice)
+			{
+				if (InDevice)
+				{
+					InDevice->FlushAudioRenderingCommands(false /* bPumpSynchronously */);
+				}
+			});
+		}
+	}
+}
 #endif
 
 bool USoundWave::HasStreamingChunks()
@@ -3865,8 +3886,14 @@ void USoundWave::OverrideLoadingBehavior(ESoundWaveLoadingBehavior InLoadingBeha
 #if WITH_EDITOR
 void USoundWave::UpdateAsset(bool bMarkDirty)
 {
-	InvalidateCompressedData();
+	// Free playing sounds
 	FreeResources();
+	
+	// Flush audio rendering commands to ensure stopped sounds have finished decoding
+	// before invalidating compressed data.
+	FlushAudioRenderingCommands(); 
+
+	InvalidateCompressedData();
 	UpdatePlatformData();
 	
 	// Reset any error state we've encountered.
