Remove weak pointer collection from UNetDriver::AddReferencedObjects

IsValid() won't be true if there is only these reference to an object, since the object will remain Unreachable. When other strong references exist, we get a non-deterministic amount of references collected during reachability analysis depending on if the target object has been marked as reachable yet.

#rnx

diff --git a/Engine/Source/Runtime/Engine/Private/NetDriver.cpp b/Engine/Source/Runtime/Engine/Private/NetDriver.cpp
--- a/Engine/Source/Runtime/Engine/Private/NetDriver.cpp
+++ b/Engine/Source/Runtime/Engine/Private/NetDriver.cpp
@@ -3771,11 +3771,6 @@ void UNetDriver::AddReferencedObjects(UObject* InThis, FReferenceCollector& Coll
 	
 	for (FObjectReplicator* Replicator : This->AllOwnedReplicators)
 	{
-		if (Replicator->GetWeakObjectPtr().IsValid())
-		{
-			Collector.AddStableReference(&Replicator->ObjectPtr);
-		}
-
 		Collector.AddStableReference(&Replicator->ObjectClass);
 	}
 
