Guard against integer overflow under allocation in Custom Instance Data by not allocating. Access to this buffer is not occurring when NumCustomDataFloats is 0.

diff --git a/Engine/Source/Runtime/Engine/Private/InstancedStaticMesh.cpp b/Engine/Source/Runtime/Engine/Private/InstancedStaticMesh.cpp
--- a/Engine/Source/Runtime/Engine/Private/InstancedStaticMesh.cpp
+++ b/Engine/Source/Runtime/Engine/Private/InstancedStaticMesh.cpp
@@ -742,6 +742,14 @@ void FStaticMeshInstanceData::Serialize(FArchive& Ar)
 
 	if (Ar.IsLoading())
 	{
+		const int64 NumTotalCustomDataFloats = (int64)NumCustomDataFloats * NumInstances;
+		if (!IntFitsIn<int32>(NumTotalCustomDataFloats))
+		{
+			// Sanitize inputs. Allocate no custom data to avoid out of range access
+			NumCustomDataFloats = 0;
+			ensureMsgf(false, TEXT("Total Custom Instance Data Floats count is out of range."));
+		}
+
 		AllocateBuffers(NumInstances);
 	}
 
