Moved interface for Horde environment to its own lib so it can be shared between EpicStudioAnalyticsPlugin and EngineAnalytics

diff --git a/Engine/Source/Developer/Horde/Horde.Build.cs b/Engine/Source/Developer/Horde/Horde.Build.cs
new file mode 100644
--- /dev/null
+++ b/Engine/Source/Developer/Horde/Horde.Build.cs
@@ -0,0 +1,12 @@
+// Copyright Epic Games, Inc. All Rights Reserved.
+
+using UnrealBuildTool;
+using System.IO;
+
+public class Horde : ModuleRules
+{
+	public Horde(ReadOnlyTargetRules Target) : base(Target)
+	{
+		PrivateDependencyModuleNames.AddRange(new string[] { "Core" });
+	}
+}
diff --git a/Engine/Source/Developer/Horde/Private/Horde.cpp b/Engine/Source/Developer/Horde/Private/Horde.cpp
new file mode 100644
--- /dev/null
+++ b/Engine/Source/Developer/Horde/Private/Horde.cpp
@@ -0,0 +1,105 @@
+// Copyright Epic Games, Inc. All Rights Reserved.
+
+#include "Horde.h"
+#include "Misc/CommandLine.h"
+#include "Misc/Parse.h"
+
+#if WITH_EDITOR
+
+FString FHorde::GetTemplateName()
+{
+	static FString TemplateName;
+
+	if (TemplateName.IsEmpty())
+	{
+		if (false == FParse::Value(FCommandLine::Get(), TEXT("HordeTemplateName="), TemplateName))
+		{
+			TemplateName = FPlatformMisc::GetEnvironmentVariable(TEXT("UE_HORDE_TEMPLATENAME"));
+		}
+	}
+	return TemplateName;
+}
+
+FString FHorde::GetTemplateId()
+{
+	static FString TemplateId;
+
+	if (TemplateId.IsEmpty())
+	{
+		if (false == FParse::Value(FCommandLine::Get(), TEXT("HordeTemplateId="), TemplateId))
+		{
+			TemplateId = FPlatformMisc::GetEnvironmentVariable(TEXT("UE_HORDE_TEMPLATEID"));
+		}
+	}
+
+	return TemplateId;
+}
+
+FString FHorde::GetJobId()
+{
+	static FString JobId;
+
+	if (JobId.IsEmpty())
+	{
+		if (false == FParse::Value(FCommandLine::Get(), TEXT("HordeJobId="), JobId))
+		{
+			JobId = FPlatformMisc::GetEnvironmentVariable(TEXT("UE_HORDE_JOBID"));
+		}
+	}
+
+	return JobId;
+}
+
+FString FHorde::GetJobURL()
+{
+	return FString::Printf(TEXT("https://horde.devtools.epicgames.com/job/%s"), *GetJobId());
+}
+
+FString FHorde::GetStepId()
+{
+	static FString StepId;
+
+	if (StepId.IsEmpty())
+	{
+		if (false == FParse::Value(FCommandLine::Get(), TEXT("HordeStepId="), StepId))
+		{
+			StepId = FPlatformMisc::GetEnvironmentVariable(TEXT("UE_HORDE_STEPID"));
+		}
+	}
+	return StepId;
+}
+
+FString FHorde::GetStepURL()
+{
+	return FString::Printf(TEXT("%s?step=%s"), *GetJobURL(), *GetStepId());
+}
+
+FString FHorde::GetStepName()
+{
+	FString StepName;
+
+	if (StepName.IsEmpty())
+	{
+		if (false == FParse::Value(FCommandLine::Get(), TEXT("HordeStepName="), StepName))
+		{
+			StepName = FPlatformMisc::GetEnvironmentVariable(TEXT("UE_HORDE_STEPNAME"));
+		}
+	}
+	return StepName;
+}
+
+
+FString FHorde::GetBatchId()
+{
+	FString BatchId;
+	if (BatchId.IsEmpty())
+	{
+		if (false == FParse::Value(FCommandLine::Get(), TEXT("HordeBatchId="), BatchId))
+		{
+			BatchId = FPlatformMisc::GetEnvironmentVariable(TEXT("UE_HORDE_BATCHID"));
+		}
+	}
+	return BatchId;
+} 
+
+#endif
\ No newline at end of file
diff --git a/Engine/Source/Developer/Horde/Public/Horde.h b/Engine/Source/Developer/Horde/Public/Horde.h
new file mode 100644
--- /dev/null
+++ b/Engine/Source/Developer/Horde/Public/Horde.h
@@ -0,0 +1,24 @@
+// Copyright Epic Games, Inc. All Rights Reserved.
+
+#pragma once
+
+#include "CoreMinimal.h"
+#include "Containers/UnrealString.h"
+
+#if WITH_EDITOR
+
+#define UE_API HORDE_API
+
+struct FHorde
+{
+	UE_API static FString GetTemplateName();
+	UE_API static FString GetTemplateId();
+	UE_API static FString GetJobId();
+	UE_API static FString GetJobURL();
+	UE_API static FString GetStepId();
+	UE_API static FString GetStepURL();
+	UE_API static FString GetStepName();
+	UE_API static FString GetBatchId();
+};
+
+#endif
diff --git a/Engine/Source/Runtime/Engine/Engine.Build.cs b/Engine/Source/Runtime/Engine/Engine.Build.cs
--- a/Engine/Source/Runtime/Engine/Engine.Build.cs
+++ b/Engine/Source/Runtime/Engine/Engine.Build.cs
@@ -22,7 +22,7 @@ public class Engine : ModuleRules
 				"AudioMixerCore",
 				"MovieSceneCapture", 
 				"PacketHandler", 
-				"Renderer", 
+				"Renderer"
 			}
 		);
 
@@ -187,7 +187,8 @@ public class Engine : ModuleRules
 
 			PublicDependencyModuleNames.AddRange(
 				new string[] {
-					"TextureBuildUtilities"
+					"TextureBuildUtilities",
+					"Horde"
 				}
 			);
 
diff --git a/Engine/Source/Runtime/Engine/Private/EngineAnalytics.cpp b/Engine/Source/Runtime/Engine/Private/EngineAnalytics.cpp
--- a/Engine/Source/Runtime/Engine/Private/EngineAnalytics.cpp
+++ b/Engine/Source/Runtime/Engine/Private/EngineAnalytics.cpp
@@ -10,6 +10,7 @@
 #include "IAnalyticsProviderET.h"
 #include "GeneralProjectSettings.h"
 #include "Misc/EngineVersion.h"
+#include "BuildSettings.h"
 #include "RHI.h"
 #include "GenericPlatform/GenericPlatformCrashContext.h"
 #include "StudioAnalytics.h"
@@ -20,6 +21,7 @@
 #include "AnalyticsSessionSummarySender.h"
 #include "Analytics/EditorAnalyticsSessionSummary.h"
 #include "EditorAnalyticsSession.h" // DEPRECATED: kept around to clean up expired old sessions.
+#include "Horde.h"
 #endif
 
 bool FEngineAnalytics::bIsInitialized;
@@ -125,6 +127,7 @@ void FEngineAnalytics::Initialize()
 				Analytics->SetUserID(FString::Printf(TEXT("%s|%s|%s"), *FPlatformMisc::GetLoginId(), *FPlatformMisc::GetEpicAccountId(), *FPlatformMisc::GetOperatingSystemId()));
 			}
 
+			const FString UserID = FPlatformProcess::UserName(false);
 			const UGeneralProjectSettings& ProjectSettings = *GetDefault<UGeneralProjectSettings>();
 
 			TArray<FAnalyticsEventAttribute> StartSessionAttributes;
@@ -137,6 +140,19 @@ void FEngineAnalytics::Initialize()
 			StartSessionAttributes.Emplace(TEXT("ProjectID"), ProjectSettings.ProjectID);
 			StartSessionAttributes.Emplace(TEXT("ProjectDescription"), ProjectSettings.Description);
 			StartSessionAttributes.Emplace(TEXT("ProjectVersion"), ProjectSettings.ProjectVersion);
+			StartSessionAttributes.Emplace(TEXT("Application.Commandline"), FCommandLine::Get());
+			StartSessionAttributes.Emplace(TEXT("User.ID"), UserID);
+			StartSessionAttributes.Emplace(TEXT("Build.Configuration"), LexToString(FApp::GetBuildConfiguration()));
+			StartSessionAttributes.Emplace(TEXT("Build.IsInternalBuild"), FEngineBuildSettings::IsInternalBuild());
+			StartSessionAttributes.Emplace(TEXT("Build.IsPerforceBuild"), FEngineBuildSettings::IsPerforceBuild());
+			StartSessionAttributes.Emplace(TEXT("Build.IsPromotedBuild"), FApp::GetEngineIsPromotedBuild() == 0 ? false : true);
+			StartSessionAttributes.Emplace(TEXT("Build.BranchName"), FApp::GetBranchName());
+			StartSessionAttributes.Emplace(TEXT("Build.Changelist"), BuildSettings::GetCurrentChangelist());
+			StartSessionAttributes.Emplace(TEXT("Config.IsEditor"), GIsEditor);
+			StartSessionAttributes.Emplace(TEXT("Config.IsUnattended"), FApp::IsUnattended());
+			StartSessionAttributes.Emplace(TEXT("Config.IsBuildMachine"), GIsBuildMachine);
+			StartSessionAttributes.Emplace(TEXT("Config.IsRunningCommandlet"), IsRunningCommandlet());
+			StartSessionAttributes.Emplace(TEXT("Platform.IsRemoteSession"), FPlatformMisc::IsRemoteSession());
 			StartSessionAttributes.Emplace(TEXT("GPUVendorID"), GRHIVendorId);
 			StartSessionAttributes.Emplace(TEXT("GPUDeviceID"), GRHIDeviceId);
 			StartSessionAttributes.Emplace(TEXT("GRHIDeviceRevision"), GRHIDeviceRevision);
@@ -154,6 +170,18 @@ void FEngineAnalytics::Initialize()
 			StartSessionAttributes.Emplace(TEXT("OSMinor"), OSMinor);
 			StartSessionAttributes.Emplace(TEXT("OSVersion"), FPlatformMisc::GetOSVersion());
 			StartSessionAttributes.Emplace(TEXT("Is64BitOS"), FPlatformMisc::Is64bitOperatingSystem());
+
+#if WITH_EDITOR
+			StartSessionAttributes.Emplace(TEXT("Horde.TemplateID"), FHorde::GetTemplateId());
+			StartSessionAttributes.Emplace(TEXT("Horde.TemplateName"), FHorde::GetTemplateName());
+			StartSessionAttributes.Emplace(TEXT("Horde.JobURL"), FHorde::GetJobURL());
+			StartSessionAttributes.Emplace(TEXT("Horde.JobID"), FHorde::GetJobId());
+			StartSessionAttributes.Emplace(TEXT("Horde.StepName"), FHorde::GetStepName());
+			StartSessionAttributes.Emplace(TEXT("Horde.StepID"), FHorde::GetStepId());
+			StartSessionAttributes.Emplace(TEXT("Horde.StepURL"), FHorde::GetStepURL());
+			StartSessionAttributes.Emplace(TEXT("Horde.BatchID"), FHorde::GetBatchId());
+#endif
+
 #if PLATFORM_MAC
 #if PLATFORM_MAC_ARM64
             StartSessionAttributes.Emplace(TEXT("UEBuildArch"), FString(TEXT("AppleSilicon")));
