World Partition: when parsing actor descriptors from a container, do validations through the asset registry for missing BP classes, following redirectors if needed.

diff --git a/Engine/Source/Runtime/Engine/Private/WorldPartition/ActorDescContainer.cpp b/Engine/Source/Runtime/Engine/Private/WorldPartition/ActorDescContainer.cpp
--- a/Engine/Source/Runtime/Engine/Private/WorldPartition/ActorDescContainer.cpp
+++ b/Engine/Source/Runtime/Engine/Private/WorldPartition/ActorDescContainer.cpp
@@ -48,7 +48,7 @@ void UActorDescContainer::Initialize(const FInitializeParams& InitParams)
 	{
 		const FString ContainerExternalActorsPath = GetExternalActorPath();
 
-		// Do a synchronous scan of the level external actors path.			
+		// Do a synchronous scan of the level external actors path.					
 		IAssetRegistry& AssetRegistry = FModuleManager::LoadModuleChecked<FAssetRegistryModule>(TEXT("AssetRegistry")).Get();
 		{
 			TRACE_CPUPROFILER_EVENT_SCOPE(ScanPathsSynchronous);
@@ -70,14 +70,13 @@ void UActorDescContainer::Initialize(const FInitializeParams& InitParams)
 		{
 			TUniquePtr<FWorldPartitionActorDesc> ActorDesc = FWorldPartitionActorDescUtils::GetActorDescriptorFromAssetData(Asset);
 
-			if (ActorDesc.IsValid() && ActorDesc->GetNativeClass().IsValid() && (!InitParams.FilterActorDesc || InitParams.FilterActorDesc(ActorDesc.Get())))
-			{
-				AddActorDescriptor(ActorDesc.Release(), OwningWorld);
-			}
-			else
+			if (!ActorDesc.IsValid() || !FWorldPartitionActorDescUtils::ValidateActorDescClass(ActorDesc.Get()) || (InitParams.FilterActorDesc && !InitParams.FilterActorDesc(ActorDesc.Get())))
 			{
 				InvalidActors.Emplace(MoveTemp(ActorDesc));
+				continue;
 			}
+
+			AddActorDescriptor(ActorDesc.Release(), OwningWorld);
 		}
 	}
 
diff --git a/Engine/Source/Runtime/Engine/Private/WorldPartition/WorldPartitionActorDescUtils.cpp b/Engine/Source/Runtime/Engine/Private/WorldPartition/WorldPartitionActorDescUtils.cpp
--- a/Engine/Source/Runtime/Engine/Private/WorldPartition/WorldPartitionActorDescUtils.cpp
+++ b/Engine/Source/Runtime/Engine/Private/WorldPartition/WorldPartitionActorDescUtils.cpp
@@ -4,6 +4,7 @@
 #include "WorldPartition/WorldPartitionActorDescUtils.h"
 #include "AssetRegistry/ARFilter.h"
 #include "AssetRegistry/AssetData.h"
+#include "AssetRegistry/AssetRegistryModule.h"
 #include "Modules/ModuleManager.h"
 #include "WorldPartition/WorldPartitionLog.h"
 #include "Engine/Level.h"
@@ -26,6 +27,22 @@ FName FWorldPartitionActorDescUtils::ActorMetaDataTagName()
 	return NAME_ActorMetaData;
 }
 
+static FString ResolveClassRedirector(const FString& InClassName)
+{
+	FString ClassName;
+	FString ClassPackageName;
+	if (!InClassName.Split(TEXT("."), &ClassPackageName, &ClassName))
+	{
+		ClassName = *InClassName;
+	}
+
+	// Look for a class redirectors
+	const FCoreRedirectObjectName OldClassName = FCoreRedirectObjectName(*ClassName, NAME_None, *ClassPackageName);
+	const FCoreRedirectObjectName NewClassName = FCoreRedirects::GetRedirectedName(ECoreRedirectFlags::Type_Class, OldClassName);
+
+	return NewClassName.ToString();
+}
+
 bool FWorldPartitionActorDescUtils::IsValidActorDescriptorFromAssetData(const FAssetData& InAssetData)
 {
 	return InAssetData.FindTag(NAME_ActorMetaDataClass) && InAssetData.FindTag(NAME_ActorMetaData);
@@ -36,19 +53,11 @@ UClass* FWorldPartitionActorDescUtils::GetActorNativeClassFromAssetData(const FA
 	FString ActorMetaDataClass;
 	if (InAssetData.GetTagValue(NAME_ActorMetaDataClass, ActorMetaDataClass))
 	{
-		FString ActorClassName;
-		FString ActorPackageName;
-		if (!ActorMetaDataClass.Split(TEXT("."), &ActorPackageName, &ActorClassName))
-		{
-			ActorClassName = *ActorMetaDataClass;
-		}
-
 		// Look for a class redirectors
-		const FCoreRedirectObjectName OldClassName = FCoreRedirectObjectName(*ActorClassName, NAME_None, *ActorPackageName);
-		const FCoreRedirectObjectName NewClassName = FCoreRedirects::GetRedirectedName(ECoreRedirectFlags::Type_Class, OldClassName);
+		const FString ActorNativeClassName = ResolveClassRedirector(ActorMetaDataClass);
 		
 		// Handle deprecated short class names
-		const FTopLevelAssetPath ClassPath = FAssetData::TryConvertShortClassNameToPathName(*NewClassName.ToString(), ELogVerbosity::Log);
+		const FTopLevelAssetPath ClassPath = FAssetData::TryConvertShortClassNameToPathName(*ActorNativeClassName, ELogVerbosity::Log);
 
 		// Lookup the native class
 		return UClass::TryFindTypeSlow<UClass>(ClassPath.ToString(), EFindFirstObjectOptions::ExactClass);
@@ -75,26 +84,9 @@ TUniquePtr<FWorldPartitionActorDesc> FWorldPartitionActorDescUtils::GetActorDesc
 			
 		if (!ActorDescInitData.NativeClass)
 		{
-			UE_LOG(LogWorldPartition, Log, TEXT("Invalid class for actor guid `%s` ('%s') from package '%s'"), *NewActorDesc->GetGuid().ToString(), *NewActorDesc->GetActorName().ToString(), *NewActorDesc->GetActorPackage().ToString());
+			UE_LOG(LogWorldPartition, Warning, TEXT("Invalid class for actor guid `%s` ('%s') from package '%s'"), *NewActorDesc->GetGuid().ToString(), *NewActorDesc->GetActorName().ToString(), *NewActorDesc->GetActorPackage().ToString());
 			NewActorDesc->NativeClass.Reset();
-			return NewActorDesc;
 		}
-		/*else if (UClass* Class = FindObject<UClass>(InAssetData.AssetClassPath); !Class)
-		{
-			// We can't detect mising BP classes for inactive plugins, etc.
-			if (InAssetData.AssetClassPath.GetPackageName().ToString().StartsWith(TEXT("/Game/")))
-			{
-				TArray<FAssetData> BlueprintClass;
-				IAssetRegistry& AssetRegistry = FModuleManager::LoadModuleChecked<FAssetRegistryModule>(TEXT("AssetRegistry")).Get();
-				AssetRegistry.GetAssetsByPackageName(*InAssetData.AssetClassPath.GetPackageName().ToString(), BlueprintClass, true);
-
-				if (!BlueprintClass.Num())
-				{
-					UE_LOG(LogWorldPartition, Log, TEXT("Failed to find class '%s' for actor '%s"), *InAssetData.AssetClassPath.ToString(), *NewActorDesc->GetActorSoftPath().ToString());
-					return nullptr;
-				}
-			}
-		}*/
 
 		return NewActorDesc;
 	}
@@ -166,4 +158,76 @@ void FWorldPartitionActorDescUtils::ReplaceActorDescriptorPointerFromActor(const
 	check(!InActorDesc->ActorPtr.IsValid() || (InActorDesc->ActorPtr == InOldActor));
 	InActorDesc->ActorPtr = InNewActor;
 }
+
+bool FWorldPartitionActorDescUtils::ValidateActorDescClass(FWorldPartitionActorDesc* InActorDesc)
+{
+	// If the native class in invalid (potentially deleted), it means we parsed the actor descriptor with AActor::StaticClass and explicitly
+	// marked the class as invalid.
+	if (!InActorDesc->GetNativeClass().IsValid())
+	{
+		UE_LOG(LogWorldPartition, Warning, TEXT("Failed to find native class for actor '%s"), *InActorDesc->GetActorSoftPath().ToString());
+		return false;
+	}
+
+	// If the base class is invalid, it means the actor is from a native class.
+	if (!InActorDesc->GetBaseClass().IsValid())
+	{
+		return true;
+	}
+
+	// Lookup the Bp class, if it's already loaded we don't need to validate anything.
+	if (UClass* BaseClass = FindObject<UClass>(InActorDesc->GetBaseClass()))
+	{
+		return true;
+	}
+
+	// The actor is from a BP class which isn't loaded. To avoid loading the class, go through the asset registry to validate the class.
+	const FAssetData* ClassData;
+	FString ActorDescBaseClass = InActorDesc->GetBaseClass().ToString();
+
+	IAssetRegistry& AssetRegistry = FModuleManager::LoadModuleChecked<FAssetRegistryModule>(TEXT("AssetRegistry")).Get();
+
+	for (;;)
+	{
+		ActorDescBaseClass.RemoveFromEnd(TEXT("_C"), ESearchCase::CaseSensitive);
+
+		// Look for a class redirectors
+		const FString ActorClassName = ResolveClassRedirector(ActorDescBaseClass);
+
+		FTopLevelAssetPath AssetClassPath(*ActorClassName);
+
+		TArray<FAssetData> BlueprintAssets;							
+		AssetRegistry.ScanFilesSynchronous({ AssetClassPath.GetPackageName().ToString() }, /*bForceRescan*/false);
+		AssetRegistry.GetAssetsByPackageName(AssetClassPath.GetPackageName(), BlueprintAssets, /*bIncludeOnlyOnDiskAssets*/true);
+
+		if (!BlueprintAssets.Num())
+		{
+			UE_LOG(LogWorldPartition, Warning, TEXT("Failed to find assets for class '%s' for actor '%s"), *AssetClassPath.ToString(), *InActorDesc->GetActorSoftPath().ToString());
+			return false;
+		}
+
+		ClassData = BlueprintAssets.FindByPredicate([&AssetClassPath](const FAssetData& AssetData) { return AssetData.ToSoftObjectPath().GetAssetPath() == AssetClassPath; });
+		if (!ClassData)
+		{
+			UE_LOG(LogWorldPartition, Warning, TEXT("Failed to find class asset '%s' for actor '%s"), *AssetClassPath.ToString(), *InActorDesc->GetActorSoftPath().ToString());
+			return false;
+		}
+
+		if (!ClassData->IsRedirector())
+		{
+			break;
+		}
+
+		FString DestinationObjectPath;
+		if (!ClassData->GetTagValue(TEXT("DestinationObject"), DestinationObjectPath))
+		{
+			UE_LOG(LogWorldPartition, Warning, TEXT("Failed to follow class redirector for '%s' for actor '%s"), *AssetClassPath.ToString(), *InActorDesc->GetActorSoftPath().ToString());
+			return false;
+		}
+
+		ActorDescBaseClass = DestinationObjectPath;
+	}
+
+	return true;
+};
 #endif
diff --git a/Engine/Source/Runtime/Engine/Public/WorldPartition/WorldPartitionActorDescUtils.h b/Engine/Source/Runtime/Engine/Public/WorldPartition/WorldPartitionActorDescUtils.h
--- a/Engine/Source/Runtime/Engine/Public/WorldPartition/WorldPartitionActorDescUtils.h
+++ b/Engine/Source/Runtime/Engine/Public/WorldPartition/WorldPartitionActorDescUtils.h
@@ -72,5 +72,12 @@ struct ENGINE_API FWorldPartitionActorDescUtils
 	 * @param InActorDesc	Actor descriptor that will get its actor pointer updated.
 	 */
 	static void ReplaceActorDescriptorPointerFromActor(const AActor* InOldActor, AActor* InNewActor, FWorldPartitionActorDesc* InActorDesc);
+
+	/** 
+	 * Validate if the actor descriptor's class exists without loading anything by going through the asset registry, walking the redirector chain if needed.
+	 * @param InActorDesc		The actor descriptor to resolve the class from.
+	 * @return					Whether the actor descriptor's class is valid or not.
+	 */
+	static bool ValidateActorDescClass(FWorldPartitionActorDesc* InActorDesc);
 };
 #endif
\ No newline at end of file
