Cloth - Fixed crash when using Collide with Cloth Children.

#rnx

diff --git a/Engine/Source/Runtime/Engine/Private/SkeletalMeshComponentPhysics.cpp b/Engine/Source/Runtime/Engine/Private/SkeletalMeshComponentPhysics.cpp
--- a/Engine/Source/Runtime/Engine/Private/SkeletalMeshComponentPhysics.cpp
+++ b/Engine/Source/Runtime/Engine/Private/SkeletalMeshComponentPhysics.cpp
@@ -2972,7 +2972,7 @@ void USkeletalMeshComponent::CopyClothCollisionsToChildren()
 	for (USceneComponent* AttachedChild : GetAttachChildren())
 	{
 		USkeletalMeshComponent* pChild = Cast<USkeletalMeshComponent>(AttachedChild);
-		if(pChild && pChild->ClothingSimulation->ShouldSimulate())
+		if(pChild && pChild->ClothingSimulation && pChild->ClothingSimulation->ShouldSimulate())
 		{
 			ClothChildren.Add(pChild);
 		}
