Fixed morph corruption issue in GPUMorphUpdateCS by adding an early out when GroupIndex >= NumGroups.
This can happen when the dispatch is so large that it starts using wrapped 2d coordinates.
Also added early out in GPUMorphNormalizeCS for good measure.

diff --git a/Engine/Shaders/Private/MorphTargets.usf b/Engine/Shaders/Private/MorphTargets.usf
--- a/Engine/Shaders/Private/MorphTargets.usf
+++ b/Engine/Shaders/Private/MorphTargets.usf
@@ -28,6 +28,7 @@ float4 MorphTargetWeights[GMorphTargetDispatchBatchSize / 4];
 
 float4 PositionScale;
 float2 Precision;
+uint NumGroups;
 
 groupshared uint LocalGroupOffsets[GMorphTargetDispatchBatchSize];
 
@@ -114,6 +115,11 @@ void GPUMorphUpdateCS(uint3 GroupId : SV_GroupID, uint GroupThreadIndex : SV_Gro
 	LoadLocalData(GroupThreadIndex);
 
 	const uint GroupIndex = GetUnWrappedDispatchGroupId(GroupId);
+	if (GroupIndex >= NumGroups)
+	{
+		return;
+	}
+
 	const uint MorphTargetLocalIndex = GroupIndexToMorphTargetLocalIndex(GroupIndex);
 
 	const uint BatchStartIndex = MorphTargetBatchOffsets[MorphTargetLocalIndex >> 2][MorphTargetLocalIndex & 3u];			//TODO: use a buffer instead of packing into constant buffer?
@@ -152,12 +158,18 @@ void GPUMorphUpdateCS(uint3 GroupId : SV_GroupID, uint GroupThreadIndex : SV_Gro
 	}
 }
 
+uint NumVertices;
+
 [numthreads(64, 1, 1)]
 void GPUMorphNormalizeCS(uint3 GroupId : SV_GroupID, uint GroupThreadIndex : SV_GroupIndex)
 {
 	const uint DestVertexIndex = GetUnWrappedDispatchThreadId(GroupId, GroupThreadIndex, 64);
-	const uint DestVertexAddress = DestVertexIndex * 6;
+	if (DestVertexIndex >= NumVertices)
+	{
+		return;
+	}
 
+	const uint DestVertexAddress = DestVertexIndex * 6;
 	const float3 Position = float3(int3(MorphVertexBuffer[DestVertexAddress + 0], MorphVertexBuffer[DestVertexAddress + 1], MorphVertexBuffer[DestVertexAddress + 2])) * PositionScale.xyz;
 
 	MorphVertexBuffer[DestVertexAddress + 0] = asuint(Position.x);
@@ -168,4 +180,4 @@ void GPUMorphNormalizeCS(uint3 GroupId : SV_GroupID, uint GroupThreadIndex : SV_
 	MorphVertexBuffer[DestVertexAddress + 3] = asuint(TangentZ.x);
 	MorphVertexBuffer[DestVertexAddress + 4] = asuint(TangentZ.y);
 	MorphVertexBuffer[DestVertexAddress + 5] = asuint(TangentZ.z);
-}
+}
\ No newline at end of file
diff --git a/Engine/Source/Runtime/Engine/Private/SkeletalRenderGPUSkin.cpp b/Engine/Source/Runtime/Engine/Private/SkeletalRenderGPUSkin.cpp
--- a/Engine/Source/Runtime/Engine/Private/SkeletalRenderGPUSkin.cpp
+++ b/Engine/Source/Runtime/Engine/Private/SkeletalRenderGPUSkin.cpp
@@ -1113,11 +1113,12 @@ FGPUMorphUpdateCS::FGPUMorphUpdateCS(const ShaderMetaType::CompiledShaderInitial
 	MorphTargetGroupOffsetsParameter.Bind(Initializer.ParameterMap, TEXT("MorphTargetGroupOffsets"));
 	PositionScaleParameter.Bind(Initializer.ParameterMap, TEXT("PositionScale"));
 	PrecisionParameter.Bind(Initializer.ParameterMap, TEXT("Precision"));
+	NumGroupsParameter.Bind(Initializer.ParameterMap, TEXT("NumGroups"));
 
 	MorphDataBufferParameter.Bind(Initializer.ParameterMap, TEXT("MorphDataBuffer"));
 }
 
-void FGPUMorphUpdateCS::SetParameters(FRHICommandList& RHICmdList, const FVector4& LocalScale, const FMorphTargetVertexInfoBuffers& MorphTargetVertexInfoBuffers, FMorphVertexBuffer& MorphVertexBuffer)
+void FGPUMorphUpdateCS::SetParameters(FRHICommandList& RHICmdList, const FVector4& LocalScale, const FMorphTargetVertexInfoBuffers& MorphTargetVertexInfoBuffers, FMorphVertexBuffer& MorphVertexBuffer, uint32 NumGroups)
 {
 	FRHIComputeShader* CS = RHICmdList.GetBoundComputeShader();
 
@@ -1126,6 +1127,7 @@ void FGPUMorphUpdateCS::SetParameters(FRHICommandList& RHICmdList, const FVector
 	SetShaderValue(RHICmdList, CS, PositionScaleParameter, (FVector4f)LocalScale);
 	FVector2f Precision = { MorphTargetVertexInfoBuffers.GetPositionPrecision(), MorphTargetVertexInfoBuffers.GetTangentZPrecision() };
 	SetShaderValue(RHICmdList, CS, PrecisionParameter, Precision);
+	SetShaderValue(RHICmdList, CS, NumGroupsParameter, NumGroups);
 
 	SetSRVParameter(RHICmdList, CS, MorphDataBufferParameter, MorphTargetVertexInfoBuffers.MorphDataSRV);
 }
@@ -1164,6 +1166,7 @@ FGPUMorphNormalizeCS::FGPUMorphNormalizeCS(const ShaderMetaType::CompiledShaderI
 {
 	MorphVertexBufferParameter.Bind(Initializer.ParameterMap, TEXT("MorphVertexBuffer"));
 	PositionScaleParameter.Bind(Initializer.ParameterMap, TEXT("PositionScale"));
+	NumVerticesParameter.Bind(Initializer.ParameterMap, TEXT("NumVertices"));
 }
 
 bool FGPUMorphNormalizeCS::ShouldCompilePermutation(const FGlobalShaderPermutationParameters& Parameters)
@@ -1171,11 +1174,12 @@ bool FGPUMorphNormalizeCS::ShouldCompilePermutation(const FGlobalShaderPermutati
 	return IsFeatureLevelSupported(Parameters.Platform, ERHIFeatureLevel::SM5);
 }
 
-void FGPUMorphNormalizeCS::SetParameters(FRHICommandList& RHICmdList, const FVector4& InvLocalScale, const FMorphTargetVertexInfoBuffers& MorphTargetVertexInfoBuffers, FMorphVertexBuffer& MorphVertexBuffer)
+void FGPUMorphNormalizeCS::SetParameters(FRHICommandList& RHICmdList, const FVector4& InvLocalScale, const FMorphTargetVertexInfoBuffers& MorphTargetVertexInfoBuffers, FMorphVertexBuffer& MorphVertexBuffer, uint32 NumVertices)
 {
 	FRHIComputeShader* CS = RHICmdList.GetBoundComputeShader();
 	SetUAVParameter(RHICmdList, CS, MorphVertexBufferParameter, MorphVertexBuffer.GetUAV());
 	SetShaderValue(RHICmdList, CS, PositionScaleParameter, (FVector4f)InvLocalScale);
+	SetShaderValue(RHICmdList, CS, NumVerticesParameter, NumVertices);
 }
 
 void FGPUMorphNormalizeCS::Dispatch(FRHICommandList& RHICmdList, uint32 NumVertices)
@@ -1275,7 +1279,7 @@ void FSkeletalMeshObjectGPUSkin::FSkeletalMeshObjectLOD::UpdateMorphVertexBuffer
 
 					SetComputePipelineState(RHICmdList, GPUMorphUpdateCS.GetComputeShader());
 
-					GPUMorphUpdateCS->SetParameters(RHICmdList, MorphScale, MorphTargetVertexInfoBuffers, MorphVertexBuffer);
+					GPUMorphUpdateCS->SetParameters(RHICmdList, MorphScale, MorphTargetVertexInfoBuffers, MorphVertexBuffer, NumBatches);
 					GPUMorphUpdateCS->SetMorphOffsetsAndWeights(RHICmdList, BatchOffsets, GroupOffsets, Weights);
 					GPUMorphUpdateCS->Dispatch(RHICmdList, NumBatches);
 				}
@@ -1299,7 +1303,7 @@ void FSkeletalMeshObjectGPUSkin::FSkeletalMeshObjectLOD::UpdateMorphVertexBuffer
 
 				SetComputePipelineState(RHICmdList, GPUMorphNormalizeCS.GetComputeShader());
 
-				GPUMorphNormalizeCS->SetParameters(RHICmdList, InvMorphScale, MorphTargetVertexInfoBuffers, MorphVertexBuffer);
+				GPUMorphNormalizeCS->SetParameters(RHICmdList, InvMorphScale, MorphTargetVertexInfoBuffers, MorphVertexBuffer, MorphVertexBuffer.GetNumVerticies());
 				GPUMorphNormalizeCS->Dispatch(RHICmdList, MorphVertexBuffer.GetNumVerticies());
 				GPUMorphNormalizeCS->EndAllDispatches(RHICmdList);
 				RHICmdList.Transition(FRHITransitionInfo(MorphVertexBuffer.GetUAV(), ERHIAccess::UAVCompute, ERHIAccess::VertexOrIndexBuffer | ERHIAccess::SRVMask));
diff --git a/Engine/Source/Runtime/Engine/Private/SkeletalRenderGPUSkin.h b/Engine/Source/Runtime/Engine/Private/SkeletalRenderGPUSkin.h
--- a/Engine/Source/Runtime/Engine/Private/SkeletalRenderGPUSkin.h
+++ b/Engine/Source/Runtime/Engine/Private/SkeletalRenderGPUSkin.h
@@ -712,7 +712,7 @@ public:
 
 	static const uint32 MorphTargetDispatchBatchSize = 128;
 
-	void SetParameters(FRHICommandList& RHICmdList, const FVector4& LocalScale, const FMorphTargetVertexInfoBuffers& MorphTargetVertexInfoBuffers, FMorphVertexBuffer& MorphVertexBuffer);
+	void SetParameters(FRHICommandList& RHICmdList, const FVector4& LocalScale, const FMorphTargetVertexInfoBuffers& MorphTargetVertexInfoBuffers, FMorphVertexBuffer& MorphVertexBuffer, uint32 NumGroups);
 	void SetMorphOffsetsAndWeights(FRHICommandList& RHICmdList, uint32 BatchOffsets[MorphTargetDispatchBatchSize], uint32 GroupOffsets[MorphTargetDispatchBatchSize], float Weights[MorphTargetDispatchBatchSize]);
 
 	void Dispatch(FRHICommandList& RHICmdList, uint32 Size);
@@ -729,6 +729,7 @@ protected:
 	LAYOUT_FIELD(FShaderParameter, MorphTargetGroupOffsetsParameter);
 	LAYOUT_FIELD(FShaderParameter, PositionScaleParameter);
 	LAYOUT_FIELD(FShaderParameter, PrecisionParameter);
+	LAYOUT_FIELD(FShaderParameter, NumGroupsParameter);
 
 	LAYOUT_FIELD(FShaderResourceParameter, MorphDataBufferParameter);
 };
@@ -743,13 +744,14 @@ public:
 
 	static bool ShouldCompilePermutation(const FGlobalShaderPermutationParameters& Parameters);
 
-	void SetParameters(FRHICommandList& RHICmdList, const FVector4& LocalScale, const FMorphTargetVertexInfoBuffers& MorphTargetVertexInfoBuffers, FMorphVertexBuffer& MorphVertexBuffer);
+	void SetParameters(FRHICommandList& RHICmdList, const FVector4& LocalScale, const FMorphTargetVertexInfoBuffers& MorphTargetVertexInfoBuffers, FMorphVertexBuffer& MorphVertexBuffer, uint32 NumVertices);
 
-	void Dispatch(FRHICommandList& RHICmdList, uint32 NumVerticies);
+	void Dispatch(FRHICommandList& RHICmdList, uint32 NumVertices);
 	void EndAllDispatches(FRHICommandList& RHICmdList);
 
 protected:
 	LAYOUT_FIELD(FShaderResourceParameter, MorphVertexBufferParameter);
 
 	LAYOUT_FIELD(FShaderParameter, PositionScaleParameter);
+	LAYOUT_FIELD(FShaderParameter, NumVerticesParameter);
 };
