Distinguish between the game and editor shaders in messaging and DDC key dumps.

- Necessary for debugging DerivedDataCache commandlet which is often run like -targetplatform=MacEditor+Mac

diff --git a/Engine/Source/Runtime/Engine/Private/Materials/MaterialShader.cpp b/Engine/Source/Runtime/Engine/Private/Materials/MaterialShader.cpp
--- a/Engine/Source/Runtime/Engine/Private/Materials/MaterialShader.cpp
+++ b/Engine/Source/Runtime/Engine/Private/Materials/MaterialShader.cpp
@@ -1454,7 +1454,10 @@ TSharedRef<FMaterialShaderMap::FAsyncLoadContext> FMaterialShaderMap::BeginLoadF
 
 			if (UNLIKELY(ShouldDumpShaderDDCKeys()))
 			{
-				const FString FileName = FString::Printf(TEXT("%s-%s-%s.txt"), *Material->GetAssetName(), *LexToString(ShaderMapId.FeatureLevel), *LexToString(ShaderMapId.QualityLevel));
+				const FString FileName = FString::Printf(TEXT("%s-%s-%s-%s.txt"), 
+					*Material->GetAssetName(), *LexToString(ShaderMapId.FeatureLevel), *LexToString(ShaderMapId.QualityLevel),
+					ShaderMapId.LayoutParams.WithEditorOnly() ? TEXT("Editor") : TEXT("Game")
+					);
 				DumpShaderDDCKeyToFile(InPlatform, ShaderMapId.LayoutParams.WithEditorOnly(), FileName, Result->DataKey);
 			}
 
diff --git a/Engine/Source/Runtime/Engine/Private/Materials/MaterialShared.cpp b/Engine/Source/Runtime/Engine/Private/Materials/MaterialShared.cpp
--- a/Engine/Source/Runtime/Engine/Private/Materials/MaterialShared.cpp
+++ b/Engine/Source/Runtime/Engine/Private/Materials/MaterialShared.cpp
@@ -2720,12 +2720,13 @@ bool FMaterial::CacheShaders(const FMaterialShaderMapId& ShaderMapId, EShaderPla
 				}
 #if WITH_EDITOR
 				FString ShaderPlatformName = FGenericDataDrivenShaderPlatformInfo::GetName(Platform).ToString();
-				UE_LOG(LogMaterial, Display, TEXT("%s cached shadermap for %s in %s, %s, %s (DDC key hash: %s), compiling. %s"),
+				UE_LOG(LogMaterial, Display, TEXT("%s cached shadermap for %s in %s, %s, %s, %s (DDC key hash: %s), compiling. %s"),
 					ShaderMapCondition,
 					*GetAssetName(),
 					*ShaderPlatformName,
 					*LexToString(ShaderMapId.QualityLevel),
 					*LexToString(ShaderMapId.FeatureLevel),
+					ShaderMapId.LayoutParams.WithEditorOnly() ? TEXT("Editor") : TEXT("Game"),
 					*DDCKeyHash,
 					IsSpecialEngineMaterial() ? TEXT("Is special engine material.") : TEXT("")
 				);
