- Added AActor::ShouldSkipFromLevelInstance to decide whether an actor part of a Level Instance should be considered in its parent world.
It will call a virtual ActorTypeShouldSkipFromLevelInstance() on the CDO so that any actor call can override it.
- Refactored existing FWorldPartitionActorDesc::IsRuntimeRelevant to rely on this method.
- Also use this method to skip actors that return true when breaking a level instance.

diff --git a/Engine/Source/Runtime/Engine/Classes/Engine/LevelBounds.h b/Engine/Source/Runtime/Engine/Classes/Engine/LevelBounds.h
--- a/Engine/Source/Runtime/Engine/Classes/Engine/LevelBounds.h
+++ b/Engine/Source/Runtime/Engine/Classes/Engine/LevelBounds.h
@@ -106,8 +106,8 @@ protected:
 	FDelegateHandle OnLevelActorMovedDelegateHandle;
 	FDelegateHandle OnLevelActorDeletedDelegateHandle;
 	FDelegateHandle OnLevelActorAddedDelegateHandle;
-#endif
-};
-
-
 
+private:
+	virtual bool ActorTypeShouldSkipFromLevelInstance() const override { return true; }
+#endif
+};
\ No newline at end of file
diff --git a/Engine/Source/Runtime/Engine/Classes/GameFramework/Actor.h b/Engine/Source/Runtime/Engine/Classes/GameFramework/Actor.h
--- a/Engine/Source/Runtime/Engine/Classes/GameFramework/Actor.h
+++ b/Engine/Source/Runtime/Engine/Classes/GameFramework/Actor.h
@@ -934,6 +934,11 @@ public:
 	 */
 	virtual bool ShouldLevelKeepRefIfExternal() const { return false; }
 
+	/**
+	 * Whether this actor should be ignored when it is loaded from a level instance
+	 */
+	bool ShouldSkipFromLevelInstance() const;
+
 	FActorOnPackagingModeChanged OnPackagingModeChanged;
 
 	/** Returns this actor's current target runtime grid. */
@@ -1231,6 +1236,7 @@ public:
 #if WITH_EDITOR
 private:
 	virtual bool ActorTypeSupportsDataLayer() const { return true; }
+	virtual bool ActorTypeShouldSkipFromLevelInstance() const { return false; }
 public:
 	bool AddDataLayer(const UDataLayerAsset* DataLayerAsset);
 	bool AddDataLayer(const UDataLayerInstance* DataLayerInstance);
diff --git a/Engine/Source/Runtime/Engine/Private/ActorEditor.cpp b/Engine/Source/Runtime/Engine/Private/ActorEditor.cpp
--- a/Engine/Source/Runtime/Engine/Private/ActorEditor.cpp
+++ b/Engine/Source/Runtime/Engine/Private/ActorEditor.cpp
@@ -1015,6 +1015,18 @@ void AActor::SetHLODLayer(class UHLODLayer* InHLODLayer)
 	HLODLayer = InHLODLayer;
 }
 
+bool AActor::ShouldSkipFromLevelInstance() const
+{
+	if (HasAnyFlags(RF_ClassDefaultObject))
+	{
+		return ActorTypeShouldSkipFromLevelInstance();
+	}
+	else
+	{
+		return CastChecked<AActor>(GetClass()->GetDefaultObject())->ShouldSkipFromLevelInstance();
+	}
+}
+
 void AActor::SetPackageExternal(bool bExternal, bool bShouldDirty)
 {
 	// @todo_ow: Call FExternalPackageHelper::SetPackagingMode and keep calling the actor specific code here (components). 
diff --git a/Engine/Source/Runtime/Engine/Private/LevelInstance/LevelInstanceSubsystem.cpp b/Engine/Source/Runtime/Engine/Private/LevelInstance/LevelInstanceSubsystem.cpp
--- a/Engine/Source/Runtime/Engine/Private/LevelInstance/LevelInstanceSubsystem.cpp
+++ b/Engine/Source/Runtime/Engine/Private/LevelInstance/LevelInstanceSubsystem.cpp
@@ -1300,13 +1300,9 @@ void ULevelInstanceSubsystem::BreakLevelInstance_Impl(ILevelInstanceInterface* L
 			}
 
 			// Skip some actor types
-			// @todo_ow: Move this logic in a new virtual function
 			if ((Actor != Actor->GetLevel()->GetDefaultBrush()) &&
-				!Actor->IsA<ALevelBounds>() &&
 				!Actor->IsA<AWorldSettings>() &&
-				!Actor->IsA<ALevelInstanceEditorInstanceActor>() &&
-				!Actor->IsA<AWorldDataLayers>() &&
-				!Actor->IsA<AWorldPartitionMiniMap>())
+				!Actor->ShouldSkipFromLevelInstance())
 			{
 				if (CanMoveActorToLevel(Actor))
 				{
diff --git a/Engine/Source/Runtime/Engine/Private/WorldPartition/DataLayer/WorldDataLayersActorDesc.cpp b/Engine/Source/Runtime/Engine/Private/WorldPartition/DataLayer/WorldDataLayersActorDesc.cpp
--- a/Engine/Source/Runtime/Engine/Private/WorldPartition/DataLayer/WorldDataLayersActorDesc.cpp
+++ b/Engine/Source/Runtime/Engine/Private/WorldPartition/DataLayer/WorldDataLayersActorDesc.cpp
@@ -169,12 +169,6 @@ void FWorldDataLayersActorDesc::Serialize(FArchive& Ar)
 	}
 }
 
-bool FWorldDataLayersActorDesc::IsRuntimeRelevant(const FActorContainerID& InContainerID) const
-{
-	// @todo_ow : For now we don't consider WorldDataLayers from content bundles
-	return InContainerID.IsMainContainer() && !ContentBundleGuid.IsValid();
-}
-
 const FDataLayerInstanceDesc* FWorldDataLayersActorDesc::GetDataLayerInstanceFromInstanceName(FName InDataLayerInstanceName) const
 {
 	for (const FDataLayerInstanceDesc& DataLayerInstance : DataLayerInstances)
diff --git a/Engine/Source/Runtime/Engine/Private/WorldPartition/WorldPartitionActorDesc.cpp b/Engine/Source/Runtime/Engine/Private/WorldPartition/WorldPartitionActorDesc.cpp
--- a/Engine/Source/Runtime/Engine/Private/WorldPartition/WorldPartitionActorDesc.cpp
+++ b/Engine/Source/Runtime/Engine/Private/WorldPartition/WorldPartitionActorDesc.cpp
@@ -23,6 +23,7 @@
 #include "WorldPartition/DataLayer/DataLayerAsset.h"
 #include "WorldPartition/DataLayer/DataLayerUtils.h"
 #include "WorldPartition/ContentBundle/ContentBundlePaths.h"
+#include "WorldPartition/WorldPartitionActorContainerID.h"
 #include "WorldPartition/ErrorHandling/WorldPartitionStreamingGenerationErrorHandler.h"
 #include "ActorReferencesUtils.h"
 
@@ -590,6 +591,11 @@ void FWorldPartitionActorDesc::CheckForErrors(IStreamingGenerationErrorHandler*
 	}
 }
 
+bool FWorldPartitionActorDesc::IsRuntimeRelevant(const FActorContainerID& InContainerID) const
+{
+	return InContainerID.IsMainContainer() || !CastChecked<AActor>(ActorNativeClass->GetDefaultObject())->ShouldSkipFromLevelInstance();
+}
+
 bool FWorldPartitionActorDesc::IsLoaded(bool bEvenIfPendingKill) const
 {
 	if (ActorPtr.IsExplicitlyNull() || ActorPtr.IsStale())
diff --git a/Engine/Source/Runtime/Engine/Public/LevelInstance/LevelInstanceEditorInstanceActor.h b/Engine/Source/Runtime/Engine/Public/LevelInstance/LevelInstanceEditorInstanceActor.h
--- a/Engine/Source/Runtime/Engine/Public/LevelInstance/LevelInstanceEditorInstanceActor.h
+++ b/Engine/Source/Runtime/Engine/Public/LevelInstance/LevelInstanceEditorInstanceActor.h
@@ -32,6 +32,7 @@ public:
 	virtual AActor* GetSelectionParent() const override;
 
 private:
+	virtual bool ActorTypeShouldSkipFromLevelInstance() const override { return true; }
 	friend class ULevelInstanceComponent;
 	void UpdateWorldTransform(const FTransform& WorldTransform);
 
diff --git a/Engine/Source/Runtime/Engine/Public/WorldPartition/DataLayer/WorldDataLayers.h b/Engine/Source/Runtime/Engine/Public/WorldPartition/DataLayer/WorldDataLayers.h
--- a/Engine/Source/Runtime/Engine/Public/WorldPartition/DataLayer/WorldDataLayers.h
+++ b/Engine/Source/Runtime/Engine/Public/WorldPartition/DataLayer/WorldDataLayers.h
@@ -183,6 +183,8 @@ private:
 	void DumpDataLayerRecursively(const UDataLayerInstance* DataLayer, FString Prefix, FOutputDevice& OutputDevice) const;
 
 #if WITH_EDITOR
+	virtual bool ActorTypeShouldSkipFromLevelInstance() const override { return true; }
+
 	void ConvertDataLayerToInstancces();
 	void UpdateContainsDeprecatedDataLayers();
 	void ResolveActorDescContainers();
diff --git a/Engine/Source/Runtime/Engine/Public/WorldPartition/DataLayer/WorldDataLayersActorDesc.h b/Engine/Source/Runtime/Engine/Public/WorldPartition/DataLayer/WorldDataLayersActorDesc.h
--- a/Engine/Source/Runtime/Engine/Public/WorldPartition/DataLayer/WorldDataLayersActorDesc.h
+++ b/Engine/Source/Runtime/Engine/Public/WorldPartition/DataLayer/WorldDataLayersActorDesc.h
@@ -63,7 +63,6 @@ protected:
 	virtual bool Equals(const FWorldPartitionActorDesc* Other) const override;
 	virtual void Serialize(FArchive& Ar) override;
 	virtual bool IsResaveNeeded() const override { return !IsValid(); }
-	virtual bool IsRuntimeRelevant(const FActorContainerID& InContainerID) const override;
 	//~ End FWorldPartitionActorDesc Interface.
 
 private:
diff --git a/Engine/Source/Runtime/Engine/Public/WorldPartition/WorldPartitionActorDesc.h b/Engine/Source/Runtime/Engine/Public/WorldPartition/WorldPartitionActorDesc.h
--- a/Engine/Source/Runtime/Engine/Public/WorldPartition/WorldPartitionActorDesc.h
+++ b/Engine/Source/Runtime/Engine/Public/WorldPartition/WorldPartitionActorDesc.h
@@ -143,7 +143,7 @@ PRAGMA_ENABLE_DEPRECATION_WARNINGS
 
 	virtual const FGuid& GetSceneOutlinerParent() const { return GetParentActor(); }
 	virtual bool IsResaveNeeded() const { return false; }
-	virtual bool IsRuntimeRelevant(const FActorContainerID& InContainerID) const { return true; }
+	virtual bool IsRuntimeRelevant(const FActorContainerID& InContainerID) const;
 	virtual void CheckForErrors(IStreamingGenerationErrorHandler* ErrorHandler) const;
 
 	UE_DEPRECATED(5.2, "ShouldValidateRuntimeGrid is deprecated and should not be used.")
diff --git a/Engine/Source/Runtime/Engine/Public/WorldPartition/WorldPartitionMiniMap.h b/Engine/Source/Runtime/Engine/Public/WorldPartition/WorldPartitionMiniMap.h
--- a/Engine/Source/Runtime/Engine/Public/WorldPartition/WorldPartitionMiniMap.h
+++ b/Engine/Source/Runtime/Engine/Public/WorldPartition/WorldPartitionMiniMap.h
@@ -19,6 +19,7 @@ class ENGINE_API AWorldPartitionMiniMap : public AInfo
 private:
 #if WITH_EDITOR
 	virtual bool ActorTypeSupportsDataLayer() const final { return false; }
+	virtual bool ActorTypeShouldSkipFromLevelInstance() const override  { return true; }
 #endif
 
 public:
