Add compiler invocations to shader compilation stats.

Example output with D3D11 RHI: "Compiler invocations: DXC (72 pre-compile steps / 74), Legacy FXC (5736), Platform compiler (NO)"

#rnx

diff --git a/Engine/Source/Developer/Windows/ShaderFormatD3D/Private/D3DShaderCompiler.cpp b/Engine/Source/Developer/Windows/ShaderFormatD3D/Private/D3DShaderCompiler.cpp
--- a/Engine/Source/Developer/Windows/ShaderFormatD3D/Private/D3DShaderCompiler.cpp
+++ b/Engine/Source/Developer/Windows/ShaderFormatD3D/Private/D3DShaderCompiler.cpp
@@ -1069,13 +1069,7 @@ void CompileD3DShader(const FShaderCompilerInput& Input, FShaderCompilerOutput&
 {
 	TRACE_CPUPROFILER_EVENT_SCOPE(CompileD3DShader);
 	FString PreprocessedShaderSource;
-	const bool bIsRayTracingShader = Input.IsRayTracingShader();
-	const bool bUseDXC =
-		Language == ELanguage::SM6
-		|| bIsRayTracingShader
-		|| Input.Environment.CompilerFlags.Contains(CFLAG_WaveOperations)
-		|| Input.Environment.CompilerFlags.Contains(CFLAG_ForceDXC)
-		|| Input.Environment.CompilerFlags.Contains(CFLAG_InlineRayTracing);
+	const bool bUseDXC = !Input.CanCompileWithLegacyFxc();
 	const TCHAR* ShaderProfile = GetShaderProfileName(Language, Input.Target.Frequency, bUseDXC);
 
 	if(!ShaderProfile)
diff --git a/Engine/Source/Runtime/Engine/Private/ShaderCompiler/ShaderCompiler.cpp b/Engine/Source/Runtime/Engine/Private/ShaderCompiler/ShaderCompiler.cpp
--- a/Engine/Source/Runtime/Engine/Private/ShaderCompiler/ShaderCompiler.cpp
+++ b/Engine/Source/Runtime/Engine/Private/ShaderCompiler/ShaderCompiler.cpp
@@ -3312,6 +3312,39 @@ static FString FormatNumber(T Number)
 	return FText::AsNumber(Number, &FormattingOptions).ToString();
 }
 
+static FString WriteShaderCompilerInvocationStats(const FShaderCompilerStats::FCompilerInvocations& Invocations)
+{
+	FString Stats;
+
+	auto AppendStat = [&Stats](const TCHAR* Text)
+	{
+		Stats += FString::Printf(TEXT("%s%s"), (Stats.IsEmpty() ? TEXT("") : TEXT(", ")), Text);
+	};
+	auto AppendInvocationStat = [&AppendStat](const TCHAR* CompilerName, int32 Invocations)
+	{
+		if (Invocations > 0)
+		{
+			AppendStat(*FString::Printf(TEXT("%s (%d)"), CompilerName, Invocations));
+		}
+	};
+	if (Invocations.Dxc > 0)
+	{
+		if (Invocations.DxcPrecompileSteps > 0)
+		{
+			AppendStat(*FString::Printf(TEXT("DXC (%d pre-compile steps / %d)"), Invocations.DxcPrecompileSteps, Invocations.Dxc));
+		}
+		else
+		{
+			AppendInvocationStat(TEXT("DXC"), Invocations.Dxc);
+		}
+	}
+	AppendInvocationStat(TEXT("Legacy FXC"), Invocations.Fxc);
+	AppendInvocationStat(TEXT("HLSLcc"), Invocations.Hlslcc);
+	AppendStat(*FString::Printf(TEXT("Platform compiler (%s)"), Invocations.bWasPlatformCompilerUsed ? TEXT("YES") : TEXT("NO")));
+
+	return Stats;
+}
+
 void FShaderCompilerStats::WriteStatSummary()
 {
 	const uint32 TotalCompiled = GetTotalShadersCompiled();
@@ -3326,6 +3359,8 @@ void FShaderCompilerStats::WriteStatSummary()
 		*FormatNumber(JobsCompleted),
 		(JobsAssigned > 0.0) ? 100.0 * JobsCompleted / JobsAssigned : 0.0);
 
+	UE_LOG(LogShaderCompilers, Display, TEXT("Compiler invocations: %s"), *WriteShaderCompilerInvocationStats(CompilerInvocations));
+
 	if (TimesLocalWorkersWereIdle > 0.0)
 	{
 		UE_LOG(LogShaderCompilers, Display, TEXT("Average time worker was idle: %.2f s"), AccumulatedLocalWorkerIdleTime / TimesLocalWorkersWereIdle);
@@ -3541,10 +3576,10 @@ void FShaderCompilerStats::RegisterFinishedJob(FShaderCommonCompileJob& Job)
 	ensure(Job.TimeAddedToPendingQueue != 0.0 && Job.TimeAddedToPendingQueue <= Job.TimeExecutionCompleted);
 	AddToInterval(JobLifeTimeIntervals, TInterval<double>(Job.TimeAddedToPendingQueue, Job.TimeExecutionCompleted));
 
-	if (const FShaderCompileJob* SingleJob = Job.GetSingleShaderJob())
+	auto RegisterStatsFromSingleJob = [this](const FShaderCompileJob& SingleJob)
 	{
 		// Register min/max/average shader code sizes for single job output
-		const int32 ShaderCodeSize = SingleJob->Output.ShaderCode.GetShaderCodeSize();
+		const int32 ShaderCodeSize = SingleJob.Output.ShaderCode.GetShaderCodeSize();
 		if (ShaderCodeSize > 0)
 		{
 			MinShaderCodeSize = (MinShaderCodeSize > 0 ? FMath::Min(MinShaderCodeSize, ShaderCodeSize) : ShaderCodeSize);
@@ -3553,13 +3588,56 @@ void FShaderCompilerStats::RegisterFinishedJob(FShaderCommonCompileJob& Job)
 			++NumAccumulatedShaderCodes;
 		}
 
-		const FString ShaderName(SingleJob->Key.ShaderType->GetName());
+		// Register shader compiler invocations
+		{
+			// Record FXC invocation if target is D3D11 (no other paltform uses FXC)
+			const EShaderPlatform InputShaderPlatform = SingleJob.Input.Target.GetPlatform();
+			const bool bIsD3DPlatform = FDataDrivenShaderPlatformInfo::GetIsLanguageD3D(InputShaderPlatform);
+			const bool bUsedFxc = bIsD3DPlatform && SingleJob.Input.CanCompileWithLegacyFxc();
+			if (bUsedFxc)
+			{
+				CompilerInvocations.Fxc++;
+			}
+
+			const bool bSupportsDxc = FDataDrivenShaderPlatformInfo::GetSupportsDxc(InputShaderPlatform);
+			const bool bSupportsDxcPrecompileStep = FDataDrivenShaderPlatformInfo::GetIsPC(InputShaderPlatform);
+			if (bSupportsDxc || bSupportsDxcPrecompileStep)
+			{
+				// Record DXC pre-compile step if FXC was used and HLSL 2021 or an explicit pre-compile step was requested
+				const bool bHlslVersion2021 = SingleJob.Input.Environment.CompilerFlags.Contains(CFLAG_HLSL2021);
+				const bool bPrecompileWithDXC = SingleJob.Input.Environment.CompilerFlags.Contains(CFLAG_PrecompileWithDXC);
+				const bool bWasPrecompiledWithDXC = bUsedFxc && (bHlslVersion2021 || bPrecompileWithDXC);
+				if (bWasPrecompiledWithDXC)
+				{
+					CompilerInvocations.DxcPrecompileSteps++;
+				}
+
+				// Record DXC invocation if it was explicitly requested or if there was a pre-compile step
+				const bool bForceDxc = SingleJob.Input.Environment.CompilerFlags.Contains(CFLAG_ForceDXC);
+				if (bForceDxc || bWasPrecompiledWithDXC)
+				{
+					CompilerInvocations.Dxc++;
+				}
+			}
+
+			// Record HLSLcc invocation if output is marked as such
+			if (SingleJob.Output.bUsedHLSLccCompiler)
+			{
+				CompilerInvocations.Hlslcc++;
+			}
+
+			// Assume any console target requires a platform compiler
+			const bool bIsConsole = FDataDrivenShaderPlatformInfo::GetIsConsole(InputShaderPlatform);
+			CompilerInvocations.bWasPlatformCompilerUsed = bIsConsole;
+		}
+
+		const FString ShaderName(SingleJob.Key.ShaderType->GetName());
 		if (FShaderTimings* Existing = ShaderTimings.Find(ShaderName))
 		{
-			Existing->MinCompileTime = FMath::Min(Existing->MinCompileTime, static_cast<float>(SingleJob->Output.CompileTime));
-			Existing->MaxCompileTime = FMath::Max(Existing->MaxCompileTime, static_cast<float>(SingleJob->Output.CompileTime));
-			Existing->TotalCompileTime += SingleJob->Output.CompileTime;
-			Existing->TotalPreprocessTime += SingleJob->Output.PreprocessTime;
+			Existing->MinCompileTime = FMath::Min(Existing->MinCompileTime, static_cast<float>(SingleJob.Output.CompileTime));
+			Existing->MaxCompileTime = FMath::Max(Existing->MaxCompileTime, static_cast<float>(SingleJob.Output.CompileTime));
+			Existing->TotalCompileTime += SingleJob.Output.CompileTime;
+			Existing->TotalPreprocessTime += SingleJob.Output.PreprocessTime;
 			Existing->NumCompiled++;
 			// calculate as an optimization to make sorting later faster
 			Existing->AverageCompileTime = Existing->TotalCompileTime / static_cast<float>(Existing->NumCompiled);
@@ -3567,16 +3645,17 @@ void FShaderCompilerStats::RegisterFinishedJob(FShaderCommonCompileJob& Job)
 		else
 		{
 			FShaderTimings New;
-			New.MinCompileTime = SingleJob->Output.CompileTime;
+			New.MinCompileTime = SingleJob.Output.CompileTime;
 			New.MaxCompileTime = New.MinCompileTime;
 			New.TotalCompileTime = New.MinCompileTime;
-			New.TotalPreprocessTime += SingleJob->Output.PreprocessTime;
+			New.TotalPreprocessTime += SingleJob.Output.PreprocessTime;
 			New.AverageCompileTime = New.MinCompileTime;
 			New.NumCompiled = 1;
 			ShaderTimings.Add(ShaderName, New);
 		}
-	}
+	};
 
+	Job.ForEachSingleShaderJob(RegisterStatsFromSingleJob);
 }
 
 void FShaderCompilerStats::RegisterJobBatch(int32 NumJobs, EExecutionType ExecType)
diff --git a/Engine/Source/Runtime/Engine/Public/ShaderCompiler.h b/Engine/Source/Runtime/Engine/Public/ShaderCompiler.h
--- a/Engine/Source/Runtime/Engine/Public/ShaderCompiler.h
+++ b/Engine/Source/Runtime/Engine/Public/ShaderCompiler.h
@@ -454,7 +454,6 @@ public:
 		uint32 CompiledDouble = 0;
 		uint32 CookedDouble = 0;
 		float CompileTime = 0.f;
-
 	};
 	using ShaderCompilerStats = TMap<FString, FShaderStats>;
 
@@ -469,6 +468,15 @@ public:
 		float AverageCompileTime = 0.0f;	// stored explicitly as an optimization
 	};
 
+	struct FCompilerInvocations
+	{
+		int32 Dxc = 0;
+		int32 DxcPrecompileSteps = 0;
+		int32 Fxc = 0;
+		int32 Hlslcc = 0;
+		bool bWasPlatformCompilerUsed = false;
+	};
+
 	ENGINE_API void RegisterCookedShaders(uint32 NumCooked, float CompileTime, EShaderPlatform Platform, const FString MaterialPath, FString PermutationString = FString(""));
 	ENGINE_API void RegisterCompiledShaders(uint32 NumPermutations, EShaderPlatform Platform, const FString MaterialPath, FString PermutationString = FString(""));
 	ENGINE_API const TSparseArray<ShaderCompilerStats>& GetShaderCompilerStats() { return CompileStats; }
@@ -575,6 +583,9 @@ private:
 
 	/** Total number of DDC hits on shader maps. */
 	uint32 ShaderMapDDCHits = 0;
+
+	/** Total number of shader compiler invocations. */
+	FCompilerInvocations CompilerInvocations;
 };
 
 
diff --git a/Engine/Source/Runtime/RenderCore/Public/ShaderCompilerCore.h b/Engine/Source/Runtime/RenderCore/Public/ShaderCompilerCore.h
--- a/Engine/Source/Runtime/RenderCore/Public/ShaderCompilerCore.h
+++ b/Engine/Source/Runtime/RenderCore/Public/ShaderCompilerCore.h
@@ -364,6 +364,17 @@ struct FShaderCompilerInput
 	{
 		return IsRayTracingShaderFrequency(Target.GetFrequency());
 	}
+
+	/** Returns whether this shader input *can* be compiled with the legacy FXC compiler. */
+	bool CanCompileWithLegacyFxc() const
+	{
+		return !(Target.GetPlatform() == SP_PCD3D_SM6
+			|| IsRayTracingShader()
+			|| Environment.CompilerFlags.Contains(CFLAG_WaveOperations)
+			|| Environment.CompilerFlags.Contains(CFLAG_ForceDXC)
+			|| Environment.CompilerFlags.Contains(CFLAG_InlineRayTracing)
+			);
+	}
 };
 
 /** A shader compiler error or warning. */
diff --git a/Engine/Source/Runtime/RenderCore/Public/ShaderCompilerJobTypes.h b/Engine/Source/Runtime/RenderCore/Public/ShaderCompilerJobTypes.h
--- a/Engine/Source/Runtime/RenderCore/Public/ShaderCompilerJobTypes.h
+++ b/Engine/Source/Runtime/RenderCore/Public/ShaderCompilerJobTypes.h
@@ -6,6 +6,7 @@
 #include "Hash/Blake3.h"
 #include "Shader.h"
 #include "VertexFactory.h"
+#include "Templates/Function.h"
 
 class FShaderCommonCompileJob;
 class FShaderCompileJob;
@@ -114,7 +115,10 @@ public:
 	const FShaderPipelineCompileJob* GetShaderPipelineJob() const;
 
 	bool Equals(const FShaderCommonCompileJob& Rhs) const;
-	
+
+	/** Calls the specified predicate for each single compile job, i.e. FShaderCompileJob and each stage of FShaderPipelineCompileJob. */
+	void ForEachSingleShaderJob(const TFunction<void(const FShaderCompileJob& SingleJob)>& Predicate) const;
+
 	/** This returns a unique id for a shader compiler job */
 	RENDERCORE_API static uint32 GetNextJobId();
 
@@ -276,4 +280,22 @@ inline void FShaderCommonCompileJob::Destroy() const
 	case EShaderCompileJobType::Pipeline: delete static_cast<const FShaderPipelineCompileJob*>(this); break;
 	default: checkNoEntry();
 	}
-}
\ No newline at end of file
+}
+
+inline void FShaderCommonCompileJob::ForEachSingleShaderJob(const TFunction<void(const FShaderCompileJob&)>& Predicate) const
+{
+	if (const FShaderCompileJob* SingleJob = GetSingleShaderJob())
+	{
+		Predicate(*SingleJob);
+	}
+	else if (const FShaderPipelineCompileJob* PipelineJob = GetShaderPipelineJob())
+	{
+		for (const TRefCountPtr<FShaderCompileJob>& StageJob : PipelineJob->StageJobs)
+		{
+			if (const FShaderCompileJob* SingleStageJob = StageJob->GetSingleShaderJob())
+			{
+				Predicate(*SingleStageJob);
+			}
+		}
+	}
+}
