Engine: Debug: Extended 'debug eatmem' debug console command to allow for eating memory spread across multiple frames with the defined free memory target.

It can be used as folllows:

debug eatmem tick target=1gb tickmax=5mb

which will make the engine allocate up to 5 MB of memory every tick until there is only 1 GB left. It supports kb, mb and gb suffixes for sizes (values in bytes are assumed when no suffix is given).

#rnx

diff --git a/Engine/Source/Runtime/Engine/Private/UnrealEngine.cpp b/Engine/Source/Runtime/Engine/Private/UnrealEngine.cpp
--- a/Engine/Source/Runtime/Engine/Private/UnrealEngine.cpp
+++ b/Engine/Source/Runtime/Engine/Private/UnrealEngine.cpp
@@ -9777,6 +9777,201 @@ FORCENOINLINE static void DebugDummyWorkAfterFailure()
 	UE_LOG(LogEngine, Log, TEXT("A dummy log message"));
 }
 
+#if !UE_BUILD_SHIPPING
+
+class FDebugMemoryEater
+{
+public:
+	static FDebugMemoryEater& Get()
+	{
+		static FDebugMemoryEater Instance;
+
+		return Instance;
+	}
+
+	bool Start(const TCHAR* Cmd, FOutputDevice& Ar)
+	{
+		Target       = 0;
+		LimitPerTick = 10 * 1024 * 1024;
+
+		if (!ParseSizeParam(Cmd, TEXT("Target="), Target))
+		{
+			Ar.Logf(TEXT("Failed to parse Target parameter!"));
+
+			return false;
+		}
+
+		if (!ParseSizeParam(Cmd, TEXT("TickMax="), LimitPerTick))
+		{
+			Ar.Logf(TEXT("Failed to parse TickMax parameter!"));
+
+			return false;
+		}
+
+		Ar.Logf(TEXT("Target Free Memory: %llu, Max Allocation Per Tick: %llu"), Target, LimitPerTick);
+
+		StartTicking();
+
+		return true;
+	}
+
+	void Stop()
+	{
+		StopTicking();
+	}
+
+	void Purge()
+	{
+		while (void* MemoryBlock = PopMemoryBlock())
+		{
+			FMemory::Free(MemoryBlock);
+		}
+	}
+
+	bool Tick(float DeltaSeconds)
+	{
+		constexpr uint64 MemoryBlockSize = 64 * 1024;
+
+		uint64 AllocatedThisTick = 0;
+
+		FPlatformMemoryStats Stats = FPlatformMemory::GetStats();
+
+		while (Stats.AvailablePhysical > Target && AllocatedThisTick < LimitPerTick)
+		{
+			void* Ptr = FMemory::Malloc(MemoryBlockSize);
+
+			if (Ptr)
+			{
+				AllocatedThisTick += MemoryBlockSize;
+
+				PushMemoryBlock(Ptr);
+			}
+			else
+			{
+				break;
+			}
+
+			// We could just subtract the amount allocated from the amount of free memory queried before the loop but let's be precise and query again.
+			Stats = FPlatformMemory::GetStats();
+		}
+
+		return (Stats.AvailablePhysical > Target);
+	}
+
+private:
+	bool ParseSizeParam(const TCHAR* Cmd, const TCHAR* Match, uint64& Value)
+	{
+		FString SizeString;
+		uint64  SizeMultiplier = 1;
+
+		if (!FParse::Value(Cmd, Match, SizeString))
+		{
+			return true;
+		}
+
+		static constexpr TCHAR const* SuffixKB  = TEXT("kb");
+		static constexpr TCHAR const* SuffixMB  = TEXT("mb");
+		static constexpr TCHAR const* SuffixGB  = TEXT("gb");
+		static constexpr uint32       SuffixLen = 2;
+
+		if (SizeString.EndsWith(SuffixKB, SuffixLen))
+		{
+			SizeMultiplier = 1024;
+
+			SizeString = SizeString.LeftChop(SuffixLen);
+		}
+		else if (SizeString.EndsWith(SuffixMB, SuffixLen))
+		{
+			SizeMultiplier = 1024 * 1024;
+
+			SizeString = SizeString.LeftChop(SuffixLen);
+		}
+		else if (SizeString.EndsWith(SuffixGB, SuffixLen))
+		{
+			SizeMultiplier = 1024 * 1024 * 1024;
+
+			SizeString = SizeString.LeftChop(SuffixLen);
+		}
+
+		if (SizeString.IsEmpty())
+		{
+			return false;
+		}
+
+		for (int32 i = 0; i < SizeString.Len(); ++i)
+		{
+			if (!FChar::IsDigit(SizeString[i]))
+			{
+				return false;
+			}
+		}
+
+		Value = FPlatformString::Atoi64(*SizeString) * SizeMultiplier;
+
+		return true;
+	}
+
+	void StartTicking()
+	{
+		if (TickHandle.IsValid())
+		{
+			StopTicking();
+		}
+
+		FTSTicker& Ticker = FTSTicker::GetCoreTicker();
+
+		TickHandle = Ticker.AddTicker(FTickerDelegate::CreateRaw(this, &FDebugMemoryEater::Tick));
+	}
+
+	void StopTicking()
+	{
+		if (TickHandle.IsValid())
+		{
+			FTSTicker& Ticker = FTSTicker::GetCoreTicker();
+
+			Ticker.RemoveTicker(TickHandle);
+
+			TickHandle.Reset();
+		}
+	}
+
+	void PushMemoryBlock(void* Ptr)
+	{
+		BlockHeader* MemoryBlock = (BlockHeader*)Ptr;
+
+		MemoryBlock->Next = MemoryListHead;
+		MemoryListHead    = MemoryBlock;
+	}
+
+	void* PopMemoryBlock()
+	{
+		if (!MemoryListHead)
+		{
+			return nullptr;
+		}
+
+		BlockHeader* PrevHead = MemoryListHead;
+
+		MemoryListHead = MemoryListHead->Next;
+
+		return PrevHead;
+	}
+
+	FTSTicker::FDelegateHandle TickHandle;
+
+	uint64 Target;
+	uint64 LimitPerTick;
+
+	struct BlockHeader
+	{
+		BlockHeader* Next;
+	};
+
+	BlockHeader* MemoryListHead = nullptr;
+};
+
+#endif
+
 bool UEngine::PerformError(const TCHAR* Cmd, FOutputDevice& Ar)
 {
 #if !UE_BUILD_SHIPPING
@@ -10264,13 +10459,38 @@ bool UEngine::PerformError(const TCHAR* Cmd, FOutputDevice& Ar)
 	}
 	else if (FParse::Command(&Cmd, TEXT("EATMEM")))
 	{
-		Ar.Log(TEXT("Eating up all available memory"));
-		FGenericCrashContext::SetCrashTrigger(ECrashTrigger::Debug);
-		while (1)
+		if (FParse::Command(&Cmd, TEXT("TICK")))
 		{
-			void* Eat = FMemory::Malloc(65536);
-			FMemory::Memset(Eat, 0, 65536);
+			if (FParse::Command(&Cmd, TEXT("STOP")))
+			{
+				FDebugMemoryEater::Get().Stop();
+			}
+			else if (FParse::Command(&Cmd, TEXT("PURGE")))
+			{
+				FDebugMemoryEater::Get().Stop();
+				FDebugMemoryEater::Get().Purge();
+			}
+			else
+			{
+				Ar.Log(TEXT("Gradually eat memory every Tick until there's less than the requested target left"));
+
+				if (!FDebugMemoryEater::Get().Start(Cmd, Ar))
+				{
+					Ar.Log(TEXT("Failed to start the memory eater!"));
+				}
+			}
+		}
+		else
+		{
+			Ar.Log(TEXT("Eating up all available memory"));
+			FGenericCrashContext::SetCrashTrigger(ECrashTrigger::Debug);
+			while (1)
+			{
+				void* Eat = FMemory::Malloc(65536);
+				FMemory::Memset(Eat, 0, 65536);
+			}
 		}
+
 		return true;
 	}
 	else if (FParse::Command(&Cmd, TEXT("OOM")))
