- Level Instance will only propagate its spatially loaded flag down to its child actors when set to false.
- Non-spatial actors part of a spatially loaded level instance now remain non-spatially loaded.
- Allow editing of bIsSpatiallyLoaded flag when editing a partitioned Level Instance.

diff --git a/Engine/Source/Runtime/Engine/Private/ActorEditor.cpp b/Engine/Source/Runtime/Engine/Private/ActorEditor.cpp
--- a/Engine/Source/Runtime/Engine/Private/ActorEditor.cpp
+++ b/Engine/Source/Runtime/Engine/Private/ActorEditor.cpp
@@ -86,7 +86,7 @@ bool AActor::CanEditChange(const FProperty* PropertyThatWillChange) const
 		{
 			UWorld* World = GetTypedOuter<UWorld>();
 			UWorldPartition* WorldPartition = World ? World->GetWorldPartition() : nullptr;
-			if (!WorldPartition || (!WorldPartition->IsStreamingEnabled() && !bIsDataLayersProperty && !bIsHLODLayerProperty))
+			if (!WorldPartition || (!WorldPartition->IsStreamingEnabled() && bIsRuntimeGridProperty))
 			{
 				return false;
 			}
diff --git a/Engine/Source/Runtime/Engine/Private/WorldPartition/WorldPartitionStreamingGeneration.cpp b/Engine/Source/Runtime/Engine/Private/WorldPartition/WorldPartitionStreamingGeneration.cpp
--- a/Engine/Source/Runtime/Engine/Private/WorldPartition/WorldPartitionStreamingGeneration.cpp
+++ b/Engine/Source/Runtime/Engine/Private/WorldPartition/WorldPartitionStreamingGeneration.cpp
@@ -132,6 +132,9 @@ class FWorldPartitionStreamingGenerator
 					ActorSetInstance.ContainerID = ContainerInstanceDescriptor.ID;
 					ActorSetInstance.Transform = ContainerInstanceDescriptor.Transform;
 
+					// Apply AND logic on spatially loaded flag
+					ActorSetInstance.bIsSpatiallyLoaded = ReferenceActorDescView.GetIsSpatiallyLoaded() && ContainerInstanceDescriptor.bIsSpatiallyLoaded;
+
 					// Since Content Bundles streaming generation happens in its own context, all actor set instances must have the same content bundle GUID for now, so Level Instances
 					// placed inside a Content Bundle will propagate their Content Bundle GUID to child instances.
 					ActorSetInstance.ContentBundleID = MainWorldContainer->GetContentBundleGuid();
@@ -140,15 +143,12 @@ class FWorldPartitionStreamingGenerator
 					{
 						// Main container will get inherited properties from the actor descriptors
 						ActorSetInstance.RuntimeGrid = ReferenceActorDescView.GetRuntimeGrid();
-						ActorSetInstance.bIsSpatiallyLoaded = ReferenceActorDescView.GetIsSpatiallyLoaded();
 						ActorSetInstance.DataLayers = StreamingGenerator->GetRuntimeDataLayerInstances(ReferenceActorDescView.GetRuntimeDataLayerInstanceNames());
 					}
 					else
 					{
 						// Sub containers will inherit some properties from the parent actor descriptors
 						ActorSetInstance.RuntimeGrid = ContainerInstanceDescriptor.RuntimeGrid;
-						ActorSetInstance.bIsSpatiallyLoaded = ContainerInstanceDescriptor.bIsSpatiallyLoaded;
-
 						TSet<FName> CombinedDataLayers = ContainerInstanceDescriptor.RuntimeDataLayers;
 						CombinedDataLayers.Append(ReferenceActorDescView.GetRuntimeDataLayerInstanceNames());
 						ActorSetInstance.DataLayers = StreamingGenerator->GetRuntimeDataLayerInstances(CombinedDataLayers.Array());
@@ -408,13 +408,15 @@ class FWorldPartitionStreamingGenerator
 		// Inherited parent properties logic
 		auto InheritParentContainerProperties = [](const FActorContainerID& InParentContainerID, const FWorldPartitionActorDescView& InParentActorDescView, TSet<FName>& InOutRuntimeDataLayers, FName& InOutRuntimeGrid, bool& bInOutIsSpatiallyLoaded)
 		{
-			// Runtime grid and spatially loaded flag are only inherited from the main world, since level instance don't support setting these values on actors
+			// Runtime grid is only inherited from the main world, since level instance doesn't support setting this value on actors
 			if (InParentContainerID.IsMainContainer())
 			{
 				InOutRuntimeGrid = InParentActorDescView.GetRuntimeGrid();
-				bInOutIsSpatiallyLoaded = InParentActorDescView.GetIsSpatiallyLoaded();
 			}
 
+			// Apply AND logic on spatially loaded flag
+			bInOutIsSpatiallyLoaded = bInOutIsSpatiallyLoaded && InParentActorDescView.GetIsSpatiallyLoaded();
+
 			// Data layers are accumulated down the hierarchy chain, since level instances supports data layers assignation on actors
 			InOutRuntimeDataLayers.Append(InParentActorDescView.GetRuntimeDataLayerInstanceNames());
 		};
@@ -474,7 +476,9 @@ class FWorldPartitionStreamingGenerator
 	{
 		TRACE_CPUPROFILER_EVENT_SCOPE(FWorldPartitionStreamingGenerator::CreateActorContainers);
 
-		CreateActorDescriptorViewsRecursive(InContainer, FTransform::Identity, NAME_None, false, TSet<FName>(), FActorContainerID(), FActorContainerID(), EContainerClusterMode::Partitioned, TEXT("MainContainer"));
+		// Since we apply AND logic on spatially loaded flag recursively, startup value must be true
+		const bool bIsSpatiallyLoaded = true;
+		CreateActorDescriptorViewsRecursive(InContainer, FTransform::Identity, NAME_None, bIsSpatiallyLoaded, TSet<FName>(), FActorContainerID(), FActorContainerID(), EContainerClusterMode::Partitioned, TEXT("MainContainer"));
 
 		// Update container instances bounds
 		for (auto& [ContainerID, ContainerInstanceDescriptor] : ContainerInstanceDescriptorsMap)
