Add a ReceivedPlayerController function to UPlayer, called from it's owning Player Controller when the owning APlayerController pointer is now valid on that player. This is useful piece of state for games to use during pawn initalization.

Alongside this, add a "OnPlayerControllerChanged" delegate to ULocalPlayer and Local Player Subsystems for an easy hook in.

diff --git a/Engine/Source/Runtime/Engine/Classes/Engine/LocalPlayer.h b/Engine/Source/Runtime/Engine/Classes/Engine/LocalPlayer.h
--- a/Engine/Source/Runtime/Engine/Classes/Engine/LocalPlayer.h
+++ b/Engine/Source/Runtime/Engine/Classes/Engine/LocalPlayer.h
@@ -18,6 +18,7 @@
 #include "SceneTypes.h"
 #include "Engine/Player.h"
 #include "GameFramework/OnlineReplStructs.h"
+#include "GameFramework/PlayerController.h"
 
 #include "Subsystems/LocalPlayerSubsystem.h"
 #include "Subsystems/SubsystemCollection.h"
@@ -212,6 +213,10 @@ public:
 	DECLARE_EVENT_TwoParams(ULocalPlayer, FOnPlatformUserIdChanged, FPlatformUserId /*NewId*/, FPlatformUserId /*OldId*/);
 	FOnPlatformUserIdChanged& OnPlatformUserIdChanged() { return OnPlatformUserIdChangedEvent; }
 
+	/** Event called when this local player has had a new outer PlayerController set */
+	DECLARE_EVENT_OneParam(ULocalPlayer, FOnPlayerControllerChanged, APlayerController* /*NewPC*/);
+	FOnPlayerControllerChanged& OnPlayerControllerChanged() { return OnPlayerControllerChangedEvent; }
+	
 private:
 	TArray<FSceneViewStateReference> ViewStates;
 
@@ -227,6 +232,9 @@ private:
 	/** Event called when platform user id changes */
 	FOnPlatformUserIdChanged OnPlatformUserIdChangedEvent;
 
+	/** Event called when the outer player controller changes */
+	FOnPlayerControllerChanged OnPlayerControllerChangedEvent;
+
 	FObjectSubsystemCollection<ULocalPlayerSubsystem> SubsystemCollection;
 
 public:
@@ -235,6 +243,10 @@ public:
 	static void AddReferencedObjects(UObject* InThis, FReferenceCollector& Collector);
 	// End of UObject interface
 
+	// Begin UPlayer interface
+	virtual void ReceivedPlayerController(APlayerController* NewController) override;
+	// End UPlayer interface
+	
 	// FExec interface
 public:
 #if UE_ALLOW_EXEC_COMMANDS
diff --git a/Engine/Source/Runtime/Engine/Classes/Engine/Player.h b/Engine/Source/Runtime/Engine/Classes/Engine/Player.h
--- a/Engine/Source/Runtime/Engine/Classes/Engine/Player.h
+++ b/Engine/Source/Runtime/Engine/Classes/Engine/Player.h
@@ -63,4 +63,9 @@ public:
 	 * @return The controller associated with this player in InWorld, if one exists.
 	 */
 	ENGINE_API APlayerController* GetPlayerController(const UWorld* const InWorld) const;
+
+	/**
+	 * Called when this player has had it's outer Player Controller set for remote net connections
+	 */
+	ENGINE_API virtual void ReceivedPlayerController(APlayerController* NewController);
 };
diff --git a/Engine/Source/Runtime/Engine/Private/LocalPlayer.cpp b/Engine/Source/Runtime/Engine/Private/LocalPlayer.cpp
--- a/Engine/Source/Runtime/Engine/Private/LocalPlayer.cpp
+++ b/Engine/Source/Runtime/Engine/Private/LocalPlayer.cpp
@@ -714,6 +714,23 @@ void ULocalPlayer::GetViewPoint(FMinimalViewInfo& OutViewInfo) const
 	OutViewInfo.DesiredFOV = OutViewInfo.FOV;
 }
 
+void ULocalPlayer::ReceivedPlayerController(APlayerController* NewController)
+{
+	QUICK_SCOPE_CYCLE_COUNTER(STAT_LocalPlayer_HandlePlayerControllerChanged);
+	
+	Super::ReceivedPlayerController(NewController);
+	
+	// Broadcast an event for anyone that may be listening
+	OnPlayerControllerChanged().Broadcast(NewController);
+
+	// Tell any local player subsystems
+	const TArray<ULocalPlayerSubsystem*>& LPSubsystems = SubsystemCollection.GetSubsystemArray<ULocalPlayerSubsystem>(ULocalPlayerSubsystem::StaticClass());
+	for (ULocalPlayerSubsystem* WorldSubsystem : LPSubsystems)
+	{
+		WorldSubsystem->PlayerControllerChanged(NewController);
+	}
+}
+
 bool ULocalPlayer::CalcSceneViewInitOptions(
 	struct FSceneViewInitOptions& ViewInitOptions,
 	FViewport* Viewport,
diff --git a/Engine/Source/Runtime/Engine/Private/Player.cpp b/Engine/Source/Runtime/Engine/Private/Player.cpp
--- a/Engine/Source/Runtime/Engine/Private/Player.cpp
+++ b/Engine/Source/Runtime/Engine/Private/Player.cpp
@@ -166,5 +166,10 @@ void UPlayer::SwitchController(class APlayerController* PC)
 	// Set the viewport.
 	PC->Player = this;
 	this->PlayerController = PC;
+	
+	ReceivedPlayerController(PC);
 }
 
+void UPlayer::ReceivedPlayerController(class APlayerController* NewController)
+{
+}
\ No newline at end of file
diff --git a/Engine/Source/Runtime/Engine/Private/PlayerController.cpp b/Engine/Source/Runtime/Engine/Private/PlayerController.cpp
--- a/Engine/Source/Runtime/Engine/Private/PlayerController.cpp
+++ b/Engine/Source/Runtime/Engine/Private/PlayerController.cpp
@@ -887,6 +887,11 @@ void APlayerController::ReceivedPlayer()
 			BeginSpectatingState();
 		}
 	}
+
+	if (Player)
+	{
+		Player->ReceivedPlayerController(this);
+	}
 }
 
 FVector APlayerController::GetFocalLocation() const
diff --git a/Engine/Source/Runtime/Engine/Private/Subsystems/LocalPlayerSubsystem.cpp b/Engine/Source/Runtime/Engine/Private/Subsystems/LocalPlayerSubsystem.cpp
--- a/Engine/Source/Runtime/Engine/Private/Subsystems/LocalPlayerSubsystem.cpp
+++ b/Engine/Source/Runtime/Engine/Private/Subsystems/LocalPlayerSubsystem.cpp
@@ -1,6 +1,7 @@
 // Copyright Epic Games, Inc. All Rights Reserved.
 
 #include "Subsystems/LocalPlayerSubsystem.h"
+#include "GameFramework/PlayerController.h"
 
 #include UE_INLINE_GENERATED_CPP_BY_NAME(LocalPlayerSubsystem)
 
@@ -9,3 +10,8 @@ ULocalPlayerSubsystem::ULocalPlayerSubsystem()
 {
 }
 
+void ULocalPlayerSubsystem::PlayerControllerChanged(APlayerController* NewPlayerController)
+{
+	
+}
+
diff --git a/Engine/Source/Runtime/Engine/Public/Subsystems/LocalPlayerSubsystem.h b/Engine/Source/Runtime/Engine/Public/Subsystems/LocalPlayerSubsystem.h
--- a/Engine/Source/Runtime/Engine/Public/Subsystems/LocalPlayerSubsystem.h
+++ b/Engine/Source/Runtime/Engine/Public/Subsystems/LocalPlayerSubsystem.h
@@ -7,6 +7,7 @@
 #include "LocalPlayerSubsystem.generated.h"
 
 class ULocalPlayer;
+class APlayerController;
 
 /**
  * ULocalPlayerSubsystem
@@ -31,4 +32,7 @@ public:
 	{
 		return CastChecked<LocalPlayerType>(GetOuter());
 	}
+
+	/** Callback for when the player controller is changed on this subsystem's owning local player */
+	virtual void PlayerControllerChanged(APlayerController* NewPlayerController);
 };
