WorldPartition: ActorDesc now have a IsEditorRelevant() used by IWorldPartitionActorLoaderInterface::ILoaderAdapter::PassActorDescFilter
- IsEditorRelevant will check a couple conditions(GetActorIsRuntimeOnly(), ShouldSkipFromLevelInstance() for non main WP actors)

diff --git a/Engine/Source/Editor/SceneOutliner/Private/ActorHierarchy.cpp b/Engine/Source/Editor/SceneOutliner/Private/ActorHierarchy.cpp
--- a/Engine/Source/Editor/SceneOutliner/Private/ActorHierarchy.cpp
+++ b/Engine/Source/Editor/SceneOutliner/Private/ActorHierarchy.cpp
@@ -456,7 +456,7 @@ void FActorHierarchy::CreateWorldChildren(UWorld* World, TArray<FSceneOutlinerTr
 				}
 				FWorldPartitionHelpers::ForEachActorDesc(WorldPartition, [this, &OutItems](const FWorldPartitionActorDesc* ActorDesc)
 				{
-					if (ActorDesc != nullptr && !ActorDesc->IsLoaded(true))
+					if (ActorDesc != nullptr && !ActorDesc->IsLoaded(true) && ActorDesc->IsEditorRelevant())
 					{
 						if (const FSceneOutlinerTreeItemPtr ActorDescItem = Mode->CreateItemFor<FActorDescTreeItem>(FActorDescTreeItem(ActorDesc->GetGuid(), ActorDesc->GetContainer())))
 						{
diff --git a/Engine/Source/Runtime/Engine/Private/WorldPartition/LoaderAdapter/LoaderAdapterPinnedActors.cpp b/Engine/Source/Runtime/Engine/Private/WorldPartition/LoaderAdapter/LoaderAdapterPinnedActors.cpp
--- a/Engine/Source/Runtime/Engine/Private/WorldPartition/LoaderAdapter/LoaderAdapterPinnedActors.cpp
+++ b/Engine/Source/Runtime/Engine/Private/WorldPartition/LoaderAdapter/LoaderAdapterPinnedActors.cpp
@@ -27,7 +27,7 @@ FText FLoaderAdapterPinnedActors::GetUnloadedReason(FWorldPartitionActorDesc* In
 
 		UActorDescContainer* ActorDescContainer = InActorDesc->GetContainer();
 		UWorldPartition* WorldPartition = ActorDescContainer != nullptr ? ActorDescContainer->GetWorldPartition() : nullptr;
-		bool bShouldBeLoaded = !InActorDesc->GetIsSpatiallyLoaded() && !InActorDesc->GetActorIsRuntimeOnly();
+		bool bShouldBeLoaded = !InActorDesc->GetIsSpatiallyLoaded() && InActorDesc->IsEditorRelevant();
 		if (WorldPartition && (bShouldBeLoaded || WorldPartition->IsActorPinned(InActorDesc->GetGuid())))
 		{
 			return LOCTEXT("UnloadedDataLayerReason", "Unloaded DataLayer");
diff --git a/Engine/Source/Runtime/Engine/Private/WorldPartition/WorldPartitionActorDesc.cpp b/Engine/Source/Runtime/Engine/Private/WorldPartition/WorldPartitionActorDesc.cpp
--- a/Engine/Source/Runtime/Engine/Private/WorldPartition/WorldPartitionActorDesc.cpp
+++ b/Engine/Source/Runtime/Engine/Private/WorldPartition/WorldPartitionActorDesc.cpp
@@ -591,6 +591,21 @@ void FWorldPartitionActorDesc::CheckForErrors(IStreamingGenerationErrorHandler*
 	}
 }
 
+bool FWorldPartitionActorDesc::IsEditorRelevant() const
+{
+	if (GetActorIsRuntimeOnly())
+	{
+		return false;
+	}
+
+	if (CastChecked<AActor>(ActorNativeClass->GetDefaultObject())->ShouldSkipFromLevelInstance())
+	{
+		return GetContainer() && !GetContainer()->IsTemplateContainer() && GetContainer()->GetWorldPartition()->IsMainWorldPartition();
+	}
+
+	return true;
+}
+
 bool FWorldPartitionActorDesc::IsRuntimeRelevant(const FActorContainerID& InContainerID) const
 {
 	return InContainerID.IsMainContainer() || !CastChecked<AActor>(ActorNativeClass->GetDefaultObject())->ShouldSkipFromLevelInstance();
diff --git a/Engine/Source/Runtime/Engine/Private/WorldPartition/WorldPartitionActorLoaderInterface.cpp b/Engine/Source/Runtime/Engine/Private/WorldPartition/WorldPartitionActorLoaderInterface.cpp
--- a/Engine/Source/Runtime/Engine/Private/WorldPartition/WorldPartitionActorLoaderInterface.cpp
+++ b/Engine/Source/Runtime/Engine/Private/WorldPartition/WorldPartitionActorLoaderInterface.cpp
@@ -116,7 +116,7 @@ void IWorldPartitionActorLoaderInterface::ILoaderAdapter::UnregisterDelegates()
 
 bool IWorldPartitionActorLoaderInterface::ILoaderAdapter::PassActorDescFilter(const FWorldPartitionHandle& Actor) const
 {
-	return !Actor->GetActorIsRuntimeOnly();
+	return Actor->IsEditorRelevant();
 }
 
 // Helper added to since the nested class IWorldPartitionActorLoaderInterface::ILoaderAdapter can't be friend of UDataLayerManager.
diff --git a/Engine/Source/Runtime/Engine/Public/WorldPartition/WorldPartitionActorDesc.h b/Engine/Source/Runtime/Engine/Public/WorldPartition/WorldPartitionActorDesc.h
--- a/Engine/Source/Runtime/Engine/Public/WorldPartition/WorldPartitionActorDesc.h
+++ b/Engine/Source/Runtime/Engine/Public/WorldPartition/WorldPartitionActorDesc.h
@@ -144,6 +144,7 @@ PRAGMA_ENABLE_DEPRECATION_WARNINGS
 	virtual const FGuid& GetSceneOutlinerParent() const { return GetParentActor(); }
 	virtual bool IsResaveNeeded() const { return false; }
 	virtual bool IsRuntimeRelevant(const FActorContainerID& InContainerID) const;
+	virtual bool IsEditorRelevant() const;
 	virtual void CheckForErrors(IStreamingGenerationErrorHandler* ErrorHandler) const;
 
 	UE_DEPRECATED(5.2, "ShouldValidateRuntimeGrid is deprecated and should not be used.")
@@ -228,8 +229,8 @@ public:
 	AActor* Load() const;
 	virtual void Unload();
 
-	UE_DEPRECATED(5.1, "ShouldBeLoadedByEditorCells is deprecated, GetActorIsRuntimeOnly should be used instead.")
-	bool ShouldBeLoadedByEditorCells() const { return !GetActorIsRuntimeOnly(); }
+	UE_DEPRECATED(5.1, "ShouldBeLoadedByEditorCells is deprecated, IsEditorRelevant should be used instead.")
+	bool ShouldBeLoadedByEditorCells() const { return IsEditorRelevant(); }
 
 	virtual void Init(const AActor* InActor);
 	virtual void Init(const FWorldPartitionActorDescInitData& DescData);
