Audio header cleanup

diff --git a/Engine/Source/Runtime/Engine/Private/AudioDevice.cpp b/Engine/Source/Runtime/Engine/Private/AudioDevice.cpp
--- a/Engine/Source/Runtime/Engine/Private/AudioDevice.cpp
+++ b/Engine/Source/Runtime/Engine/Private/AudioDevice.cpp
@@ -394,6 +394,30 @@ FAudioDevice::FAudioDevice()
 {
 }
 
+PRAGMA_DISABLE_DEPRECATION_WARNINGS // suppress deprecation warning in default dtor
+FAudioDevice::~FAudioDevice() = default;
+PRAGMA_ENABLE_DEPRECATION_WARNINGS
+
+FAudioDeviceHandle FAudioDevice::GetMainAudioDevice()
+{
+	// Try to get GEngine's main audio device
+	FAudioDeviceHandle AudioDevice = GEngine->GetMainAudioDevice();
+
+	// If we don't have a main audio device (maybe we're running in a non-standard mode like a commandlet)
+	if (!AudioDevice)
+	{
+		// We should have an active device for device manager
+		FAudioDeviceManager* DeviceManager = GEngine->GetAudioDeviceManager();
+		return DeviceManager->GetActiveAudioDevice();
+	}
+	return AudioDevice;
+}
+
+FAudioDeviceManager* FAudioDevice::GetAudioDeviceManager()
+{
+	return GEngine->GetAudioDeviceManager();
+}
+
 FAudioEffectsManager* FAudioDevice::CreateEffectsManager()
 {
 	return new FAudioEffectsManager(this);
@@ -778,6 +802,33 @@ bool FAudioDevice::SpatializationPluginInterfacesAvailable()
 	return GetCurrentSpatializationPluginInterfaceInfo().IsValid();
 }
 
+bool FAudioDevice::IsOcclusionPluginLoaded()
+{
+	if (FAudioDeviceHandle MainAudioDevice = GEngine->GetMainAudioDevice())
+	{
+		return MainAudioDevice->bOcclusionInterfaceEnabled;
+	}
+	return false;
+}
+
+bool FAudioDevice::IsReverbPluginLoaded()
+{
+	if (FAudioDeviceHandle MainAudioDevice = GEngine->GetMainAudioDevice())
+	{
+		return MainAudioDevice->bReverbInterfaceEnabled;
+	}
+	return false;
+}
+
+bool FAudioDevice::IsSourceDataOverridePluginLoaded()
+{
+	if (FAudioDeviceHandle MainAudioDevice = GEngine->GetMainAudioDevice())
+	{
+		return MainAudioDevice->bSourceDataOverrideInterfaceEnabled;
+	}
+	return false;
+}
+
 void FAudioDevice::PrecacheStartupSounds()
 {
 	// Iterate over all already loaded sounds and precache them. This relies on Super::Init in derived classes to be called last.
@@ -6814,6 +6865,42 @@ void FAudioDevice::UnregisterSoundClass(USoundClass* InSoundClass)
 	}
 }
 
+void FAudioDevice::RegisterSubmixBufferListener(ISubmixBufferListener* InSubmixBufferListener, USoundSubmix* SoundSubmix)
+{
+	UE_LOG(LogAudio, Error, TEXT("Submix buffer listener only works with the audio mixer. Please run with audio mixer enabled."));
+}
+
+void FAudioDevice::UnregisterSubmixBufferListener(ISubmixBufferListener* InSubmixBufferListener, USoundSubmix* SoundSubmix)
+{
+	UE_LOG(LogAudio, Error, TEXT("Submix buffer listener only works with the audio mixer. Please run with audio mixer enabled."));
+}
+
+Audio::FPatchOutputStrongPtr FAudioDevice::AddPatchForSubmix(uint32 InObjectId, float InPatchGain)
+{
+	UE_LOG(LogAudio, Error, TEXT("Submix patching only works with the audio mixer. Please run with audio mixer enabled."));
+	return nullptr;
+}
+
+Audio::FPatchOutputStrongPtr FAudioDevice::AddPatchForAudioBus(uint32 InAudioBusId, float InPatchGain)
+{
+	return nullptr;
+}
+
+Audio::FPatchOutputStrongPtr FAudioDevice::AddPatchForAudioBus_GameThread(uint32 InAudioBusId, float InPatchGain)
+{
+	return nullptr;
+}
+
+Audio::FPatchInput FAudioDevice::AddPatchInputForAudioBus(uint32 InAudioBusId, int32 InFrames, int32 InChannels, float InGain)
+{
+	return Audio::FPatchInput();
+}
+
+Audio::FPatchOutputStrongPtr FAudioDevice::AddPatchOutputForAudioBus(uint32 InAudioBusId, int32 InFrames, int32 InChannels, float InGain)
+{
+	return Audio::FPatchOutputStrongPtr();
+}
+
 FSoundClassProperties* FAudioDevice::GetSoundClassCurrentProperties(USoundClass* InSoundClass)
 {
 	if (InSoundClass)
@@ -7234,6 +7321,12 @@ void FAudioDevice::SetPlatformAudioHeadroom(const float InPlatformHeadRoom)
 	PlatformAudioHeadroom = InPlatformHeadRoom;
 }
 
+bool FAudioDevice::IsMainAudioDevice() const
+{
+	FAudioDeviceHandle MainAudioDevice = GEngine->GetMainAudioDevice();
+	return (!MainAudioDevice || MainAudioDevice.GetAudioDevice() == this);
+}
+
 const TArray<FWaveInstance*>& FAudioDevice::GetActiveWaveInstances() const
 {
 	check(IsInAudioThread());
diff --git a/Engine/Source/Runtime/Engine/Public/AudioDevice.h b/Engine/Source/Runtime/Engine/Public/AudioDevice.h
--- a/Engine/Source/Runtime/Engine/Public/AudioDevice.h
+++ b/Engine/Source/Runtime/Engine/Public/AudioDevice.h
@@ -1,32 +1,39 @@
 // Copyright Epic Games, Inc. All Rights Reserved.
+
 #pragma once
 
+#include "AudioDeviceHandle.h"
+#include "AudioDynamicParameter.h"
+#include "AudioVirtualLoop.h"
+#include "Components/AudioComponent.h"
+#include "HAL/LowLevelMemStats.h"
+#include "Math/Transform.h"
+#include "Math/Vector.h"
+#include "Sound/AudioVolume.h"
+#include "Sound/SoundClass.h"
+#include "Sound/SoundMix.h"
+#include "Subsystems/AudioEngineSubsystem.h"
+#include "UObject/StrongObjectPtr.h"
+
+#if UE_ENABLE_INCLUDE_ORDER_DEPRECATED_IN_5_3
 #include "Audio.h"
 #include "AudioDeviceManager.h"
-#include "Components/AudioComponent.h"
+#include "AudioMixer.h"
 #include "CoreMinimal.h"
 #include "DSP/MultithreadedPatching.h"
 #include "Engine/Engine.h"
 #include "EngineGlobals.h"
-#include "HAL/LowLevelMemStats.h"
 #include "IAudioExtensionPlugin.h"
 #include "ISubmixBufferListener.h"
-#include "AudioDynamicParameter.h"
 #include "Sound/AudioSettings.h"
-#include "Sound/AudioVolume.h"
 #include "Sound/SoundAttenuation.h"
-#include "Sound/SoundClass.h"
 #include "Sound/SoundConcurrency.h"
-#include "Sound/SoundMix.h"
+#include "Sound/SoundModulationDestination.h"
+#include "Sound/SoundSourceBus.h"
 #include "Sound/SoundSubmix.h"
 #include "Sound/SoundSubmixSend.h"
-#include "Sound/SoundSourceBus.h"
-#include "Sound/SoundModulationDestination.h"
-#include "Subsystems/AudioEngineSubsystem.h"
 #include "Subsystems/SubsystemCollection.h"
-#include "AudioVirtualLoop.h"
-#include "AudioMixer.h"
-#include "UObject/StrongObjectPtr.h"
+#endif
 
 #if UE_ENABLE_INCLUDE_ORDER_DEPRECATED_IN_5_1
 #include "DSP/SpectrumAnalyzer.h"
@@ -48,9 +55,11 @@ class FViewport;
 class FViewportClient;
 class IAudioSpatialization;
 class ICompressedAudioInfo;
+class ISubmixBufferListener;
 class UReverbEffect;
 class USoundAttenuation;
 class USoundBase;
+class USoundClass;
 class USoundConcurrency;
 class USoundEffectSourcePreset;
 class USoundEffectSubmixPreset;
@@ -65,8 +74,23 @@ struct FActiveSound;
 struct FAttenuationFocusData;
 struct FAudioComponentParam;
 struct FAudioQualitySettings;
+struct FAudioVirtualLoop;
+struct FSourceEffectChainEntry;
+struct FSoundClassDynamicProperties;
+struct FSoundClassProperties;
+struct FSoundGroup;
+struct FSoundParseParameters;
+struct FSoundSpectrumAnalyzerDelegateSettings;
+struct FSoundSpectrumAnalyzerSettings;
 struct FWaveInstance;
 
+namespace Audio
+{
+	class FPatchInput;
+	struct FPatchOutput;
+	typedef TSharedPtr<FPatchOutput, ESPMode::ThreadSafe> FPatchOutputStrongPtr;
+}
+
 namespace Audio
 {
 	class FAudioDebugger;
@@ -498,9 +522,7 @@ public:
 	 */
 	FAudioDevice();
 
-	PRAGMA_DISABLE_DEPRECATION_WARNINGS // supress deprecation warning in default dtor
-	virtual ~FAudioDevice() = default;
-	PRAGMA_ENABLE_DEPRECATION_WARNINGS
+	virtual ~FAudioDevice();
 
 	/** Returns an array of available audio devices names for the platform */
 	virtual void GetAudioDeviceList(TArray<FString>& OutAudioDeviceNames) const
@@ -940,25 +962,15 @@ public:
 	 * Registers the submix buffer listener with the given submix.
 	 * A nullptr for SoundSubmix will register the listener with the master submix.
 	*/
-	virtual void RegisterSubmixBufferListener(ISubmixBufferListener* InSubmixBufferListener, USoundSubmix* SoundSubmix = nullptr)
-	{
-		UE_LOG(LogAudio, Error, TEXT("Submix buffer listener only works with the audio mixer. Please run with audio mixer enabled."));
-	}
+	virtual void RegisterSubmixBufferListener(ISubmixBufferListener* InSubmixBufferListener, USoundSubmix* SoundSubmix = nullptr);
 
 	/**
 	 * Unregisters the submix buffer listener with the given submix.
 	 * A nullptr for SoundSubmix will unregister the listener with the master submix.
 	*/
-	virtual void UnregisterSubmixBufferListener(ISubmixBufferListener* InSubmixBufferListener, USoundSubmix* SoundSubmix = nullptr)
-	{
-		UE_LOG(LogAudio, Error, TEXT("Submix buffer listener only works with the audio mixer. Please run with audio mixer enabled."));
-	}
+	virtual void UnregisterSubmixBufferListener(ISubmixBufferListener* InSubmixBufferListener, USoundSubmix* SoundSubmix = nullptr);
 
-	virtual Audio::FPatchOutputStrongPtr AddPatchForSubmix(uint32 InObjectId, float InPatchGain)
-	{
-		UE_LOG(LogAudio, Error, TEXT("Submix patching only works with the audio mixer. Please run with audio mixer enabled."));
-		return nullptr;
-	}
+	virtual Audio::FPatchOutputStrongPtr AddPatchForSubmix(uint32 InObjectId, float InPatchGain);
 
 	virtual void StartAudioBus(uint32 InAudioBusId, int32 InNumChannels, bool bInIsAutomatic)
 	{
@@ -974,16 +986,10 @@ public:
 	}
 
 	UE_DEPRECATED(5.2, "AddPatchForAudioBus is deprecated.  Use AddPatchOutputForAudioBus.")
-	virtual Audio::FPatchOutputStrongPtr AddPatchForAudioBus(uint32 InAudioBusId, float InPatchGain = 1.0f)
-	{
-		return nullptr;
-	}
+	virtual Audio::FPatchOutputStrongPtr AddPatchForAudioBus(uint32 InAudioBusId, float InPatchGain = 1.0f);
 
 	UE_DEPRECATED(5.2, "AddPatchForAudioBus_GameThread is deprecated.  Use AddPatchOutputForAudioBus.")
-	virtual Audio::FPatchOutputStrongPtr AddPatchForAudioBus_GameThread(uint32 InAudioBusId, float InPatchGain = 1.0f)
-	{
-		return nullptr;
-	}
+	virtual Audio::FPatchOutputStrongPtr AddPatchForAudioBus_GameThread(uint32 InAudioBusId, float InPatchGain = 1.0f);
 
 	UE_DEPRECATED(5.2, "This overload of AddPatchInputForAudioBus is deprecated.  Use the overload that takes the number of frames and channels as parameters.")
 	virtual void AddPatchInputForAudioBus(const Audio::FPatchInput& InPatchInput, uint32 InAudioBusId, float InPatchGain = 1.0f)
@@ -995,15 +1001,9 @@ public:
 	{
 	}
 
-	virtual Audio::FPatchInput AddPatchInputForAudioBus(uint32 InAudioBusId, int32 InFrames, int32 InChannels, float InGain = 1.f)
-	{
-		return Audio::FPatchInput();
-	}
+	virtual Audio::FPatchInput AddPatchInputForAudioBus(uint32 InAudioBusId, int32 InFrames, int32 InChannels, float InGain = 1.f);
 
-	virtual Audio::FPatchOutputStrongPtr AddPatchOutputForAudioBus(uint32 InAudioBusId, int32 InFrames, int32 InChannels, float InGain = 1.f)
-	{
-		return Audio::FPatchOutputStrongPtr();
-	}
+	virtual Audio::FPatchOutputStrongPtr AddPatchOutputForAudioBus(uint32 InAudioBusId, int32 InFrames, int32 InChannels, float InGain = 1.f);
 
 	virtual void InitSoundEffectPresets() {}
 
@@ -1311,26 +1311,10 @@ public:
 	}
 
 	/** Returns the main audio device of the engine */
-	static FAudioDeviceHandle GetMainAudioDevice()
-	{
-		// Try to get GEngine's main audio device
-		FAudioDeviceHandle AudioDevice = GEngine->GetMainAudioDevice();
-
-		// If we don't have a main audio device (maybe we're running in a non-standard mode like a commandlet)
-		if (!AudioDevice)
-		{
-			// We should have an active device for device manager
-			FAudioDeviceManager* DeviceManager = GEngine->GetAudioDeviceManager();
-			return DeviceManager->GetActiveAudioDevice();
-		}
-		return AudioDevice;
-	}
+	static FAudioDeviceHandle GetMainAudioDevice();
 
 	/** Returns the audio device manager */
-	static FAudioDeviceManager* GetAudioDeviceManager()
-	{
-		return GEngine->GetAudioDeviceManager();
-	}
+	static FAudioDeviceManager* GetAudioDeviceManager();
 
 	/** Low pass filter OneOverQ value */
 	float GetLowPassFilterResonance() const;
@@ -1386,14 +1370,7 @@ public:
 		return bOcclusionInterfaceEnabled;
 	}
 
-	static bool IsOcclusionPluginLoaded()
-	{
-		if (FAudioDeviceHandle MainAudioDevice = GEngine->GetMainAudioDevice())
-		{
-			return MainAudioDevice->bOcclusionInterfaceEnabled;
-		}
-		return false;
-	}
+	static bool IsOcclusionPluginLoaded();
 
 	/** Whether or not there's a reverb plugin enabled. */
 	bool IsReverbPluginEnabled() const
@@ -1401,28 +1378,14 @@ public:
 		return bReverbInterfaceEnabled;
 	}
 
-	static bool IsReverbPluginLoaded()
-	{
-		if (FAudioDeviceHandle MainAudioDevice = GEngine->GetMainAudioDevice())
-		{
-			return MainAudioDevice->bReverbInterfaceEnabled;
-		}
-		return false;
-	}
+	static bool IsReverbPluginLoaded();
 
 	bool IsSourceDataOverridePluginEnabled() const
 	{
 		return bSourceDataOverrideInterfaceEnabled;
 	}
 
-	static bool IsSourceDataOverridePluginLoaded()
-	{
-		if (FAudioDeviceHandle MainAudioDevice = GEngine->GetMainAudioDevice())
-		{
-			return MainAudioDevice->bSourceDataOverrideInterfaceEnabled;
-		}
-		return false;
-	}
+	static bool IsSourceDataOverridePluginLoaded();
 
 	/** Returns if this is the multi-platform audio mixer. */
 	bool IsAudioMixerEnabled() const
@@ -1972,11 +1935,7 @@ public:
 	/** Whether play when silent is enabled for all sounds associated with this audio device*/
 	bool PlayWhenSilentEnabled() const { return bAllowPlayWhenSilent; }
 
-	bool IsMainAudioDevice()
-	{
-		FAudioDeviceHandle MainAudioDevice = GEngine->GetMainAudioDevice();
-		return (!MainAudioDevice || MainAudioDevice.GetAudioDevice() == this);
-	}
+	bool IsMainAudioDevice() const;
 
 	/** Set whether or not we force the use of attenuation for non-game worlds (as by default we only care about game worlds) */
 	void SetUseAttenuationForNonGameWorlds(bool bInUseAttenuationForNonGameWorlds)
diff --git a/Engine/Source/Runtime/Engine/Public/ISubmixBufferListener.h b/Engine/Source/Runtime/Engine/Public/ISubmixBufferListener.h
--- a/Engine/Source/Runtime/Engine/Public/ISubmixBufferListener.h
+++ b/Engine/Source/Runtime/Engine/Public/ISubmixBufferListener.h
@@ -1,10 +1,16 @@
 // Copyright Epic Games, Inc. All Rights Reserved.
 #pragma once
 
+#include "HAL/Platform.h"
+
+#if UE_ENABLE_INCLUDE_ORDER_DEPRECATED_IN_5_3
 #include "Sound/SoundSubmix.h"
+#endif
+
+class USoundSubmix;
 
 /** Abstract interface for receiving audio data from a given submix. */
-class ENGINE_API ISubmixBufferListener
+class ISubmixBufferListener
 {
 public:
 	virtual ~ISubmixBufferListener() = default;
@@ -32,4 +38,4 @@ public:
 	{
 		return false;
 	}
-};
\ No newline at end of file
+};
