Fix bug where Audio Component Modulation ignores inheritance from SoundWave

diff --git a/Engine/Source/Runtime/Engine/Private/ActiveSound.cpp b/Engine/Source/Runtime/Engine/Private/ActiveSound.cpp
--- a/Engine/Source/Runtime/Engine/Private/ActiveSound.cpp
+++ b/Engine/Source/Runtime/Engine/Private/ActiveSound.cpp
@@ -97,6 +97,7 @@ FActiveSound::FActiveSound()
 	, bEnableSubmixSendRoutingOverride(false)
 	, bIsFirstAttenuationUpdate(true)
 	, bStartedWithinNonBinauralRadius(false)
+	, bModulationRoutingUpdated(false)
 	, UserIndex(0)
 	, FadeOut(EFadeOut::None)
 	, bIsOccluded(false)
