Fix for material parameter collection leak.

diff --git a/Engine/Source/Runtime/Engine/Classes/Engine/World.h b/Engine/Source/Runtime/Engine/Classes/Engine/World.h
--- a/Engine/Source/Runtime/Engine/Classes/Engine/World.h
+++ b/Engine/Source/Runtime/Engine/Classes/Engine/World.h
@@ -3033,6 +3033,9 @@ public:
 	/** Updates this world's scene with the list of instances, and optionally updates each instance's uniform buffer. */
 	void UpdateParameterCollectionInstances(bool bUpdateInstanceUniformBuffers, bool bRecreateUniformBuffer);
 
+	/* clean up any material parameter collection instances which have had their collections destroyed */
+	void OnPostGC(); 
+
 private:
 	UMaterialParameterCollectionInstance* CreateParameterCollectionInstance(int32 ExistingIndex, UMaterialParameterCollection* Collection, bool bUpdateScene);
 public:
diff --git a/Engine/Source/Runtime/Engine/Classes/Materials/MaterialParameterCollectionInstance.h b/Engine/Source/Runtime/Engine/Classes/Materials/MaterialParameterCollectionInstance.h
--- a/Engine/Source/Runtime/Engine/Classes/Materials/MaterialParameterCollectionInstance.h
+++ b/Engine/Source/Runtime/Engine/Classes/Materials/MaterialParameterCollectionInstance.h
@@ -33,6 +33,8 @@ class ENGINE_API UMaterialParameterCollectionInstance : public UObject
 	/** Initializes the instance with the collection it is based off of and the world it is owned by. */
 	void SetCollection(UMaterialParameterCollection* InCollection, UWorld* InWorld);
 
+	bool IsCollectionValid() const { return Collection.IsValid(); }
+
 	/** Sets parameter value overrides, returns false if the parameter was not found. */
 	bool SetScalarParameterValue(FName ParameterName, float ParameterValue);
 	bool SetVectorParameterValue(FName ParameterName, const FLinearColor& ParameterValue);
@@ -54,7 +56,7 @@ class ENGINE_API UMaterialParameterCollectionInstance : public UObject
 
 	const UMaterialParameterCollection* GetCollection() const
 	{
-		return Collection;
+		return Collection.Get();
 	}
 
 	using ScalarParameterUpdate = TPair<FName, float>;
@@ -82,7 +84,7 @@ protected:
 
 	/** Collection resource this instance is based off of. */
 	UPROPERTY()
-	TObjectPtr<UMaterialParameterCollection> Collection;
+	TWeakObjectPtr<UMaterialParameterCollection> Collection;
 
 	/** World that owns this instance. */
 	TWeakObjectPtr<UWorld> World;
diff --git a/Engine/Source/Runtime/Engine/Private/Materials/ParameterCollection.cpp b/Engine/Source/Runtime/Engine/Private/Materials/ParameterCollection.cpp
--- a/Engine/Source/Runtime/Engine/Private/Materials/ParameterCollection.cpp
+++ b/Engine/Source/Runtime/Engine/Private/Materials/ParameterCollection.cpp
@@ -615,7 +615,7 @@ void UMaterialParameterCollectionInstance::SetCollection(UMaterialParameterColle
 
 bool UMaterialParameterCollectionInstance::SetScalarParameterValue(FName ParameterName, float ParameterValue)
 {
-	check(World.IsValid() && Collection);
+	check(World.IsValid() && Collection.IsValid());
 
 	if (Collection->GetScalarParameterByName(ParameterName))
 	{
@@ -649,7 +649,7 @@ bool UMaterialParameterCollectionInstance::SetScalarParameterValue(FName Paramet
 
 bool UMaterialParameterCollectionInstance::SetVectorParameterValue(FName ParameterName, const FLinearColor& ParameterValue)
 {
-	check(World.IsValid() && Collection);
+	check(World.IsValid() && Collection.IsValid());
 
 	if (Collection->GetVectorParameterByName(ParameterName))
 	{
@@ -741,7 +741,7 @@ void UMaterialParameterCollectionInstance::DeferredUpdateRenderState(bool bRecre
 		// Propagate the new values to the rendering thread
 		TArray<FVector4f> ParameterData;
 		GetParameterData(ParameterData);
-		Resource->GameThread_UpdateContents(Collection ? Collection->StateId : FGuid(), ParameterData, GetFName(), bRecreateUniformBuffer);
+		Resource->GameThread_UpdateContents(Collection.IsValid() ? Collection->StateId : FGuid(), ParameterData, GetFName(), bRecreateUniformBuffer);
 	}
 
 	bNeedsRenderStateUpdate = false;
@@ -751,7 +751,7 @@ void UMaterialParameterCollectionInstance::GetParameterData(TArray<FVector4f>& P
 {
 	// The memory layout created here must match the index assignment in UMaterialParameterCollection::GetParameterIndex
 
-	if (Collection)
+	if (Collection.IsValid())
 	{
 		ParameterData.Empty(FMath::DivideAndRoundUp(Collection->ScalarParameters.Num(), 4) + Collection->VectorParameters.Num());
 
diff --git a/Engine/Source/Runtime/Engine/Private/World.cpp b/Engine/Source/Runtime/Engine/Private/World.cpp
--- a/Engine/Source/Runtime/Engine/Private/World.cpp
+++ b/Engine/Source/Runtime/Engine/Private/World.cpp
@@ -1459,6 +1459,20 @@ UMaterialParameterCollectionInstance* UWorld::CreateParameterCollectionInstance(
 	return NewInstance;
 }
 
+
+void UWorld::OnPostGC()
+{
+	for (int32 InstanceIndex = ParameterCollectionInstances.Num()-1; InstanceIndex >= 0; InstanceIndex--)
+	{
+		if (!ParameterCollectionInstances[InstanceIndex]->IsCollectionValid())
+		{
+			ParameterCollectionInstances.RemoveAt(InstanceIndex);
+		}
+	}
+}
+
+
+
 UCanvas* UWorld::GetCanvasForRenderingToTarget()
 {
 	if (!CanvasForRenderingToTarget)
@@ -1828,6 +1842,8 @@ void UWorld::InitWorld(const InitializationValues IVS)
 		return;
 	}
 
+	FCoreUObjectDelegates::GetPostGarbageCollect().AddUObject(this, &UWorld::OnPostGC);
+
 	InitializeSubsystems();
 
 	FWorldDelegates::OnPreWorldInitialization.Broadcast(this, IVS);
