Improve Pose Asset GUID caching behaviour. Whenever the user performs an action that invalidates the data, and therefore the comparison - the GUID will be cleared and exempted from future comparisons. Actions include:
    - Removing a pose
    - Adding an additional pose
    - Update an existing pose with current view-port pose
#misc dynamic button label, depending on whether or not user will be updating the poses (GUID out-of-date), or restoring the poses (GUID invalidated)

diff --git a/Engine/Source/Editor/DetailCustomizations/Private/PoseAssetDetails.cpp b/Engine/Source/Editor/DetailCustomizations/Private/PoseAssetDetails.cpp
--- a/Engine/Source/Editor/DetailCustomizations/Private/PoseAssetDetails.cpp
+++ b/Engine/Source/Editor/DetailCustomizations/Private/PoseAssetDetails.cpp
@@ -306,7 +306,7 @@ void FPoseAssetDetails::CustomizeDetails(IDetailLayoutBuilder& DetailBuilder)
 		.AutoHeight()
 		[
 			SNew(SButton)
-			.Text(LOCTEXT("UpdateSource_Lable", "Update Source"))
+			.Text(this, &FPoseAssetDetails::GetAnimationUpdateButtonText)
 			.ToolTipText(LOCTEXT("UpdateSource_Tooltip", "Update Poses From Source Animation"))
 			.OnClicked(this, &FPoseAssetDetails::OnUpdatePoseSourceAnimation)
 			.IsEnabled(this, &FPoseAssetDetails::IsUpdateSourceEnabled)
@@ -702,4 +702,14 @@ bool FPoseAssetDetails::IsUpdateSourceEnabled() const
 	return false;
 }
 
+FText FPoseAssetDetails::GetAnimationUpdateButtonText() const
+{
+	if (PoseAsset.IsValid() && !PoseAsset->SourceAnimationRawDataGUID.IsValid())
+	{
+		return LOCTEXT("RestoreSource_Label", "Restore Source");
+	}
+
+	return LOCTEXT("UpdateSource_Label", "Update Source");
+}
+
 #undef LOCTEXT_NAMESPACE
diff --git a/Engine/Source/Editor/DetailCustomizations/Private/PoseAssetDetails.h b/Engine/Source/Editor/DetailCustomizations/Private/PoseAssetDetails.h
--- a/Engine/Source/Editor/DetailCustomizations/Private/PoseAssetDetails.h
+++ b/Engine/Source/Editor/DetailCustomizations/Private/PoseAssetDetails.h
@@ -95,6 +95,7 @@ private:
 	bool ShouldFilterAsset(const FAssetData& AssetData);
 	FReply OnUpdatePoseSourceAnimation();
 	bool IsUpdateSourceEnabled() const;
+	FText GetAnimationUpdateButtonText() const;
 
 	FText GetButtonText() const;
 };
diff --git a/Engine/Source/Editor/Persona/Private/AnimationEditorViewportClient.cpp b/Engine/Source/Editor/Persona/Private/AnimationEditorViewportClient.cpp
--- a/Engine/Source/Editor/Persona/Private/AnimationEditorViewportClient.cpp
+++ b/Engine/Source/Editor/Persona/Private/AnimationEditorViewportClient.cpp
@@ -1265,7 +1265,7 @@ FText FAnimationViewportClient::GetDisplayInfo(bool bDisplayAllInfo) const
 
 	if (const UPoseAsset* PoseAsset = Cast<UPoseAsset>(GetAnimPreviewScene()->GetPreviewAnimationAsset()))
 	{
-		if (PoseAsset->GetLinkerCustomVersion(FUE5MainStreamObjectVersion::GUID) >= FUE5MainStreamObjectVersion::PoseAssetRawDataGUIDUpdate && PoseAsset->SourceAnimation && PoseAsset->SourceAnimation->GetDataModel()->GenerateGuid() != PoseAsset->SourceAnimationRawDataGUID)
+		if (PoseAsset->GetLinkerCustomVersion(FUE5MainStreamObjectVersion::GUID) >= FUE5MainStreamObjectVersion::PoseAssetRawDataGUIDUpdate && PoseAsset->SourceAnimation && PoseAsset->SourceAnimationRawDataGUID.IsValid() && PoseAsset->SourceAnimation->GetDataModel()->GenerateGuid() != PoseAsset->SourceAnimationRawDataGUID)
 		{
 			TextValue = ConcatenateLine(TextValue, LOCTEXT("PoseAssetOutOfDateWarning", "<AnimViewport.WarningText>Poses are out-of-sync with the source animation. To update them click \"Update Source\"</>"));
 		}
diff --git a/Engine/Source/Runtime/Engine/Classes/Animation/PoseAsset.h b/Engine/Source/Runtime/Engine/Classes/Animation/PoseAsset.h
--- a/Engine/Source/Runtime/Engine/Classes/Animation/PoseAsset.h
+++ b/Engine/Source/Runtime/Engine/Classes/Animation/PoseAsset.h
@@ -336,6 +336,7 @@ private:
 	// After any update to SourceLocalPoses, this does update runtime data
 	void AddOrUpdatePose(const FSmartName& PoseName, const TArray<FName>& TrackNames, const TArray<FTransform>& LocalTransform, const TArray<float>& CurveValues);
 	void PostProcessData();
+	void BreakAnimationSequenceGUIDComparison();
 #endif // WITH_EDITOR	
 
 private:
diff --git a/Engine/Source/Runtime/Engine/Private/Animation/PoseAsset.cpp b/Engine/Source/Runtime/Engine/Private/Animation/PoseAsset.cpp
--- a/Engine/Source/Runtime/Engine/Private/Animation/PoseAsset.cpp
+++ b/Engine/Source/Runtime/Engine/Private/Animation/PoseAsset.cpp
@@ -918,7 +918,7 @@ void UPoseAsset::PostLoad()
 			}
 			SourceAnimation->ConditionalPostLoad();
 			
-			if (!SourceAnimationRawDataGUID.IsValid() || SourceAnimationRawDataGUID != SourceAnimation->GetDataModel()->GenerateGuid())
+			if (SourceAnimationRawDataGUID.IsValid() && SourceAnimationRawDataGUID != SourceAnimation->GetDataModel()->GenerateGuid())
 			{
 				FFormatNamedArguments Args;
 				Args.Add(TEXT("AssetName"), FText::FromString(GetPathName()));
@@ -928,7 +928,7 @@ void UPoseAsset::PostLoad()
 				Args.Add(TEXT("Found"), FText::FromString(SourceAnimation->GetDataModel()->GenerateGuid().ToString()));
 				
 				const FText ResultText = FText::Format(LOCTEXT("PoseAssetSourceOutOfDate", "PoseAsset {AssetName} is out-of-date with its source animation {SourceAsset} {Stored} vs {Found}"), Args);
-				//UE_LOG(LogAnimation, Warning,TEXT("%s"), *ResultText.ToString());
+				UE_LOG(LogAnimation, Warning,TEXT("%s"), *ResultText.ToString());
 			}
 		}	
 	}	
@@ -1165,6 +1165,11 @@ void UPoseAsset::PostProcessData()
 	UpdateTrackBoneIndices();
 }
 
+void UPoseAsset::BreakAnimationSequenceGUIDComparison()
+{
+	SourceAnimationRawDataGUID.Invalidate();
+}
+
 bool UPoseAsset::AddOrUpdatePoseWithUniqueName(const USkeletalMeshComponent* MeshComponent, FSmartName* OutPoseName /*= nullptr*/)
 {
 	bool bSavedAdditivePose = bAdditivePose;
@@ -1248,6 +1253,8 @@ void UPoseAsset::AddOrUpdatePose(const FSmartName& PoseName, const USkeletalMesh
 				}
 			}
 		}
+		
+		BreakAnimationSequenceGUIDComparison();
 
 		// Only update curves if user has requested so - or when setting up a new pose
 		const FPoseData* PoseData = PoseContainer.FindPoseData(PoseName);
@@ -1296,6 +1303,8 @@ void UPoseAsset::AddOrUpdatePose(const FSmartName& PoseName, const TArray<FName>
 			// copy to the internal track index
 			PoseData->SourceLocalSpacePose[InternalTrackIndex] = LocalTransform[Index];
 		}
+				
+		BreakAnimationSequenceGUIDComparison();
 
 		if (bNewPose)
 		{
@@ -1333,7 +1342,7 @@ void UPoseAsset::PostEditChangeProperty(FPropertyChangedEvent& PropertyChangedEv
 
 		if (PropertyChangedEvent.Property->GetFName() == GET_MEMBER_NAME_CHECKED(UPoseAsset, SourceAnimation))
 		{
-			SourceAnimationRawDataGUID.Invalidate();
+			BreakAnimationSequenceGUIDComparison();
 		}
 	}
 }
@@ -1401,7 +1410,6 @@ void UPoseAsset::CreatePoseFromAnimation(class UAnimSequence* AnimSequence, cons
 		{
 			SetSkeleton(TargetSkeleton);
 			SourceAnimation = AnimSequence;
-			SourceAnimationRawDataGUID = AnimSequence->GetDataModel()->GenerateGuid();
 
 			// reinitialize, now we're making new pose from this animation
 			Reinitialize();
@@ -1492,6 +1500,8 @@ void UPoseAsset::CreatePoseFromAnimation(class UAnimSequence* AnimSequence, cons
 
 				PostProcessData();
 			}
+
+			SourceAnimationRawDataGUID = AnimSequence->GetDataModel()->GenerateGuid();
 		}
 	}
 }
@@ -1587,6 +1597,11 @@ int32 UPoseAsset::DeletePoses(TArray<FName> PoseNamesToDelete)
 			}
 		}
 	}
+	
+	if (ItemsDeleted)
+	{
+		BreakAnimationSequenceGUIDComparison();
+	}
 
 	PostProcessData();
 	OnPoseListChanged.Broadcast();
