Missing dllexport for UTexture2D StreamIn() & StreamOut() methods

diff --git a/Engine/Source/Runtime/Engine/Classes/Engine/Texture2D.h b/Engine/Source/Runtime/Engine/Classes/Engine/Texture2D.h
--- a/Engine/Source/Runtime/Engine/Classes/Engine/Texture2D.h
+++ b/Engine/Source/Runtime/Engine/Classes/Engine/Texture2D.h
@@ -142,8 +142,8 @@ public:
 	//~ Begin UStreamableRenderAsset Interface
 	virtual int32 CalcCumulativeLODSize(int32 NumLODs) const final override { return CalcTextureMemorySize(NumLODs); }
 	ENGINE_API int32 GetNumResidentMips() const;
-	virtual bool StreamOut(int32 NewMipCount) final override;
-	virtual bool StreamIn(int32 NewMipCount, bool bHighPrio) final override;
+	ENGINE_API virtual bool StreamOut(int32 NewMipCount) final override;
+	ENGINE_API virtual bool StreamIn(int32 NewMipCount, bool bHighPrio) final override;
 	//~ End UStreamableRenderAsset Interface
 
 	/** Trivial accessors. */
