Don't filter unresolved DataLayerAssets in WorldPartitionActorDesc as it might mismatch the Actor's version (now that the actor's DataLayerAsset is not fixup).

diff --git a/Engine/Source/Runtime/Engine/Private/WorldPartition/WorldPartitionActorDesc.cpp b/Engine/Source/Runtime/Engine/Private/WorldPartition/WorldPartitionActorDesc.cpp
--- a/Engine/Source/Runtime/Engine/Private/WorldPartition/WorldPartitionActorDesc.cpp
+++ b/Engine/Source/Runtime/Engine/Private/WorldPartition/WorldPartitionActorDesc.cpp
@@ -75,32 +75,29 @@ void FWorldPartitionActorDesc::Init(const AActor* InActor)
 		// Use Actor's DataLayerManager since the fixup is relative to this level
 		if (UWorldPartition* WorldPartition = FWorldPartitionHelpers::GetWorldPartition(InActor))
 		{
-			UDataLayerManager* DataLayerManager = WorldPartition->GetDataLayerManager();
-			if (ensure(DataLayerManager))
+			LocalDataLayerAssetPaths.Reserve(InActor->GetDataLayerAssets().Num());
+			for (const TObjectPtr<const UDataLayerAsset>& DataLayerAsset : InActor->GetDataLayerAssets())
 			{
-				LocalDataLayerAssetPaths.Reserve(InActor->GetDataLayerAssets().Num());
-				for (const TObjectPtr<const UDataLayerAsset>& DataLayerAsset : InActor->GetDataLayerAssets())
+				if (DataLayerAsset)
 				{
-					// Pass Actor Level when resolving the DataLayerInstance as FWorldPartitionActorDesc always represents the state of the actor local to its outer level
-					if (DataLayerAsset && DataLayerManager->GetDataLayerInstance(DataLayerAsset))
-					{
-						LocalDataLayerAssetPaths.Add(*DataLayerAsset->GetPathName());
-					}
+					LocalDataLayerAssetPaths.Add(*DataLayerAsset->GetPathName());
 				}
+			}
+			bIsUsingDataLayerAsset = LocalDataLayerAssetPaths.Num() > 0;
 
-				PRAGMA_DISABLE_DEPRECATION_WARNINGS
+			UDataLayerManager* DataLayerManager = WorldPartition->GetDataLayerManager();
+			if (!bIsUsingDataLayerAsset)
+			{
+				if (ensure(DataLayerManager))
+				{
+					PRAGMA_DISABLE_DEPRECATION_WARNINGS
 					// Pass Actor Level when resolving the DataLayerInstance as FWorldPartitionActorDesc always represents the state of the actor local to its outer level
 					LocalDataLayerInstanceNames = DataLayerManager->GetDataLayerInstanceNames(InActor->GetActorDataLayers());
-				PRAGMA_ENABLE_DEPRECATION_WARNINGS
+					PRAGMA_ENABLE_DEPRECATION_WARNINGS
+				}
 			}
 
-			// Validation
-			const bool bHasDataLayerAssets = LocalDataLayerAssetPaths.Num() > 0;
-			const bool bHasDeprecatedDataLayers = LocalDataLayerInstanceNames.Num() > 0;
-			check((!bHasDataLayerAssets && !bHasDeprecatedDataLayers) || (bHasDataLayerAssets != bHasDeprecatedDataLayers));
-
 			// Init DataLayers persistent info
-			bIsUsingDataLayerAsset = bHasDataLayerAssets;
 			DataLayers = bIsUsingDataLayerAsset ? MoveTemp(LocalDataLayerAssetPaths) : MoveTemp(LocalDataLayerInstanceNames);
 
 			// Init DataLayers transient info
