Fix directional light preshadow breaking due to LWC WORLD_MAX.

diff --git a/Engine/Source/Runtime/Engine/Private/Components/DirectionalLightComponent.cpp b/Engine/Source/Runtime/Engine/Private/Components/DirectionalLightComponent.cpp
--- a/Engine/Source/Runtime/Engine/Private/Components/DirectionalLightComponent.cpp
+++ b/Engine/Source/Runtime/Engine/Private/Components/DirectionalLightComponent.cpp
@@ -494,7 +494,7 @@ public:
 		OutInitializer.WAxis = FVector4(0,0,0,1);
 		// Use the minimum of half the world, things further away do not cast shadows,
 		// However, if the cascade bounds are larger, then extend the casting distance far enough to encompass the cascade.
-		OutInitializer.MinLightW = FMath::Min<float>(-HALF_WORLD_MAX, -SubjectBounds.SphereRadius);
+		OutInitializer.MinLightW = FMath::Min<float>(-UE_OLD_HALF_WORLD_MAX, -SubjectBounds.SphereRadius);
 		// Range must extend to end of cascade bounds
 		const float MaxLightW = SubjectBounds.SphereRadius;
 		OutInitializer.MaxDistanceToCastInLightW = MaxLightW - OutInitializer.MinLightW;
@@ -538,7 +538,7 @@ public:
 		OutInitializer.Scales = FVector2D(1.0f / SubjectBounds.SphereRadius,1.0f / SubjectBounds.SphereRadius);
 		OutInitializer.SubjectBounds = FBoxSphereBounds(FVector::ZeroVector,SubjectBounds.BoxExtent,SubjectBounds.SphereRadius);
 		OutInitializer.WAxis = FVector4(0,0,0,1);
-		OutInitializer.MinLightW = -HALF_WORLD_MAX;
+		OutInitializer.MinLightW = -UE_OLD_HALF_WORLD_MAX;
 		// Reduce casting distance on a directional light
 		// This is necessary to improve floating point precision in several places, especially when deriving frustum verts from InvReceiverMatrix
 		// This takes the object size into account to ensure that large objects get an extended distance
