- Add mechanism for wrapping exported text items in quotes when required.

(FUniqueNetIdRepl::ImportTextItem relies on FPropertyHelpers::ReadToken, but some net id types use characters outside the Token set, and do not import correctly. This detects such cases at export time and wraps them in quotes.

diff --git a/Engine/Source/Runtime/Engine/Classes/GameFramework/OnlineReplStructs.h b/Engine/Source/Runtime/Engine/Classes/GameFramework/OnlineReplStructs.h
--- a/Engine/Source/Runtime/Engine/Classes/GameFramework/OnlineReplStructs.h
+++ b/Engine/Source/Runtime/Engine/Classes/GameFramework/OnlineReplStructs.h
@@ -124,6 +124,8 @@ protected:
 	/** Network serialized data cache */
 	UPROPERTY(Transient)
 	TArray<uint8> ReplicationBytes;
+	
+	static bool ShouldExportTextItemAsQuotedString(const FString& NetIdStr);
 };
 
 /** Specify type trait support for various low level UPROPERTY overrides */
diff --git a/Engine/Source/Runtime/Engine/Private/OnlineReplStructs.cpp b/Engine/Source/Runtime/Engine/Private/OnlineReplStructs.cpp
--- a/Engine/Source/Runtime/Engine/Private/OnlineReplStructs.cpp
+++ b/Engine/Source/Runtime/Engine/Private/OnlineReplStructs.cpp
@@ -11,6 +11,7 @@
 #include "Dom/JsonValue.h"
 #include "EngineLogs.h"
 #include "Net/OnlineEngineInterface.h"
+#include "Misc/AsciiSet.h"
 #include "Misc/ConfigCacheIni.h"
 #include "Misc/OutputDeviceNull.h"
 #include "UObject/UnrealType.h"
@@ -461,6 +462,28 @@ bool FUniqueNetIdRepl::Serialize(FArchive& Ar)
 	return true;
 }
 
+bool FUniqueNetIdRepl::ShouldExportTextItemAsQuotedString(const FString& NetIdStr)
+{
+	// Logic derived from FPropertyHelpers::ReadToken. If that would not parse the string, we need to wrap in quotes.
+
+	if (NetIdStr.IsEmpty())
+	{
+		return false;
+	}
+
+	constexpr FAsciiSet AlphaNumericChars = "abcdefghijklmnopqrstuvwxyzABCDEFGHIJKLMNOPQRSTUVWXYZ0123456789";
+	constexpr FAsciiSet ValidChars = AlphaNumericChars + '_' + '-' + '+';
+
+	TCHAR FirstChar = NetIdStr[0];
+	const bool bFirstCharValid = AlphaNumericChars.Test(FirstChar) || FirstChar > 255;
+	if (!bFirstCharValid)
+	{
+		return true;
+	}
+
+	return !FAsciiSet::HasOnly(*NetIdStr, ValidChars);
+}
+
 bool FUniqueNetIdRepl::ExportTextItem(FString& ValueStr, FUniqueNetIdRepl const& DefaultValue, UObject* Parent, int32 PortFlags, UObject* ExportRootScope) const
 {
 	if (IsValid())
@@ -468,13 +491,29 @@ bool FUniqueNetIdRepl::ExportTextItem(FString& ValueStr, FUniqueNetIdRepl const&
 		if (IsV1())
 		{
 			FName Type = GetType();
+			const FString NetIdStr = ToString();
 			if (Type == UOnlineEngineInterface::Get()->GetDefaultOnlineSubsystemName())
 			{
-				ValueStr += ToString();
+				if (ShouldExportTextItemAsQuotedString(NetIdStr))
+				{
+					ValueStr += FString::Printf(TEXT("\"%s\""), *NetIdStr);
+				}
+				else
+				{
+					ValueStr += NetIdStr;
+				}
 			}
 			else
 			{
-				ValueStr += FString::Printf(TEXT("%s:%s"), *Type.ToString(), *ToString());
+				if (ShouldExportTextItemAsQuotedString(NetIdStr))
+				{
+					ValueStr += FString::Printf(TEXT("%s:\"%s\""), *Type.ToString(), *NetIdStr);
+				}
+				else
+				{
+					ValueStr += FString::Printf(TEXT("%s:%s"), *Type.ToString(), *NetIdStr);
+				}
+				
 			}
 		}
 		else
