Fix an ensure that would happen upon input device property reset if the triggers were set to none.

#rnx

diff --git a/Engine/Source/Runtime/Engine/Private/GameFramework/InputDeviceProperties.cpp b/Engine/Source/Runtime/Engine/Private/GameFramework/InputDeviceProperties.cpp
--- a/Engine/Source/Runtime/Engine/Private/GameFramework/InputDeviceProperties.cpp
+++ b/Engine/Source/Runtime/Engine/Private/GameFramework/InputDeviceProperties.cpp
@@ -197,7 +197,7 @@ FInputDeviceProperty* UInputDeviceTriggerEffect::GetInternalDeviceProperty()
 
 void UInputDeviceTriggerEffect::ResetDeviceProperty_Implementation(const FPlatformUserId PlatformUser, const FInputDeviceId DeviceId)
 {
-	if (BaseTriggerData.bResetUponCompletion)
+	if (BaseTriggerData.bResetUponCompletion && BaseTriggerData.AffectedTriggers != EInputDeviceTriggerMask::None)
 	{
 		// Pass in our reset property
 		ResetProperty.AffectedTriggers = BaseTriggerData.AffectedTriggers;
