Making Attenuation Settings slot on Sound Cues greyed out when Override Attenuation is selected

diff --git a/Engine/Source/Runtime/Engine/Classes/Sound/SoundBase.h b/Engine/Source/Runtime/Engine/Classes/Sound/SoundBase.h
--- a/Engine/Source/Runtime/Engine/Classes/Sound/SoundBase.h
+++ b/Engine/Source/Runtime/Engine/Classes/Sound/SoundBase.h
@@ -188,7 +188,7 @@ public:
 	float Priority;
 
 	/** Attenuation settings package for the sound */
-	UPROPERTY(EditAnywhere, Category = Attenuation)
+	UPROPERTY(EditAnywhere, Category = Attenuation, meta = (EditCondition = "IsAttenuationSettingsEditable"))
 	TObjectPtr<USoundAttenuation> AttenuationSettings;
 
 	/** Submix to route sound output to. If unset, falls back to referenced SoundClass submix.
@@ -345,6 +345,10 @@ public:
 	/** Whether or not this sound allows submix sends on preview. */
 	virtual bool EnableSubmixSendsOnPreview() const { return false; }
 
+	/** Only used as an edit condition for AttenuationSettings member, as base classes may choose to provide an attenuation override implementation */
+	UFUNCTION()
+	virtual bool IsAttenuationSettingsEditable() const { return true; }
+
 #if WITH_EDITORONLY_DATA
 	void SetTimecodeOffset(const FSoundTimecodeOffset& InTimecodeOffset);
 	TOptional<FSoundTimecodeOffset> GetTimecodeOffset() const;
diff --git a/Engine/Source/Runtime/Engine/Classes/Sound/SoundCue.h b/Engine/Source/Runtime/Engine/Classes/Sound/SoundCue.h
--- a/Engine/Source/Runtime/Engine/Classes/Sound/SoundCue.h
+++ b/Engine/Source/Runtime/Engine/Classes/Sound/SoundCue.h
@@ -182,6 +182,7 @@ public:
 	virtual bool HasCookedFFTData() const override;
 	virtual bool HasCookedAmplitudeEnvelopeData() const override;
 	virtual TSharedPtr<Audio::IParameterTransmitter> CreateParameterTransmitter(Audio::FParameterTransmitterInitParams&& InParams) const override;
+	virtual bool IsAttenuationSettingsEditable() const override { return !bOverrideAttenuation; }
 	//~ End USoundBase Interface.
 
 	/** Construct and initialize a node within this Cue */
