FExec: Fixed commands not being routed through UEngine when there's a UGameEngine. Split the call to UEngine::Exec into its components so that we don't go through the vtable and get back to calling the UGameEngine implementations.

diff --git a/Engine/Source/Runtime/Engine/Private/GameEngine.cpp b/Engine/Source/Runtime/Engine/Private/GameEngine.cpp
--- a/Engine/Source/Runtime/Engine/Private/GameEngine.cpp
+++ b/Engine/Source/Runtime/Engine/Private/GameEngine.cpp
@@ -1408,7 +1408,19 @@ bool UGameEngine::Exec( UWorld* InWorld, const TCHAR* Cmd, FOutputDevice& Ar )
 			GDisallowNetworkTravel = true;
 		}
 #endif // UE_BUILD_SHIPPING
-		if (UEngine::Exec(InWorld, Cmd, Ar))
+#if UE_ALLOW_EXEC_EDITOR
+		if (UEngine::Exec_Editor(InWorld, Cmd, Ar))
+		{
+			return true;
+		}
+#endif
+#if UE_ALLOW_EXEC_DEV
+		if (UEngine::Exec_Dev(InWorld, Cmd, Ar))
+		{
+			return true;
+		}
+#endif
+		if (UEngine::Exec_Runtime(InWorld, Cmd, Ar))
 		{
 			return true;
 		}
