Changes to address AABBTree crash:
* Change DeferredKinematicUpdateSkelMeshes to use TWeakObjectPtr to avoid garbage collected upointers getting into the tree update
* Add call to CopySolverAccelerationStructure(); to make sure pending deletes are processed before the deferred updates which may share nodes with deletes
* Add check for GetSyncTimestamp != nullptr before calling GetMarkedDeleted which depends on SyncTimestamp
#discussion: https://epicgames.slack.com/archives/C04P0RT6MBN/p1676076326371339

diff --git a/Engine/Source/Runtime/Engine/Private/PhysicsEngine/Experimental/PhysScene_Chaos.cpp b/Engine/Source/Runtime/Engine/Private/PhysicsEngine/Experimental/PhysScene_Chaos.cpp
--- a/Engine/Source/Runtime/Engine/Private/PhysicsEngine/Experimental/PhysScene_Chaos.cpp
+++ b/Engine/Source/Runtime/Engine/Private/PhysicsEngine/Experimental/PhysScene_Chaos.cpp
@@ -48,6 +48,8 @@ FAutoConsoleVariableRef CVar_EnableKinematicDeferralStartPhysicsCondition(TEXT("
 bool GKinematicDeferralCheckValidBodies = true;
 FAutoConsoleVariableRef CVar_KinematicDeferralCheckValidBodies(TEXT("p.KinematicDeferralCheckValidBodies"), GKinematicDeferralCheckValidBodies, TEXT("If true, don't attempt to update deferred kinematic skeletal mesh bodies which are pending delete."));
 
+bool GKinematicDeferralUpdateExternalAccelerationStructure = false;
+FAutoConsoleVariableRef CVar_KinematicDeferralUpdateExternalAccelerationStructure(TEXT("p.KinematicDeferralUpdateExternalAccelerationStructure"), GKinematicDeferralUpdateExternalAccelerationStructure, TEXT("If true, process any operations in PendingSpatialOperations_External before doing deferred kinematic updates."));
 bool GKinematicDeferralLogInvalidBodies = true;
 FAutoConsoleVariableRef CVar_KinematicDeferralLogInvalidBodies(TEXT("p.KinematicDeferralLogInvalidBodies"), GKinematicDeferralLogInvalidBodies, TEXT("If true and p.KinematicDeferralCheckValidBodies is true, log when an invalid body is found on kinematic update."));
 
@@ -1269,7 +1271,10 @@ void FPhysScene_Chaos::ClearPreSimKinematicUpdate(USkeletalMeshComponent* InSkel
 		const int32 DeferredKinematicUpdateIndex = InSkelComp->DeferredKinematicUpdateIndex;
 		if (DeferredKinematicUpdateIndex != INDEX_NONE)
 		{
-			DeferredKinematicUpdateSkelMeshes.Last().Key->DeferredKinematicUpdateIndex = DeferredKinematicUpdateIndex;
+			if (USkeletalMeshComponent* SkelComp = DeferredKinematicUpdateSkelMeshes.Last().Key.Get())
+			{
+				SkelComp->DeferredKinematicUpdateIndex = DeferredKinematicUpdateIndex;
+			}
 			DeferredKinematicUpdateSkelMeshes.RemoveAtSwap(InSkelComp->DeferredKinematicUpdateIndex);
 			InSkelComp->DeferredKinematicUpdateIndex = INDEX_NONE;
 		}
@@ -1444,14 +1449,16 @@ void FPhysScene_Chaos::UpdateKinematicsOnDeferredSkelMeshes()
 		SkeletalMeshStartIndexArray.Reserve(DeferredKinematicUpdateSkelMeshes.Num());
 
 		int32 TotalBodies = 0;
-		for (const TPair<USkeletalMeshComponent*, FDeferredKinematicUpdateInfo>& DeferredKinematicUpdate : DeferredKinematicUpdateSkelMeshes)
+		for (const TPair<TWeakObjectPtr<USkeletalMeshComponent>, FDeferredKinematicUpdateInfo>& DeferredKinematicUpdate : DeferredKinematicUpdateSkelMeshes)
 		{
 			SkeletalMeshStartIndexArray.Add(TotalBodies);
 
-			USkeletalMeshComponent* SkelComp = DeferredKinematicUpdate.Key;
-			if (!SkelComp->bEnablePerPolyCollision)
+			if (USkeletalMeshComponent* SkelComp = DeferredKinematicUpdate.Key.Get())
 			{
-				TotalBodies += SkelComp->Bodies.Num();
+				if (!SkelComp->bEnablePerPolyCollision)
+				{
+					TotalBodies += SkelComp->Bodies.Num();
+				}
 			}
 		}
 
@@ -1462,9 +1469,14 @@ void FPhysScene_Chaos::UpdateKinematicsOnDeferredSkelMeshes()
 
 	// Gather proxies that need to be dirtied before paralell loop, and update any per poly collision skeletal meshes.
 	{
-		for (const TPair<USkeletalMeshComponent*, FDeferredKinematicUpdateInfo>& DeferredKinematicUpdate : DeferredKinematicUpdateSkelMeshes)
+		for (const TPair<TWeakObjectPtr<USkeletalMeshComponent>, FDeferredKinematicUpdateInfo>& DeferredKinematicUpdate : DeferredKinematicUpdateSkelMeshes)
 		{
-			USkeletalMeshComponent* SkelComp = DeferredKinematicUpdate.Key;
+			USkeletalMeshComponent* SkelComp = DeferredKinematicUpdate.Key.Get();
+			if (SkelComp == nullptr)
+			{
+				continue;
+			}
+
 			const FDeferredKinematicUpdateInfo& Info = DeferredKinematicUpdate.Value;
 
 			if (!SkelComp->bEnablePerPolyCollision)
@@ -1488,7 +1500,7 @@ void FPhysScene_Chaos::UpdateKinematicsOnDeferredSkelMeshes()
 					}
 
 					FPhysicsActorHandle& ActorHandle = BodyInst->ActorHandle;
-					if (GKinematicDeferralCheckValidBodies && (ActorHandle == nullptr || ActorHandle->GetMarkedDeleted()))
+					if (GKinematicDeferralCheckValidBodies && (ActorHandle == nullptr || ActorHandle->GetSyncTimestamp() == nullptr || ActorHandle->GetMarkedDeleted()))
 					{
 						if (GKinematicDeferralLogInvalidBodies)
 						{
@@ -1536,8 +1548,13 @@ void FPhysScene_Chaos::UpdateKinematicsOnDeferredSkelMeshes()
 	{
 		Chaos::PhysicsParallelFor(DeferredKinematicUpdateSkelMeshes.Num(), [&](int32 Index)
 		{
-			const TPair<USkeletalMeshComponent*, FDeferredKinematicUpdateInfo>& DeferredKinematicUpdate = DeferredKinematicUpdateSkelMeshes[Index];
-			USkeletalMeshComponent* SkelComp = DeferredKinematicUpdate.Key;
+			const TPair<TWeakObjectPtr<USkeletalMeshComponent>, FDeferredKinematicUpdateInfo>& DeferredKinematicUpdate = DeferredKinematicUpdateSkelMeshes[Index];
+			USkeletalMeshComponent* SkelComp = DeferredKinematicUpdate.Key.Get();
+			if (SkelComp == nullptr)
+			{
+				return;
+			}
+
 			const FDeferredKinematicUpdateInfo& Info = DeferredKinematicUpdate.Value;
 
 			SkelComp->DeferredKinematicUpdateIndex = INDEX_NONE;
@@ -1568,7 +1585,7 @@ void FPhysScene_Chaos::UpdateKinematicsOnDeferredSkelMeshes()
 					}
 
 					FPhysicsActorHandle& ActorHandle = BodyInst->ActorHandle;
-					if (GKinematicDeferralCheckValidBodies && (ActorHandle == nullptr || ActorHandle->GetMarkedDeleted()))
+					if (GKinematicDeferralCheckValidBodies && (ActorHandle == nullptr || ActorHandle->GetSyncTimestamp() == nullptr || ActorHandle->GetMarkedDeleted()))
 					{
 						if (GKinematicDeferralLogInvalidBodies)
 						{
@@ -1625,6 +1642,17 @@ void FPhysScene_Chaos::UpdateKinematicsOnDeferredSkelMeshes()
 		BodyScalePair.BodyInstance->UpdateBodyScale(BodyScalePair.Scale);
 	}
 
+	// If there are pending deletions that share an AABBTree node with an actor
+	// that is to be updated in TeleportActorsPool, then we must remove it before
+	// trying to update the node.
+	//
+	// If the deleted particle has had its memory cleared or overwritten, then it
+	// may have invalid object bounds, so when the tree attempts to update the bounds
+	// of the parent node, NaNs or other invalid numbers may appear in the tree
+	if (GKinematicDeferralUpdateExternalAccelerationStructure)
+	{
+		CopySolverAccelerationStructure();
+	}
 
 	UpdateActorsInAccelerationStructure(TeleportActorsPool);
 
diff --git a/Engine/Source/Runtime/Engine/Public/Physics/Experimental/PhysScene_Chaos.h b/Engine/Source/Runtime/Engine/Public/Physics/Experimental/PhysScene_Chaos.h
--- a/Engine/Source/Runtime/Engine/Public/Physics/Experimental/PhysScene_Chaos.h
+++ b/Engine/Source/Runtime/Engine/Public/Physics/Experimental/PhysScene_Chaos.h
@@ -339,7 +339,7 @@ private:
 	};
 
 	/** Map of SkeletalMeshComponents that need their bone transforms sent to the physics engine before simulation. */
-	TArray<TPair<USkeletalMeshComponent*,FDeferredKinematicUpdateInfo>>	DeferredKinematicUpdateSkelMeshes;
+	TArray<TPair<TWeakObjectPtr<USkeletalMeshComponent>, FDeferredKinematicUpdateInfo>> DeferredKinematicUpdateSkelMeshes;
 
 	TSet<UPrimitiveComponent*> DeferredCreatePhysicsStateComponents;
 	//Body Instances
