Added early out and stats for HandleCollisionEvents and HandlePhysicsCollisions_AssumesLocked.

diff --git a/Engine/Source/Runtime/Engine/Private/PhysicsEngine/Experimental/PhysScene_Chaos.cpp b/Engine/Source/Runtime/Engine/Private/PhysicsEngine/Experimental/PhysScene_Chaos.cpp
--- a/Engine/Source/Runtime/Engine/Private/PhysicsEngine/Experimental/PhysScene_Chaos.cpp
+++ b/Engine/Source/Runtime/Engine/Private/PhysicsEngine/Experimental/PhysScene_Chaos.cpp
@@ -732,6 +732,7 @@ FCollisionNotifyInfo& FPhysScene_Chaos::GetPendingCollisionForContactPair(const
 
 void FPhysScene_Chaos::HandleCollisionEvents(const Chaos::FCollisionEventData& Event)
 {
+	QUICK_SCOPE_CYCLE_COUNTER(STAT_HandleCollisionEvents);
 	ContactPairToPendingNotifyMap.Reset();
 
 	TMap<IPhysicsProxyBase*, TArray<int32>> const& PhysicsProxyToCollisionIndicesMap = Event.PhysicsProxyToCollisionIndices.PhysicsProxyToIndicesMap;
@@ -851,21 +852,23 @@ void FPhysScene_Chaos::DispatchPendingCollisionNotifies()
 		OwningWorld->PhysicsCollisionHandler->HandlePhysicsCollisions_AssumesLocked(PendingCollisionNotifies);
 	}
 
-	// Fire any collision notifies in the queue.
-	for (FCollisionNotifyInfo& NotifyInfo : PendingCollisionNotifies)
+	if (!PendingCollisionNotifies.IsEmpty())
 	{
-		//		if (NotifyInfo.RigidCollisionData.ContactInfos.Num() > 0)
+		// Fire any collision notifies in the queue.
+		for (FCollisionNotifyInfo& NotifyInfo : PendingCollisionNotifies)
 		{
-			if (NotifyInfo.bCallEvent0 && /*NotifyInfo.IsValidForNotify() && */ NotifyInfo.Info0.Actor.IsValid())
+			//if (NotifyInfo.RigidCollisionData.ContactInfos.Num() > 0)
 			{
-				NotifyInfo.Info0.Actor->DispatchPhysicsCollisionHit(NotifyInfo.Info0, NotifyInfo.Info1, NotifyInfo.RigidCollisionData);
-			}
+				if (NotifyInfo.bCallEvent0 && /*NotifyInfo.IsValidForNotify() && */ NotifyInfo.Info0.Actor.IsValid())
+				{
+					NotifyInfo.Info0.Actor->DispatchPhysicsCollisionHit(NotifyInfo.Info0, NotifyInfo.Info1, NotifyInfo.RigidCollisionData);
+				}
 
-			// CHAOS: don't call event 1, because the code below will generate the reflexive hit data as separate entries
+				// CHAOS: don't call event 1, because the code below will generate the reflexive hit data as separate entries
+			}
 		}
+		PendingCollisionNotifies.Reset();
 	}
-
-	PendingCollisionNotifies.Reset();
 }
 
 #if CHAOS_WITH_PAUSABLE_SOLVER
