Export buffer Serialize() functions.
Want this for storing vertex buffers in GeometryCollection object.

diff --git a/Engine/Source/Runtime/Engine/Public/Rendering/PositionVertexBuffer.h b/Engine/Source/Runtime/Engine/Public/Rendering/PositionVertexBuffer.h
--- a/Engine/Source/Runtime/Engine/Public/Rendering/PositionVertexBuffer.h
+++ b/Engine/Source/Runtime/Engine/Public/Rendering/PositionVertexBuffer.h
@@ -45,7 +45,7 @@ public:
 	* Initializes this vertex buffer with the contents of the given vertex buffer.
 	* @param InVertexBuffer - The vertex buffer to initialize from.
 	*/
-	void Init(const FPositionVertexBuffer& InVertexBuffer, bool bInNeedsCPUAccess = true);
+	ENGINE_API void Init(const FPositionVertexBuffer& InVertexBuffer, bool bInNeedsCPUAccess = true);
 
 	ENGINE_API void Init(const TArray<FVector3f>& InPositions, bool bInNeedsCPUAccess = true);
 
