Added STAT_DispatchPhysicsCollisionHit to differentiate between separate collision hits and bp calls.

diff --git a/Engine/Source/Runtime/Engine/Private/PhysicsEngine/Experimental/PhysScene_Chaos.cpp b/Engine/Source/Runtime/Engine/Private/PhysicsEngine/Experimental/PhysScene_Chaos.cpp
--- a/Engine/Source/Runtime/Engine/Private/PhysicsEngine/Experimental/PhysScene_Chaos.cpp
+++ b/Engine/Source/Runtime/Engine/Private/PhysicsEngine/Experimental/PhysScene_Chaos.cpp
@@ -859,9 +859,14 @@ void FPhysScene_Chaos::DispatchPendingCollisionNotifies()
 		{
 			//if (NotifyInfo.RigidCollisionData.ContactInfos.Num() > 0)
 			{
-				if (NotifyInfo.bCallEvent0 && /*NotifyInfo.IsValidForNotify() && */ NotifyInfo.Info0.Actor.IsValid())
+				if (NotifyInfo.bCallEvent0 /*&& NotifyInfo.IsValidForNotify()*/)
 				{
-					NotifyInfo.Info0.Actor->DispatchPhysicsCollisionHit(NotifyInfo.Info0, NotifyInfo.Info1, NotifyInfo.RigidCollisionData);
+					if (AActor* Actor = NotifyInfo.Info0.Actor.Get())
+					{
+						QUICK_SCOPE_CYCLE_COUNTER(STAT_DispatchPhysicsCollisionHit);
+						SCOPE_CYCLE_UOBJECT(NotifyActor, Actor);
+						Actor->DispatchPhysicsCollisionHit(NotifyInfo.Info0, NotifyInfo.Info1, NotifyInfo.RigidCollisionData);
+					}
 				}
 
 				// CHAOS: don't call event 1, because the code below will generate the reflexive hit data as separate entries
