motion trajectory - defaulting extrapolation to true

diff --git a/Engine/Source/Runtime/Engine/Public/Animation/MotionTrajectoryTypes.h b/Engine/Source/Runtime/Engine/Public/Animation/MotionTrajectoryTypes.h
--- a/Engine/Source/Runtime/Engine/Public/Animation/MotionTrajectoryTypes.h
+++ b/Engine/Source/Runtime/Engine/Public/Animation/MotionTrajectoryTypes.h
@@ -71,7 +71,7 @@ struct ENGINE_API FTrajectorySampleRange
 	// Determine if all trajectory samples are default values
 	bool HasOnlyZeroSamples() const;
 
-	FTrajectorySample GetSampleAtTime(float Time, bool bExtrapolate = false) const;
+	FTrajectorySample GetSampleAtTime(float Time, bool bExtrapolate = true) const;
 
 	// Debug draw in-world trajectory samples and optional corresponding information
 	void DebugDrawTrajectory(bool bEnable
