Make RT shaders unaffected by the ISR/MMV settings (UE-175983).

- Resolves issuues with a path tracer crash due to an unbound instanced view UB.
- Also incorporates CL 23934017 which fixes SetViewParameters as after CL 23341764 / UE-163717 Instanced View UB cannot be substituted by instanced view's View UB.

diff --git a/Engine/Source/Runtime/Engine/Private/ShaderCompiler/ShaderCompiler.cpp b/Engine/Source/Runtime/Engine/Private/ShaderCompiler/ShaderCompiler.cpp
--- a/Engine/Source/Runtime/Engine/Private/ShaderCompiler/ShaderCompiler.cpp
+++ b/Engine/Source/Runtime/Engine/Private/ShaderCompiler/ShaderCompiler.cpp
@@ -5823,6 +5823,13 @@ void GlobalBeginCompileShader(
 		);
 }
 
+namespace
+{
+	bool ShaderFrequencyNeedsInstancedStereoMods(const FShaderType* ShaderType)
+	{
+		return !(IsRayTracingShaderFrequency(ShaderType->GetFrequency()));
+	}
+}
 
 void GlobalBeginCompileShader(
 	const FString& DebugGroupName,
@@ -6047,6 +6054,7 @@ void GlobalBeginCompileShader(
 	}
 
 	// Set VR definitions
+	if (ShaderFrequencyNeedsInstancedStereoMods(ShaderType))
 	{
 		const UE::StereoRenderUtils::FStereoShaderAspects Aspects(ShaderPlatform);
 
@@ -6077,6 +6085,12 @@ void GlobalBeginCompileShader(
 			GShaderCompilingManager->SuppressWarnings(ShaderPlatform);
 		}
 	}
+	else
+	{
+		Input.Environment.SetDefine(TEXT("INSTANCED_STEREO"), 0);
+		Input.Environment.SetDefine(TEXT("MULTI_VIEW"), 0);
+		Input.Environment.SetDefine(TEXT("MOBILE_MULTI_VIEW"), 0);
+	}
 
 	ShaderType->AddUniformBufferIncludesToEnvironment(Input.Environment, ShaderPlatform);
 
@@ -6085,7 +6099,7 @@ void GlobalBeginCompileShader(
 		VFType->AddUniformBufferIncludesToEnvironment(Input.Environment, ShaderPlatform);
 	}
 
-	// Add generated instanced stereo code
+	// Add generated instanced stereo code (this code also generates ViewState, so needed not just for ISR)
 	{
 		// this function may be called on multiple threads, so protect the storage
 		FScopeLock GeneratedInstancedCodeLock(&GCachedGeneratedInstancedStereoCodeLock);
