Fix an issue where hlods subactors of Levellnstances wouldnt not receive the proper container transform.
Hlod subactors are serialize in AWorldPartitionHlod and invalidated is subsequent hlod generations when they mistmatch their current state in the level.
Fix ContainerTransform not affecting the hlod subactors invalidation.

diff --git a/Engine/Source/Runtime/Engine/Public/WorldPartition/HLOD/HLODSubActor.h b/Engine/Source/Runtime/Engine/Public/WorldPartition/HLOD/HLODSubActor.h
--- a/Engine/Source/Runtime/Engine/Public/WorldPartition/HLOD/HLODSubActor.h
+++ b/Engine/Source/Runtime/Engine/Public/WorldPartition/HLOD/HLODSubActor.h
@@ -79,7 +79,12 @@ struct FHLODSubActor
 #if WITH_EDITORONLY_DATA
 	friend bool operator==(const FHLODSubActor& Lhs, const FHLODSubActor& Rhs)
 	{
-		return Lhs.ActorGuid == Rhs.ActorGuid && Lhs.ContainerID == Rhs.ContainerID;
+		return Lhs.ActorGuid == Rhs.ActorGuid 
+			&& Lhs.ActorPackage == Rhs.ActorPackage
+			&& Lhs.ActorPath == Rhs.ActorPath
+			&& Lhs.ContainerID == Rhs.ContainerID
+			&& Lhs.ContainerPackage == Rhs.ContainerPackage
+			&& Lhs.ContainerTransform.Equals(Rhs.ContainerTransform);
 	}
 
 	friend bool operator!=(const FHLODSubActor& Lhs, const FHLODSubActor& Rhs)
