Refactored SparseVolumeTextures to be their own type in the material compiler, similar to VirtualTextures. This makes the implementation a bit cleaner and more similar to how other texture types work. In particular, there is now a SVTObject(Parameter) node that can be used to override the SVT in a SVTSample node and as input to the TextureProperty node, giving access to TextureSize and TexelSize of SVT. Also exposed SVT texture size to blueprints.

#rnx

diff --git a/Engine/Shaders/Private/MaterialTemplate.ush b/Engine/Shaders/Private/MaterialTemplate.ush
--- a/Engine/Shaders/Private/MaterialTemplate.ush
+++ b/Engine/Shaders/Private/MaterialTemplate.ush
@@ -23,6 +23,7 @@
 #include "/Engine/Private/HairStrands/HairCardsAttributeCommon.ush"
 #include "/Engine/Private/HairStrands/HairStrandsAttributeCommon.ush"
 #include "/Engine/Private/DeferredShadingCommon.ush"
+#include "/Engine/Private/SparseVolumeTexture/SparseVolumeTextureCommon.ush"
 
 // Reduce latency of Nanite material base pass pixel shaders caused by register pressure due to the calculation of WPO
 #if PIXELSHADER && IS_NANITE_SHADING_PASS && IS_BASE_PASS
@@ -810,17 +811,6 @@ float2 GetPixelPosition(FMaterialVertexParameters Parameters)
 	return GetViewportUV(Parameters) * View.ViewSizeAndInvSize.xy;
 }
 
-uint3 SparseVolumeTextureGetVoxelCoord(uint PackedPhysicalTileCoord, uint PhysicalTileSize, uint3 PageTableCoord, uint3 VolumeCoord)
-{
-	const int3 PhysicalTileCoord = int3(
-		PackedPhysicalTileCoord & 0x7FF,
-		(PackedPhysicalTileCoord >> 11) & 0x7FF,
-		(PackedPhysicalTileCoord >> 22) & 0x3FF);
-
-	int3 VoxelCoord = PhysicalTileCoord * PhysicalTileSize + int3(VolumeCoord - PageTableCoord * PhysicalTileSize);
-	return VoxelCoord;
-}
-
 #if POST_PROCESS_MATERIAL
 
 float2 GetPixelPosition(FMaterialPixelParameters Parameters)
diff --git a/Engine/Shaders/Private/SparseVolumeTexture/SparseVolumeTextureCommon.ush b/Engine/Shaders/Private/SparseVolumeTexture/SparseVolumeTextureCommon.ush
new file mode 100644
--- /dev/null
+++ b/Engine/Shaders/Private/SparseVolumeTexture/SparseVolumeTextureCommon.ush
@@ -0,0 +1,90 @@
+// Copyright Epic Games, Inc. All Rights Reserved.
+
+#pragma once
+
+struct FSparseVolumeTexture
+{
+	Texture3D<uint>	PageTable;
+	Texture3D		PhysicalTileDataA;
+	Texture3D		PhysicalTileDataB;
+	float3			PageTableSize;
+	float3			UVScale;
+	float3			UVBias;
+	float			TileSize;
+};
+
+FSparseVolumeTexture SparseVolumeTextureCreateStruct(
+	Texture3D<uint> PageTable, 
+	Texture3D PhysicalTileDataA, 
+	Texture3D PhysicalTileDataB, 
+	float TileSize,
+	float3 PageTableSize,
+	float3 UVScale,
+	float3 UVBias)
+{
+	FSparseVolumeTexture Result;
+	Result.PageTable = PageTable;
+	Result.PhysicalTileDataA = PhysicalTileDataA;
+	Result.PhysicalTileDataB = PhysicalTileDataB;
+	Result.PageTableSize = PageTableSize;
+	Result.UVScale = UVScale;
+	Result.UVBias = UVBias;
+	Result.TileSize = TileSize;
+	return Result;
+}
+
+FSparseVolumeTexture SparseVolumeTextureCreateStruct(
+	Texture3D<uint> PageTable,
+	Texture3D PhysicalTileDataA,
+	Texture3D PhysicalTileDataB,
+	uint4 Packed0,
+	uint4 Packed1)
+{
+	FSparseVolumeTexture Result;
+	Result.PageTable = PageTable;
+	Result.PhysicalTileDataA = PhysicalTileDataA;
+	Result.PhysicalTileDataB = PhysicalTileDataB;
+	Result.UVScale = asfloat(Packed0.xyz);
+	Result.UVBias = asfloat(Packed1.xyz);
+	Result.TileSize = asfloat(Packed1.w);
+
+	// Packed as 11|11|10. Since this limits us to 2047x2047x1023, we add an implicit +1 during unpacking.
+	Result.PageTableSize = float3(
+		(Packed0.w & 0x7FF) + 1,
+		((Packed0.w >> 11) & 0x7FF) + 1,
+		((Packed0.w >> 22) & 0x3FF) + 1);
+
+	return Result;
+}
+
+float3 SparseVolumeTextureGetVoxelCoord(uint PackedPhysicalTileCoord, float TileSize, float3 PageTableCoord)
+{
+	const int3 PhysicalTileCoord = int3(
+		PackedPhysicalTileCoord & 0x7FF,
+		(PackedPhysicalTileCoord >> 11) & 0x7FF,
+		(PackedPhysicalTileCoord >> 22) & 0x3FF);
+
+	const float3 VoxelCoord = (float3(PhysicalTileCoord) + frac(PageTableCoord)) * TileSize;
+	return VoxelCoord;
+}
+
+float3 SparseVolumeTextureSamplePageTable(FSparseVolumeTexture SparseVolumeTexture, float3 UVW)
+{
+	const float3 UVWScaledBiased = UVW * SparseVolumeTexture.UVScale + SparseVolumeTexture.UVBias;
+	const float3 PageTableCoord = (UVWScaledBiased * SparseVolumeTexture.PageTableSize);
+	
+	const uint PackedPhysicalTileCoord = SparseVolumeTexture.PageTable.Load(int4(int3(PageTableCoord), 0)).x;
+	
+	const float3 VoxelCoord = SparseVolumeTextureGetVoxelCoord(PackedPhysicalTileCoord, SparseVolumeTexture.TileSize, PageTableCoord);
+	return VoxelCoord;
+}
+
+float4 SparseVolumeTextureSamplePhysicalTileData(FSparseVolumeTexture SparseVolumeTexture, float3 VoxelCoord, int PhysicalTileDataIndex)
+{
+	switch (PhysicalTileDataIndex)
+	{
+	case 0: return SparseVolumeTexture.PhysicalTileDataA.Load(int4(int3(VoxelCoord), 0));
+	case 1: return SparseVolumeTexture.PhysicalTileDataB.Load(int4(int3(VoxelCoord), 0));
+	default: return 0.0f;
+	}
+} 
diff --git a/Engine/Shaders/Public/ShaderVersion.ush b/Engine/Shaders/Public/ShaderVersion.ush
--- a/Engine/Shaders/Public/ShaderVersion.ush
+++ b/Engine/Shaders/Public/ShaderVersion.ush
@@ -2,4 +2,4 @@
 // Each time the console command is executed it generates a new GUID. As this file is included
 // in Platform.ush (which should be included in any shader) it allows to invalidate the shader DDC.
 // 
-// GUID = DB03DB78-EB6A-46FD-8894-A1BB690894F9
+// GUID = 9C9AD04A-1448-4F79-9227-A11F49F17F90
\ No newline at end of file
diff --git a/Engine/Source/Runtime/Engine/Classes/Materials/MaterialExpressionSparseVolumeTextureBase.h b/Engine/Source/Runtime/Engine/Classes/Materials/MaterialExpressionSparseVolumeTextureBase.h
new file mode 100644
--- /dev/null
+++ b/Engine/Source/Runtime/Engine/Classes/Materials/MaterialExpressionSparseVolumeTextureBase.h
@@ -0,0 +1,28 @@
+// Copyright Epic Games, Inc. All Rights Reserved.
+
+#pragma once
+
+#include "CoreMinimal.h"
+#include "Materials/MaterialExpression.h"
+#include "UObject/ObjectMacros.h"
+#include "MaterialTypes.h"
+
+#include "MaterialExpressionSparseVolumeTextureBase.generated.h"
+
+class USparseVolumeTexture;
+struct FMaterialParameterMetadata;
+
+UCLASS(abstract, hidecategories = Object)
+class ENGINE_API UMaterialExpressionSparseVolumeTextureBase : public UMaterialExpression
+{
+	GENERATED_UCLASS_BODY()
+
+	/** The Sparse Virtual Texture to sample. */
+	UPROPERTY(EditAnywhere, BlueprintReadWrite, Category = SparseVolumeTexture)
+	TObjectPtr<USparseVolumeTexture> SparseVolumeTexture;
+
+	//~ Begin UMaterialExpression Interface
+	virtual UObject* GetReferencedTexture() const override;
+	virtual bool CanReferenceTexture() const { return true; }
+	//~ End UMaterialExpression Interface
+};
diff --git a/Engine/Source/Runtime/Engine/Classes/Materials/MaterialExpressionSparseVolumeTextureObject.h b/Engine/Source/Runtime/Engine/Classes/Materials/MaterialExpressionSparseVolumeTextureObject.h
new file mode 100644
--- /dev/null
+++ b/Engine/Source/Runtime/Engine/Classes/Materials/MaterialExpressionSparseVolumeTextureObject.h
@@ -0,0 +1,44 @@
+// Copyright Epic Games, Inc. All Rights Reserved.
+
+#pragma once
+
+#include "CoreMinimal.h"
+#include "UObject/ObjectMacros.h"
+#include "Materials/MaterialExpressionSparseVolumeTextureBase.h"
+#include "Materials/MaterialExpressionSparseVolumeTextureSample.h"
+
+#include "MaterialExpressionSparseVolumeTextureObject.generated.h"
+
+struct FPropertyChangedEvent;
+
+UCLASS(collapsecategories, hidecategories = Object, MinimalAPI)
+class UMaterialExpressionSparseVolumeTextureObject : public UMaterialExpressionSparseVolumeTextureBase
+{
+	GENERATED_UCLASS_BODY()
+
+#if WITH_EDITOR
+	virtual void PostEditChangeProperty(FPropertyChangedEvent& PropertyChangedEvent) override;
+#endif // WITH_EDITOR
+	//~ Begin UMaterialExpression Interface
+#if WITH_EDITOR
+	virtual int32 Compile(class FMaterialCompiler* Compiler, int32 OutputIndex) override;
+	virtual void GetCaption(TArray<FString>& OutCaptions) const override;
+	virtual uint32 GetOutputType(int32 OutputIndex) override { return MCT_SparseVolumeTexture; }
+	//~ End UMaterialExpression Interface
+#endif // WITH_EDITOR
+};
+
+UCLASS(collapsecategories, hidecategories = (Object, MaterialExpressionTextureSample), MinimalAPI)
+class UMaterialExpressionSparseVolumeTextureTextureObjectParameter : public UMaterialExpressionSparseVolumeTextureSampleParameter
+{
+	GENERATED_UCLASS_BODY()
+
+	//~ Begin UMaterialExpression Interface
+#if WITH_EDITOR
+	virtual void GetCaption(TArray<FString>& OutCaptions) const override;
+	virtual int32 Compile(class FMaterialCompiler* Compiler, int32 OutputIndex) override;
+	virtual uint32 GetOutputType(int32 OutputIndex) override { return MCT_SparseVolumeTexture; }
+	virtual const TArray<FExpressionInput*> GetInputs() override;
+#endif
+	//~ End UMaterialExpression Interface
+};
diff --git a/Engine/Source/Runtime/Engine/Classes/Materials/MaterialExpressionSparseVolumeTextureSample.h b/Engine/Source/Runtime/Engine/Classes/Materials/MaterialExpressionSparseVolumeTextureSample.h
--- a/Engine/Source/Runtime/Engine/Classes/Materials/MaterialExpressionSparseVolumeTextureSample.h
+++ b/Engine/Source/Runtime/Engine/Classes/Materials/MaterialExpressionSparseVolumeTextureSample.h
@@ -6,6 +6,7 @@
 #include "Materials/MaterialExpression.h"
 #include "UObject/ObjectMacros.h"
 #include "MaterialTypes.h"
+#include "Materials/MaterialExpressionSparseVolumeTextureBase.h"
 
 #include "MaterialExpressionSparseVolumeTextureSample.generated.h"
 
@@ -14,7 +15,7 @@ struct FMaterialParameterMetadata;
 
 /** Material expression for sampling from a runtime virtual texture. */
 UCLASS(collapsecategories, hidecategories=Object)
-class ENGINE_API UMaterialExpressionSparseVolumeTextureSample : public UMaterialExpression
+class ENGINE_API UMaterialExpressionSparseVolumeTextureSample : public UMaterialExpressionSparseVolumeTextureBase
 {
 	GENERATED_UCLASS_BODY()
 
@@ -22,16 +23,11 @@ class ENGINE_API UMaterialExpressionSparseVolumeTextureSample : public UMaterial
 	UPROPERTY(meta = (RequiredInput = "false"))
 	FExpressionInput Coordinates;
 
-	/** The Sparse Virtual Texture to sample. */
-	UPROPERTY(EditAnywhere, Category = SparseVolumeTexture)
-	TObjectPtr<class USparseVolumeTexture> SparseVolumeTexture;
+	UPROPERTY(meta = (RequiredInput = "false", ToolTip = "Defaults to 'SparseVolumeTexture' if not specified"))
+	FExpressionInput TextureObject;
 
 protected:
 
-	//~ Begin UMaterialExpression Interface
-	virtual UObject* GetReferencedTexture() const override;
-	virtual bool CanReferenceTexture() const { return true; }
-
 #if WITH_EDITOR
 	virtual void PostLoad() override;
 	virtual int32 Compile(class FMaterialCompiler* Compiler, int32 OutputIndex) override;
diff --git a/Engine/Source/Runtime/Engine/Classes/SparseVolumeTexture/SparseVolumeTexture.h b/Engine/Source/Runtime/Engine/Classes/SparseVolumeTexture/SparseVolumeTexture.h
--- a/Engine/Source/Runtime/Engine/Classes/SparseVolumeTexture/SparseVolumeTexture.h
+++ b/Engine/Source/Runtime/Engine/Classes/SparseVolumeTexture/SparseVolumeTexture.h
@@ -184,7 +184,16 @@ public:
 	USparseVolumeTexture();
 	virtual ~USparseVolumeTexture() = default;
 
-	UFUNCTION(BlueprintCallable, Category = "Animation")
+	UFUNCTION(BlueprintCallable, Category = "Texture")
+	virtual int32 GetSizeX() const { return FMath::CeilToInt32(GetVolumeBounds().GetSize().X); }
+
+	UFUNCTION(BlueprintCallable, Category = "Texture")
+	virtual int32 GetSizeY() const { return FMath::CeilToInt32(GetVolumeBounds().GetSize().Y); }
+
+	UFUNCTION(BlueprintCallable, Category = "Texture")
+	virtual int32 GetSizeZ() const { return FMath::CeilToInt32(GetVolumeBounds().GetSize().Z); }
+
+	UFUNCTION(BlueprintCallable, Category = "Texture")
 	virtual int32 GetFrameCount() const { return 0; }
 	virtual FBox GetVolumeBounds() const { return FBox(); }
 	virtual const FSparseVolumeTextureSceneProxy* GetSparseVolumeTextureSceneProxy() const { return nullptr; }
@@ -192,6 +201,8 @@ public:
 	/** Getter for the shader uniform parameters with index as ESparseVolumeTextureShaderUniform. */
 	FVector4 GetUniformParameter(int32 Index) const;
 
+	void GetPackedUniforms(FUintVector4& OutPacked0, FUintVector4& OutPacked1) const;
+
 	/** In order to keep the contents of an animated SVT sequence stable in world space, we need to account for the fact that
 		different frames of the sequence have different AABBs. We solve this by scaling and biasing UVs that are relative to
 		the volume bounds into the UV space represented by the AABB of each animation frame.*/
diff --git a/Engine/Source/Runtime/Engine/Private/Materials/HLSLMaterialTranslator.cpp b/Engine/Source/Runtime/Engine/Private/Materials/HLSLMaterialTranslator.cpp
--- a/Engine/Source/Runtime/Engine/Private/Materials/HLSLMaterialTranslator.cpp
+++ b/Engine/Source/Runtime/Engine/Private/Materials/HLSLMaterialTranslator.cpp
@@ -2867,6 +2867,7 @@ const TCHAR* FHLSLMaterialTranslator::DescribeType(EMaterialValueType Type) cons
 	case MCT_MaterialAttributes:	return TEXT("MaterialAttributes");
 	case MCT_TextureExternal:		return TEXT("TextureExternal");
 	case MCT_TextureVirtual:		return TEXT("TextureVirtual");
+	case MCT_SparseVolumeTexture:	return TEXT("SparseVolumeTexture");
 	case MCT_VTPageTableResult:		return TEXT("VTPageTableResult");
 	case MCT_ShadingModel:			return TEXT("ShadingModel");
 	case MCT_UInt:					return TEXT("uint");
@@ -2903,6 +2904,7 @@ const TCHAR* FHLSLMaterialTranslator::HLSLTypeString(EMaterialValueType Type) co
 	case MCT_MaterialAttributes:	return TEXT("FMaterialAttributes");
 	case MCT_TextureExternal:		return TEXT("TextureExternal");
 	case MCT_TextureVirtual:		return TEXT("TextureVirtual");
+	case MCT_SparseVolumeTexture:	return TEXT("FSparseVolumeTexture");
 	case MCT_VTPageTableResult:		return TEXT("VTPageTableResult");
 	case MCT_ShadingModel:			return TEXT("uint");
 	case MCT_UInt:					return TEXT("uint");
@@ -2939,6 +2941,7 @@ const TCHAR* FHLSLMaterialTranslator::HLSLTypeStringDeriv(EMaterialValueType Typ
 	case MCT_MaterialAttributes:	return TEXT("FMaterialAttributes");
 	case MCT_TextureExternal:		return TEXT("TextureExternal");
 	case MCT_TextureVirtual:		return TEXT("TextureVirtual");
+	case MCT_SparseVolumeTexture:	return TEXT("FSparseVolumeTexture");
 	case MCT_VTPageTableResult:		return TEXT("VTPageTableResult");
 	case MCT_ShadingModel:			return TEXT("uint");
 	case MCT_UInt:					return TEXT("uint");
@@ -3047,13 +3050,13 @@ int32 FHLSLMaterialTranslator::AddCodeChunkInner(uint64 Hash, const TCHAR* Forma
 		if (CodeIndex == INDEX_NONE)
 		{
 			CodeIndex = CurrentScopeChunks->Num();
-		// Allocate a local variable name
-		const FString SymbolName = CreateSymbolName(TEXT("Local"));
-		// Construct the definition string which stores the result in a temporary and adds a newline for readability
+			// Allocate a local variable name
+			const FString SymbolName = CreateSymbolName(TEXT("Local"));
+			// Construct the definition string which stores the result in a temporary and adds a newline for readability
 			const FString LocalVariableDefinitionFinite = FString("	") + HLSLTypeString(Type) + TEXT(" ") + SymbolName + TEXT(" = ") + FormattedCode + TEXT(";") + LINE_TERMINATOR;
 			// Construct the definition string which stores the result in a temporary and adds a newline for readability
 			const FString LocalVariableDefinitionAnalytic = FString("	") + HLSLTypeString(Type) + TEXT(" ") + SymbolName + TEXT(" = ") + FormattedCode + TEXT(";") + LINE_TERMINATOR;
-		// Adding a code chunk that creates a local variable
+			// Adding a code chunk that creates a local variable
 			new(*CurrentScopeChunks) FShaderCodeChunk(Hash, *LocalVariableDefinitionFinite, *LocalVariableDefinitionAnalytic, SymbolName, Type, DerivativeStatus, false);
 		}
 	}
@@ -3511,6 +3514,18 @@ int32 FHLSLMaterialTranslator::AccessUniformExpression(int32 Index)
 			FormattedCode.Appendf(TEXT("Material.%s_%u"), BaseName, TextureInputIndex);
 		}
 	}
+	else if (CodeChunk.Type == MCT_SparseVolumeTexture)
+	{
+		int32 TextureInputIndex = UniformTextureExpressions[(uint32)EMaterialTextureParameterType::SparseVolume].AddUnique(TextureUniformExpression);
+		FormattedCode.Appendf(
+			TEXT("SparseVolumeTextureCreateStruct("
+				"Material.SparseVolumeTexturePageTable_%d,"
+				"Material.SparseVolumeTexturePhysicalA_%d,"
+				"Material.SparseVolumeTexturePhysicalB_%d,"
+				"Material.SVTPackedUniform[%d*2],"
+				"Material.SVTPackedUniform[%d*2+1])"),
+			TextureInputIndex, TextureInputIndex, TextureInputIndex, TextureInputIndex, TextureInputIndex);
+	}
 	else
 	{
 		UE_LOG(LogMaterial, Fatal,TEXT("User input of unknown type: %s"),DescribeType(CodeChunk.Type));
@@ -6652,9 +6667,10 @@ int32 FHLSLMaterialTranslator::TextureProperty(int32 TextureIndex, EMaterialExpo
 	if (TextureType != MCT_Texture2D &&
 		TextureType != MCT_TextureVirtual &&
 		TextureType != MCT_VolumeTexture &&
-		TextureType != MCT_Texture2DArray)
+		TextureType != MCT_Texture2DArray &&
+		TextureType != MCT_SparseVolumeTexture)
 	{
-		return Errorf(TEXT("Texture size only available for Texture2D, TextureVirtual, Texture2DArray, and VolumeTexture, not %s"), DescribeType(TextureType));
+		return Errorf(TEXT("Texture size only available for Texture2D, TextureVirtual, Texture2DArray, VolumeTexture and SparseVolumeTexture, not %s"), DescribeType(TextureType));
 	}
 
 	FMaterialUniformExpressionTexture* TextureExpression = (*CurrentScopeChunks)[TextureIndex].UniformExpression->GetTextureUniformExpression();
@@ -6663,7 +6679,7 @@ int32 FHLSLMaterialTranslator::TextureProperty(int32 TextureIndex, EMaterialExpo
 		return Errorf(TEXT("Expected a texture expression"));
 	}
 
-	const EMaterialValueType ValueType = (TextureType == MCT_VolumeTexture || TextureType == MCT_Texture2DArray) ? MCT_Float3 : MCT_Float2;
+	const EMaterialValueType ValueType = (TextureType == MCT_VolumeTexture || TextureType == MCT_Texture2DArray || TextureType == MCT_SparseVolumeTexture) ? MCT_Float3 : MCT_Float2;
 	return AddUniformExpression(new FMaterialUniformExpressionTextureProperty(TextureExpression, Property), ValueType, TEXT(""));
 }
 
@@ -12537,15 +12553,15 @@ bool FHLSLMaterialTranslator::IsDevelopmentFeatureEnabled(const FName& FeatureNa
 	return true;
 }
 
-int32 FHLSLMaterialTranslator::SparseVolumeTexture(USparseVolumeTexture* Texture, int32 SubTextureId, int32& TextureReferenceIndex, EMaterialSamplerType SamplerType)
+int32 FHLSLMaterialTranslator::SparseVolumeTexture(USparseVolumeTexture* Texture, int32& TextureReferenceIndex, EMaterialSamplerType SamplerType)
 {
 	TextureReferenceIndex = Material->GetReferencedTextures().Find(Texture);
 	checkf(TextureReferenceIndex != INDEX_NONE, TEXT("Material expression called Compiler->SparseVolumeTexture() without implementing UMaterialExpression::GetReferencedTexture properly"));
 
-	return AddUniformExpression(new FMaterialUniformExpressionTexture(TextureReferenceIndex, SubTextureId, SamplerType), MCT_VolumeTexture, TEXT(""));
+	return AddUniformExpression(new FMaterialUniformExpressionTexture(TextureReferenceIndex, SamplerType), MCT_SparseVolumeTexture, TEXT(""));
 }
 
-int32 FHLSLMaterialTranslator::SparseVolumeTextureParameter(FName ParameterName, USparseVolumeTexture* InDefaultTexture, int32 SubTextureId, int32& TextureReferenceIndex, EMaterialSamplerType SamplerType)
+int32 FHLSLMaterialTranslator::SparseVolumeTextureParameter(FName ParameterName, USparseVolumeTexture* InDefaultTexture, int32& TextureReferenceIndex, EMaterialSamplerType SamplerType)
 {
 	USparseVolumeTexture* DefaultTexture = InDefaultTexture;
 
@@ -12562,103 +12578,75 @@ int32 FHLSLMaterialTranslator::SparseVolumeTextureParameter(FName ParameterName,
 	FMaterialParameterInfo ParameterInfo = GetParameterAssociationInfo();
 	ParameterInfo.Name = ParameterName;
 
-	return AddUniformExpression(new FMaterialUniformExpressionTextureParameter(ParameterInfo, TextureReferenceIndex, SubTextureId, SamplerType), MCT_VolumeTexture, TEXT(""));
+	return AddUniformExpression(new FMaterialUniformExpressionTextureParameter(ParameterInfo, TextureReferenceIndex, SamplerType), MCT_SparseVolumeTexture, TEXT(""));
 }
 
-int32 FHLSLMaterialTranslator::SparseVolumeTextureGetVoxelCoord(int32 PackedPhysicalTileCoord, int32 TileSize, int32 CoordPageTable, int32 CoordVolume)
+int32 FHLSLMaterialTranslator::SparseVolumeTextureUniform(int32 TextureIndex, int32 VectorIndex, UE::Shader::EValueType Type)
 {
-	FString SampleCode = FString::Printf(TEXT("SparseVolumeTextureGetVoxelCoord(%s, %s, %s, %s)"),
-		*GetParameterCode(PackedPhysicalTileCoord), *GetParameterCode(TileSize), *GetParameterCode(CoordPageTable), *GetParameterCode(CoordVolume));
-	return AddCodeChunk(MCT_UInt3, *SampleCode);
+	return AddUniformExpression(new FMaterialUniformExpressionSparseVolumeTextureUniform(TextureIndex, VectorIndex), GetMaterialValueType(Type), TEXT(""));
 }
 
-int32 FHLSLMaterialTranslator::SparseVolumeTextureSample(int32 TextureIndex, int32 CoordinateIndex)
+int32 FHLSLMaterialTranslator::SparseVolumeTextureUniformParameter(FName ParameterName, int32 TextureIndex, int32 VectorIndex, UE::Shader::EValueType Type)
 {
-	if (TextureIndex == INDEX_NONE || CoordinateIndex == INDEX_NONE)
+	FMaterialParameterInfo ParameterInfo = GetParameterAssociationInfo();
+	ParameterInfo.Name = ParameterName;
+	return AddUniformExpression(new FMaterialUniformExpressionSparseVolumeTextureUniform(ParameterInfo, TextureIndex, VectorIndex), GetMaterialValueType(Type), TEXT(""));
+}
+
+int32 FHLSLMaterialTranslator::SparseVolumeTextureSamplePageTable(int32 SparseVolumeTextureIndex, int32 UVWIndex)
+{
+	if (SparseVolumeTextureIndex == INDEX_NONE || UVWIndex == INDEX_NONE)
 	{
 		return INDEX_NONE;
 	}
 
-	EMaterialValueType TextureType = GetParameterType(TextureIndex);
-
-	int32 SamplingCodeIndex = INDEX_NONE;
-	FString TextureName;
-	if (TextureType == MCT_TextureCube)
+	EMaterialValueType TextureType = GetParameterType(SparseVolumeTextureIndex);
+	if ((TextureType & MCT_SparseVolumeTexture) == 0)
 	{
-		return Errorf(TEXT("FHLSLMaterialTranslator::SparseVolumeTextureSample MCT_TextureCube ERROR."));
+		return Errorf(TEXT("FHLSLMaterialTranslator::SparseVolumeTextureSamplePageTable expects MCT_SparseVolumeTexture but was passed %s ERROR."), DescribeType(TextureType));
 	}
-	else if (TextureType == MCT_Texture2DArray)
-	{
-		return Errorf(TEXT("FHLSLMaterialTranslator::SparseVolumeTextureSample MCT_Texture2DArray ERROR."));
-	}
-	else if (TextureType == MCT_TextureCubeArray)
+
+	EMaterialValueType UVWType = GetParameterType(UVWIndex);
+	if ((UVWType & MCT_Float3) == 0)
 	{
-		return Errorf(TEXT("FHLSLMaterialTranslator::SparseVolumeTextureSample MCT_TextureCubeArray ERROR."));
+		return Errorf(TEXT("FHLSLMaterialTranslator::SparseVolumeTextureSamplePageTable expects MCT_Float3 as UVW input but was passed %s ERROR."), DescribeType(UVWType));
 	}
-	else if (TextureType == MCT_VolumeTexture)
-	{
-		TextureName = CoerceParameter(TextureIndex, MCT_VolumeTexture);		// TODO this hsould be some texture reserved for sub textures
-
-
-		FMaterialUniformExpression* UniformExpression = GetParameterUniformExpression(TextureIndex);
-		if (UniformExpression == nullptr)
-		{
-			return Errorf(TEXT("Unable to find expression for Sparse Volume Texture."));
-		}
-		FMaterialUniformExpressionTexture* TextureUniformExpression = UniformExpression->GetTextureUniformExpression();
-		if (TextureUniformExpression == nullptr || TextureUniformExpression->GetPageTableLayerIndex() != INDEX_NONE || TextureUniformExpression->GetTextureLayerIndex() == INDEX_NONE)
-		{
-			// Sparse volume texture do not use PageTableLayerIndex but only TextureLayerIndex.
-			return Errorf(TEXT("The provided uniform expression is cannot be identified as a Sparse Volume Texture."));
-		}
-
-		int32 VirtualTextureIndex = UniformTextureExpressions[(uint32)EMaterialTextureParameterType::Volume].Find(TextureUniformExpression);
-		check(UniformTextureExpressions[(uint32)EMaterialTextureParameterType::Volume].IsValidIndex(VirtualTextureIndex));
 
-		const int32 TextureLayerIndex = TextureUniformExpression->GetTextureLayerIndex();
-		check(TextureLayerIndex != INDEX_NONE);
-		if (TextureLayerIndex == 0)
-		{
-			FString SampleCode = FString::Printf(TEXT("%s.Load(int4(int3(%s), 0)).x"), *TextureName, *GetParameterCode(CoordinateIndex));
-			AddEstimatedTextureSample();
-			// We ignore derivative for now, see IsAnalyticDerivEnabled() and IsDerivativeValid(UvDerivativeStatus)
-			SamplingCodeIndex = AddCodeChunk(MCT_UInt1, *SampleCode);
-		}
-		else if (TextureLayerIndex == 1 || TextureLayerIndex == 2)
-		{
-			FString SampleCode = FString::Printf(TEXT("%s.Load(int4(int3(%s), 0))"), *TextureName, *GetParameterCode(CoordinateIndex));
-			AddEstimatedTextureSample();
-			// We ignore derivative for now, see IsAnalyticDerivEnabled() and IsDerivativeValid(UvDerivativeStatus)
-			SamplingCodeIndex = AddCodeChunk(MCT_Float4, *SampleCode);
-		}
-		else
-		{
-			return Errorf(TEXT("Sparse volume texture cannot be interpreted."));
-		}
+	AddEstimatedTextureSample();
+	FString SampleCode = FString::Printf(TEXT("SparseVolumeTextureSamplePageTable(%s, %s)"),
+		*GetParameterCode(SparseVolumeTextureIndex), *GetParameterCode(UVWIndex));
+	return AddCodeChunk(MCT_Float3, *SampleCode);
+}
 
-	}
-	else if (TextureType == MCT_TextureExternal)
+int32 FHLSLMaterialTranslator::SparseVolumeTextureSamplePhysicalTileData(int32 SparseVolumeTextureIndex, int32 VoxelCoordIndex, int32 PhysicalTileDataIdxIndex)
+{
+	if (SparseVolumeTextureIndex == INDEX_NONE || VoxelCoordIndex == INDEX_NONE || PhysicalTileDataIdxIndex == INDEX_NONE)
 	{
-		return Errorf(TEXT("FHLSLMaterialTranslator::SparseVolumeTextureSample MCT_TextureExternal ERROR."));
+		return INDEX_NONE;
 	}
-	else // MCT_Texture2D
+
+	EMaterialValueType TextureType = GetParameterType(SparseVolumeTextureIndex);
+	if ((TextureType & MCT_SparseVolumeTexture) == 0)
 	{
-		return Errorf(TEXT("FHLSLMaterialTranslator::SparseVolumeTextureSample MCT_Texture2D ERROR."));
+		return Errorf(TEXT("FHLSLMaterialTranslator::SparseVolumeTextureSamplePhysicalTileData expects MCT_SparseVolumeTexture but was passed %s ERROR."), DescribeType(TextureType));
 	}
 
-	return SamplingCodeIndex;
-}
+	EMaterialValueType VoxelCoordType = GetParameterType(VoxelCoordIndex);
+	if ((VoxelCoordType & MCT_Float3) == 0)
+	{
+		return Errorf(TEXT("FHLSLMaterialTranslator::SparseVolumeTextureSamplePhysicalTileData expects MCT_Float3 as VoxelCoord input but was passed %s ERROR."), DescribeType(VoxelCoordType));
+	}
 
-int32 FHLSLMaterialTranslator::SparseVolumeTextureUniform(int32 TextureIndex, int32 VectorIndex, UE::Shader::EValueType Type)
-{
-	return AddUniformExpression(new FMaterialUniformExpressionSparseVolumeTextureUniform(TextureIndex, VectorIndex), GetMaterialValueType(Type), TEXT(""));
-}
+	EMaterialValueType IndexType = GetParameterType(PhysicalTileDataIdxIndex);
+	if ((IndexType & MCT_Float1) == 0 && (IndexType & MCT_UInt1) == 0)
+	{
+		return Errorf(TEXT("FHLSLMaterialTranslator::SparseVolumeTextureSamplePhysicalTileData expects MCT_Float1 or MCT_UInt1 as PhysicalTileDataIndex input but was passed %s ERROR."), DescribeType(IndexType));
+	}
 
-int32 FHLSLMaterialTranslator::SparseVolumeTextureUniformParameter(FName ParameterName, int32 TextureIndex, int32 VectorIndex, UE::Shader::EValueType Type)
-{
-	FMaterialParameterInfo ParameterInfo = GetParameterAssociationInfo();
-	ParameterInfo.Name = ParameterName;
-	return AddUniformExpression(new FMaterialUniformExpressionSparseVolumeTextureUniform(ParameterInfo, TextureIndex, VectorIndex), GetMaterialValueType(Type), TEXT(""));
+	AddEstimatedTextureSample();
+	FString SampleCode = FString::Printf(TEXT("SparseVolumeTextureSamplePhysicalTileData(%s, %s, %s)"),
+		*GetParameterCode(SparseVolumeTextureIndex), *GetParameterCode(VoxelCoordIndex), *GetParameterCode(PhysicalTileDataIdxIndex));
+	return AddCodeChunk(MCT_Float4, *SampleCode);
 }
 
 
diff --git a/Engine/Source/Runtime/Engine/Private/Materials/HLSLMaterialTranslator.h b/Engine/Source/Runtime/Engine/Private/Materials/HLSLMaterialTranslator.h
--- a/Engine/Source/Runtime/Engine/Private/Materials/HLSLMaterialTranslator.h
+++ b/Engine/Source/Runtime/Engine/Private/Materials/HLSLMaterialTranslator.h
@@ -922,12 +922,12 @@ protected:
 	virtual int32 ExternalTextureCoordinateOffset(int32 TextureReferenceIndex, TOptional<FName> ParameterName) override;
 	virtual int32 ExternalTextureCoordinateOffset(const FGuid& ExternalTextureGuid) override;
 
-	virtual int32 SparseVolumeTexture(USparseVolumeTexture* Texture, int32 SubTextureId, int32& TextureReferenceIndex, EMaterialSamplerType SamplerType) override;
-	virtual int32 SparseVolumeTextureParameter(FName ParameterName, USparseVolumeTexture* InDefaultTexture, int32 SubTextureId, int32& TextureReferenceIndex, EMaterialSamplerType SamplerType) override;
-	virtual int32 SparseVolumeTextureSample(int32 TextureIndex, int32 CoordinateIndex) override;
+	virtual int32 SparseVolumeTexture(USparseVolumeTexture* Texture, int32& TextureReferenceIndex, EMaterialSamplerType SamplerType) override;
+	virtual int32 SparseVolumeTextureParameter(FName ParameterName, USparseVolumeTexture* InDefaultTexture, int32& TextureReferenceIndex, EMaterialSamplerType SamplerType) override;
 	virtual int32 SparseVolumeTextureUniform(int32 TextureIndex, int32 VectorIndex, UE::Shader::EValueType Type) override;
 	virtual int32 SparseVolumeTextureUniformParameter(FName ParameterName, int32 TextureIndex, int32 VectorIndex, UE::Shader::EValueType Type) override;
-	virtual int32 SparseVolumeTextureGetVoxelCoord(int32 PackedPhysicalTileCoord, int32 TileSize, int32 CoordPageTable, int32 CoordVolume) override;
+	virtual int32 SparseVolumeTextureSamplePageTable(int32 SparseVolumeTextureIndex, int32 UVWIndex) override;
+	virtual int32 SparseVolumeTextureSamplePhysicalTileData(int32 SparseVolumeTextureIndex, int32 VoxelCoordIndex, int32 PhysicalTileDataIdxIndex) override;
 
 	virtual UObject* GetReferencedTexture(int32 Index);
 
diff --git a/Engine/Source/Runtime/Engine/Private/Materials/MaterialExpressions.cpp b/Engine/Source/Runtime/Engine/Private/Materials/MaterialExpressions.cpp
--- a/Engine/Source/Runtime/Engine/Private/Materials/MaterialExpressions.cpp
+++ b/Engine/Source/Runtime/Engine/Private/Materials/MaterialExpressions.cpp
@@ -269,6 +269,8 @@
 #include "Materials/MaterialExpressionGetLocal.h"
 #include "Materials/MaterialExpressionBinaryOp.h"
 #include "Materials/MaterialExpressionGenericConstant.h"
+#include "Materials/MaterialExpressionSparseVolumeTextureBase.h"
+#include "Materials/MaterialExpressionSparseVolumeTextureObject.h"
 #include "Materials/MaterialExpressionSparseVolumeTextureSample.h"
 #include "EditorSupportDelegates.h"
 #if WITH_EDITOR
@@ -4019,7 +4021,7 @@ void UMaterialExpressionTextureProperty::GetCaption(TArray<FString>& OutCaptions
 uint32 UMaterialExpressionTextureProperty::GetInputType(int32 InputIndex)
 {
 	// TextureObject
-	IF_INPUT_RETURN(MCT_Texture);
+	IF_INPUT_RETURN(MCT_Texture | MCT_SparseVolumeTexture);
 	return MCT_Unknown;
 }
 #undef IF_INPUT_RETURN
@@ -25830,6 +25832,145 @@ void UMaterialExpressionSkyLightEnvMapSample::GetCaption(TArray<FString>& OutCap
 }
 #endif // WITH_EDITOR
 
+///////////////////////////////////////////////////////////////////////////////
+// UMaterialExpressionSparseVolumeTextureBase
+///////////////////////////////////////////////////////////////////////////////
+
+UMaterialExpressionSparseVolumeTextureBase::UMaterialExpressionSparseVolumeTextureBase(const FObjectInitializer& ObjectInitializer)
+	: Super(ObjectInitializer)
+{
+
+}
+
+UObject* UMaterialExpressionSparseVolumeTextureBase::GetReferencedTexture() const
+{
+	return SparseVolumeTexture;
+}
+
+///////////////////////////////////////////////////////////////////////////////
+// UMaterialExpressionSparseVolumeTextureObject
+///////////////////////////////////////////////////////////////////////////////
+
+UMaterialExpressionSparseVolumeTextureObject::UMaterialExpressionSparseVolumeTextureObject(const FObjectInitializer& ObjectInitializer)
+	: Super(ObjectInitializer)
+{
+	// Structure to hold one-time initialization
+	struct FConstructorStatics
+	{
+		FText NAME_Texture;
+		FText NAME_Functions;
+		FConstructorStatics()
+			: NAME_Texture(LOCTEXT("Texture", "Texture"))
+			, NAME_Functions(LOCTEXT("Functions", "Functions"))
+		{
+		}
+	};
+	static FConstructorStatics ConstructorStatics;
+
+#if WITH_EDITORONLY_DATA
+	MenuCategories.Add(ConstructorStatics.NAME_Texture);
+	MenuCategories.Add(ConstructorStatics.NAME_Functions);
+
+	Outputs.Empty();
+	Outputs.Add(FExpressionOutput(TEXT("")));
+
+	bCollapsed = true;
+	bHidePreviewWindow = true;
+#endif // WITH_EDITORONLY_DATA
+}
+
+#if WITH_EDITOR
+
+void UMaterialExpressionSparseVolumeTextureObject::PostEditChangeProperty(FPropertyChangedEvent& PropertyChangedEvent)
+{
+	// Update what needs to be when the texture is changed
+	if (PropertyChangedEvent.Property && PropertyChangedEvent.Property->GetName() == TEXT("SparseVirtualTexture"))
+	{
+		if (SparseVolumeTexture != nullptr)
+		{
+			FEditorSupportDelegates::ForcePropertyWindowRebuild.Broadcast(this);
+		}
+	}
+
+	Super::PostEditChangeProperty(PropertyChangedEvent);
+}
+
+int32 UMaterialExpressionSparseVolumeTextureObject::Compile(class FMaterialCompiler* Compiler, int32 OutputIndex)
+{
+	if (!SparseVolumeTexture)
+	{
+		return CompilerError(Compiler, TEXT("Requires valid SparseVolumeTexture"));
+	}
+
+	int32 TextureReferenceIndex = INDEX_NONE;
+	return Compiler->SparseVolumeTexture(SparseVolumeTexture, TextureReferenceIndex, SAMPLERTYPE_LinearColor);
+}
+
+void UMaterialExpressionSparseVolumeTextureObject::GetCaption(TArray<FString>& OutCaptions) const
+{
+	OutCaptions.Add(TEXT("SparseVolumeTexture Object"));
+}
+
+#endif // WITH_EDITOR
+
+///////////////////////////////////////////////////////////////////////////////
+// UMaterialExpressionSparseVolumeTextureTextureObjectParameter
+///////////////////////////////////////////////////////////////////////////////
+
+UMaterialExpressionSparseVolumeTextureTextureObjectParameter::UMaterialExpressionSparseVolumeTextureTextureObjectParameter(const FObjectInitializer& ObjectInitializer)
+	: Super(ObjectInitializer)
+{
+	// Structure to hold one-time initialization
+	struct FConstructorStatics
+	{
+		FText NAME_Texture;
+		FText NAME_Parameters;
+		FConstructorStatics()
+			: NAME_Texture(LOCTEXT("Texture", "Texture"))
+			, NAME_Parameters(LOCTEXT("Parameters", "Parameters"))
+		{
+		}
+	};
+	static FConstructorStatics ConstructorStatics;
+
+#if WITH_EDITORONLY_DATA
+	MenuCategories.Empty();
+	MenuCategories.Add(ConstructorStatics.NAME_Texture);
+	MenuCategories.Add(ConstructorStatics.NAME_Parameters);
+
+	Outputs.Empty();
+	Outputs.Add(FExpressionOutput(TEXT("")));
+#endif // WITH_EDITORONLY_DATA
+}
+
+#if WITH_EDITOR
+
+void UMaterialExpressionSparseVolumeTextureTextureObjectParameter::GetCaption(TArray<FString>& OutCaptions) const
+{
+	OutCaptions.Add(TEXT("Param Tex Object"));
+	OutCaptions.Add(FString::Printf(TEXT("'%s'"), *ParameterName.ToString()));
+}
+
+int32 UMaterialExpressionSparseVolumeTextureTextureObjectParameter::Compile(class FMaterialCompiler* Compiler, int32 OutputIndex)
+{
+	if (!SparseVolumeTexture)
+	{
+		return INDEX_NONE;
+	}
+
+	int32 TextureReferenceIndex = INDEX_NONE;
+	return Compiler->SparseVolumeTextureParameter(ParameterName, SparseVolumeTexture, TextureReferenceIndex, SAMPLERTYPE_LinearColor);
+}
+
+const TArray<FExpressionInput*> UMaterialExpressionSparseVolumeTextureTextureObjectParameter::GetInputs()
+{
+	// Hide the texture coordinate input
+	return TArray<FExpressionInput*>();
+}
+
+#endif // WITH_EDITOR
+
+
 ///////////////////////////////////////////////////////////////////////////////
 // UMaterialExpressionSparseVolumeTextureSample
 ///////////////////////////////////////////////////////////////////////////////
@@ -25840,16 +25981,16 @@ UMaterialExpressionSparseVolumeTextureSample::UMaterialExpressionSparseVolumeTex
 	// Structure to hold one-time initialization
 	struct FConstructorStatics
 	{
-		FText NAME_VirtualTexture;
+		FText NAME_Texture;
 		FConstructorStatics()
-			: NAME_VirtualTexture(LOCTEXT("Texture", "Texture"))
+			: NAME_Texture(LOCTEXT("Texture", "Texture"))
 		{
 		}
 	};
 	static FConstructorStatics ConstructorStatics;
 
 #if WITH_EDITORONLY_DATA
-	MenuCategories.Add(ConstructorStatics.NAME_VirtualTexture);
+	MenuCategories.Add(ConstructorStatics.NAME_Texture);
 #endif
 
 #if WITH_EDITOR
@@ -25861,11 +26002,6 @@ UMaterialExpressionSparseVolumeTextureSample::UMaterialExpressionSparseVolumeTex
 #endif
 }
 
-UObject* UMaterialExpressionSparseVolumeTextureSample::GetReferencedTexture() const
-{
-	return SparseVolumeTexture;
-}
-
 #if WITH_EDITOR
 
 void UMaterialExpressionSparseVolumeTextureSample::PostLoad()
@@ -25907,6 +26043,8 @@ uint32 UMaterialExpressionSparseVolumeTextureSample::GetInputType(int32 InputInd
 	{
 	case 0:
 		return MCT_Float3;
+	case 1:
+		return MCT_SparseVolumeTexture;
 		break;
 	}
 
@@ -25916,67 +26054,41 @@ uint32 UMaterialExpressionSparseVolumeTextureSample::GetInputType(int32 InputInd
 
 int32 UMaterialExpressionSparseVolumeTextureSample::Compile(class FMaterialCompiler* Compiler, int32 OutputIndex)
 {
-	if (SparseVolumeTexture != nullptr)
+	UMaterialExpression* InputExpression = TextureObject.GetTracedInput().Expression;
+
+	if (SparseVolumeTexture || InputExpression)
 	{
 		const bool bIsParameter = HasAParameterName() && GetParameterName().IsValid() && !GetParameterName().IsNone();
 
-		int32 PageTableTextureReferenceIndex = INDEX_NONE;
-		int32 PageTableSubTextureIndex = 0;
-		int32 PageTableTextureParamCodeChunkIndex = INDEX_NONE;
-		if (bIsParameter)
-		{
-			PageTableTextureParamCodeChunkIndex = Compiler->SparseVolumeTextureParameter(GetParameterName(), SparseVolumeTexture, PageTableSubTextureIndex, PageTableTextureReferenceIndex, SAMPLERTYPE_LinearColor);
-		}
-		else
-		{
-			PageTableTextureParamCodeChunkIndex = Compiler->SparseVolumeTexture(SparseVolumeTexture, PageTableSubTextureIndex, PageTableTextureReferenceIndex, SAMPLERTYPE_LinearColor);
-		}
-
-		int32 PhysicalTileTextureReferenceIndex = INDEX_NONE;
-		int32 PhysicalTileSubTextureIndex = 1 + OutputIndex; // index 0 is used by the page table texture, so we start at 1
-		int32 PhysicalTileTextureParamCodeChunkIndex = INDEX_NONE;
-		if (bIsParameter)
+		int32 SparseVolumeTextureIndex = INDEX_NONE;
+		if (InputExpression)
 		{
-			PhysicalTileTextureParamCodeChunkIndex = Compiler->SparseVolumeTextureParameter(GetParameterName(), SparseVolumeTexture, PhysicalTileSubTextureIndex, PhysicalTileTextureReferenceIndex, SAMPLERTYPE_LinearColor);
+			SparseVolumeTextureIndex = TextureObject.Compile(Compiler);
 		}
 		else
 		{
-			PhysicalTileTextureParamCodeChunkIndex = Compiler->SparseVolumeTexture(SparseVolumeTexture, PhysicalTileSubTextureIndex, PhysicalTileTextureReferenceIndex, SAMPLERTYPE_LinearColor);
-		}
-
-		// Now register the uniform we need for sampling the sparse volume texture
-		int32 UniformCodeChunkIndex[ESparseVolumeTexture_Count];
-		for (int32 UniformIndex = 0; UniformIndex < ESparseVolumeTexture_Count; ++UniformIndex)
-		{
-			const UE::Shader::EValueType Type = USparseVolumeTexture::GetUniformParameterType(UniformIndex);
+			int32 TextureReferenceIndex = INDEX_NONE;
 			if (bIsParameter)
 			{
-				UniformCodeChunkIndex[UniformIndex] = Compiler->SparseVolumeTextureUniformParameter(GetParameterName(), PageTableTextureReferenceIndex, UniformIndex, Type);
+				SparseVolumeTextureIndex = Compiler->SparseVolumeTextureParameter(GetParameterName(), SparseVolumeTexture, TextureReferenceIndex, SAMPLERTYPE_LinearColor);
 			}
 			else
 			{
-				UniformCodeChunkIndex[UniformIndex] = Compiler->SparseVolumeTextureUniform(PageTableTextureReferenceIndex, UniformIndex, Type);
+				SparseVolumeTextureIndex = Compiler->SparseVolumeTexture(SparseVolumeTexture, TextureReferenceIndex, SAMPLERTYPE_LinearColor);
 			}
 		}
 
-		// SVT_TODO Send parameters and do all the following math in shared functions (one to sample the page table and one to sample the attribute from the page table sample output)
-
-		int32 UVW = CompileWithDefaultFloat3(Compiler, Coordinates, 0.0f, 0.0f, 0.0f);
-		int32 UVWScaled = Compiler->Mul(UVW, UniformCodeChunkIndex[ESparseVolumeTexture_UVScale]);
-		int32 UVWPageTable = Compiler->Add(UVWScaled, UniformCodeChunkIndex[ESparseVolumeTexture_UVBias]);
-		int32 CoordPageTable = Compiler->Mul(UVWPageTable, UniformCodeChunkIndex[ESparseVolumeTexture_PageTableSize]);
-
-		// Code chunks have the same ID when they map to the same string of code. PackedPhysicalTileCoord will automatically end up being the same code chunk used
-		// when calling Compile() for Output 0 (Attributes A) and Output 1 (Attributes B). This is perfect for this use case: 
-		// PackedPhysicalTileCoord will end up being evaluated only once for both outputs, so they both share a single page table lookup.
-		int32 PackedPhysicalTileCoord = Compiler->SparseVolumeTextureSample(PageTableTextureParamCodeChunkIndex, CoordPageTable);
-
-		int32 CoordVolume = Compiler->Mul(CoordPageTable, UniformCodeChunkIndex[ESparseVolumeTexture_TileSize]);
-		int32 CoordPageTableFloor = Compiler->Floor(CoordPageTable);
+		if (SparseVolumeTextureIndex == INDEX_NONE)
+		{
+			// Can't continue without a texture to sample
+			return INDEX_NONE;
+		}
 
-		int32 CoordPhysicalTileVoxel = Compiler->SparseVolumeTextureGetVoxelCoord(PackedPhysicalTileCoord, UniformCodeChunkIndex[ESparseVolumeTexture_TileSize], CoordPageTableFloor, CoordVolume);
+		int32 UVWIndex = CompileWithDefaultFloat3(Compiler, Coordinates, 0.0f, 0.0f, 0.0f);
+		int32 VoxelCoordIndex = Compiler->SparseVolumeTextureSamplePageTable(SparseVolumeTextureIndex, UVWIndex);
+		int32 PhysicalTileDataIdxIndex = Compiler->Constant(OutputIndex);
 
-		return Compiler->SparseVolumeTextureSample(PhysicalTileTextureParamCodeChunkIndex, CoordPhysicalTileVoxel);
+		return Compiler->SparseVolumeTextureSamplePhysicalTileData(SparseVolumeTextureIndex, VoxelCoordIndex, PhysicalTileDataIdxIndex);
 	}
 
 	return INDEX_NONE;
@@ -25987,37 +26099,10 @@ void UMaterialExpressionSparseVolumeTextureSample::GetCaption(TArray<FString>& O
 	OutCaptions.Add(FString(TEXT("Sparse Volume Texture Sample")));
 }
 
-bool UMaterialExpressionSparseVolumeTextureSampleParameter::GetParameterValue(FMaterialParameterMetadata& OutMeta) const
-{
-	OutMeta.Value = SparseVolumeTexture;
-	OutMeta.Description = Desc;
-	OutMeta.ExpressionGuid = ExpressionGUID;
-	OutMeta.Group = Group;
-	OutMeta.SortPriority = SortPriority;
-	OutMeta.AssetPath = GetAssetPathName();
-	return true;
-}
-bool UMaterialExpressionSparseVolumeTextureSampleParameter::SetParameterValue(const FName& Name, const FMaterialParameterMetadata& Meta, EMaterialExpressionSetParameterValueFlags Flags)
-{
-	if (Meta.Value.Type == EMaterialParameterType::SparseVolumeTexture)
-	{
-		if (SetParameterValue(Name, Meta.Value.SparseVolumeTexture, Flags))
-		{
-			if (EnumHasAnyFlags(Flags, EMaterialExpressionSetParameterValueFlags::AssignGroupAndSortPriority))
-			{
-				Group = Meta.Group;
-				SortPriority = Meta.SortPriority;
-			}
-			return true;
-		}
-	}
-	return false;
-}
-
 #endif // WITH_EDITOR
 
 ///////////////////////////////////////////////////////////////////////////////
-// UMaterialExpressionSparseVolumeTextureSample
+// UMaterialExpressionSparseVolumeTextureSampleParameter
 ///////////////////////////////////////////////////////////////////////////////
 
 UMaterialExpressionSparseVolumeTextureSampleParameter::UMaterialExpressionSparseVolumeTextureSampleParameter(const FObjectInitializer& ObjectInitializer)
@@ -26077,6 +26162,35 @@ bool UMaterialExpressionSparseVolumeTextureSampleParameter::MatchesSearchQuery(c
 	}
 	return Super::MatchesSearchQuery(SearchQuery);
 }
+
+bool UMaterialExpressionSparseVolumeTextureSampleParameter::GetParameterValue(FMaterialParameterMetadata& OutMeta) const
+{
+	OutMeta.Value = SparseVolumeTexture;
+	OutMeta.Description = Desc;
+	OutMeta.ExpressionGuid = ExpressionGUID;
+	OutMeta.Group = Group;
+	OutMeta.SortPriority = SortPriority;
+	OutMeta.AssetPath = GetAssetPathName();
+	return true;
+}
+
+bool UMaterialExpressionSparseVolumeTextureSampleParameter::SetParameterValue(const FName& Name, const FMaterialParameterMetadata& Meta, EMaterialExpressionSetParameterValueFlags Flags)
+{
+	if (Meta.Value.Type == EMaterialParameterType::SparseVolumeTexture)
+	{
+		if (SetParameterValue(Name, Meta.Value.SparseVolumeTexture, Flags))
+		{
+			if (EnumHasAnyFlags(Flags, EMaterialExpressionSetParameterValueFlags::AssignGroupAndSortPriority))
+			{
+				Group = Meta.Group;
+				SortPriority = Meta.SortPriority;
+			}
+			return true;
+		}
+	}
+	return false;
+}
+
 #endif
 
 ///////////////////////////////////////////////////////////////////////////////
diff --git a/Engine/Source/Runtime/Engine/Private/Materials/MaterialInstance.cpp b/Engine/Source/Runtime/Engine/Private/Materials/MaterialInstance.cpp
--- a/Engine/Source/Runtime/Engine/Private/Materials/MaterialInstance.cpp
+++ b/Engine/Source/Runtime/Engine/Private/Materials/MaterialInstance.cpp
@@ -1258,8 +1258,17 @@ void UMaterialInstance::ValidateTextureOverrides(ERHIFeatureLevel::Type InFeatur
 		for (const FMaterialTextureParameterInfo& TextureInfo : CurrentResource->GetUniformTextureExpressions(ParameterType))
 		{
 			UTexture* Texture = nullptr;
-			TextureInfo.GetGameThreadTextureValue(this, *CurrentResource, Texture);
-			if (Texture)
+			USparseVolumeTexture* SVTexture = nullptr;
+			if (ParameterType != EMaterialTextureParameterType::SparseVolume)
+			{
+				TextureInfo.GetGameThreadTextureValue(this, *CurrentResource, Texture);
+			}
+			else
+			{
+				TextureInfo.GetGameThreadTextureValue(this, *CurrentResource, SVTexture);
+			}
+			
+			if (Texture || SVTexture)
 			{
 				const EMaterialValueType TextureType = Texture->GetMaterialType();
 				switch (ParameterType)
@@ -1308,6 +1317,12 @@ void UMaterialInstance::ValidateTextureOverrides(ERHIFeatureLevel::Type InFeatur
 						UE_LOG(LogMaterial, Error, TEXT("MaterialInstance \"%s\" parameter '%s' assigned texture \"%s\" requires virtual texture"), *MaterialName, *TextureInfo.GetParameterName().ToString(), *Texture->GetName());
 					}
 					break;
+				case EMaterialTextureParameterType::SparseVolume:
+					if (!(TextureType & MCT_SparseVolumeTexture))
+					{
+						UE_LOG(LogMaterial, Error, TEXT("MaterialInstance \"%s\" parameter '%s' assigned texture \"%s\" has invalid type, required sparse volume texture"), *MaterialName, *TextureInfo.GetParameterName().ToString(), *SVTexture->GetName());
+					}
+					break;
 				default:
 					checkNoEntry();
 					break;
diff --git a/Engine/Source/Runtime/Engine/Private/Materials/MaterialUniformExpressions.cpp b/Engine/Source/Runtime/Engine/Private/Materials/MaterialUniformExpressions.cpp
--- a/Engine/Source/Runtime/Engine/Private/Materials/MaterialUniformExpressions.cpp
+++ b/Engine/Source/Runtime/Engine/Private/Materials/MaterialUniformExpressions.cpp
@@ -19,10 +19,6 @@
 #include "SparseVolumeTexture/SparseVolumeTexture.h"
 #include "VT/RuntimeVirtualTexture.h"
 #include "TextureResource.h"
-static bool IsASparseVolumeTexture(const FMaterialTextureParameterInfo& Parameter)
-{
-	return Parameter.VirtualTextureLayerIndex != 255; // See FMaterialUniformExpressionTexture::GetTextureParameterInfo and constructor
-}
 
 void WriteMaterialUniformAccess(UE::Shader::EValueComponentType ComponentType, uint32 NumComponents, uint32 UniformOffset, FStringBuilderBase& OutResult)
 {
@@ -328,7 +324,7 @@ bool FUniformExpressionSet::operator==(const FUniformExpressionSet& ReferenceSet
 
 FString FUniformExpressionSet::GetSummaryString() const
 {
-	return FString::Printf(TEXT("(%u preshaders, %u 2d tex, %u cube tex, %u 2darray tex, %u cubearray tex, %u 3d tex, %u virtual tex, %u external tex, %u VT stacks, %u collections)"),
+	return FString::Printf(TEXT("(%u preshaders, %u 2d tex, %u cube tex, %u 2darray tex, %u cubearray tex, %u 3d tex, %u virtual tex, %u sparse volume tex, %u external tex, %u VT stacks, %u collections)"),
 		UniformPreshaders.Num(),
 		UniformTextureParameters[(uint32)EMaterialTextureParameterType::Standard2D].Num(),
 		UniformTextureParameters[(uint32)EMaterialTextureParameterType::Cube].Num(),
@@ -336,6 +332,7 @@ FString FUniformExpressionSet::GetSummaryString() const
 		UniformTextureParameters[(uint32)EMaterialTextureParameterType::ArrayCube].Num(),
 		UniformTextureParameters[(uint32)EMaterialTextureParameterType::Volume].Num(),
 		UniformTextureParameters[(uint32)EMaterialTextureParameterType::Virtual].Num(),
+		UniformTextureParameters[(uint32)EMaterialTextureParameterType::SparseVolume].Num(),
 		UniformExternalTextureParameters.Num(),
 		VTStacks.Num(),
 		ParameterCollections.Num()
@@ -373,6 +370,14 @@ FShaderParametersMetadata* FUniformExpressionSet::CreateBufferStruct()
 		NextMemberOffset += NumVirtualTextures * sizeof(FUintVector4);
 	}
 
+	const int32 NumSparseVolumeTextures = UniformTextureParameters[(uint32)EMaterialTextureParameterType::SparseVolume].Num();
+	if (NumSparseVolumeTextures > 0)
+	{
+		// 2x uint4 per SVT
+		new(Members) FShaderParametersMetadata::FMember(TEXT("SVTPackedUniform"), TEXT(""), __LINE__, NextMemberOffset, UBMT_UINT32, EShaderPrecisionModifier::Float, 1, 4, NumSparseVolumeTextures * 2, NULL);
+		NextMemberOffset += NumSparseVolumeTextures * sizeof(FUintVector4) * 2;
+	}
+
 	if (UniformPreshaderBufferSize > 0u)
 	{
 		new(Members) FShaderParametersMetadata::FMember(TEXT("PreshaderBuffer"),TEXT(""),__LINE__,NextMemberOffset,UBMT_FLOAT32,EShaderPrecisionModifier::Float,1,4, UniformPreshaderBufferSize,NULL);
@@ -398,6 +403,10 @@ FShaderParametersMetadata* FUniformExpressionSet::CreateBufferStruct()
 	static FString VirtualTexturePageTableIndirectionNames[128];
 	static FString VirtualTexturePhysicalNames[128];
 	static FString VirtualTexturePhysicalSamplerNames[128];
+	static FString SparseVolumeTexturePageTableNames[128];
+	static FString SparseVolumeTexturePhysicalANames[128];
+	static FString SparseVolumeTexturePhysicalBNames[128];
+	static FString SparseVolumeTexturePhysicalSamplerNames[128];
 	static bool bInitializedTextureNames = false;
 	if (!bInitializedTextureNames)
 	{
@@ -421,6 +430,10 @@ FShaderParametersMetadata* FUniformExpressionSet::CreateBufferStruct()
 			VirtualTexturePageTableIndirectionNames[i] = FString::Printf(TEXT("VirtualTexturePageTableIndirection_%d"), i);
 			VirtualTexturePhysicalNames[i] = FString::Printf(TEXT("VirtualTexturePhysical_%d"), i);
 			VirtualTexturePhysicalSamplerNames[i] = FString::Printf(TEXT("VirtualTexturePhysical_%dSampler"), i);
+			SparseVolumeTexturePageTableNames[i] = FString::Printf(TEXT("SparseVolumeTexturePageTable_%d"), i);
+			SparseVolumeTexturePhysicalANames[i] = FString::Printf(TEXT("SparseVolumeTexturePhysicalA_%d"), i);
+			SparseVolumeTexturePhysicalBNames[i] = FString::Printf(TEXT("SparseVolumeTexturePhysicalB_%d"), i);
+			SparseVolumeTexturePhysicalSamplerNames[i] = FString::Printf(TEXT("SparseVolumeTexturePhysical_%dSampler"), i);
 		}
 	}
 
@@ -469,44 +482,37 @@ FShaderParametersMetadata* FUniformExpressionSet::CreateBufferStruct()
 	for (int32 i = 0; i < UniformTextureParameters[(uint32)EMaterialTextureParameterType::Volume].Num(); ++i)
 	{
 		const FMaterialTextureParameterInfo& Parameter = UniformTextureParameters[(uint32)EMaterialTextureParameterType::Volume][i];
-		if (!IsASparseVolumeTexture(Parameter))
-		{
-			check((NextMemberOffset % SHADER_PARAMETER_POINTER_ALIGNMENT) == 0);
-			new(Members) FShaderParametersMetadata::FMember(*VolumeTextureNames[i], TEXT("Texture3D"), __LINE__, NextMemberOffset, UBMT_TEXTURE, EShaderPrecisionModifier::Float, 1, 1, 0, NULL);
-			NextMemberOffset += SHADER_PARAMETER_POINTER_ALIGNMENT;
-			new(Members) FShaderParametersMetadata::FMember(*VolumeTextureSamplerNames[i], TEXT("SamplerState"), __LINE__, NextMemberOffset, UBMT_SAMPLER, EShaderPrecisionModifier::Float, 1, 1, 0, NULL);
-			NextMemberOffset += SHADER_PARAMETER_POINTER_ALIGNMENT;
-		}
-		else
-		{
-			// Sparse volume texture bindings
-			if (Parameter.VirtualTextureLayerIndex == 0)
-			{
-				check((NextMemberOffset % SHADER_PARAMETER_POINTER_ALIGNMENT) == 0);
-				new(Members) FShaderParametersMetadata::FMember(*VolumeTextureNames[i], TEXT("Texture3D<uint>"), __LINE__, NextMemberOffset, UBMT_TEXTURE, EShaderPrecisionModifier::Float, 1, 1, 0, NULL);
-				NextMemberOffset += SHADER_PARAMETER_POINTER_ALIGNMENT;
-				new(Members) FShaderParametersMetadata::FMember(*VolumeTextureSamplerNames[i], TEXT("SamplerState"), __LINE__, NextMemberOffset, UBMT_SAMPLER, EShaderPrecisionModifier::Float, 1, 1, 0, NULL);
-				NextMemberOffset += SHADER_PARAMETER_POINTER_ALIGNMENT;
-			}
-			else if (Parameter.VirtualTextureLayerIndex == 1)
-			{
-				check((NextMemberOffset % SHADER_PARAMETER_POINTER_ALIGNMENT) == 0);
-				new(Members) FShaderParametersMetadata::FMember(*VolumeTextureNames[i], TEXT("Texture3D"), __LINE__, NextMemberOffset, UBMT_TEXTURE, EShaderPrecisionModifier::Float, 1, 1, 0, NULL);
-				NextMemberOffset += SHADER_PARAMETER_POINTER_ALIGNMENT;
-				new(Members) FShaderParametersMetadata::FMember(*VolumeTextureSamplerNames[i], TEXT("SamplerState"), __LINE__, NextMemberOffset, UBMT_SAMPLER, EShaderPrecisionModifier::Float, 1, 1, 0, NULL);
-				NextMemberOffset += SHADER_PARAMETER_POINTER_ALIGNMENT;
-			}
-			else if (Parameter.VirtualTextureLayerIndex == 2)
-			{
-				check((NextMemberOffset % SHADER_PARAMETER_POINTER_ALIGNMENT) == 0);
-				new(Members) FShaderParametersMetadata::FMember(*VolumeTextureNames[i], TEXT("Texture3D"), __LINE__, NextMemberOffset, UBMT_TEXTURE, EShaderPrecisionModifier::Float, 1, 1, 0, NULL);
-				NextMemberOffset += SHADER_PARAMETER_POINTER_ALIGNMENT;
-				new(Members) FShaderParametersMetadata::FMember(*VolumeTextureSamplerNames[i], TEXT("SamplerState"), __LINE__, NextMemberOffset, UBMT_SAMPLER, EShaderPrecisionModifier::Float, 1, 1, 0, NULL);
-				NextMemberOffset += SHADER_PARAMETER_POINTER_ALIGNMENT;
-			}
-		}
+		
+		check((NextMemberOffset % SHADER_PARAMETER_POINTER_ALIGNMENT) == 0);
+		new(Members) FShaderParametersMetadata::FMember(*VolumeTextureNames[i], TEXT("Texture3D"), __LINE__, NextMemberOffset, UBMT_TEXTURE, EShaderPrecisionModifier::Float, 1, 1, 0, NULL);
+		NextMemberOffset += SHADER_PARAMETER_POINTER_ALIGNMENT;
+		new(Members) FShaderParametersMetadata::FMember(*VolumeTextureSamplerNames[i], TEXT("SamplerState"), __LINE__, NextMemberOffset, UBMT_SAMPLER, EShaderPrecisionModifier::Float, 1, 1, 0, NULL);
+		NextMemberOffset += SHADER_PARAMETER_POINTER_ALIGNMENT;
+	}
+
+	for (int32 i = 0; i < UniformTextureParameters[(uint32)EMaterialTextureParameterType::SparseVolume].Num(); ++i)
+	{
+		const FMaterialTextureParameterInfo& Parameter = UniformTextureParameters[(uint32)EMaterialTextureParameterType::SparseVolume][i];
 
+		// Page Table
+		check((NextMemberOffset % SHADER_PARAMETER_POINTER_ALIGNMENT) == 0);
+		new(Members) FShaderParametersMetadata::FMember(*SparseVolumeTexturePageTableNames[i], TEXT("Texture3D<uint>"), __LINE__, NextMemberOffset, UBMT_TEXTURE, EShaderPrecisionModifier::Float, 1, 1, 0, NULL);
+		NextMemberOffset += SHADER_PARAMETER_POINTER_ALIGNMENT;
+		
+		// Physical A
+		check((NextMemberOffset% SHADER_PARAMETER_POINTER_ALIGNMENT) == 0);
+		new(Members) FShaderParametersMetadata::FMember(*SparseVolumeTexturePhysicalANames[i], TEXT("Texture3D"), __LINE__, NextMemberOffset, UBMT_TEXTURE, EShaderPrecisionModifier::Float, 1, 1, 0, NULL);
+		NextMemberOffset += SHADER_PARAMETER_POINTER_ALIGNMENT;
 
+		// Physical B
+		check((NextMemberOffset% SHADER_PARAMETER_POINTER_ALIGNMENT) == 0);
+		new(Members) FShaderParametersMetadata::FMember(*SparseVolumeTexturePhysicalBNames[i], TEXT("Texture3D"), __LINE__, NextMemberOffset, UBMT_TEXTURE, EShaderPrecisionModifier::Float, 1, 1, 0, NULL);
+		NextMemberOffset += SHADER_PARAMETER_POINTER_ALIGNMENT;
+
+		// Sampler
+		check((NextMemberOffset% SHADER_PARAMETER_POINTER_ALIGNMENT) == 0);
+		new(Members) FShaderParametersMetadata::FMember(*SparseVolumeTexturePhysicalSamplerNames[i], TEXT("SamplerState"), __LINE__, NextMemberOffset, UBMT_SAMPLER, EShaderPrecisionModifier::Float, 1, 1, 0, NULL);
+		NextMemberOffset += SHADER_PARAMETER_POINTER_ALIGNMENT;
 	}
 
 	for (int32 i = 0; i < UniformExternalTextureParameters.Num(); ++i)
@@ -717,10 +723,10 @@ void FUniformExpressionSet::GetTextureValue(int32 Index, const FMaterialRenderCo
 void FUniformExpressionSet::GetTextureValue(int32 Index, const FMaterialRenderContext& Context, const FMaterial& Material, const USparseVolumeTexture*& OutValue) const
 {
 	check(IsInParallelRenderingThread());
-	const int32 VolumeTexturesNum = GetNumTextures(EMaterialTextureParameterType::Volume);
+	const int32 VolumeTexturesNum = GetNumTextures(EMaterialTextureParameterType::SparseVolume);
 	if (ensure(Index < VolumeTexturesNum))
 	{
-		const FMaterialTextureParameterInfo& Parameter = GetTextureParameter(EMaterialTextureParameterType::Volume, Index);
+		const FMaterialTextureParameterInfo& Parameter = GetTextureParameter(EMaterialTextureParameterType::SparseVolume, Index);
 		Context.GetTextureParameterValue(Parameter.ParameterInfo, Parameter.TextureIndex, OutValue);
 	}
 	else
@@ -808,6 +814,27 @@ void FUniformExpressionSet::FillUniformBuffer(const FMaterialRenderContext& Mate
 			}
 		}
 
+		// Dump per SparseVolumeTexture uniform data
+		for (int32 ExpressionIndex = 0; ExpressionIndex < GetNumTextures(EMaterialTextureParameterType::SparseVolume); ++ExpressionIndex)
+		{
+			const FMaterialTextureParameterInfo& Parameter = GetTextureParameter(EMaterialTextureParameterType::SparseVolume, ExpressionIndex);
+
+			FUintVector4* SVTPackedUniform = (FUintVector4*)BufferCursor;
+			BufferCursor = SVTPackedUniform + 2;
+
+			const USparseVolumeTexture* SVTexture = nullptr;
+			GetTextureValue(ExpressionIndex, MaterialRenderContext, MaterialRenderContext.Material, SVTexture);
+			if (SVTexture != nullptr)
+			{
+				SVTexture->GetPackedUniforms(SVTPackedUniform[0], SVTPackedUniform[1]);
+			}
+			else
+			{
+				SVTPackedUniform[0] = FUintVector4();
+				SVTPackedUniform[1] = FUintVector4();
+			}
+		}
+
 		// Dump preshader results into buffer.
 		float* const PreshaderBuffer = (float*)BufferCursor;
 		FPreshaderStack PreshaderStack;
@@ -959,6 +986,7 @@ void FUniformExpressionSet::FillUniformBuffer(const FMaterialRenderContext& Mate
 				+ UniformTextureParameters[(uint32)EMaterialTextureParameterType::Array2D].Num() * 2
 				+ UniformTextureParameters[(uint32)EMaterialTextureParameterType::ArrayCube].Num() * 2
 				+ UniformTextureParameters[(uint32)EMaterialTextureParameterType::Volume].Num() * 2
+				+ UniformTextureParameters[(uint32)EMaterialTextureParameterType::SparseVolume].Num() * 4
 				+ UniformExternalTextureParameters.Num() * 2
 				+ UniformTextureParameters[(uint32)EMaterialTextureParameterType::Virtual].Num() * 2
 				+ NumPageTableTextures
@@ -1192,83 +1220,75 @@ void FUniformExpressionSet::FillUniformBuffer(const FMaterialRenderContext& Mate
 			BufferCursor = ((uint8*)BufferCursor) + (SHADER_PARAMETER_POINTER_ALIGNMENT * 2);
 			check(BufferCursor <= TempBuffer + TempBufferSize);
 
-			if (IsASparseVolumeTexture(Parameter))
-			{
-				const USparseVolumeTexture* SVTexture = nullptr;
-				GetTextureValue(ExpressionIndex, MaterialRenderContext, MaterialRenderContext.Material, SVTexture);
-				if (SVTexture != nullptr)
-				{
-					const int32 LayerIndex = Parameter.VirtualTextureLayerIndex;
-					// 0 is page table
-					// 1 is tile data A
-					// 2 is tile data B
+			const UTexture* Value = nullptr;
+			GetTextureValue(EMaterialTextureParameterType::Volume, ExpressionIndex, MaterialRenderContext, MaterialRenderContext.Material, Value);
 
-					const FSparseVolumeTextureSceneProxy* SVTextureProxy = SVTexture->GetSparseVolumeTextureSceneProxy();
-					if (SVTextureProxy)
-					{
-						if (LayerIndex == 0)
-						{
-							check(SVTextureProxy->GetPageTableTextureRHI());
-							*ResourceTableTexturePtr = SVTextureProxy->GetPageTableTextureRHI();
-						}
-						else if(LayerIndex == 1)
-						{
-							FTextureRHIRef TileDataATexture = SVTextureProxy->GetPhysicalTileDataATextureRHI();
-							*ResourceTableTexturePtr = TileDataATexture ? TileDataATexture : GBlackVolumeTexture->TextureRHI;
-						}
-						else if (LayerIndex == 2)
-						{
-							FTextureRHIRef TileDataBTexture = SVTextureProxy->GetPhysicalTileDataBTextureRHI();
-							*ResourceTableTexturePtr = TileDataBTexture ? TileDataBTexture : GBlackVolumeTexture->TextureRHI;
-						}
-						else
-						{
-							check(false);
-						}
+			if (Value && Value->GetResource() && Value->TextureReference.TextureReferenceRHI && (Value->GetMaterialType() & MCT_VolumeTexture) != 0u)
+			{
+				*ResourceTableTexturePtr = Value->TextureReference.TextureReferenceRHI;
+				const FSamplerStateRHIRef* SamplerSource = &Value->GetResource()->SamplerStateRHI;
 
-						*ResourceTableSamplerPtr = TStaticSamplerState<SF_Point, AM_Clamp, AM_Clamp, AM_Clamp, 0, 8>::GetRHI();
-					}
-					else
-					{
-						check(GBlackVolumeTexture->TextureRHI);
-						*ResourceTableTexturePtr = GBlackVolumeTexture->TextureRHI;
-						check(GBlackVolumeTexture->SamplerStateRHI);
-						*ResourceTableSamplerPtr = GBlackVolumeTexture->SamplerStateRHI;
-					}
+				const ESamplerSourceMode SourceMode = Parameter.SamplerSource;
+				if (SourceMode == SSM_Wrap_WorldGroupSettings)
+				{
+					SamplerSource = &Wrap_WorldGroupSettings->SamplerStateRHI;
+				}
+				else if (SourceMode == SSM_Clamp_WorldGroupSettings)
+				{
+					SamplerSource = &Clamp_WorldGroupSettings->SamplerStateRHI;
 				}
+
+				check(*SamplerSource);
+				*ResourceTableSamplerPtr = *SamplerSource;
 			}
 			else
 			{
-				const UTexture* Value = nullptr;
-				GetTextureValue(EMaterialTextureParameterType::Volume, ExpressionIndex, MaterialRenderContext, MaterialRenderContext.Material, Value);
+				check(GBlackVolumeTexture->TextureRHI);
+				*ResourceTableTexturePtr = GBlackVolumeTexture->TextureRHI;
+				check(GBlackVolumeTexture->SamplerStateRHI);
+				*ResourceTableSamplerPtr = GBlackVolumeTexture->SamplerStateRHI;
+			}
 
-				if (Value && Value->GetResource() && Value->TextureReference.TextureReferenceRHI && (Value->GetMaterialType() & MCT_VolumeTexture) != 0u)
-				{
-					*ResourceTableTexturePtr = Value->TextureReference.TextureReferenceRHI;
-					const FSamplerStateRHIRef* SamplerSource = &Value->GetResource()->SamplerStateRHI;
+		}
 
-					const ESamplerSourceMode SourceMode = Parameter.SamplerSource;
-					if (SourceMode == SSM_Wrap_WorldGroupSettings)
-					{
-						SamplerSource = &Wrap_WorldGroupSettings->SamplerStateRHI;
-					}
-					else if (SourceMode == SSM_Clamp_WorldGroupSettings)
-					{
-						SamplerSource = &Clamp_WorldGroupSettings->SamplerStateRHI;
-					}
+		// Cache sparse volume texture uniform expressions.
+		for (int32 ExpressionIndex = 0; ExpressionIndex < GetNumTextures(EMaterialTextureParameterType::SparseVolume); ExpressionIndex++)
+		{
+			const FMaterialTextureParameterInfo& Parameter = GetTextureParameter(EMaterialTextureParameterType::SparseVolume, ExpressionIndex);
 
-					check(*SamplerSource);
-					*ResourceTableSamplerPtr = *SamplerSource;
-				}
-				else
+			void** ResourceTableTexturePageTablePtr = (void**)((uint8*)BufferCursor + 0 * SHADER_PARAMETER_POINTER_ALIGNMENT);
+			void** ResourceTableTexturePhysicalAPtr = (void**)((uint8*)BufferCursor + 1 * SHADER_PARAMETER_POINTER_ALIGNMENT);
+			void** ResourceTableTexturePhysicalBPtr = (void**)((uint8*)BufferCursor + 2 * SHADER_PARAMETER_POINTER_ALIGNMENT);
+			void** ResourceTablePhysicalSamplerPtr = (void**)((uint8*)BufferCursor + 3 * SHADER_PARAMETER_POINTER_ALIGNMENT);
+			BufferCursor = ((uint8*)BufferCursor) + (SHADER_PARAMETER_POINTER_ALIGNMENT * 4);
+			check(BufferCursor <= TempBuffer + TempBufferSize);
+
+			check(GBlackVolumeTexture->TextureRHI);
+			*ResourceTableTexturePageTablePtr = GBlackVolumeTexture->TextureRHI;
+			*ResourceTableTexturePhysicalAPtr = GBlackVolumeTexture->TextureRHI;
+			*ResourceTableTexturePhysicalBPtr = GBlackVolumeTexture->TextureRHI;
+			check(GBlackVolumeTexture->SamplerStateRHI);
+			*ResourceTablePhysicalSamplerPtr = GBlackVolumeTexture->SamplerStateRHI;
+			
+			const USparseVolumeTexture* SVTexture = nullptr;
+			GetTextureValue(ExpressionIndex, MaterialRenderContext, MaterialRenderContext.Material, SVTexture);
+			if (SVTexture != nullptr)
+			{
+				const FSparseVolumeTextureSceneProxy* SVTextureProxy = SVTexture->GetSparseVolumeTextureSceneProxy();
+				if (SVTextureProxy)
 				{
-					check(GBlackVolumeTexture->TextureRHI);
-					*ResourceTableTexturePtr = GBlackVolumeTexture->TextureRHI;
-					check(GBlackVolumeTexture->SamplerStateRHI);
-					*ResourceTableSamplerPtr = GBlackVolumeTexture->SamplerStateRHI;
+					check(SVTextureProxy->GetPageTableTextureRHI());
+					*ResourceTableTexturePageTablePtr = SVTextureProxy->GetPageTableTextureRHI();
+
+					FTextureRHIRef TileDataATexture = SVTextureProxy->GetPhysicalTileDataATextureRHI();
+					*ResourceTableTexturePhysicalAPtr = TileDataATexture ? TileDataATexture : GBlackVolumeTexture->TextureRHI;
+
+					FTextureRHIRef TileDataBTexture = SVTextureProxy->GetPhysicalTileDataBTextureRHI();
+					*ResourceTableTexturePhysicalBPtr = TileDataBTexture ? TileDataBTexture : GBlackVolumeTexture->TextureRHI;
+
+					*ResourceTablePhysicalSamplerPtr = TStaticSamplerState<SF_Point, AM_Clamp, AM_Clamp, AM_Clamp, 0, 8>::GetRHI();
 				}
 			}
-
 		}
 
 		// Cache external texture uniform expressions.
@@ -1509,10 +1529,10 @@ FMaterialUniformExpressionTexture::FMaterialUniformExpressionTexture(int32 InTex
 }
 
 // This constructor is called for setting up SparseVolumeTextures
-FMaterialUniformExpressionTexture::FMaterialUniformExpressionTexture(int32 InTextureIndex, int16 InSubTextureIndex, EMaterialSamplerType InSamplerType)
+FMaterialUniformExpressionTexture::FMaterialUniformExpressionTexture(int32 InTextureIndex, EMaterialSamplerType InSamplerType)
 	: TextureIndex(InTextureIndex)
-	, TextureLayerIndex(InSubTextureIndex)
-	, PageTableLayerIndex(INDEX_NONE)		// Sparse Volume Textures are using PageTableLayerIndex as a sub texture index (for page table and attributes).
+	, TextureLayerIndex(INDEX_NONE)
+	, PageTableLayerIndex(INDEX_NONE)
 #if WITH_EDITORONLY_DATA
 	, SamplerType(InSamplerType)
 #endif
@@ -1607,6 +1627,14 @@ void FMaterialTextureParameterInfo::GetGameThreadTextureValue(const UMaterialInt
 	}
 }
 
+void FMaterialTextureParameterInfo::GetGameThreadTextureValue(const UMaterialInterface* MaterialInterface, const FMaterial& Material, USparseVolumeTexture*& OutValue) const
+{
+	if (ParameterInfo.Name.IsNone() || !MaterialInterface->GetSparseVolumeTextureParameterValue(ParameterInfo, OutValue))
+	{
+		OutValue = GetIndexedTexture<USparseVolumeTexture>(Material, TextureIndex);
+	}
+}
+
 bool FMaterialExternalTextureParameterInfo::GetExternalTexture(const FMaterialRenderContext& Context, FTextureRHIRef& OutTextureRHI, FSamplerStateRHIRef& OutSamplerStateRHI) const
 {
 	check(IsInParallelRenderingThread());
diff --git a/Engine/Source/Runtime/Engine/Private/Materials/MaterialUniformExpressions.h b/Engine/Source/Runtime/Engine/Private/Materials/MaterialUniformExpressions.h
--- a/Engine/Source/Runtime/Engine/Private/Materials/MaterialUniformExpressions.h
+++ b/Engine/Source/Runtime/Engine/Private/Materials/MaterialUniformExpressions.h
@@ -83,7 +83,7 @@ public:
 	FMaterialUniformExpressionTexture();
 	FMaterialUniformExpressionTexture(int32 InTextureIndex, EMaterialSamplerType InSamplerType, ESamplerSourceMode InSamplerSource, bool InVirtualTexture);
 	FMaterialUniformExpressionTexture(int32 InTextureIndex, int16 InTextureLayerIndex, int16 InPageTableLayerIndex, EMaterialSamplerType InSamplerType);
-	FMaterialUniformExpressionTexture(int32 InTextureIndex, int16 InSubTextureIndex, EMaterialSamplerType InSamplerType);
+	FMaterialUniformExpressionTexture(int32 InTextureIndex, EMaterialSamplerType InSamplerType);
 
 	//~ Begin FMaterialUniformExpression Interface.
 	virtual class FMaterialUniformExpressionTexture* GetTextureUniformExpression() { return this; }
@@ -409,8 +409,8 @@ public:
 		, ParameterInfo(InParameterInfo)
 	{}
 
-	FMaterialUniformExpressionTextureParameter(const FMaterialParameterInfo& InParameterInfo, int32 InTextureIndex, int16 InSubTextureIndex, EMaterialSamplerType InSamplerType)
-		: Super(InTextureIndex, InSubTextureIndex, InSamplerType)
+	FMaterialUniformExpressionTextureParameter(const FMaterialParameterInfo& InParameterInfo, int32 InTextureIndex, EMaterialSamplerType InSamplerType)
+		: Super(InTextureIndex, InSamplerType)
 		, ParameterInfo(InParameterInfo)
 	{}
 
diff --git a/Engine/Source/Runtime/Engine/Private/Shader/Preshader.cpp b/Engine/Source/Runtime/Engine/Private/Shader/Preshader.cpp
--- a/Engine/Source/Runtime/Engine/Private/Shader/Preshader.cpp
+++ b/Engine/Source/Runtime/Engine/Private/Shader/Preshader.cpp
@@ -717,19 +717,26 @@ static void EvaluateComponentSwizzle(FPreshaderStack& Stack, FPreshaderDataConte
 	}
 }
 
-static const UTexture* GetTextureParameter(const FMaterialRenderContext& Context, FPreshaderDataContext& RESTRICT Data)
+static void GetTextureParameter(const FMaterialRenderContext& Context, FPreshaderDataContext& RESTRICT Data, const UTexture*& OutTexture, const USparseVolumeTexture*& OutSparseVolumeTexture)
 {
 	const FHashedMaterialParameterInfo ParameterInfo = ReadPreshaderValue<FHashedMaterialParameterInfo>(Data);
 	const int32 TextureIndex = ReadPreshaderValue<int32>(Data);
 
-	const UTexture* Texture = nullptr;
-	Context.GetTextureParameterValue(ParameterInfo, TextureIndex, Texture);
-	return Texture;
+	OutTexture = nullptr;
+	OutSparseVolumeTexture = nullptr;
+
+	Context.GetTextureParameterValue(ParameterInfo, TextureIndex, OutTexture);
+	if (!OutTexture)
+	{
+		Context.GetTextureParameterValue(ParameterInfo, TextureIndex, OutSparseVolumeTexture);
+	}
 }
 
 static void EvaluateTextureSize(const FMaterialRenderContext& Context, FPreshaderStack& Stack, FPreshaderDataContext& RESTRICT Data)
 {
-	const UTexture* Texture = GetTextureParameter(Context, Data);
+	const UTexture* Texture = nullptr;
+	const USparseVolumeTexture* SparseVolumeTexture = nullptr;
+	GetTextureParameter(Context, Data, Texture, SparseVolumeTexture);
 	if (Texture && Texture->GetResource())
 	{
 		const uint32 SizeX = Texture->GetResource()->GetSizeX();
@@ -737,6 +744,13 @@ static void EvaluateTextureSize(const FMaterialRenderContext& Context, FPreshade
 		const uint32 SizeZ = Texture->GetResource()->GetSizeZ();
 		Stack.PushValue(FValue((float)SizeX, (float)SizeY, (float)SizeZ));
 	}
+	else if (SparseVolumeTexture)
+	{
+		const uint32 SizeX = SparseVolumeTexture->GetSizeX();
+		const uint32 SizeY = SparseVolumeTexture->GetSizeY();
+		const uint32 SizeZ = SparseVolumeTexture->GetSizeZ();
+		Stack.PushValue(FValue((float)SizeX, (float)SizeY, (float)SizeZ));
+	}
 	else
 	{
 		Stack.PushValue(FValue(0.0f, 0.0f, 0.0f));
@@ -745,7 +759,9 @@ static void EvaluateTextureSize(const FMaterialRenderContext& Context, FPreshade
 
 static void EvaluateTexelSize(const FMaterialRenderContext& Context, FPreshaderStack& Stack, FPreshaderDataContext& RESTRICT Data)
 {
-	const UTexture* Texture = GetTextureParameter(Context, Data);
+	const UTexture* Texture = nullptr;
+	const USparseVolumeTexture* SparseVolumeTexture = nullptr;
+	GetTextureParameter(Context, Data, Texture, SparseVolumeTexture);
 	if (Texture && Texture->GetResource())
 	{
 		const uint32 SizeX = Texture->GetResource()->GetSizeX();
@@ -753,6 +769,13 @@ static void EvaluateTexelSize(const FMaterialRenderContext& Context, FPreshaderS
 		const uint32 SizeZ = Texture->GetResource()->GetSizeZ();
 		Stack.PushValue(FValue(1.0f / (float)SizeX, 1.0f / (float)SizeY, (SizeZ > 0 ? 1.0f / (float)SizeZ : 0.0f)));
 	}
+	else if (SparseVolumeTexture)
+	{
+		const uint32 SizeX = SparseVolumeTexture->GetSizeX();
+		const uint32 SizeY = SparseVolumeTexture->GetSizeY();
+		const uint32 SizeZ = SparseVolumeTexture->GetSizeZ();
+		Stack.PushValue(FValue(1.0f / (float)SizeX, 1.0f / (float)SizeY, 1.0f / (float)SizeZ));
+	}
 	else
 	{
 		Stack.PushValue(FValue(0.0f, 0.0f, 0.0f));
diff --git a/Engine/Source/Runtime/Engine/Private/SparseVolumeTexture/SparseVolumeTexture.cpp b/Engine/Source/Runtime/Engine/Private/SparseVolumeTexture/SparseVolumeTexture.cpp
--- a/Engine/Source/Runtime/Engine/Private/SparseVolumeTexture/SparseVolumeTexture.cpp
+++ b/Engine/Source/Runtime/Engine/Private/SparseVolumeTexture/SparseVolumeTexture.cpp
@@ -366,6 +366,30 @@ FVector4 USparseVolumeTexture::GetUniformParameter(int32 Index) const
 	return FVector4(ForceInitToZero);
 }
 
+void USparseVolumeTexture::GetPackedUniforms(FUintVector4& OutPacked0, FUintVector4& OutPacked1) const
+{
+	const FVector4 TileSize = GetUniformParameter(ESparseVolumeTexture_TileSize);
+	const FVector4 PageTableSize = GetUniformParameter(ESparseVolumeTexture_PageTableSize);
+	const FVector4 UVScale = GetUniformParameter(ESparseVolumeTexture_UVScale);
+	const FVector4 UVBias = GetUniformParameter(ESparseVolumeTexture_UVBias);
+	const FUintVector4 PageTableSizeUIMinusOne = FUintVector4(PageTableSize.X - 1, PageTableSize.Y - 1, PageTableSize.Z - 1, 0);
+
+	auto AsUint = [](float X)
+	{
+		union { float F; uint32 U; } FU = { X };
+		return FU.U;
+	};
+
+	OutPacked0.X = AsUint((float)UVScale.X);
+	OutPacked0.Y = AsUint((float)UVScale.Y);
+	OutPacked0.Z = AsUint((float)UVScale.Z);
+	OutPacked0.W = (PageTableSizeUIMinusOne.X & 0x7FF) | ((PageTableSizeUIMinusOne.Y & 0x7FF) << 11) | ((PageTableSizeUIMinusOne.Z & 0x3FF) << 22);
+	OutPacked1.X = AsUint((float)UVBias.X);
+	OutPacked1.Y = AsUint((float)UVBias.Y);
+	OutPacked1.Z = AsUint((float)UVBias.Z);
+	OutPacked1.W = AsUint((float)TileSize.X);
+}
+
 void USparseVolumeTexture::GetFrameUVScaleBias(FVector* OutScale, FVector* OutBias) const
 {
 	*OutScale = FVector::One();
diff --git a/Engine/Source/Runtime/Engine/Public/MaterialCompiler.h b/Engine/Source/Runtime/Engine/Public/MaterialCompiler.h
--- a/Engine/Source/Runtime/Engine/Public/MaterialCompiler.h
+++ b/Engine/Source/Runtime/Engine/Public/MaterialCompiler.h
@@ -389,31 +389,22 @@ public:
 	// Making it clear for now and setting explicitly a USparseVolumeTexture as object.
 
 	/**
-	 * Register a sparse volume texture sub-texture to be sampled.
+	 * Register a sparse volume texture to be sampled.
 	 * @param Texture The sparse volume texture to sample.
-	 * @param SubTextureId The sub-texture we want to fetch from (page table, density, etc).
-	 * @param TextureReferenceIndex Output the index of the texture in the referenced textures for that shaders.
+	 * @param TextureReferenceIndex Output the index of the texture in the referenced textures of the material.
 	 * @param SamplerType The sampler type
 	 * @return The code chunk index of the texture.
 	 */
-	virtual int32 SparseVolumeTexture(USparseVolumeTexture* Texture, int32 SubTextureId, int32& TextureReferenceIndex, EMaterialSamplerType SamplerType) = 0;
+	virtual int32 SparseVolumeTexture(USparseVolumeTexture* Texture, int32& TextureReferenceIndex, EMaterialSamplerType SamplerType) = 0;
 	/**
-	 * Register a parameterized sparse volume texture sub-texture to be sampled.
+	 * Register a parameterized sparse volume texture to be sampled.
 	 * @param ParameterName The sparse volume texture parameter name.
 	 * @param InDefaultTexture The default static sparse volume texture to sample.
-	 * @param SubTextureId The sub-texture we want to fetch from (page table, density, etc).
-	 * @param TextureReferenceIndex Output the index of the texture in the referenced textures for that shaders.
+	 * @param TextureReferenceIndex Output the index of the texture in the referenced textures of the material.
 	 * @param SamplerType The sampler type
 	 * @return The code chunk index of the texture parameter.
 	 */
-	virtual int32 SparseVolumeTextureParameter(FName ParameterName, USparseVolumeTexture* InDefaultTexture, int32 SubTextureId, int32& TextureReferenceIndex, EMaterialSamplerType SamplerType) = 0;
-	/**
-	 * Sample a sparse volume sub-texture.
-	 * @param Texture The code chunk index of the texture.
-	 * @param Coordinate The coordinate at which to sample the texture.
-	 * @return The code chunk index of result of the texture sample.
-	 */
-	virtual int32 SparseVolumeTextureSample(int32 Texture, int32 Coordinate) = 0;
+	virtual int32 SparseVolumeTextureParameter(FName ParameterName, USparseVolumeTexture* InDefaultTexture, int32& TextureReferenceIndex, EMaterialSamplerType SamplerType) = 0;
 	/**
 	 * Register a uniform parameter required to be able to sample a sparse volume texture.
 	 * @param TextureIndex The TextureReferenceIndex.
@@ -432,13 +423,20 @@ public:
 	 */
 	virtual int32 SparseVolumeTextureUniformParameter(FName ParameterName, int32 TextureIndex, int32 VectorIndex, UE::Shader::EValueType Type) = 0;
 	/**
-	 * @param PackedPhysicalTileCoord: the packed coordinate fetched form the page table.
-	 * @param TileSize: the size of a tile.
-	 * @param CoordPageTable: The tile lower voxel coordinate in the physical tile volume.
-	 * @param CoordVolume: the float UV according the entire volume.
-	 * @return The code chunk index of the voxel coord within the tile volume texture.
+	 * Sample a sparse volume texture page table.
+	 * @param SparseVolumeTextureIndex The code chunk index of the texture.
+	 * @param UVWIndex The UVW coordinate at which to sample the texture.
+	 * @return The code chunk index of the result of the texture sample.
+	 */
+	virtual int32 SparseVolumeTextureSamplePageTable(int32 SparseVolumeTextureIndex, int32 UVWIndex) = 0;
+	/**
+	 * Sample a sparse volume texture physical tile data texture.
+	 * @param SparseVolumeTextureIndex The code chunk index of the texture.
+	 * @param VoxelCoordIndex The coordinate at which to sample the texture.
+	 * @param PhysicalTileDataIdxIndex The code chunk of the index (0 or 1) of the physical tile data texture to sample.
+	 * @return The code chunk index of the result of the texture sample.
 	 */
-	virtual int32 SparseVolumeTextureGetVoxelCoord(int32 PackedPhysicalTileCoord, int32 TileSize, int32 CoordPageTable, int32 CoordVolume) = 0;
+	virtual int32 SparseVolumeTextureSamplePhysicalTileData(int32 SparseVolumeTextureIndex, int32 VoxelCoordIndex, int32 PhysicalTileDataIdxIndex) = 0;
 
 	virtual UObject* GetReferencedTexture(int32 Index) { return nullptr; }
 
@@ -448,10 +446,10 @@ public:
 		return Texture(InTexture, TextureReferenceIndex, SamplerType, SamplerSource);
 	}
 
-	int32 SparseVolumeTexture(USparseVolumeTexture* Texture, int32 SubTextureId, EMaterialSamplerType SamplerType)
+	int32 SparseVolumeTexture(USparseVolumeTexture* Texture, EMaterialSamplerType SamplerType)
 	{
 		int32 TextureReferenceIndex = INDEX_NONE;
-		return SparseVolumeTexture(Texture, SubTextureId, TextureReferenceIndex, SamplerType);
+		return SparseVolumeTexture(Texture, TextureReferenceIndex, SamplerType);
 	}
 
 	int32 VirtualTexture(URuntimeVirtualTexture* InTexture, int32 TextureLayerIndex, int32 PageTableLayerIndex, EMaterialSamplerType SamplerType)
@@ -920,12 +918,12 @@ public:
 	virtual int32 ExternalTextureCoordinateOffset(int32 TextureReferenceIndex, TOptional<FName> ParameterName) override { return Compiler->ExternalTextureCoordinateOffset(TextureReferenceIndex, ParameterName); }
 	virtual int32 ExternalTextureCoordinateOffset(const FGuid& ExternalTextureGuid) override { return Compiler->ExternalTextureCoordinateOffset(ExternalTextureGuid); }
 
-	virtual int32 SparseVolumeTexture(USparseVolumeTexture* Texture, int32 SubTextureId, int32& TextureReferenceIndex, EMaterialSamplerType SamplerType) override											{ return Compiler->SparseVolumeTexture(Texture, SubTextureId, TextureReferenceIndex, SamplerType); }
-	virtual int32 SparseVolumeTextureParameter(FName ParameterName, USparseVolumeTexture* InDefaultTexture, int32 SubTextureId, int32& TextureReferenceIndex, EMaterialSamplerType SamplerType) override	{ return Compiler->SparseVolumeTextureParameter(ParameterName, InDefaultTexture, SubTextureId, TextureReferenceIndex, SamplerType); }
-	virtual int32 SparseVolumeTextureSample(int32 Texture, int32 Coordinate) override																														{ return Compiler->SparseVolumeTextureSample(Texture, Coordinate); }
+	virtual int32 SparseVolumeTexture(USparseVolumeTexture* Texture, int32& TextureReferenceIndex, EMaterialSamplerType SamplerType)																		{ return Compiler->SparseVolumeTexture(Texture, TextureReferenceIndex, SamplerType); }
+	virtual int32 SparseVolumeTextureParameter(FName ParameterName, USparseVolumeTexture* InDefaultTexture, int32& TextureReferenceIndex, EMaterialSamplerType SamplerType)									{ return Compiler->SparseVolumeTextureParameter(ParameterName, InDefaultTexture, TextureReferenceIndex, SamplerType); }
 	virtual int32 SparseVolumeTextureUniform(int32 TextureIndex, int32 VectorIndex, UE::Shader::EValueType Type) override																					{ return Compiler->SparseVolumeTextureUniform(TextureIndex, VectorIndex, Type); }
 	virtual int32 SparseVolumeTextureUniformParameter(FName ParameterName, int32 TextureIndex, int32 VectorIndex, UE::Shader::EValueType Type) override														{ return Compiler->SparseVolumeTextureUniformParameter(ParameterName, TextureIndex, VectorIndex, Type); }
-	virtual int32 SparseVolumeTextureGetVoxelCoord(int32 PackedPhysicalTileCoord, int32 TileSize, int32 CoordPageTable, int32 CoordVolume) override															{ return Compiler->SparseVolumeTextureGetVoxelCoord(PackedPhysicalTileCoord, TileSize, CoordPageTable, CoordVolume); }
+	virtual int32 SparseVolumeTextureSamplePageTable(int32 SparseVolumeTextureIndex, int32 UVWIndex) override																								{ return Compiler->SparseVolumeTextureSamplePageTable(SparseVolumeTextureIndex, UVWIndex); }
+	virtual int32 SparseVolumeTextureSamplePhysicalTileData(int32 SparseVolumeTextureIndex, int32 VoxelCoordIndex, int32 PhysicalTileDataIdxIndex) override													{ return Compiler->SparseVolumeTextureSamplePhysicalTileData(SparseVolumeTextureIndex, VoxelCoordIndex, PhysicalTileDataIdxIndex); }
 
 	virtual UObject* GetReferencedTexture(int32 Index) override { return Compiler->GetReferencedTexture(Index); }
 
diff --git a/Engine/Source/Runtime/Engine/Public/MaterialShared.h b/Engine/Source/Runtime/Engine/Public/MaterialShared.h
--- a/Engine/Source/Runtime/Engine/Public/MaterialShared.h
+++ b/Engine/Source/Runtime/Engine/Public/MaterialShared.h
@@ -469,6 +469,7 @@ enum class EMaterialTextureParameterType : uint32
 	ArrayCube,
 	Volume,
 	Virtual,
+	SparseVolume,
 
 	Count,
 };
@@ -490,6 +491,7 @@ public:
 	inline FName GetParameterName() const { return ScriptNameToName(ParameterInfo.Name); }
 
 	void GetGameThreadTextureValue(const UMaterialInterface* MaterialInterface, const FMaterial& Material, UTexture*& OutValue) const;
+	void GetGameThreadTextureValue(const UMaterialInterface* MaterialInterface, const FMaterial& Material, USparseVolumeTexture*& OutValue) const;
 
 	LAYOUT_FIELD(FHashedMaterialParameterInfo, ParameterInfo);
 	LAYOUT_FIELD_INITIALIZED(int32, TextureIndex, INDEX_NONE);
@@ -2026,6 +2028,7 @@ public:
 	inline TArrayView<const FMaterialTextureParameterInfo> GetUniform2DArrayTextureExpressions() const { return GetUniformTextureExpressions(EMaterialTextureParameterType::Array2D); }
 	inline TArrayView<const FMaterialTextureParameterInfo> GetUniformVolumeTextureExpressions() const { return GetUniformTextureExpressions(EMaterialTextureParameterType::Volume); }
 	inline TArrayView<const FMaterialTextureParameterInfo> GetUniformVirtualTextureExpressions() const { return GetUniformTextureExpressions(EMaterialTextureParameterType::Virtual); }
+	inline TArrayView<const FMaterialTextureParameterInfo> GetUniformSparseVolumeTextureExpressions() const { return GetUniformTextureExpressions(EMaterialTextureParameterType::SparseVolume); }
 
 #if WITH_EDITOR
 	const TArray<FString>& GetCompileErrors() const { return CompileErrors; }
diff --git a/Engine/Source/Runtime/Engine/Public/MaterialValueType.h b/Engine/Source/Runtime/Engine/Public/MaterialValueType.h
--- a/Engine/Source/Runtime/Engine/Public/MaterialValueType.h
+++ b/Engine/Source/Runtime/Engine/Public/MaterialValueType.h
@@ -36,34 +36,36 @@ enum EMaterialValueType
 	MCT_MaterialAttributes	  = 1 << 11,
 	MCT_TextureExternal       = 1 << 12,
 	MCT_TextureVirtual        = 1 << 13,
+	MCT_SparseVolumeTexture   = 1 << 14,
+	/** MCT_SparseVolumeTexture is intentionally not (yet) included here because it differs a lot from the other texture types and may not be supported/appropriate for all MCT_Texture use cases. */
 	MCT_Texture               = MCT_Texture2D | MCT_TextureCube | MCT_Texture2DArray | MCT_TextureCubeArray | MCT_VolumeTexture | MCT_TextureExternal | MCT_TextureVirtual,
 
 	/** Used internally when sampling from virtual textures */
-	MCT_VTPageTableResult     = 1 << 14,
+	MCT_VTPageTableResult     = 1 << 15,
 	
-	MCT_ShadingModel		  = 1 << 15,
+	MCT_ShadingModel		  = 1 << 16,
 
-	MCT_Strata				  = 1 << 16,
+	MCT_Strata				  = 1 << 17,
 
-	MCT_LWCScalar			  = 1 << 17,
-	MCT_LWCVector2			  = 1 << 18,
-	MCT_LWCVector3			  = 1 << 19,
-	MCT_LWCVector4			  = 1 << 20,
+	MCT_LWCScalar			  = 1 << 18,
+	MCT_LWCVector2			  = 1 << 19,
+	MCT_LWCVector3			  = 1 << 20,
+	MCT_LWCVector4			  = 1 << 21,
 	MCT_LWCType				  = MCT_LWCScalar | MCT_LWCVector2 | MCT_LWCVector3 | MCT_LWCVector4,
 
-	MCT_Execution             = 1 << 21,
+	MCT_Execution             = 1 << 22,
 
 	/** Used for code chunks that are statements with no value, rather than expressions */
-	MCT_VoidStatement         = 1 << 22,
+	MCT_VoidStatement         = 1 << 23,
 
 	/** Non-static bool, only used in new HLSL translator */
-	MCT_Bool				 = 1 << 23,
+	MCT_Bool				 = 1 << 24,
 
 	/** Unsigned int types */
-	MCT_UInt1				 = 1 << 24,
-	MCT_UInt2				 = 1 << 25,
-	MCT_UInt3				 = 1 << 26,
-	MCT_UInt4				 = 1 << 27,
+	MCT_UInt1				 = 1 << 25,
+	MCT_UInt2				 = 1 << 26,
+	MCT_UInt3				 = 1 << 27,
+	MCT_UInt4				 = 1 << 28,
 	MCT_UInt				 = MCT_UInt1 | MCT_UInt2 | MCT_UInt3 | MCT_UInt4,
 
 	MCT_Numeric = MCT_Float | MCT_LWCType | MCT_Bool,
diff --git a/Engine/Source/Runtime/Renderer/Private/SparseVolumeTexture/SparseVolumeTextureViewerComponent.cpp b/Engine/Source/Runtime/Renderer/Private/SparseVolumeTexture/SparseVolumeTextureViewerComponent.cpp
--- a/Engine/Source/Runtime/Renderer/Private/SparseVolumeTexture/SparseVolumeTextureViewerComponent.cpp
+++ b/Engine/Source/Runtime/Renderer/Private/SparseVolumeTexture/SparseVolumeTextureViewerComponent.cpp
@@ -214,9 +214,13 @@ void USparseVolumeTextureViewerComponent::TickComponent(float DeltaTime, enum EL
 		}
 
 		SparseVolumeTextureFrame = USparseVolumeTextureFrame::CreateFrame(SparseVolumeTexturePreview, FrameIndex);
-
-		MarkRenderStateDirty();
 	}
+	else
+	{
+		SparseVolumeTextureFrame = nullptr;
+	}
+
+	MarkRenderStateDirty();
 }
 
 /*=============================================================================
