Spline Mesh: Fix raytracing support after the refactor to dynamic raytracing mesh creation in 23966124

diff --git a/Engine/Source/Runtime/Engine/Private/SplineMeshSceneProxy.cpp b/Engine/Source/Runtime/Engine/Private/SplineMeshSceneProxy.cpp
--- a/Engine/Source/Runtime/Engine/Private/SplineMeshSceneProxy.cpp
+++ b/Engine/Source/Runtime/Engine/Private/SplineMeshSceneProxy.cpp
@@ -77,6 +77,7 @@ FSplineMeshSceneProxy::FSplineMeshSceneProxy(USplineMeshComponent* InComponent)
 
 #if RHI_RAYTRACING
 	bDynamicRayTracingGeometry = true;
+	bNeedsDynamicRayTracingGeometries = true;
 #endif
 
 	// make sure all the materials are okay to be rendered as a spline mesh
@@ -185,6 +186,8 @@ void FSplineMeshSceneProxy::GetDynamicRayTracingInstances(struct FRayTracingMate
 		return;
 	}
 
+	checkf(!DynamicRayTracingGeometries.IsEmpty(), TEXT("DynamicRayTracingGeometries has not been initialized correctly"));
+
 	ESceneDepthPriorityGroup PrimitiveDPG = GetStaticDepthPriorityGroup();
 	const int32 LODIndex = FMath::Max(GetLOD(Context.ReferenceView), (int32)GetCurrentFirstLODIdx_RenderThread());
 	const FStaticMeshLODResources& LODModel = RenderData->LODResources[LODIndex];
