- Fixed broken foliage instance selection after deleting an instance
- Also fixed missing initializations in FInstanceUpdateCmdBuffer ctor

diff --git a/Engine/Source/Runtime/Engine/Private/InstancedStaticMesh.cpp b/Engine/Source/Runtime/Engine/Private/InstancedStaticMesh.cpp
--- a/Engine/Source/Runtime/Engine/Private/InstancedStaticMesh.cpp
+++ b/Engine/Source/Runtime/Engine/Private/InstancedStaticMesh.cpp
@@ -243,8 +243,10 @@ EMouseCursor::Type HInstancedStaticMeshInstance::GetMouseCursor()
 }
 
 FInstanceUpdateCmdBuffer::FInstanceUpdateCmdBuffer()
-	: NumAdds(0),
+	: NumCustomDataFloats(0),
+	  NumAdds(0),
 	  NumUpdates(0),
+	  NumCustomFloatUpdates(0),
 	  NumRemoves(0)
 	, NumEdits(0)
 	, NumEditInstances(0)
@@ -401,16 +403,6 @@ void FInstanceUpdateCmdBuffer::SetCustomData(int32 RenderIndex, TConstArrayView<
 	Edit();
 }
 
-void FInstanceUpdateCmdBuffer::ResetInlineCommands()
-{
-	Cmds.Empty();
-	NumCustomDataFloats = 0;
-	NumAdds = 0;
-	NumUpdates = 0;
-	NumCustomFloatUpdates = 0;
-	NumRemoves = 0;
-}
-
 void FInstanceUpdateCmdBuffer::Edit()
 {
 	NumEdits++;
@@ -425,6 +417,7 @@ void FInstanceUpdateCmdBuffer::Reset()
 	NumCustomFloatUpdates = 0;
 	NumRemoves = 0;
 	NumEdits = 0;
+	NumEditInstances = 0;
 }
 
 FStaticMeshInstanceBuffer::FStaticMeshInstanceBuffer(ERHIFeatureLevel::Type InFeatureLevel, bool InRequireCPUAccess, bool bDeferGPUUploadIn)
@@ -469,7 +462,6 @@ void FStaticMeshInstanceBuffer::UpdateFromCommandBuffer_Concurrent(FInstanceUpda
 	// Compute render instances (same value that will computed on the render thread in UpdateFromCommandBuffer_RenderThread)
 	// Any query of number of render instances on game thread should use this instead of InstanceData->GetNumInstances();
 	CmdBuffer.NumEditInstances = NewCmdBuffer->NumAdds + InstanceData->GetNumInstances();
-	CmdBuffer.ResetInlineCommands();
 		
 	ENQUEUE_RENDER_COMMAND(InstanceBuffer_UpdateFromPreallocatedData)(
 		[InstanceBuffer, NewCmdBuffer](FRHICommandListImmediate& RHICmdList)
diff --git a/Engine/Source/Runtime/Engine/Public/PrimitiveInstanceUpdateCommand.h b/Engine/Source/Runtime/Engine/Public/PrimitiveInstanceUpdateCommand.h
--- a/Engine/Source/Runtime/Engine/Public/PrimitiveInstanceUpdateCommand.h
+++ b/Engine/Source/Runtime/Engine/Public/PrimitiveInstanceUpdateCommand.h
@@ -49,7 +49,6 @@ struct FInstanceUpdateCmdBuffer
 	void SetLightMapData(int32 RenderIndex, const FVector2D& LightmapUVBias);
 	void SetShadowMapData(int32 RenderIndex, const FVector2D& ShadowmapUVBias);
 	void SetCustomData(int32 RenderIndex, TConstArrayView<float> CustomDataFloats);
-	void ResetInlineCommands();
 	int32 NumInlineCommands() const { return Cmds.Num(); }
 
 	// Command that can't be in-lined and should cause full buffer rebuild
