FExec: Put back UGameEngine as it was until I can figure out how to properly decouple command execution and command routing. This is to fix an issue where the UEngine::Exec routing was being ignored, preventing deferred commands from working.

diff --git a/Engine/Source/Runtime/Engine/Classes/Engine/GameEngine.h b/Engine/Source/Runtime/Engine/Classes/Engine/GameEngine.h
--- a/Engine/Source/Runtime/Engine/Classes/Engine/GameEngine.h
+++ b/Engine/Source/Runtime/Engine/Classes/Engine/GameEngine.h
@@ -140,15 +140,11 @@ public:
 	virtual bool NetworkRemapPath(UNetConnection* Connection, FString& Str, bool bReading = true) override;
 	virtual bool ShouldDoAsyncEndOfFrameTasks() const override;
 
-	// FExec interface
 public:
-#if UE_ALLOW_EXEC_COMMANDS
+
+	// FExec interface
+
 	virtual bool Exec( UWorld* InWorld, const TCHAR* Cmd, FOutputDevice& Ar=*GLog ) override;
-#endif
-protected:
-	virtual bool Exec_Runtime( UWorld* InWorld, const TCHAR* Cmd, FOutputDevice& Ar ) override;
-	virtual bool Exec_Dev( UWorld* InWorld, const TCHAR* Cmd, FOutputDevice& Ar ) override;
-	virtual bool Exec_Editor( UWorld* InWorld, const TCHAR* Cmd, FOutputDevice& Ar ) override;
 
 public:
 
diff --git a/Engine/Source/Runtime/Engine/Private/GameEngine.cpp b/Engine/Source/Runtime/Engine/Private/GameEngine.cpp
--- a/Engine/Source/Runtime/Engine/Private/GameEngine.cpp
+++ b/Engine/Source/Runtime/Engine/Private/GameEngine.cpp
@@ -1016,7 +1016,7 @@ public:
 		return GEngineIni;
 	}
 
-	virtual bool Exec_Runtime(UWorld* Inworld, const TCHAR* Cmd, FOutputDevice& Ar) override
+	virtual bool Exec(UWorld* Inworld, const TCHAR* Cmd, FOutputDevice& Ar) override
 	{
 		if (FParse::Command(&Cmd, TEXT("exitembedded")))
 		{
@@ -1378,63 +1378,7 @@ bool UGameEngine::ShouldDoAsyncEndOfFrameTasks() const
 	Command line executor.
 -----------------------------------------------------------------------------*/
 
-#if UE_ALLOW_EXEC_COMMANDS
 bool UGameEngine::Exec( UWorld* InWorld, const TCHAR* Cmd, FOutputDevice& Ar )
-{
-	if (FExec::Exec(InWorld, Cmd, Ar))
-	{
-		return true;
-	}
-	else if (InWorld && InWorld->Exec(InWorld, Cmd, Ar))
-	{
-		return true;
-	}
-	else if (InWorld && InWorld->GetAuthGameMode() && InWorld->GetAuthGameMode()->ProcessConsoleExec(Cmd, Ar, nullptr))
-	{
-		return true;
-	}
-	else
-	{
-#if UE_BUILD_SHIPPING
-		// disallow set of actor properties if network game
-		if ((FParse::Command(&Cmd, TEXT("SET")) || FParse::Command(&Cmd, TEXT("SETNOPEC"))))
-		{
-			FWorldContext &Context = GetWorldContextFromWorldChecked(InWorld);
-			if (Context.PendingNetGame != NULL || InWorld->GetNetMode() != NM_Standalone)
-			{
-				return true;
-			}
-			// the effects of this cannot be easily reversed, so prevent the user from playing network games without restarting to avoid potential exploits
-			GDisallowNetworkTravel = true;
-		}
-#endif // UE_BUILD_SHIPPING
-#if UE_ALLOW_EXEC_EDITOR
-		if (UEngine::Exec_Editor(InWorld, Cmd, Ar))
-		{
-			return true;
-		}
-#endif
-#if UE_ALLOW_EXEC_DEV
-		if (UEngine::Exec_Dev(InWorld, Cmd, Ar))
-		{
-			return true;
-		}
-#endif
-		if (UEngine::Exec_Runtime(InWorld, Cmd, Ar))
-		{
-			return true;
-		}
-		else if (UPlatformInterfaceBase::StaticExec(Cmd, Ar))
-		{
-			return true;
-		}
-
-		return false;
-	}
-}
-#endif // UE_ALLOW_EXEC_COMMANDS
-
-bool UGameEngine::Exec_Runtime( UWorld* InWorld, const TCHAR* Cmd, FOutputDevice& Ar )
 {
 	if( FParse::Command( &Cmd,TEXT("REATTACHCOMPONENTS")) || FParse::Command( &Cmd,TEXT("REREGISTERCOMPONENTS")))
 	{
@@ -1528,26 +1472,14 @@ bool UGameEngine::Exec_Runtime( UWorld* InWorld, const TCHAR* Cmd, FOutputDevice
 
 		return true;
 	}
-
-	return false;
-}
-
-bool UGameEngine::Exec_Dev( UWorld* InWorld, const TCHAR* Cmd, FOutputDevice& Ar )
-{
 #if !UE_BUILD_SHIPPING
-	if( FParse::Command( &Cmd, TEXT("ApplyUserSettings") ) )
+	else if( FParse::Command( &Cmd, TEXT("ApplyUserSettings") ) )
 	{
-		return HandleApplyUserSettingsCommand(Cmd, Ar);
+		return HandleApplyUserSettingsCommand( Cmd, Ar );
 	}
 #endif // !UE_BUILD_SHIPPING
-
-	return false;
-}
-
-bool UGameEngine::Exec_Editor( UWorld* InWorld, const TCHAR* Cmd, FOutputDevice& Ar )
-{
 #if WITH_EDITOR
-	if( FParse::Command(&Cmd,TEXT("STARTMOVIECAPTURE")) && GIsEditor )
+	else if( FParse::Command(&Cmd,TEXT("STARTMOVIECAPTURE")) && GIsEditor )
 	{
 		IMovieSceneCaptureInterface* CaptureInterface = IMovieSceneCaptureModule::Get().GetFirstActiveMovieSceneCapture();
 		if (CaptureInterface)
@@ -1564,9 +1496,41 @@ bool UGameEngine::Exec_Editor( UWorld* InWorld, const TCHAR* Cmd, FOutputDevice&
 		}
 		return false;
 	}
-#endif // WITH_EDITOR
-
-	return false;
+#endif
+	else if( InWorld && InWorld->Exec( InWorld, Cmd, Ar ) )
+	{
+		return true;
+	}
+	else if( InWorld && InWorld->GetAuthGameMode() && InWorld->GetAuthGameMode()->ProcessConsoleExec(Cmd,Ar,NULL) )
+	{
+		return true;
+	}
+	else
+	{
+#if UE_BUILD_SHIPPING
+		// disallow set of actor properties if network game
+		if ((FParse::Command( &Cmd, TEXT("SET")) || FParse::Command( &Cmd, TEXT("SETNOPEC"))))
+		{
+			FWorldContext &Context = GetWorldContextFromWorldChecked(InWorld);
+			if( Context.PendingNetGame != NULL || InWorld->GetNetMode() != NM_Standalone)
+			{
+				return true;
+			}
+			// the effects of this cannot be easily reversed, so prevent the user from playing network games without restarting to avoid potential exploits
+			GDisallowNetworkTravel = true;
+		}
+#endif // UE_BUILD_SHIPPING
+		if (UEngine::Exec(InWorld, Cmd, Ar))
+		{
+			return true;
+		}
+		else if (UPlatformInterfaceBase::StaticExec(Cmd, Ar))
+		{
+			return true;
+		}
+	
+		return false;
+	}
 }
 
 bool UGameEngine::HandleExitCommand( const TCHAR* Cmd, FOutputDevice& Ar )
