World Partition: don't compute cells min/max z when zculling is not enabled.

diff --git a/Engine/Source/Runtime/Engine/Private/WorldPartition/WorldPartitionRuntimeSpatialHash.cpp b/Engine/Source/Runtime/Engine/Private/WorldPartition/WorldPartitionRuntimeSpatialHash.cpp
--- a/Engine/Source/Runtime/Engine/Private/WorldPartition/WorldPartitionRuntimeSpatialHash.cpp
+++ b/Engine/Source/Runtime/Engine/Private/WorldPartition/WorldPartitionRuntimeSpatialHash.cpp
@@ -167,8 +167,15 @@ void FSpatialHashStreamingGrid::GetCells(const FWorldPartitionStreamingQuerySour
 			{
 				ForEachRuntimeCell(Coords, [&](const UWorldPartitionRuntimeCell* Cell)
 				{
-					const FVector2D CellMinMaxZ(Cell->GetContentBounds().Min.Z, Cell->GetContentBounds().Max.Z);
-					if (!bEnableZCulling || TRange<double>(CellMinMaxZ.X, CellMinMaxZ.Y).Overlaps(TRange<double>(Shape.GetCenter().Z - Shape.GetRadius(), Shape.GetCenter().Z + Shape.GetRadius())))
+					bool bIncludeCell = true;
+
+					if (bEnableZCulling)
+					{
+						const FVector2D CellMinMaxZ(Cell->GetContentBounds().Min.Z, Cell->GetContentBounds().Max.Z);
+						bIncludeCell = TRange<double>(CellMinMaxZ.X, CellMinMaxZ.Y).Overlaps(TRange<double>(Shape.GetCenter().Z - Shape.GetRadius(), Shape.GetCenter().Z + Shape.GetRadius()));
+					}
+						
+					if (bIncludeCell)
 					{
 						if (ShouldAddCell(Cell, QuerySource))
 						{
@@ -242,8 +249,15 @@ void FSpatialHashStreamingGrid::GetCells(const TArray<FWorldPartitionStreamingSo
 				{
 					ForEachRuntimeCell(Coords, [&](const UWorldPartitionRuntimeCell* Cell)
 					{
-						const FVector2D CellMinMaxZ(Cell->GetContentBounds().Min.Z, Cell->GetContentBounds().Max.Z);
-						if (!bEnableZCulling || TRange<double>(CellMinMaxZ.X, CellMinMaxZ.Y).Overlaps(TRange<double>(Shape.GetCenter().Z - Shape.GetRadius(), Shape.GetCenter().Z + Shape.GetRadius())))
+						bool bIncludeCell = true;
+
+						if (bEnableZCulling)
+						{
+							const FVector2D CellMinMaxZ(Cell->GetContentBounds().Min.Z, Cell->GetContentBounds().Max.Z);
+							bIncludeCell = TRange<double>(CellMinMaxZ.X, CellMinMaxZ.Y).Overlaps(TRange<double>(Shape.GetCenter().Z - Shape.GetRadius(), Shape.GetCenter().Z + Shape.GetRadius()));
+						}
+						
+						if (bIncludeCell)
 						{
 							if (!Cell->HasDataLayers() || Cell->HasAnyDataLayerInEffectiveRuntimeState(EDataLayerRuntimeState::Activated))
 							{
