Lights: Remove bAffectDynamicIndirectLighting as it is not used anywhere

diff --git a/Engine/Source/Runtime/Engine/Classes/Components/LightComponent.h b/Engine/Source/Runtime/Engine/Classes/Components/LightComponent.h
--- a/Engine/Source/Runtime/Engine/Classes/Components/LightComponent.h
+++ b/Engine/Source/Runtime/Engine/Classes/Components/LightComponent.h
@@ -139,12 +139,6 @@ class ENGINE_API ULightComponent : public ULightComponentBase
 	UPROPERTY(EditAnywhere, BlueprintReadOnly, Category=Light, AdvancedDisplay)
 	uint32 bCastShadowsFromCinematicObjectsOnly:1;
 
-	/**
-	 * Whether the light should be injected into the Light Propagation Volume
-	 **/
-	UPROPERTY(EditAnywhere, BlueprintReadOnly, Category=Light, AdvancedDisplay, meta=(DisplayName = "Dynamic Indirect Lighting"))
-	uint32 bAffectDynamicIndirectLighting : 1;
-
 	/**
 	* Enables cached shadows for movable primitives for this light even if r.shadow.cachedshadowscastfrommovableprimitives is 0
 	*/
@@ -279,9 +273,6 @@ public:
 	UFUNCTION(BlueprintCallable, Category="Rendering|Components|Light")
 	void SetLightFunctionDisabledBrightness(float NewValue);
 
-	UFUNCTION(BlueprintCallable, Category="Rendering|Components|Light")
-	void SetAffectDynamicIndirectLighting(bool bNewValue);
-
 	UFUNCTION(BlueprintCallable, Category="Rendering|Components|Light")
 	void SetAffectTranslucentLighting(bool bNewValue);
 
diff --git a/Engine/Source/Runtime/Engine/Private/Components/DirectionalLightComponent.cpp b/Engine/Source/Runtime/Engine/Private/Components/DirectionalLightComponent.cpp
--- a/Engine/Source/Runtime/Engine/Private/Components/DirectionalLightComponent.cpp
+++ b/Engine/Source/Runtime/Engine/Private/Components/DirectionalLightComponent.cpp
@@ -1016,7 +1016,6 @@ UDirectionalLightComponent::UDirectionalLightComponent(const FObjectInitializer&
 	CastTranslucentShadows = true;
 	bUseInsetShadowsForMovableObjects = true;
 	bCastVolumetricShadow = true;
-	bAffectDynamicIndirectLighting = true;
 
 	ModulatedShadowColor = FColor(128, 128, 128);
 	ShadowAmount = 1.0f;
diff --git a/Engine/Source/Runtime/Engine/Private/Components/LightComponent.cpp b/Engine/Source/Runtime/Engine/Private/Components/LightComponent.cpp
--- a/Engine/Source/Runtime/Engine/Private/Components/LightComponent.cpp
+++ b/Engine/Source/Runtime/Engine/Private/Components/LightComponent.cpp
@@ -1015,16 +1015,6 @@ void ULightComponent::SetLightFunctionDisabledBrightness(float NewValue)
 	}
 }
 
-void ULightComponent::SetAffectDynamicIndirectLighting(bool bNewValue)
-{
-	if (AreDynamicDataChangesAllowed()
-		&& bAffectDynamicIndirectLighting != bNewValue)
-	{
-		bAffectDynamicIndirectLighting = bNewValue;
-		MarkRenderStateDirty();
-	}
-}
-
 void ULightComponent::SetAffectTranslucentLighting(bool bNewValue)
 {
 	if (AreDynamicDataChangesAllowed()
diff --git a/Engine/Source/Runtime/Engine/Private/LightSceneProxy.cpp b/Engine/Source/Runtime/Engine/Private/LightSceneProxy.cpp
--- a/Engine/Source/Runtime/Engine/Private/LightSceneProxy.cpp
+++ b/Engine/Source/Runtime/Engine/Private/LightSceneProxy.cpp
@@ -49,7 +49,6 @@ FLightSceneProxy::FLightSceneProxy(const ULightComponent* InLightComponent)
 	, bAffectGlobalIllumination(InLightComponent->bAffectGlobalIllumination)
 	, bAffectTranslucentLighting(InLightComponent->bAffectTranslucentLighting)
 	, bUsedAsAtmosphereSunLight(InLightComponent->IsUsedAsAtmosphereSunLight())
-	, bAffectDynamicIndirectLighting(InLightComponent->bAffectDynamicIndirectLighting)
 	, bUseRayTracedDistanceFieldShadows(InLightComponent->bUseRayTracedDistanceFieldShadows)
 	, bUseVirtualShadowMaps(false)	// See below
 	, bCastModulatedShadows(false)
diff --git a/Engine/Source/Runtime/Engine/Public/LightSceneProxy.h b/Engine/Source/Runtime/Engine/Public/LightSceneProxy.h
--- a/Engine/Source/Runtime/Engine/Public/LightSceneProxy.h
+++ b/Engine/Source/Runtime/Engine/Public/LightSceneProxy.h
@@ -225,7 +225,6 @@ public:
 	inline int32 GetShadowMapChannel() const { return ShadowMapChannel; }
 	inline int32 GetPreviewShadowMapChannel() const { return PreviewShadowMapChannel; }
 
-	inline bool AffectsDynamicIndirectLighting() const { return bAffectDynamicIndirectLighting; }
 	inline const class FStaticShadowDepthMap* GetStaticShadowDepthMap() const { return StaticShadowDepthMap; }
 
 	inline bool GetForceCachedShadowsForMovablePrimitives() const { return bForceCachedShadowsForMovablePrimitives; }
@@ -412,9 +411,6 @@ protected:
 	/** Whether to consider light as a sunlight for atmospheric scattering and exponential height fog. */
 	const uint8 bUsedAsAtmosphereSunLight : 1;
 
-	/** Does the light have dynamic GI? */
-	const uint8 bAffectDynamicIndirectLighting : 1;
-
 	/** Whether to use ray traced distance field area shadows. */
 	const uint8 bUseRayTracedDistanceFieldShadows : 1;
 
