Added stat scope in FObjectReplicator::ReceivedRPC for each net driver process call.

diff --git a/Engine/Source/Runtime/Engine/Private/DataReplication.cpp b/Engine/Source/Runtime/Engine/Private/DataReplication.cpp
--- a/Engine/Source/Runtime/Engine/Private/DataReplication.cpp
+++ b/Engine/Source/Runtime/Engine/Private/DataReplication.cpp
@@ -33,6 +33,7 @@ LLM_DEFINE_TAG(NetObjReplicator, NAME_None, TEXT("Networking"), GET_STATFNAME(ST
 
 DECLARE_CYCLE_STAT(TEXT("Custom Delta Property Rep Time"), STAT_NetReplicateCustomDeltaPropTime, STATGROUP_Game);
 DECLARE_CYCLE_STAT(TEXT("ReceiveRPC"), STAT_NetReceiveRPC, STATGROUP_Game);
+DECLARE_CYCLE_STAT(TEXT("ReceiveRPC_ProcessRemoteFunction"), STAT_NetReceiveRPC_ProcessRemoteFunction, STATGROUP_Game);
 
 static TAutoConsoleVariable<int32> CVarMaxRPCPerNetUpdate(
 	TEXT("net.MaxRPCPerNetUpdate"),
@@ -1380,6 +1381,7 @@ bool FObjectReplicator::ReceivedRPC(FNetBitReader& Reader, const FReplicationFla
 					{
 						if (Driver.NetDriver != nullptr && (Driver.NetDriver != Connection->Driver) && Driver.NetDriver->ShouldReplicateFunction(OwningActor, Function))
 						{
+							SCOPE_CYCLE_COUNTER(STAT_NetReceiveRPC_ProcessRemoteFunction);
 							Driver.NetDriver->ProcessRemoteFunction(OwningActor, Function, Parms, nullptr, nullptr, SubObject);
 						}
 					}
