New function UMaterialInstance::IsRedundant to return true if a MaterialInstance can definitely be replaced by its parent, false if it cannot or if that cannot be simply determined.
New function UMaterialInterface::GetNaniteOverrideRef to return nanite override material reference without loading it.

diff --git a/Engine/Source/Runtime/Engine/Classes/Materials/Material.h b/Engine/Source/Runtime/Engine/Classes/Materials/Material.h
--- a/Engine/Source/Runtime/Engine/Classes/Materials/Material.h
+++ b/Engine/Source/Runtime/Engine/Classes/Materials/Material.h
@@ -1072,6 +1072,7 @@ public:
 	ENGINE_API virtual UPhysicalMaterialMask* GetPhysicalMaterialMask() const override;
 	ENGINE_API virtual UPhysicalMaterial* GetPhysicalMaterialFromMap(int32 Index) const override;
 	ENGINE_API virtual UMaterialInterface* GetNaniteOverride(TMicRecursionGuard RecursionGuard = TMicRecursionGuard()) override;
+	ENGINE_API virtual TSoftObjectPtr<UMaterialInterface> GetNaniteOverrideRef(TMicRecursionGuard RecursionGuard = TMicRecursionGuard()) const override; 
 	ENGINE_API virtual void GetUsedTextures(TArray<UTexture*>& OutTextures, EMaterialQualityLevel::Type QualityLevel, bool bAllQualityLevels, ERHIFeatureLevel::Type FeatureLevel, bool bAllFeatureLevels) const override;
 	ENGINE_API virtual void GetUsedTexturesAndIndices(TArray<UTexture*>& OutTextures, TArray< TArray<int32> >& OutIndices, EMaterialQualityLevel::Type QualityLevel, ERHIFeatureLevel::Type FeatureLevel) const override;
 	ENGINE_API virtual void OverrideTexture(const UTexture* InTextureToOverride, UTexture* OverrideTexture, ERHIFeatureLevel::Type InFeatureLevel) override;
diff --git a/Engine/Source/Runtime/Engine/Classes/Materials/MaterialInstance.h b/Engine/Source/Runtime/Engine/Classes/Materials/MaterialInstance.h
--- a/Engine/Source/Runtime/Engine/Classes/Materials/MaterialInstance.h
+++ b/Engine/Source/Runtime/Engine/Classes/Materials/MaterialInstance.h
@@ -670,6 +670,10 @@ public:
 
 	ENGINE_API bool Equivalent(const UMaterialInstance* CompareTo) const;
 
+	// Returns true if using this material instance would be equivalent to using its parent. 
+	// i.e. no parameters/etc of any kind have been set.
+	ENGINE_API virtual bool IsRedundant() const;
+
 private:
 
 	/** Static parameter values that are overridden in this instance. */
@@ -743,6 +747,7 @@ public:
 	virtual ENGINE_API UPhysicalMaterialMask* GetPhysicalMaterialMask() const override;
 	virtual ENGINE_API UPhysicalMaterial* GetPhysicalMaterialFromMap(int32 Index) const override;
 	virtual ENGINE_API UMaterialInterface* GetNaniteOverride(TMicRecursionGuard RecursionGuard = TMicRecursionGuard()) override;
+	virtual ENGINE_API TSoftObjectPtr<UMaterialInterface> GetNaniteOverrideRef(TMicRecursionGuard RecursionGuard = TMicRecursionGuard()) const override; 
 	virtual ENGINE_API bool UpdateLightmassTextureTracking() override;
 	virtual ENGINE_API bool GetCastShadowAsMasked() const override;
 	virtual ENGINE_API float GetEmissiveBoost() const override;
@@ -908,10 +913,6 @@ public:
 	 */
 	ENGINE_API virtual void GetLightingGuidChain(bool bIncludeTextures, TArray<FGuid>& OutGuids) const override;
 
-#if WITH_EDITOR
-	ENGINE_API virtual uint32 ComputeAllStateCRC() const override;
-#endif // WITH_EDITOR
-
 	ENGINE_API virtual void DumpDebugInfo(FOutputDevice& OutputDevice) const override;
 	void SaveShaderStableKeys(const class ITargetPlatform* TP);
 	ENGINE_API virtual void SaveShaderStableKeysInner(const class ITargetPlatform* TP, const struct FStableShaderKeyAndValue& SaveKeyVal) override;
diff --git a/Engine/Source/Runtime/Engine/Classes/Materials/MaterialInterface.h b/Engine/Source/Runtime/Engine/Classes/Materials/MaterialInterface.h
--- a/Engine/Source/Runtime/Engine/Classes/Materials/MaterialInterface.h
+++ b/Engine/Source/Runtime/Engine/Classes/Materials/MaterialInterface.h
@@ -544,10 +544,16 @@ public:
 
 	/**
 	 * Get the associated nanite override material
-	 * @return - nanite override, will be this material if none was set
+	 * @return - nanite override, if none was set returns null as a signal to use this material instead
 	 */
 	virtual UMaterialInterface* GetNaniteOverride(TMicRecursionGuard RecursionGuard = TMicRecursionGuard()) PURE_VIRTUAL(UMaterialInterface::GetNaniteOverride, return nullptr;);
 
+	/**
+	 * Get the associated nanite override material as a soft reference without loading it.
+	 * @return - nanite override, if none was set returns null as a signal to use this material instead
+	 */
+	virtual TSoftObjectPtr<UMaterialInterface> GetNaniteOverrideRef(TMicRecursionGuard RecursionGuard = TMicRecursionGuard()) const PURE_VIRTUAL(UMaterialInterface::GetNaniteOverrideRef, return {};);
+
 	/**
 	 * Precache PSOs which can be used for this material for the given vertex factory type and material paramaters
 	 */
diff --git a/Engine/Source/Runtime/Engine/Private/Materials/Material.cpp b/Engine/Source/Runtime/Engine/Private/Materials/Material.cpp
--- a/Engine/Source/Runtime/Engine/Private/Materials/Material.cpp
+++ b/Engine/Source/Runtime/Engine/Private/Materials/Material.cpp
@@ -2135,6 +2135,11 @@ UMaterialInterface* UMaterial::GetNaniteOverride(TMicRecursionGuard RecursionGua
 	return NaniteOverrideMaterial.GetOverrideMaterial(this);
 }
 
+TSoftObjectPtr<UMaterialInterface> UMaterial::GetNaniteOverrideRef(TMicRecursionGuard RecursionGuard) const
+{
+	return NaniteOverrideMaterial.OverrideMaterialRef;
+}
+
 /** Helper functions for text output of properties... */
 #ifndef CASE_ENUM_TO_TEXT
 #define CASE_ENUM_TO_TEXT(txt) case txt: return TEXT(#txt);
diff --git a/Engine/Source/Runtime/Engine/Private/Materials/MaterialInstance.cpp b/Engine/Source/Runtime/Engine/Private/Materials/MaterialInstance.cpp
--- a/Engine/Source/Runtime/Engine/Private/Materials/MaterialInstance.cpp
+++ b/Engine/Source/Runtime/Engine/Private/Materials/MaterialInstance.cpp
@@ -1805,7 +1805,7 @@ UMaterialInterface* UMaterialInstance::GetNaniteOverride(TMicRecursionGuard Recu
 {
 	if (NaniteOverrideMaterial.bEnableOverride)
 	{
-		return NaniteOverrideMaterial.GetOverrideMaterial(this);
+		return NaniteOverrideMaterial.GetOverrideMaterial();
 	}
 	else if (Parent && !RecursionGuard.Contains(this))
 	{
@@ -1818,6 +1818,23 @@ UMaterialInterface* UMaterialInstance::GetNaniteOverride(TMicRecursionGuard Recu
 	}
 }
 
+TSoftObjectPtr<UMaterialInterface> UMaterialInstance::GetNaniteOverrideRef(TMicRecursionGuard RecursionGuard) const
+{
+	if (NaniteOverrideMaterial.bEnableOverride)
+	{
+		return NaniteOverrideMaterial.OverrideMaterialRef;
+	}
+	else if (Parent && !RecursionGuard.Contains(this))
+	{
+		RecursionGuard.Set(this);
+		return Parent->GetNaniteOverrideRef(RecursionGuard);
+	}
+	else
+	{
+		return nullptr;
+	}
+}
+
 #if WITH_EDITORONLY_DATA
 
 void UMaterialInstance::SetStaticSwitchParameterValueEditorOnly(const FMaterialParameterInfo& ParameterInfo, bool Value)
@@ -3852,7 +3869,7 @@ void UMaterialInstance::PostEditChangeProperty(FPropertyChangedEvent& PropertyCh
 
 	if (PropertyChangedEvent.MemberProperty != nullptr && PropertyChangedEvent.MemberProperty->GetFName() == GET_MEMBER_NAME_CHECKED(UMaterial, NaniteOverrideMaterial))
 	{
-		NaniteOverrideMaterial.PostEditChange(this);
+		NaniteOverrideMaterial.PostEditChange();
 
 		// Update primitives that might depend on the nanite override material.
 		FGlobalComponentRecreateRenderStateContext RecreateComponentsRenderState;
@@ -3898,7 +3915,7 @@ void UMaterialInstance::PostEditUndo()
 {
 	Super::PostEditUndo();
 
-	NaniteOverrideMaterial.PostEditChange(this);
+	NaniteOverrideMaterial.PostEditChange();
 }
 
 #endif // WITH_EDITOR
@@ -4452,19 +4469,6 @@ void UMaterialInstance::GetLightingGuidChain(bool bIncludeTextures, TArray<FGuid
 #endif
 }
 
-#if WITH_EDITOR
-uint32 UMaterialInstance::ComputeAllStateCRC() const
-{
-	uint32 CRC = Super::ComputeAllStateCRC();
-	if (Parent)
-	{
-		uint32 ParentCRC = Parent->ComputeAllStateCRC();
-		CRC = FCrc::TypeCrc32(ParentCRC, CRC);
-	}
-	return CRC;
-}
-#endif
-
 void UMaterialInstance::PreSave(const class ITargetPlatform* TargetPlatform)
 {
 	PRAGMA_DISABLE_DEPRECATION_WARNINGS;
@@ -4543,6 +4547,65 @@ bool UMaterialInstance::Equivalent(const UMaterialInstance* CompareTo) const
 	return true;
 }
 
+bool UMaterialInstance::IsRedundant() const
+{
+	if (NaniteOverrideMaterial.bEnableOverride)
+	{
+		// Check if we resolve to a different material to our parent 
+		TSoftObjectPtr<UMaterialInterface> MyOverride = GetNaniteOverrideRef();
+		TSoftObjectPtr<UMaterialInterface> ParentOverride = Parent->GetNaniteOverrideRef();
+		// Possible refinement: Could check if they are equivalent MIDs, or a redundant MID and its parent, but that would require loading them. 
+		if (MyOverride != ParentOverride)
+		{
+			return false;
+		}
+	}
+
+	if (HasStaticParameters())
+	{
+		return false;
+	}
+	if (GetPhysicalMaterial() != Parent->GetPhysicalMaterial())
+	{
+		return false;
+	}
+	for (int32 i = 0; i < EPhysicalMaterialMaskColor::MAX; ++i)
+	{
+		if (GetPhysicalMaterialFromMap(i) != Parent->GetPhysicalMaterialFromMap(i))
+		{
+			return false;
+		}
+	}
+	if (GetSubsurfaceProfile_Internal() != Parent->GetSubsurfaceProfile_Internal())
+	{
+		return false;
+	}
+	// Assume that if any properties are overridden they are different to their parent
+	if (
+	   BasePropertyOverrides.bOverride_OpacityMaskClipValue 
+	|| BasePropertyOverrides.bOverride_BlendMode
+	|| BasePropertyOverrides.bOverride_ShadingModel
+	|| BasePropertyOverrides.bOverride_DitheredLODTransition 
+	|| BasePropertyOverrides.bOverride_CastDynamicShadowAsMasked
+	|| BasePropertyOverrides.bOverride_TwoSided 
+	|| BasePropertyOverrides.bOutputTranslucentVelocity)
+	{
+		return false;
+	}
+	if (
+		!TextureParameterValues.IsEmpty()
+		|| !ScalarParameterValues.IsEmpty()
+		|| !VectorParameterValues.IsEmpty()
+		|| !DoubleVectorParameterValues.IsEmpty()
+		|| !RuntimeVirtualTextureParameterValues.IsEmpty()
+		|| !FontParameterValues.IsEmpty())
+	{
+		return false;
+	}
+
+	return true;
+}
+
 #if !UE_BUILD_SHIPPING
 
 static void FindRedundantMICS(const TArray<FString>& Args)
