When building WP HLODs, check if material instances are redundant and if so, use the parent material in the ISMC.

diff --git a/Engine/Source/Runtime/Engine/Private/ISMPartition/ISMComponentDescriptor.cpp b/Engine/Source/Runtime/Engine/Private/ISMPartition/ISMComponentDescriptor.cpp
--- a/Engine/Source/Runtime/Engine/Private/ISMPartition/ISMComponentDescriptor.cpp
+++ b/Engine/Source/Runtime/Engine/Private/ISMPartition/ISMComponentDescriptor.cpp
@@ -177,16 +177,24 @@ void FISMComponentDescriptor::InitComponent(UInstancedStaticMeshComponent* ISMCo
 	{
 		if (OverrideMaterial && !OverrideMaterial->IsAsset())
 		{
-			// As override materials are normally outered to their owner component, we need to duplicate them here to make sure we don't create
-			// references to actors in other levels (for packed level instances or HLOD actors).
-			OverrideMaterial = DuplicateObject<UMaterialInterface>(OverrideMaterial, ISMComponent);
-
-			// If the MID we just duplicated has a nanite override that's also not an asset, duplicate that too
-			UMaterialInstanceDynamic* OverrideMID = Cast<UMaterialInstanceDynamic>(OverrideMaterial);
-			UMaterialInterface* NaniteOverride = OverrideMID ? OverrideMID->GetNaniteOverride() : nullptr; 
-			if (NaniteOverride && !NaniteOverride->IsAsset())
+			// If the material is equivalent to its parent, just take a reference to its parent rather than making another redundant object 
+			if (UMaterialInstance* Instance = Cast<UMaterialInstance>(OverrideMaterial); Instance && Instance->IsRedundant())
 			{
-				OverrideMID->SetNaniteOverride(DuplicateObject<UMaterialInterface>(NaniteOverride, ISMComponent));
+				OverrideMaterial = Instance->Parent;
+			}
+			else
+			{
+				// As override materials are normally outered to their owner component, we need to duplicate them here to make sure we don't create
+				// references to actors in other levels (for packed level instances or HLOD actors).
+				OverrideMaterial = DuplicateObject<UMaterialInterface>(OverrideMaterial, ISMComponent);
+
+				// If the MID we just duplicated has a nanite override that's also not an asset, duplicate that too
+				UMaterialInstanceDynamic* OverrideMID = Cast<UMaterialInstanceDynamic>(OverrideMaterial);
+				UMaterialInterface* NaniteOverride = OverrideMID ? OverrideMID->GetNaniteOverride() : nullptr; 
+				if (NaniteOverride && !NaniteOverride->IsAsset())
+				{
+					OverrideMID->SetNaniteOverride(DuplicateObject<UMaterialInterface>(NaniteOverride, ISMComponent));
+				}
 			}
 		}
 
