Fixed compile error in MaterialExpressionSparseVolumeTextureObject.h in non-unity builds by moving function definition into the .cpp file.

#rnx

diff --git a/Engine/Source/Runtime/Engine/Classes/Materials/MaterialExpressionSparseVolumeTextureObject.h b/Engine/Source/Runtime/Engine/Classes/Materials/MaterialExpressionSparseVolumeTextureObject.h
--- a/Engine/Source/Runtime/Engine/Classes/Materials/MaterialExpressionSparseVolumeTextureObject.h
+++ b/Engine/Source/Runtime/Engine/Classes/Materials/MaterialExpressionSparseVolumeTextureObject.h
@@ -23,7 +23,7 @@ class UMaterialExpressionSparseVolumeTextureObject : public UMaterialExpressionS
 #if WITH_EDITOR
 	virtual int32 Compile(class FMaterialCompiler* Compiler, int32 OutputIndex) override;
 	virtual void GetCaption(TArray<FString>& OutCaptions) const override;
-	virtual uint32 GetOutputType(int32 OutputIndex) override { return MCT_SparseVolumeTexture; }
+	virtual uint32 GetOutputType(int32 OutputIndex) override;
 	//~ End UMaterialExpression Interface
 #endif // WITH_EDITOR
 };
@@ -37,7 +37,7 @@ class UMaterialExpressionSparseVolumeTextureTextureObjectParameter : public UMat
 #if WITH_EDITOR
 	virtual void GetCaption(TArray<FString>& OutCaptions) const override;
 	virtual int32 Compile(class FMaterialCompiler* Compiler, int32 OutputIndex) override;
-	virtual uint32 GetOutputType(int32 OutputIndex) override { return MCT_SparseVolumeTexture; }
+	virtual uint32 GetOutputType(int32 OutputIndex) override;
 	virtual const TArray<FExpressionInput*> GetInputs() override;
 #endif
 	//~ End UMaterialExpression Interface
diff --git a/Engine/Source/Runtime/Engine/Private/Materials/MaterialExpressions.cpp b/Engine/Source/Runtime/Engine/Private/Materials/MaterialExpressions.cpp
--- a/Engine/Source/Runtime/Engine/Private/Materials/MaterialExpressions.cpp
+++ b/Engine/Source/Runtime/Engine/Private/Materials/MaterialExpressions.cpp
@@ -25911,6 +25911,11 @@ void UMaterialExpressionSparseVolumeTextureObject::GetCaption(TArray<FString>& O
 	OutCaptions.Add(TEXT("SparseVolumeTexture Object"));
 }
 
+uint32 UMaterialExpressionSparseVolumeTextureObject::GetOutputType(int32 OutputIndex)
+{
+	return MCT_SparseVolumeTexture;
+}
+
 #endif // WITH_EDITOR
 
 ///////////////////////////////////////////////////////////////////////////////
@@ -25962,6 +25967,11 @@ int32 UMaterialExpressionSparseVolumeTextureTextureObjectParameter::Compile(clas
 	return Compiler->SparseVolumeTextureParameter(ParameterName, SparseVolumeTexture, TextureReferenceIndex, SAMPLERTYPE_LinearColor);
 }
 
+uint32 UMaterialExpressionSparseVolumeTextureTextureObjectParameter::GetOutputType(int32 OutputIndex)
+{
+	return MCT_SparseVolumeTexture;
+}
+
 const TArray<FExpressionInput*> UMaterialExpressionSparseVolumeTextureTextureObjectParameter::GetInputs()
 {
 	// Hide the texture coordinate input
