[shaders] move generation of OutputHash.txt debug file into common compile function (instead of being part of the worker output readback code); as a result it is executed both for the in-process and worker-based shader compilation paths.

diff --git a/Engine/Source/Runtime/Engine/Private/ShaderCompiler/ShaderCompiler.cpp b/Engine/Source/Runtime/Engine/Private/ShaderCompiler/ShaderCompiler.cpp
--- a/Engine/Source/Runtime/Engine/Private/ShaderCompiler/ShaderCompiler.cpp
+++ b/Engine/Source/Runtime/Engine/Private/ShaderCompiler/ShaderCompiler.cpp
@@ -1780,12 +1780,6 @@ static bool ReadSingleJob(FShaderCompileJob* CurrentJob, FArchive& WorkerOutputF
 	// The job should already have a non-zero output hash
 	checkf(CurrentJob->Output.OutputHash != FSHAHash() || !CurrentJob->bSucceeded, TEXT("OutputHash for a successful job was not set in the shader compile worker!"));
 
-	if (CurrentJob->bSucceeded && CurrentJob->Input.DumpDebugInfoPath.Len() > 0)
-	{
-		// write down the output hash as a file
-		FString HashFileName = FPaths::Combine(CurrentJob->Input.DumpDebugInfoPath, TEXT("OutputHash.txt"));
-		FFileHelper::SaveStringToFile(CurrentJob->Output.OutputHash.ToString(), *HashFileName, FFileHelper::EEncodingOptions::ForceAnsi);
-	}
 
 	// Support dumping debug info for only failed compilations or those with warnings
 	if (GShaderCompilingManager->ShouldRecompileToDumpShaderDebugInfo(*CurrentJob))
diff --git a/Engine/Source/Runtime/RenderCore/Private/ShaderCore.cpp b/Engine/Source/Runtime/RenderCore/Private/ShaderCore.cpp
--- a/Engine/Source/Runtime/RenderCore/Private/ShaderCore.cpp
+++ b/Engine/Source/Runtime/RenderCore/Private/ShaderCore.cpp
@@ -900,6 +900,14 @@ void CompileShader(const TArray<const IShaderFormat*>& ShaderFormats, FShaderCom
 #endif
 
 	Job.bSucceeded = Job.Output.bSucceeded;
+
+	if (Job.bSucceeded && Job.Input.DumpDebugInfoPath.Len() > 0)
+	{
+		// write down the output hash as a file
+		FString HashFileName = FPaths::Combine(Job.Input.DumpDebugInfoPath, TEXT("OutputHash.txt"));
+		FFileHelper::SaveStringToFile(Job.Output.OutputHash.ToString(), *HashFileName, FFileHelper::EEncodingOptions::ForceAnsi);
+	}
+
 }
 
 void CompileShaderPipeline(const TArray<const IShaderFormat*>& ShaderFormats, FShaderPipelineCompileJob* PipelineJob, const FString& Dir, int32* CompileCount)
