Updated tooltip of FarShadowCascadeCount to remove deprecated property names.

diff --git a/Engine/Source/Runtime/Engine/Classes/Components/DirectionalLightComponent.h b/Engine/Source/Runtime/Engine/Classes/Components/DirectionalLightComponent.h
--- a/Engine/Source/Runtime/Engine/Classes/Components/DirectionalLightComponent.h
+++ b/Engine/Source/Runtime/Engine/Classes/Components/DirectionalLightComponent.h
@@ -107,7 +107,7 @@ class ENGINE_API UDirectionalLightComponent : public ULightComponent
 	UPROPERTY(EditAnywhere, AdvancedDisplay, BlueprintReadOnly, Category=CascadedShadowMaps, DisplayName = "Inset Shadows For Movable Objects")
 	uint32 bUseInsetShadowsForMovableObjects : 1;
 
-	/** 0: no DistantShadowCascades, otherwise the count of cascades between WholeSceneDynamicShadowRadius and DistantShadowDistance that are covered by distant shadow cascades. */
+	/** 0: no Far Shadow Cascades, otherwise the number of cascades between DynamicShadowDistance and FarShadowDistance that are covered by Far Shadow Cascades. */
 	UPROPERTY(EditAnywhere, AdvancedDisplay, BlueprintReadOnly, Category=CascadedShadowMaps, meta=(UIMin = "0", UIMax = "4"), DisplayName = "Far Shadow Cascade Count")
 	int32 FarShadowCascadeCount;
 
