Non-deterministic GUID generation from Animation Data
#fix Handle Attribute curve keys individually, as the key array entries themselves hold raw ptrs (not the actual value) causing indeterminism
#misc Extra logging when storing/fetching a animationsequencecompression to cache

diff --git a/Engine/Plugins/Animation/AnimationData/Source/AnimationData/Private/AnimSequencerDataModel.cpp b/Engine/Plugins/Animation/AnimationData/Source/AnimationData/Private/AnimSequencerDataModel.cpp
--- a/Engine/Plugins/Animation/AnimationData/Source/AnimationData/Private/AnimSequencerDataModel.cpp
+++ b/Engine/Plugins/Animation/AnimationData/Source/AnimationData/Private/AnimSequencerDataModel.cpp
@@ -632,7 +632,12 @@ FGuid UAnimationSequencerDataModel::GenerateGuid() const
 		for (const FAnimatedBoneAttribute& Attribute : AnimatedBoneAttributes)
 		{
 			UpdateWithData(Attribute.Identifier);
-			UpdateSHAWithArray(Attribute.Curve.GetConstRefOfKeys());
+			const uint32 StructSize = Attribute.Identifier.GetType()->GetStructureSize();
+			for (const FAttributeKey& Key : Attribute.Curve.GetConstRefOfKeys())
+			{
+				UpdateWithData(Key.Time);
+				Sha.Update(Key.GetValuePtr<uint8>(), StructSize);
+			}
 		}
 
 		auto UpdateWithFloatCurve = [&UpdateWithData, &UpdateSHAWithArray](const FRichCurve& Curve)
diff --git a/Engine/Source/Runtime/Engine/Private/Animation/AnimDataModel.cpp b/Engine/Source/Runtime/Engine/Private/Animation/AnimDataModel.cpp
--- a/Engine/Source/Runtime/Engine/Private/Animation/AnimDataModel.cpp
+++ b/Engine/Source/Runtime/Engine/Private/Animation/AnimDataModel.cpp
@@ -561,7 +561,12 @@ FGuid UAnimDataModel::GenerateGuid() const
 	for (const FAnimatedBoneAttribute& Attribute : AnimatedBoneAttributes)
 	{
 		UpdateWithData(Attribute.Identifier);
-		UpdateSHAWithArray(Attribute.Curve.GetConstRefOfKeys());
+		const uint32 StructSize = Attribute.Identifier.GetType()->GetStructureSize();
+		for (const FAttributeKey& Key : Attribute.Curve.GetConstRefOfKeys())
+		{
+			UpdateWithData(Key.Time);
+			Sha.Update(Key.GetValuePtr<uint8>(), StructSize);
+		}
 	}
 
 	Sha.Final();
diff --git a/Engine/Source/Runtime/Engine/Private/Animation/AnimationCompressionDerivedData.cpp b/Engine/Source/Runtime/Engine/Private/Animation/AnimationCompressionDerivedData.cpp
--- a/Engine/Source/Runtime/Engine/Private/Animation/AnimationCompressionDerivedData.cpp
+++ b/Engine/Source/Runtime/Engine/Private/Animation/AnimationCompressionDerivedData.cpp
@@ -138,6 +138,7 @@ void FAnimationSequenceAsyncCacheTask::EndCache(UE::DerivedData::FCacheGetValueR
 				FMemoryReaderView Ar(RecordData, /*bIsPersistent*/ true);
 				CompressedData->SerializeCompressedData(Ar, true, AnimSequence, AnimSequence->GetSkeleton(), CompressibleAnimPtr->BoneCompressionSettings, CompressibleAnimPtr->CurveCompressionSettings);
 
+				UE_LOG(LogAnimationCompression, Display, TEXT("Fetched compressed animation data for %s"), *CompressibleAnimPtr->FullName);
 				if (Compression::FAnimationCompressionMemorySummaryScope::ShouldStoreCompressionResults())
 				{
 					const double CompressionEndTime = FPlatformTime::Seconds();
@@ -171,7 +172,8 @@ void FAnimationSequenceAsyncCacheTask::EndCache(UE::DerivedData::FCacheGetValueR
 				TArray64<uint8> RecordData;
 				FMemoryWriter64 Ar(RecordData, /*bIsPersistent*/ true);
 				CompressedData->SerializeCompressedData(Ar, true, nullptr, nullptr, CompressibleAnimPtr->BoneCompressionSettings, CompressibleAnimPtr->CurveCompressionSettings);
-				
+
+				UE_LOG(LogAnimationCompression, Display, TEXT("Storing compressed animation data for %s, at %s/%s"), *Name, *FString(Key.Bucket.ToString()), *LexToString(Key.Hash));
 				GetCache().PutValue({ {Name, Key, FValue::Compress(MakeSharedBufferFromArray(MoveTemp(RecordData)))} }, Owner);
 				
 				if (Compression::FAnimationCompressionMemorySummaryScope::ShouldStoreCompressionResults())
