Added stat scopes for ULevelStreamingProfilingSubsystem::StopTrackingAndReport ticker and its background thread task.

diff --git a/Engine/Source/Runtime/Engine/Private/ProfilingDebugging/LevelStreamingProfilingSubsystem.cpp b/Engine/Source/Runtime/Engine/Private/ProfilingDebugging/LevelStreamingProfilingSubsystem.cpp
--- a/Engine/Source/Runtime/Engine/Private/ProfilingDebugging/LevelStreamingProfilingSubsystem.cpp
+++ b/Engine/Source/Runtime/Engine/Private/ProfilingDebugging/LevelStreamingProfilingSubsystem.cpp
@@ -9,6 +9,7 @@
 #include "Streaming/LevelStreamingDelegates.h"
 #include "WorldPartition/WorldPartition.h"
 #include "WorldPartition/WorldPartitionLevelStreamingDynamic.h"
+#include "Stats/Stats.h"
 
 #if !WITH_EDITOR
 #include "Engine/World.h"
@@ -153,6 +154,7 @@ void ULevelStreamingProfilingSubsystem::StartTracking()
 
 void ULevelStreamingProfilingSubsystem::StopTrackingAndReport()
 {
+	QUICK_SCOPE_CYCLE_COUNTER(STAT_LevelStreamingProfilingSubsystem_StopTrackingAndReport);
 	if (!bIsTracking)
 	{
 		UE_LOG(LogLevelStreamingProfiling, Warning, TEXT("StopTrackingAndReport called without StartTracking"));
@@ -165,6 +167,7 @@ void ULevelStreamingProfilingSubsystem::StopTrackingAndReport()
 	ReportWritingTask = UE::Tasks::Launch(UE_SOURCE_LOCATION,
 	[this]() 
 	{
+		QUICK_SCOPE_CYCLE_COUNTER(STAT_LevelStreamingProfilingSubsystem_ReportWritingTask);
 		FString OutDir = FPaths::ProfilingDir() / TEXT("LevelStreaming");
 		IFileManager::Get().MakeDirectory(*OutDir, true);
 
