Moving some files around to remove dependencies on Engine/Private

diff --git a/Engine/Source/Runtime/Engine/Internal/PhysicsEngine/Experimental/ChaosCooking.h b/Engine/Source/Runtime/Engine/Internal/PhysicsEngine/Experimental/ChaosCooking.h
new file mode 100644
--- /dev/null
+++ b/Engine/Source/Runtime/Engine/Internal/PhysicsEngine/Experimental/ChaosCooking.h
@@ -0,0 +1,54 @@
+// Copyright Epic Games, Inc. All Rights Reserved.
+
+#pragma once
+
+#include "Containers/Array.h"
+#include "Templates/UniquePtr.h"
+#include "Async/TaskGraphInterfaces.h"
+#include "PhysicsEngine/BodySetup.h"
+
+struct FTriMeshCollisionData;
+
+namespace Chaos
+{
+	class FImplicitObject;
+	class FTriangleMeshImplicitObject;
+}
+
+namespace Chaos
+{
+	namespace Cooking
+	{
+		ENGINE_API TUniquePtr<Chaos::FTriangleMeshImplicitObject> BuildSingleTrimesh(const FTriMeshCollisionData& Desc, TArray<int32>& OutFaceRemap, TArray<int32>& OutVertexRemap);
+
+		void BuildConvexMeshes(TArray<TUniquePtr<Chaos::FImplicitObject>>& OutTriangleMeshes, const FCookBodySetupInfo& InParams);
+		void BuildTriangleMeshes(TArray<TUniquePtr<Chaos::FTriangleMeshImplicitObject>>& OutTriangleMeshes, TArray<int32>& OutFaceRemap, TArray<int32>& OutVertexRemap, const FCookBodySetupInfo& InParams);
+	}
+
+	struct FCookHelper
+	{
+		FCookHelper() = delete;
+		FCookHelper(UBodySetup* InSetup);
+
+		TArray<TUniquePtr<Chaos::FImplicitObject>> SimpleImplicits;
+		TArray<TUniquePtr<Chaos::FTriangleMeshImplicitObject>> ComplexImplicits;
+		FBodySetupUVInfo UVInfo;
+		TArray<int32> FaceRemap;
+		TArray<int32> VertexRemap;
+
+		void Cook();
+		void CookAsync(FSimpleDelegateGraphTask::FDelegate CompletionDelegate);
+		bool HasWork() const;
+
+		// CancelCookAsync is not guaranteed to have any effect on the work done.
+		// If it is called the cook work may be abandoned and the CookAsync may return early.
+		// If bCancel is true in the CompletionDelegate the results must be ignored.
+		void CancelCookAsync() { bCanceled = true; }
+		bool WasCanceled() const { return bCanceled; }
+
+	private:
+		UBodySetup* SourceSetup;
+		FCookBodySetupInfo CookInfo;
+		std::atomic<bool> bCanceled;
+	};
+}
\ No newline at end of file
diff --git a/Engine/Source/Runtime/Engine/Private/KismetTraceUtils.h b/Engine/Source/Runtime/Engine/Private/KismetTraceUtils.h
deleted file mode 100644
diff --git a/Engine/Source/Runtime/Engine/Private/PhysicsEngine/BodySetup.cpp b/Engine/Source/Runtime/Engine/Private/PhysicsEngine/BodySetup.cpp
--- a/Engine/Source/Runtime/Engine/Private/PhysicsEngine/BodySetup.cpp
+++ b/Engine/Source/Runtime/Engine/Private/PhysicsEngine/BodySetup.cpp
@@ -36,7 +36,7 @@
 #include "Experimental/ChaosDerivedData.h"
 #include "Physics/Experimental/ChaosDerivedDataReader.h"
 #include "Chaos/CollisionConvexMesh.h"
-#include "Experimental/ChaosCooking.h"
+#include "PhysicsEngine/Experimental/ChaosCooking.h"
 #include "PhysicsEngine/SphylElem.h"
 #include "PhysicsEngine/TaperedCapsuleElem.h"
 
diff --git a/Engine/Source/Runtime/Engine/Private/PhysicsEngine/Experimental/ChaosCooking.h b/Engine/Source/Runtime/Engine/Private/PhysicsEngine/Experimental/ChaosCooking.h
deleted file mode 100644
diff --git a/Engine/Source/Runtime/Engine/Private/PhysicsEngine/Experimental/ChaosDerivedData.cpp b/Engine/Source/Runtime/Engine/Private/PhysicsEngine/Experimental/ChaosDerivedData.cpp
--- a/Engine/Source/Runtime/Engine/Private/PhysicsEngine/Experimental/ChaosDerivedData.cpp
+++ b/Engine/Source/Runtime/Engine/Private/PhysicsEngine/Experimental/ChaosDerivedData.cpp
@@ -5,7 +5,7 @@
 #include "Chaos/TriangleMeshImplicitObject.h"
 
 #include "Serialization/MemoryWriter.h"
-#include "ChaosCooking.h"
+#include "PhysicsEngine/Experimental/ChaosCooking.h"
 
 class FChaosDerivedDataCookerRefHolder : public FGCObject
 {
diff --git a/Engine/Source/Runtime/Engine/Public/KismetTraceUtils.h b/Engine/Source/Runtime/Engine/Public/KismetTraceUtils.h
new file mode 100644
--- /dev/null
+++ b/Engine/Source/Runtime/Engine/Public/KismetTraceUtils.h
@@ -0,0 +1,35 @@
+// Copyright Epic Games, Inc. All Rights Reserved.
+
+
+
+#pragma once 
+
+#include "CoreMinimal.h"
+#include "DrawDebugHelpers.h"
+#include "CollisionQueryParams.h"
+#include "Kismet/KismetSystemLibrary.h"
+
+FCollisionQueryParams ConfigureCollisionParams(FName TraceTag, bool bTraceComplex, const TArray<AActor*>& ActorsToIgnore, bool bIgnoreSelf, const UObject* WorldContextObject);
+
+FCollisionObjectQueryParams ConfigureCollisionObjectParams(const TArray<TEnumAsByte<EObjectTypeQuery> > & ObjectTypes);
+
+#if ENABLE_DRAW_DEBUG
+
+ENGINE_API void DrawDebugSweptSphere(const UWorld* InWorld, FVector const& Start, FVector const& End, float Radius, FColor const& Color, bool bPersistentLines = false, float LifeTime = -1.f, uint8 DepthPriority = 0);
+ENGINE_API void DrawDebugSweptBox(const UWorld* InWorld, FVector const& Start, FVector const& End, FRotator const & Orientation, FVector const & HalfSize, FColor const& Color, bool bPersistentLines = false, float LifeTime = -1.f, uint8 DepthPriority = 0);
+
+ENGINE_API void DrawDebugLineTraceSingle(const UWorld* World, const FVector& Start, const FVector& End, EDrawDebugTrace::Type DrawDebugType, bool bHit, const FHitResult& OutHit, FLinearColor TraceColor, FLinearColor TraceHitColor, float DrawTime);
+ENGINE_API void DrawDebugLineTraceMulti(const UWorld* World, const FVector& Start, const FVector& End, EDrawDebugTrace::Type DrawDebugType, bool bHit, const TArray<FHitResult>& OutHits, FLinearColor TraceColor, FLinearColor TraceHitColor, float DrawTime);
+
+ENGINE_API void DrawDebugBoxTraceSingle(const UWorld* World, const FVector& Start, const FVector& End, const FVector HalfSize, const FRotator Orientation, EDrawDebugTrace::Type DrawDebugType, bool bHit, const FHitResult& OutHit, FLinearColor TraceColor, FLinearColor TraceHitColor, float DrawTime);
+ENGINE_API void DrawDebugBoxTraceMulti(const UWorld* World, const FVector& Start, const FVector& End, const FVector HalfSize, const FRotator Orientation, EDrawDebugTrace::Type DrawDebugType, bool bHit, const TArray<FHitResult>& OutHits, FLinearColor TraceColor, FLinearColor TraceHitColor, float DrawTime);
+
+ENGINE_API void DrawDebugSphereTraceSingle(const UWorld* World, const FVector& Start, const FVector& End, float Radius, EDrawDebugTrace::Type DrawDebugType, bool bHit, const FHitResult& OutHit, FLinearColor TraceColor, FLinearColor TraceHitColor, float DrawTime);
+ENGINE_API void DrawDebugSphereTraceMulti(const UWorld* World, const FVector& Start, const FVector& End, float Radius, EDrawDebugTrace::Type DrawDebugType, bool bHit, const TArray<FHitResult>& OutHits, FLinearColor TraceColor, FLinearColor TraceHitColor, float DrawTime);
+
+ENGINE_API void DrawDebugCapsuleTraceSingle(const UWorld* World, const FVector& Start, const FVector& End, float Radius, float HalfHeight, EDrawDebugTrace::Type DrawDebugType, bool bHit, const FHitResult& OutHit, FLinearColor TraceColor, FLinearColor TraceHitColor, float DrawTime);
+ENGINE_API void DrawDebugCapsuleTraceMulti(const UWorld* World, const FVector& Start, const FVector& End, float Radius, float HalfHeight, EDrawDebugTrace::Type DrawDebugType, bool bHit, const TArray<FHitResult>& OutHits, FLinearColor TraceColor, FLinearColor TraceHitColor, float DrawTime);
+
+#endif // ENABLE_DRAW_DEBUG
+
+
diff --git a/Engine/Source/Runtime/Landscape/Landscape.Build.cs b/Engine/Source/Runtime/Landscape/Landscape.Build.cs
--- a/Engine/Source/Runtime/Landscape/Landscape.Build.cs
+++ b/Engine/Source/Runtime/Landscape/Landscape.Build.cs
@@ -9,7 +9,6 @@ public class Landscape : ModuleRules
 	{
 		PrivateIncludePaths.AddRange(
 			new string[] {
-				Path.Combine(GetModuleDirectory("Engine"), "Private"), // for Engine/Private/Collision/PhysXCollision.h
 				"../Shaders/Shared"
             }
 		);
diff --git a/Engine/Source/Runtime/Landscape/Private/LandscapeCollision.cpp b/Engine/Source/Runtime/Landscape/Private/LandscapeCollision.cpp
--- a/Engine/Source/Runtime/Landscape/Private/LandscapeCollision.cpp
+++ b/Engine/Source/Runtime/Landscape/Private/LandscapeCollision.cpp
@@ -58,7 +58,6 @@
 #include "Chaos/HeightField.h"
 #include "Chaos/ImplicitObjectTransformed.h"
 #include "PhysicsEngine/BodySetup.h"
-#include "PhysicsEngine/Experimental/ChaosDerivedData.h"
 #include "PhysicsEngine/Experimental/ChaosCooking.h"
 #include "Chaos/ChaosArchive.h"
 #include "PhysicsProxy/SingleParticlePhysicsProxy.h"
