Strata - Fix for subsurface profile not reseting on root node.
Also subsurface profile now cannot be tweaked when IsThin is true.

diff --git a/Engine/Shaders/Private/Strata/StrataExport.ush b/Engine/Shaders/Private/Strata/StrataExport.ush
--- a/Engine/Shaders/Private/Strata/StrataExport.ush
+++ b/Engine/Shaders/Private/Strata/StrataExport.ush
@@ -367,6 +367,7 @@ void PackStrataOut(
 				}
 				else if ((OptimisedModeMask & SingleIsSSSProfile) == SingleIsSSSProfile)
 				{
+					// SingleIsSSSProfile appears critically after SingleIsWrapThin because thin will disable subsurface profile.
 					OptimisedLegacyMode = SINGLE_OPTLEGACYMODE_SSSPROFILE;
 					SSSData.Header.Bytes = 0; // Force to skip SSSData writing
 				}
diff --git a/Engine/Source/Runtime/Engine/Classes/Materials/Material.h b/Engine/Source/Runtime/Engine/Classes/Materials/Material.h
--- a/Engine/Source/Runtime/Engine/Classes/Materials/Material.h
+++ b/Engine/Source/Runtime/Engine/Classes/Materials/Material.h
@@ -523,7 +523,7 @@ public:
 	UPROPERTY(EditAnywhere, Category=Material)
 	uint8 TwoSided : 1;
 
-	/** Indicates that the material should be rendered as a thin surface (i.e., without inner volume). Only used by Strata materials. */
+	/** Indicates that the material should be rendered as a thin surface (i.e., without inner volume). Only used by Strata materials. Enabling ThinSurface will disable subsurface profiles. */
 	UPROPERTY(EditAnywhere, Category = Material)
 	uint8 bIsThinSurface : 1;
 
diff --git a/Engine/Source/Runtime/Engine/Classes/Materials/MaterialInterface.h b/Engine/Source/Runtime/Engine/Classes/Materials/MaterialInterface.h
--- a/Engine/Source/Runtime/Engine/Classes/Materials/MaterialInterface.h
+++ b/Engine/Source/Runtime/Engine/Classes/Materials/MaterialInterface.h
@@ -270,7 +270,7 @@ public:
 #endif // WITH_EDITORONLY_DATA
 
 public:
-	/** SubsurfaceProfile, for Screen Space Subsurface Scattering */
+	/** SubsurfaceProfile, for Screen Space Subsurface Scattering. Disabled when root node IsThin is true. */
 	UPROPERTY(EditAnywhere, BlueprintReadOnly, Category = Material, meta = (DisplayName = "Subsurface Profile"))
 	TObjectPtr<class USubsurfaceProfile> SubsurfaceProfile;
 
diff --git a/Engine/Source/Runtime/Engine/Private/Materials/Material.cpp b/Engine/Source/Runtime/Engine/Private/Materials/Material.cpp
--- a/Engine/Source/Runtime/Engine/Private/Materials/Material.cpp
+++ b/Engine/Source/Runtime/Engine/Private/Materials/Material.cpp
@@ -4442,7 +4442,15 @@ bool UMaterial::CanEditChange(const FProperty* InProperty) const
 
 		if (PropertyName == GET_MEMBER_NAME_STRING_CHECKED(UMaterial, SubsurfaceProfile))
 		{
-			return MaterialDomain == MD_Surface && UseSubsurfaceProfile(ShadingModels) && IsOpaqueOrMaskedBlendMode(BlendMode);
+			const bool bCommonTest = MaterialDomain == MD_Surface && UseSubsurfaceProfile(ShadingModels) && IsOpaqueOrMaskedBlendMode(BlendMode);
+			if (bStrataEnabled)
+			{
+				return !IsThinSurface() && bCommonTest;
+			}
+			else
+			{
+				return bCommonTest;
+			}
 		}
 
 		if (PropertyName == GET_MEMBER_NAME_STRING_CHECKED(FLightmassMaterialInterfaceSettings, bCastShadowAsMasked))
@@ -4834,16 +4842,6 @@ void UMaterial::RebuildShadingModelField()
 					ShadingModels.AddShadingModel(MSM_DefaultLit);
 				}
 			}
-
-			// Subsurface profile
-			if (StrataMaterialInfo.HasShadingModel(SSM_SubsurfaceLit) && StrataMaterialInfo.CountSubsurfaceProfiles() > 0)
-			{
-				if (StrataMaterialInfo.CountSubsurfaceProfiles() > 1)
-				{
-					UE_LOG(LogMaterial, Error, TEXT("%s: Material has more than a single sub-surface profile used."), *GetName());
-				}
-				SubsurfaceProfile = StrataMaterialInfo.GetSubsurfaceProfile();
-			}
 		}
 		else
 		{
@@ -4925,16 +4923,18 @@ void UMaterial::RebuildShadingModelField()
 
 			// Also update the ShadingModels for remaining pipeline operation
 			ShadingModels.AddShadingModel(ShadingModel);
+		}
 
-			// Subsurface profil
-			if ((StrataMaterialInfo.HasOnlyShadingModel(SSM_Eye) || StrataMaterialInfo.HasOnlyShadingModel(SSM_SubsurfaceLit)) && StrataMaterialInfo.CountSubsurfaceProfiles() > 0)
+		// Now, reset the subsurface profile (in case it has been removed from any slab before) and set it only if needed.
+		SubsurfaceProfile = nullptr;
+		if ((StrataMaterialInfo.HasOnlyShadingModel(SSM_Eye) || StrataMaterialInfo.HasOnlyShadingModel(SSM_SubsurfaceLit)) && StrataMaterialInfo.CountSubsurfaceProfiles() > 0 
+			&& !IsThinSurface())	// Thin surfaces disable subsurface profile because this is not compatible with fast path.
+		{
+			if (StrataMaterialInfo.CountSubsurfaceProfiles() > 1)
 			{
-				if (StrataMaterialInfo.CountSubsurfaceProfiles() > 1)
-				{
-					UE_LOG(LogMaterial, Error, TEXT("%s: Material has more than a single sub-surface profile used."), *GetName());
-				}
-				SubsurfaceProfile = StrataMaterialInfo.GetSubsurfaceProfile();
+				UE_LOG(LogMaterial, Error, TEXT("%s: Material has more than a single sub-surface profile used."), *GetName());
 			}
+			SubsurfaceProfile = StrataMaterialInfo.GetSubsurfaceProfile();
 		}
 	}
 	// If using shading model from material expression, go through the expressions and look for the ShadingModel expression to figure out what shading models need to be supported in this material.
