Fix issue where streamed out levels preparing for GC would rename to a unique name to try and avoid being reloaded - but the unique naming pattern used matched the naming pattern for newly loaded up levels, resulting in the garbage marked levels being loaded back up.

diff --git a/Engine/Source/Runtime/Engine/Private/LevelStreamingGCHelper.cpp b/Engine/Source/Runtime/Engine/Private/LevelStreamingGCHelper.cpp
--- a/Engine/Source/Runtime/Engine/Private/LevelStreamingGCHelper.cpp
+++ b/Engine/Source/Runtime/Engine/Private/LevelStreamingGCHelper.cpp
@@ -189,7 +189,7 @@ void FLevelStreamingGCHelper::PrepareStreamedOutLevelForGC(ULevel* InLevel)
 			for (UPackage* Package : Packages)
 			{
 				FCoreUObjectInternalDelegates::GetOnLeakedPackageRenameDelegate().Broadcast(Package);
-				FName NewName = MakeUniqueObjectName(nullptr, UPackage::StaticClass(), Package->GetFName());
+				const FName NewName = MakeUniqueObjectName(nullptr, UPackage::StaticClass(), FName(FString::Printf(TEXT("%s_PreparedForGC"), *Package->GetFName().GetPlainNameString())));
 				Package->Rename(*NewName.ToString(), nullptr, REN_ForceNoResetLoaders | REN_DontCreateRedirectors | REN_NonTransactional);
 			}
 
