Function to determine wether all async loads have completed

diff --git a/Engine/Source/Runtime/Engine/Classes/Engine/StreamableManager.h b/Engine/Source/Runtime/Engine/Classes/Engine/StreamableManager.h
--- a/Engine/Source/Runtime/Engine/Classes/Engine/StreamableManager.h
+++ b/Engine/Source/Runtime/Engine/Classes/Engine/StreamableManager.h
@@ -534,6 +534,9 @@ struct ENGINE_API FStreamableManager : public FGCObject
 	 */
 	bool GetActiveHandles(const FSoftObjectPath& Target, TArray<TSharedRef<FStreamableHandle>>& HandleList, bool bOnlyManagedHandles = false) const;
 
+	/** Returns true if all pending async loads have finished for all targets */
+	bool AreAllAsyncLoadsComplete() const;
+
 	/** Returns true if all pending async loads have finished for this target */
 	bool IsAsyncLoadComplete(const FSoftObjectPath& Target) const;
 
diff --git a/Engine/Source/Runtime/Engine/Private/StreamableManager.cpp b/Engine/Source/Runtime/Engine/Private/StreamableManager.cpp
--- a/Engine/Source/Runtime/Engine/Private/StreamableManager.cpp
+++ b/Engine/Source/Runtime/Engine/Private/StreamableManager.cpp
@@ -1667,6 +1667,20 @@ void FStreamableManager::RemoveReferencedAsset(const FSoftObjectPath& Target, TS
 	}
 }
 
+bool FStreamableManager::AreAllAsyncLoadsComplete() const
+{
+	check(IsInGameThread());
+	for (const TPair<FSoftObjectPath, FStreamable*>& Item : StreamableItems)
+	{
+		const FStreamable* Streamable = Item.Value;
+		if (Streamable && Streamable->bAsyncLoadRequestOutstanding)
+		{
+			return false;
+		}
+	}
+	return true;
+}
+
 bool FStreamableManager::IsAsyncLoadComplete(const FSoftObjectPath& Target) const
 {
 	// Failed loads count as success
