Ground work for Level Instance Actor Filters:
Refactor CreateActorDescriptorViewsRecursive to take a FContainerInstanceDescriptor param (instead of adding more params)
Refactor FWorldPartitionActorDesc::GetContainerInstance to take a FContainerInstance reference (instead of adding more params)
Refactor Bounds calculation so we don't need to do it outside of the FContainerInstanceDescriptor creation

#rnx

diff --git a/Engine/Source/Runtime/Engine/Private/WorldPartition/LevelInstance/LevelInstanceActorDesc.cpp b/Engine/Source/Runtime/Engine/Private/WorldPartition/LevelInstance/LevelInstanceActorDesc.cpp
--- a/Engine/Source/Runtime/Engine/Private/WorldPartition/LevelInstance/LevelInstanceActorDesc.cpp
+++ b/Engine/Source/Runtime/Engine/Private/WorldPartition/LevelInstance/LevelInstanceActorDesc.cpp
@@ -158,16 +158,16 @@ bool FLevelInstanceActorDesc::IsContainerInstance() const
 	return ULevelInstanceSubsystem::CanUsePackage(LevelPackage);
 }
 
-bool FLevelInstanceActorDesc::GetContainerInstance(const UActorDescContainer*& OutLevelContainer, FTransform& OutLevelTransform, EContainerClusterMode& OutClusterMode) const
+bool FLevelInstanceActorDesc::GetContainerInstance(FWorldPartitionActorDesc::FContainerInstance& OutContainerInstance) const
 {
 	if (LevelInstanceContainer.IsValid())
 	{
-		OutLevelContainer = LevelInstanceContainer.Get();
-		OutClusterMode = EContainerClusterMode::Partitioned;
+		OutContainerInstance.Container = LevelInstanceContainer.Get();
+		OutContainerInstance.ClusterMode = EContainerClusterMode::Partitioned;
 
 		// Apply level instance pivot offset
 		FTransform LevelInstancePivotOffsetTransform = FTransform(ULevel::GetLevelInstancePivotOffsetFromPackage(LevelInstanceContainer->GetContainerPackage()));
-		OutLevelTransform = LevelInstancePivotOffsetTransform * LevelInstanceTransform;
+		OutContainerInstance.Transform = LevelInstancePivotOffsetTransform * LevelInstanceTransform;
 
 		return true;
 	}
diff --git a/Engine/Source/Runtime/Engine/Private/WorldPartition/WorldPartitionActorDesc.cpp b/Engine/Source/Runtime/Engine/Private/WorldPartition/WorldPartitionActorDesc.cpp
--- a/Engine/Source/Runtime/Engine/Private/WorldPartition/WorldPartitionActorDesc.cpp
+++ b/Engine/Source/Runtime/Engine/Private/WorldPartition/WorldPartitionActorDesc.cpp
@@ -23,7 +23,6 @@
 #include "WorldPartition/DataLayer/DataLayerAsset.h"
 #include "WorldPartition/DataLayer/DataLayerUtils.h"
 #include "WorldPartition/ContentBundle/ContentBundlePaths.h"
-#include "WorldPartition/WorldPartitionActorContainerID.h"
 #include "WorldPartition/ErrorHandling/WorldPartitionStreamingGenerationErrorHandler.h"
 #include "ActorReferencesUtils.h"
 
@@ -575,6 +574,20 @@ FName FWorldPartitionActorDesc::GetDisplayClassName() const
 	return BaseClass.IsNull() ? GetCleanClassName(NativeClass) : GetCleanClassName(BaseClass);
 }
 
+bool FWorldPartitionActorDesc::GetContainerInstance(const UActorDescContainer*& OutLevelContainer, FTransform& OutLevelTransform, EContainerClusterMode& OutClusterMode) const
+{
+	FContainerInstance ContainerInstance;
+	if (GetContainerInstance(ContainerInstance))
+	{
+		OutLevelContainer = ContainerInstance.Container;
+		OutLevelTransform = ContainerInstance.Transform;
+		OutClusterMode = ContainerInstance.ClusterMode;
+		return true;
+	}
+
+	return false;
+}
+
 FGuid FWorldPartitionActorDesc::GetContentBundleGuid() const
 {
 	return ContentBundleGuid;
diff --git a/Engine/Source/Runtime/Engine/Private/WorldPartition/WorldPartitionActorDescView.cpp b/Engine/Source/Runtime/Engine/Private/WorldPartition/WorldPartitionActorDescView.cpp
--- a/Engine/Source/Runtime/Engine/Private/WorldPartition/WorldPartitionActorDescView.cpp
+++ b/Engine/Source/Runtime/Engine/Private/WorldPartition/WorldPartitionActorDescView.cpp
@@ -184,9 +184,9 @@ FName FWorldPartitionActorDescView::GetLevelPackage() const
 	return ActorDesc->GetLevelPackage();
 }
 
-bool FWorldPartitionActorDescView::GetContainerInstance(const UActorDescContainer*& OutLevelContainer, FTransform& OutLevelTransform, EContainerClusterMode& OutClusterMode) const
+bool FWorldPartitionActorDescView::GetContainerInstance(FWorldPartitionActorDesc::FContainerInstance& OutContainerInstance) const
 {
-	return ActorDesc->GetContainerInstance(OutLevelContainer, OutLevelTransform, OutClusterMode);
+	return ActorDesc->GetContainerInstance(OutContainerInstance);
 }
 
 void FWorldPartitionActorDescView::CheckForErrors(IStreamingGenerationErrorHandler* ErrorHandler) const
diff --git a/Engine/Source/Runtime/Engine/Private/WorldPartition/WorldPartitionStreamingGeneration.cpp b/Engine/Source/Runtime/Engine/Private/WorldPartition/WorldPartitionStreamingGeneration.cpp
--- a/Engine/Source/Runtime/Engine/Private/WorldPartition/WorldPartitionStreamingGeneration.cpp
+++ b/Engine/Source/Runtime/Engine/Private/WorldPartition/WorldPartitionStreamingGeneration.cpp
@@ -233,6 +233,7 @@ class FWorldPartitionStreamingGenerator
 		bool bIsSpatiallyLoaded;
 		FString OwnerName;
 		FActorContainerID ID;
+		FActorContainerID ParentID;
 	};
 
 	void ResolveRuntimeSpatiallyLoaded(FWorldPartitionActorDescView& ActorDescView)
@@ -383,26 +384,44 @@ class FWorldPartitionStreamingGenerator
 		}
 	}
 
-	void CreateActorDescriptorViewsRecursive(const UActorDescContainer* InContainer, const FTransform& InTransform, FName InRuntimeGrid, bool bInIsSpatiallyLoaded, const TSet<FName>& InRuntimeDataLayers, const FActorContainerID& InContainerID, const FActorContainerID& InParentContainerID, EContainerClusterMode InClusterMode, const TCHAR* InOwnerName)
+	void CreateActorDescriptorViewsRecursive(const FContainerInstanceDescriptor& InContainerInstanceDescriptor)
 	{
 		// Create or resolve container descriptor
-		FContainerDescriptor& ContainerDescriptor = ContainerDescriptorsMap.FindOrAdd(InContainer);
-
-		if (!ContainerDescriptor.Container)
-		{
-			ContainerDescriptor.Container = InContainer;
+		FContainerDescriptor& ContainerDescriptor = ContainerDescriptorsMap.FindOrAdd(InContainerInstanceDescriptor.Container);
 		
-			// Gather actor descriptor views for this container
-			CreateActorDescViewMap(InContainer, ContainerDescriptor.ActorDescViewMap, ContainerDescriptor.FilteredActorDescViewMap, InContainerID, ContainerDescriptor.ContainerInstanceViews);
+		// Scope ContainerInstanceDescriptor so that it doesn't get used mistakenly after the ContainerInstanceViews loop
+		{
+			// Create container instance descriptor
+			check(!ContainerInstanceDescriptorsMap.Contains(InContainerInstanceDescriptor.ID));
+
+			FContainerInstanceDescriptor& ContainerInstanceDescriptor = ContainerInstanceDescriptorsMap.Add(InContainerInstanceDescriptor.ID, InContainerInstanceDescriptor);
+
+			if (!ContainerDescriptor.Container)
+			{
+				ContainerDescriptor.Container = ContainerInstanceDescriptor.Container;
+
+				// Gather actor descriptor views for this container
+				CreateActorDescViewMap(ContainerInstanceDescriptor.Container, ContainerDescriptor.ActorDescViewMap, ContainerDescriptor.FilteredActorDescViewMap, ContainerInstanceDescriptor.ID, ContainerDescriptor.ContainerInstanceViews);
 
-			// Resolve actor descriptor views before validation
-			ResolveContainerDescriptor(ContainerDescriptor);
+				// Resolve actor descriptor views before validation
+				ResolveContainerDescriptor(ContainerDescriptor);
 
-			// Validate container, fixing anything illegal, etc.
-			ValidateContainerDescriptor(ContainerDescriptor, InContainerID.IsMainContainer());
+				// Validate container, fixing anything illegal, etc.
+				ValidateContainerDescriptor(ContainerDescriptor, ContainerInstanceDescriptor.ID.IsMainContainer());
 
-			// Update container, computing cluster, bounds, etc.
-			UpdateContainerDescriptor(ContainerDescriptor);
+				// Update container, computing cluster, bounds, etc.
+				UpdateContainerDescriptor(ContainerDescriptor);
+			}
+
+			// Calculate Bounds of Non-container ActorDescViews
+			check(!ContainerInstanceDescriptor.Bounds.IsValid);
+			ContainerDescriptor.ActorDescViewMap.ForEachActorDescView([&ContainerInstanceDescriptor](const FWorldPartitionActorDescView& ActorDescView)
+			{
+				if (ActorDescView.GetIsSpatiallyLoaded())
+				{
+					ContainerInstanceDescriptor.Bounds += ActorDescView.GetRuntimeBounds().TransformBy(ContainerInstanceDescriptor.Transform);
+				}
+			});
 		}
 
 		// Inherited parent properties logic
@@ -425,47 +444,46 @@ class FWorldPartitionStreamingGenerator
 		TArray<FWorldPartitionActorDescView> ContainerInstanceViews = ContainerDescriptor.ContainerInstanceViews;
 		for (const FWorldPartitionActorDescView& ContainerInstanceView : ContainerInstanceViews)
 		{
-			const UActorDescContainer* SubContainer;
-			EContainerClusterMode SubClusterMode;
-			FTransform SubTransform;
-
-			if (!ContainerInstanceView.GetContainerInstance(SubContainer, SubTransform, SubClusterMode))
+			const FGuid ActorGuid = ContainerInstanceView.GetGuid();
+			FWorldPartitionActorDesc::FContainerInstance SubContainerInstance(FActorContainerID(InContainerInstanceDescriptor.ID, ActorGuid));
+			if (!ContainerInstanceView.GetContainerInstance(SubContainerInstance))
 			{
 				ErrorHandler->OnLevelInstanceInvalidWorldAsset(ContainerInstanceView, ContainerInstanceView.GetLevelPackage(), IStreamingGenerationErrorHandler::ELevelInstanceInvalidReason::WorldAssetHasInvalidContainer);
 				continue;
 			}
 
-			check(SubContainer);
+			check(SubContainerInstance.Container);
+			FContainerInstanceDescriptor SubContainerInstanceDescriptor;
+
+			// Inherit fields from FContainerInstance
+			SubContainerInstanceDescriptor.ID = SubContainerInstance.GetID();
+			SubContainerInstanceDescriptor.Container = SubContainerInstance.Container;
+			SubContainerInstanceDescriptor.Transform = SubContainerInstance.Transform * InContainerInstanceDescriptor.Transform;
 
-			const FGuid ActorGuid = ContainerInstanceView.GetGuid();
-			const FActorContainerID SubContainerID(InContainerID, ActorGuid);
+
+			SubContainerInstanceDescriptor.ParentID = InContainerInstanceDescriptor.ID;
+			SubContainerInstanceDescriptor.OwnerName = *ContainerInstanceView.GetActorLabelOrName().ToString();
 
 			// Inherit parent container properties
-			FName InheritedRuntimeGrid = InRuntimeGrid;
-			bool bInheritedbIsSpatiallyLoaded = bInIsSpatiallyLoaded;
-			TSet<FName> InheritedRuntimeDataLayers = InRuntimeDataLayers;
-			InheritParentContainerProperties(InContainerID, ContainerInstanceView, InheritedRuntimeDataLayers, InheritedRuntimeGrid, bInheritedbIsSpatiallyLoaded);
+			FName InheritedRuntimeGrid = InContainerInstanceDescriptor.RuntimeGrid;
+			bool bInheritedbIsSpatiallyLoaded = InContainerInstanceDescriptor.bIsSpatiallyLoaded;
+			TSet<FName> InheritedRuntimeDataLayers = InContainerInstanceDescriptor.RuntimeDataLayers;
+			InheritParentContainerProperties(InContainerInstanceDescriptor.ID, ContainerInstanceView, InheritedRuntimeDataLayers, InheritedRuntimeGrid, bInheritedbIsSpatiallyLoaded);
 
-			CreateActorDescriptorViewsRecursive(SubContainer, SubTransform * InTransform, InheritedRuntimeGrid, bInheritedbIsSpatiallyLoaded, InheritedRuntimeDataLayers, SubContainerID, InContainerID, SubClusterMode, *ContainerInstanceView.GetActorLabelOrName().ToString());
-		}
+			SubContainerInstanceDescriptor.bIsSpatiallyLoaded = bInheritedbIsSpatiallyLoaded;
+			SubContainerInstanceDescriptor.RuntimeGrid = InheritedRuntimeGrid;
+			SubContainerInstanceDescriptor.RuntimeDataLayers = InheritedRuntimeDataLayers;
 
-		// Create container instance descriptor
-		check(!ContainerInstanceDescriptorsMap.Contains(InContainerID));
-
-		FContainerInstanceDescriptor& ContainerInstanceDescriptor = ContainerInstanceDescriptorsMap.Add(InContainerID);
-		ContainerInstanceDescriptor.Container = InContainer;
-		ContainerInstanceDescriptor.Transform = InTransform;
-		ContainerInstanceDescriptor.ClusterMode = InClusterMode;
-		ContainerInstanceDescriptor.RuntimeDataLayers = InRuntimeDataLayers;
-		ContainerInstanceDescriptor.RuntimeGrid = InRuntimeGrid;
-		ContainerInstanceDescriptor.bIsSpatiallyLoaded = bInIsSpatiallyLoaded;
-		ContainerInstanceDescriptor.OwnerName = InOwnerName;
-		ContainerInstanceDescriptor.ID = InContainerID;
-
-		// Maintain containers hierarchy, bottom up
-		if (InContainerID != InParentContainerID)
+			CreateActorDescriptorViewsRecursive(SubContainerInstanceDescriptor);
+		}
+							
+		if (!InContainerInstanceDescriptor.ID.IsMainContainer())
 		{
-			ContainersHierarchy.Add(InContainerID, InParentContainerID);
+			FContainerInstanceDescriptor& ParentContainer = ContainerInstanceDescriptorsMap.FindChecked(InContainerInstanceDescriptor.ParentID);
+			
+			// Fetch version stored in map as it will be the only one that will have its bounds updated by its Children
+			const FContainerInstanceDescriptor& ContainerInstanceWithUpdatedBounds = ContainerInstanceDescriptorsMap.FindChecked(InContainerInstanceDescriptor.ID);
+			ParentContainer.Bounds += ContainerInstanceWithUpdatedBounds.Bounds;
 		}
 	}
 
@@ -477,35 +495,12 @@ class FWorldPartitionStreamingGenerator
 		TRACE_CPUPROFILER_EVENT_SCOPE(FWorldPartitionStreamingGenerator::CreateActorContainers);
 
 		// Since we apply AND logic on spatially loaded flag recursively, startup value must be true
-		const bool bIsSpatiallyLoaded = true;
-		CreateActorDescriptorViewsRecursive(InContainer, FTransform::Identity, NAME_None, bIsSpatiallyLoaded, TSet<FName>(), FActorContainerID(), FActorContainerID(), EContainerClusterMode::Partitioned, TEXT("MainContainer"));
-
-		// Update container instances bounds
-		for (auto& [ContainerID, ContainerInstanceDescriptor] : ContainerInstanceDescriptorsMap)
-		{
-			const FContainerDescriptor& ContainerDescriptor = ContainerDescriptorsMap.FindChecked(ContainerInstanceDescriptor.Container);
-
-			// Lambdas can't capture structured bindings
-			const FTransform& ContainerInstanceDescriptorTransform = ContainerInstanceDescriptor.Transform;
-			FBox& ContainerInstanceDescriptorBounds = ContainerInstanceDescriptor.Bounds;			
-			check(!ContainerInstanceDescriptorBounds.IsValid);
-
-			ContainerDescriptor.ActorDescViewMap.ForEachActorDescView([&ContainerInstanceDescriptorBounds, &ContainerInstanceDescriptorTransform](const FWorldPartitionActorDescView& ActorDescView)
-			{
-				if (ActorDescView.GetIsSpatiallyLoaded())
-				{
-					ContainerInstanceDescriptorBounds += ActorDescView.GetRuntimeBounds().TransformBy(ContainerInstanceDescriptorTransform);
-				}
-			});
-		}
-
-		// Update parent containers bounds, this relies on the fact that ContainersHierarchy is built bottom up
-		for (auto& [ChildContainerID, ParentContainerID] : ContainersHierarchy)
-		{
-			const FContainerInstanceDescriptor& CurrentContainer = ContainerInstanceDescriptorsMap.FindChecked(ChildContainerID);
-			FContainerInstanceDescriptor& ParentContainer = ContainerInstanceDescriptorsMap.FindChecked(ParentContainerID);
-			ParentContainer.Bounds += CurrentContainer.Bounds;
-		}
+		FContainerInstanceDescriptor MainContainerInstance;
+		MainContainerInstance.Container = InContainer;
+		MainContainerInstance.bIsSpatiallyLoaded = true;
+		MainContainerInstance.ClusterMode = EContainerClusterMode::Partitioned;
+		MainContainerInstance.OwnerName = TEXT("MainContainer");
+		CreateActorDescriptorViewsRecursive(MainContainerInstance);
 	}
 
 	/** 
@@ -892,9 +887,12 @@ public:
 
 		// Build the containers tree representation
 		TMultiMap<FActorContainerID, FActorContainerID> InvertedContainersHierarchy;
-		for (auto& [ChildContainerID, ParentContainerID] : ContainersHierarchy)
+		for (auto& [ContainerID, ContainerInstanceDescriptor] : ContainerInstanceDescriptorsMap)
 		{
-			InvertedContainersHierarchy.Add(ParentContainerID, ChildContainerID);
+			if (!ContainerID.IsMainContainer())
+			{
+				InvertedContainersHierarchy.Add(ContainerInstanceDescriptor.ParentID, ContainerID);
+			}
 		}
 
 		UE_SCOPED_INDENT_LOG_ARCHIVE(Ar.PrintfIndent(TEXT("Containers:")));
@@ -995,9 +993,6 @@ private:
 	/** Maps containers IDs to their container instance descriptor */
 	TMap<FActorContainerID, FContainerInstanceDescriptor> ContainerInstanceDescriptorsMap;
 
-	/** Maps containers IDs to their parent ID */
-	TMap<FActorContainerID, FActorContainerID> ContainersHierarchy;
-
 	/** Data required for streaming generation interface */
 	TUniquePtr<FStreamingGenerationContext> StreamingGenerationContext;
 
diff --git a/Engine/Source/Runtime/Engine/Public/WorldPartition/LevelInstance/LevelInstanceActorDesc.h b/Engine/Source/Runtime/Engine/Public/WorldPartition/LevelInstance/LevelInstanceActorDesc.h
--- a/Engine/Source/Runtime/Engine/Public/WorldPartition/LevelInstance/LevelInstanceActorDesc.h
+++ b/Engine/Source/Runtime/Engine/Public/WorldPartition/LevelInstance/LevelInstanceActorDesc.h
@@ -26,7 +26,7 @@ public:
 
 	virtual bool IsContainerInstance() const override;
 	virtual FName GetLevelPackage() const override { return LevelPackage; }
-	virtual bool GetContainerInstance(const UActorDescContainer*& OutLevelContainer, FTransform& OutLevelTransform, EContainerClusterMode& OutClusterMode) const override;
+	virtual bool GetContainerInstance(FWorldPartitionActorDesc::FContainerInstance& OutContainerInstance) const override;
 	virtual void CheckForErrors(IStreamingGenerationErrorHandler* ErrorHandler) const override;
 
 protected:
diff --git a/Engine/Source/Runtime/Engine/Public/WorldPartition/WorldPartitionActorDesc.h b/Engine/Source/Runtime/Engine/Public/WorldPartition/WorldPartitionActorDesc.h
--- a/Engine/Source/Runtime/Engine/Public/WorldPartition/WorldPartitionActorDesc.h
+++ b/Engine/Source/Runtime/Engine/Public/WorldPartition/WorldPartitionActorDesc.h
@@ -12,6 +12,7 @@
 #include "UObject/WeakObjectPtr.h"
 #include "Templates/SubclassOf.h"
 #include "Misc/Guid.h"
+#include "WorldPartition/WorldPartitionActorContainerID.h"
 
 // Struct used to create actor descriptor
 struct FWorldPartitionActorDescInitData
@@ -25,7 +26,6 @@ struct FWorldPartitionActorDescInitData
 class AActor;
 class UActorDescContainer;
 class IStreamingGenerationErrorHandler;
-struct FActorContainerID;
 
 enum class EContainerClusterMode : uint8
 {
@@ -73,6 +73,22 @@ class ENGINE_API FWorldPartitionActorDesc
 	friend class FAssetRootPackagePatcher;
 
 public:
+	struct FContainerInstance
+	{
+		FContainerInstance() {}
+		FContainerInstance(const FActorContainerID& InContainerID)
+			: ContainerID(InContainerID) {}
+	public:
+		const FActorContainerID& GetID() const { return ContainerID; }
+
+		const UActorDescContainer* Container = nullptr;
+		FTransform Transform = FTransform::Identity;
+		EContainerClusterMode ClusterMode;
+
+	private:
+		FActorContainerID ContainerID;
+	};
+
 	virtual ~FWorldPartitionActorDesc() {}
 
 	inline const FGuid& GetGuid() const { return Guid; }
@@ -137,7 +153,8 @@ PRAGMA_ENABLE_DEPRECATION_WARNINGS
 
 	virtual bool IsContainerInstance() const { return false; }
 	virtual FName GetLevelPackage() const { return NAME_None; }
-	virtual bool GetContainerInstance(const UActorDescContainer*& OutLevelContainer, FTransform& OutLevelTransform, EContainerClusterMode& OutClusterMode) const { return false; }
+	bool GetContainerInstance(const UActorDescContainer*& OutLevelContainer, FTransform& OutLevelTransform, EContainerClusterMode& OutClusterMode) const;
+	virtual bool GetContainerInstance(FContainerInstance& OutContainerInstance) const { return false; }
 
 	FGuid GetContentBundleGuid() const;
 
diff --git a/Engine/Source/Runtime/Engine/Public/WorldPartition/WorldPartitionActorDescView.h b/Engine/Source/Runtime/Engine/Public/WorldPartition/WorldPartitionActorDescView.h
--- a/Engine/Source/Runtime/Engine/Public/WorldPartition/WorldPartitionActorDescView.h
+++ b/Engine/Source/Runtime/Engine/Public/WorldPartition/WorldPartitionActorDescView.h
@@ -5,14 +5,13 @@
 #include "UObject/SoftObjectPath.h"
 #include "UObject/TopLevelAssetPath.h"
 #include "Misc/Optional.h"
+#include "WorldPartition/WorldPartitionActorDesc.h"
 
 #if WITH_EDITOR
 class AActor;
-class FWorldPartitionActorDesc;
 class IStreamingGenerationErrorHandler;
 class UActorDescContainer;
 class UWorldPartition;
-enum class EContainerClusterMode : uint8;
 
 /**
  * A view on top of an actor desc, used to cache information that can be (potentially) different than the actor desc
@@ -66,7 +65,7 @@ PRAGMA_ENABLE_DEPRECATION_WARNINGS
 
 	bool IsContainerInstance() const;
 	FName GetLevelPackage() const;
-	bool GetContainerInstance(const UActorDescContainer*& OutLevelContainer, FTransform& OutLevelTransform, EContainerClusterMode& OutClusterMode) const;
+	bool GetContainerInstance(FWorldPartitionActorDesc::FContainerInstance& OutContainerInstance) const;
 
 	void CheckForErrors(IStreamingGenerationErrorHandler* ErrorHandler) const;
 
