Small clean up and tooltip clarification.

diff --git a/Engine/Shaders/Private/VolumetricCloud.usf b/Engine/Shaders/Private/VolumetricCloud.usf
--- a/Engine/Shaders/Private/VolumetricCloud.usf
+++ b/Engine/Shaders/Private/VolumetricCloud.usf
@@ -640,7 +640,6 @@ void MainCommon(in FMaterialPixelParameters MaterialParameters, in float4 SvPosi
 		const float StepCount = float(IStepCount);
 		const float StepT = (TMax - TMin) / StepCount; // StepT is step distance in centimeters
 		const float dtMeters = StepT * CENTIMETER_TO_METER;
-		// TODO adaptive sample count based on distance
 
 		// This is the scattered sky light luminance assuming a uniform phase function.
 		// TODO Approximate some form of occlusion
diff --git a/Engine/Source/Runtime/Engine/Classes/Components/SkyAtmosphereComponent.h b/Engine/Source/Runtime/Engine/Classes/Components/SkyAtmosphereComponent.h
--- a/Engine/Source/Runtime/Engine/Classes/Components/SkyAtmosphereComponent.h
+++ b/Engine/Source/Runtime/Engine/Classes/Components/SkyAtmosphereComponent.h
@@ -140,7 +140,7 @@ class USkyAtmosphereComponent : public USceneComponent
 	
 
 
-	/** Scales the luminance of pixels representing the sky, i.e. not belonging to any surface. */
+	/** Scales the luminance of pixels representing the sky. This will impact the captured sky light. */
 	UPROPERTY(EditAnywhere, BlueprintReadOnly, interp, Category = "Art Direction", meta = (HideAlphaChannel))
 	FLinearColor SkyLuminanceFactor;
 
