Fix for shader compile error on mobile due to Texture3D members in the FSparseVolumeTexture struct.

#rnx

diff --git a/Engine/Shaders/Private/SparseVolumeTexture/SparseVolumeTextureCommon.ush b/Engine/Shaders/Private/SparseVolumeTexture/SparseVolumeTextureCommon.ush
--- a/Engine/Shaders/Private/SparseVolumeTexture/SparseVolumeTextureCommon.ush
+++ b/Engine/Shaders/Private/SparseVolumeTexture/SparseVolumeTextureCommon.ush
@@ -2,30 +2,17 @@
 
 #pragma once
 
-struct FSparseVolumeTexture
+struct FSparseVolumeTextureUniforms
 {
-	Texture3D<uint>	PageTable;
-	Texture3D		PhysicalTileDataA;
-	Texture3D		PhysicalTileDataB;
 	float3			PageTableSize;
 	float3			UVScale;
 	float3			UVBias;
 	float			TileSize;
 };
 
-FSparseVolumeTexture SparseVolumeTextureCreateStruct(
-	Texture3D<uint> PageTable, 
-	Texture3D PhysicalTileDataA, 
-	Texture3D PhysicalTileDataB, 
-	float TileSize,
-	float3 PageTableSize,
-	float3 UVScale,
-	float3 UVBias)
+FSparseVolumeTextureUniforms SparseVolumeTextureCreateUniforms(float TileSize, float3 PageTableSize, float3 UVScale, float3 UVBias)
 {
-	FSparseVolumeTexture Result;
-	Result.PageTable = PageTable;
-	Result.PhysicalTileDataA = PhysicalTileDataA;
-	Result.PhysicalTileDataB = PhysicalTileDataB;
+	FSparseVolumeTextureUniforms Result;
 	Result.PageTableSize = PageTableSize;
 	Result.UVScale = UVScale;
 	Result.UVBias = UVBias;
@@ -33,17 +20,9 @@ FSparseVolumeTexture SparseVolumeTextureCreateStruct(
 	return Result;
 }
 
-FSparseVolumeTexture SparseVolumeTextureCreateStruct(
-	Texture3D<uint> PageTable,
-	Texture3D PhysicalTileDataA,
-	Texture3D PhysicalTileDataB,
-	uint4 Packed0,
-	uint4 Packed1)
+FSparseVolumeTextureUniforms SparseVolumeTextureUnpackUniforms(uint4 Packed0, uint4 Packed1)
 {
-	FSparseVolumeTexture Result;
-	Result.PageTable = PageTable;
-	Result.PhysicalTileDataA = PhysicalTileDataA;
-	Result.PhysicalTileDataB = PhysicalTileDataB;
+	FSparseVolumeTextureUniforms Result;
 	Result.UVScale = asfloat(Packed0.xyz);
 	Result.UVBias = asfloat(Packed1.xyz);
 	Result.TileSize = asfloat(Packed1.w);
@@ -68,23 +47,23 @@ float3 SparseVolumeTextureGetVoxelCoord(uint PackedPhysicalTileCoord, float Tile
 	return VoxelCoord;
 }
 
-float3 SparseVolumeTextureSamplePageTable(FSparseVolumeTexture SparseVolumeTexture, float3 UVW)
+float3 SparseVolumeTextureSamplePageTable(Texture3D<uint> PageTable, FSparseVolumeTextureUniforms Uniforms, float3 UVW)
 {
-	const float3 UVWScaledBiased = UVW * SparseVolumeTexture.UVScale + SparseVolumeTexture.UVBias;
-	const float3 PageTableCoord = (UVWScaledBiased * SparseVolumeTexture.PageTableSize);
+	const float3 UVWScaledBiased = UVW * Uniforms.UVScale + Uniforms.UVBias;
+	const float3 PageTableCoord = (UVWScaledBiased * Uniforms.PageTableSize);
 	
-	const uint PackedPhysicalTileCoord = SparseVolumeTexture.PageTable.Load(int4(int3(PageTableCoord), 0)).x;
+	const uint PackedPhysicalTileCoord = PageTable.Load(int4(int3(PageTableCoord), 0)).x;
 	
-	const float3 VoxelCoord = SparseVolumeTextureGetVoxelCoord(PackedPhysicalTileCoord, SparseVolumeTexture.TileSize, PageTableCoord);
+	const float3 VoxelCoord = SparseVolumeTextureGetVoxelCoord(PackedPhysicalTileCoord, Uniforms.TileSize, PageTableCoord);
 	return VoxelCoord;
 }
 
-float4 SparseVolumeTextureSamplePhysicalTileData(FSparseVolumeTexture SparseVolumeTexture, float3 VoxelCoord, int PhysicalTileDataIndex)
+float4 SparseVolumeTextureSamplePhysicalTileData(Texture3D PhysicalTileDataA, Texture3D PhysicalTileDataB, float3 VoxelCoord, int PhysicalTileDataIndex)
 {
 	switch (PhysicalTileDataIndex)
 	{
-	case 0: return SparseVolumeTexture.PhysicalTileDataA.Load(int4(int3(VoxelCoord), 0));
-	case 1: return SparseVolumeTexture.PhysicalTileDataB.Load(int4(int3(VoxelCoord), 0));
+	case 0: return PhysicalTileDataA.Load(int4(int3(VoxelCoord), 0));
+	case 1: return PhysicalTileDataB.Load(int4(int3(VoxelCoord), 0));
 	default: return 0.0f;
 	}
 } 
diff --git a/Engine/Source/Runtime/Engine/Private/Materials/HLSLMaterialTranslator.cpp b/Engine/Source/Runtime/Engine/Private/Materials/HLSLMaterialTranslator.cpp
--- a/Engine/Source/Runtime/Engine/Private/Materials/HLSLMaterialTranslator.cpp
+++ b/Engine/Source/Runtime/Engine/Private/Materials/HLSLMaterialTranslator.cpp
@@ -2941,7 +2941,7 @@ const TCHAR* FHLSLMaterialTranslator::HLSLTypeStringDeriv(EMaterialValueType Typ
 	case MCT_MaterialAttributes:	return TEXT("FMaterialAttributes");
 	case MCT_TextureExternal:		return TEXT("TextureExternal");
 	case MCT_TextureVirtual:		return TEXT("TextureVirtual");
-	case MCT_SparseVolumeTexture:	return TEXT("FSparseVolumeTexture");
+	case MCT_SparseVolumeTexture:	return TEXT("FSparseVolumeTextureUniforms");
 	case MCT_VTPageTableResult:		return TEXT("VTPageTableResult");
 	case MCT_ShadingModel:			return TEXT("uint");
 	case MCT_UInt:					return TEXT("uint");
@@ -3517,14 +3517,7 @@ int32 FHLSLMaterialTranslator::AccessUniformExpression(int32 Index)
 	else if (CodeChunk.Type == MCT_SparseVolumeTexture)
 	{
 		int32 TextureInputIndex = UniformTextureExpressions[(uint32)EMaterialTextureParameterType::SparseVolume].AddUnique(TextureUniformExpression);
-		FormattedCode.Appendf(
-			TEXT("SparseVolumeTextureCreateStruct("
-				"Material.SparseVolumeTexturePageTable_%d,"
-				"Material.SparseVolumeTexturePhysicalA_%d,"
-				"Material.SparseVolumeTexturePhysicalB_%d,"
-				"Material.SVTPackedUniform[%d*2],"
-				"Material.SVTPackedUniform[%d*2+1])"),
-			TextureInputIndex, TextureInputIndex, TextureInputIndex, TextureInputIndex, TextureInputIndex);
+		FormattedCode.Appendf(TEXT("SparseVolumeTextureUnpackUniforms(Material.SVTPackedUniform[%d*2], Material.SVTPackedUniform[%d*2+1])"), TextureInputIndex, TextureInputIndex);
 	}
 	else
 	{
@@ -12612,9 +12605,26 @@ int32 FHLSLMaterialTranslator::SparseVolumeTextureSamplePageTable(int32 SparseVo
 		return Errorf(TEXT("FHLSLMaterialTranslator::SparseVolumeTextureSamplePageTable expects MCT_Float3 as UVW input but was passed %s ERROR."), DescribeType(UVWType));
 	}
 
+	FMaterialUniformExpression* UniformExpression = GetParameterUniformExpression(SparseVolumeTextureIndex);
+	if (UniformExpression == nullptr)
+	{
+		return Errorf(TEXT("Unable to find SVT uniform expression."));
+	}
+	FMaterialUniformExpressionTexture* TextureUniformExpression = UniformExpression->GetTextureUniformExpression();
+	if (TextureUniformExpression == nullptr)
+	{
+		return Errorf(TEXT("The provided uniform expression is not a texture"));
+	}
+
+	// Make sure the SVT has been added to UniformTextureExpressions
+	AccessUniformExpression(SparseVolumeTextureIndex);
+
+	int32 SVTReferenceIndex = UniformTextureExpressions[(uint32)EMaterialTextureParameterType::SparseVolume].Find(TextureUniformExpression);
+	check(UniformTextureExpressions[(uint32)EMaterialTextureParameterType::SparseVolume].IsValidIndex(SVTReferenceIndex));
+
 	AddEstimatedTextureSample();
-	FString SampleCode = FString::Printf(TEXT("SparseVolumeTextureSamplePageTable(%s, %s)"),
-		*GetParameterCode(SparseVolumeTextureIndex), *GetParameterCode(UVWIndex));
+	FString SampleCode = FString::Printf(TEXT("SparseVolumeTextureSamplePageTable(Material.SparseVolumeTexturePageTable_%d, %s, %s)"),
+		SVTReferenceIndex , *GetParameterCode(SparseVolumeTextureIndex), *GetParameterCode(UVWIndex));
 	return AddCodeChunk(MCT_Float3, *SampleCode);
 }
 
@@ -12643,9 +12653,26 @@ int32 FHLSLMaterialTranslator::SparseVolumeTextureSamplePhysicalTileData(int32 S
 		return Errorf(TEXT("FHLSLMaterialTranslator::SparseVolumeTextureSamplePhysicalTileData expects MCT_Float1 or MCT_UInt1 as PhysicalTileDataIndex input but was passed %s ERROR."), DescribeType(IndexType));
 	}
 
+	FMaterialUniformExpression* UniformExpression = GetParameterUniformExpression(SparseVolumeTextureIndex);
+	if (UniformExpression == nullptr)
+	{
+		return Errorf(TEXT("Unable to find SVT uniform expression."));
+	}
+	FMaterialUniformExpressionTexture* TextureUniformExpression = UniformExpression->GetTextureUniformExpression();
+	if (TextureUniformExpression == nullptr)
+	{
+		return Errorf(TEXT("The provided uniform expression is not a texture"));
+	}
+
+	// Make sure the SVT has been added to UniformTextureExpressions
+	AccessUniformExpression(SparseVolumeTextureIndex);
+
+	int32 SVTReferenceIndex = UniformTextureExpressions[(uint32)EMaterialTextureParameterType::SparseVolume].Find(TextureUniformExpression);
+	check(UniformTextureExpressions[(uint32)EMaterialTextureParameterType::SparseVolume].IsValidIndex(SVTReferenceIndex));
+
 	AddEstimatedTextureSample();
-	FString SampleCode = FString::Printf(TEXT("SparseVolumeTextureSamplePhysicalTileData(%s, %s, %s)"),
-		*GetParameterCode(SparseVolumeTextureIndex), *GetParameterCode(VoxelCoordIndex), *GetParameterCode(PhysicalTileDataIdxIndex));
+	FString SampleCode = FString::Printf(TEXT("SparseVolumeTextureSamplePhysicalTileData(Material.SparseVolumeTexturePhysicalA_%d, Material.SparseVolumeTexturePhysicalB_%d, %s, %s)"),
+		SVTReferenceIndex, SVTReferenceIndex, *GetParameterCode(VoxelCoordIndex), *GetParameterCode(PhysicalTileDataIdxIndex));
 	return AddCodeChunk(MCT_Float4, *SampleCode);
 }
 
