Add missing nullptr check (PreEditChange(nullptr) gets called in ULevelFactory::FactorCreateText

#rnx

diff --git a/Engine/Source/Runtime/Engine/Private/LevelInstance/LevelInstanceActor.cpp b/Engine/Source/Runtime/Engine/Private/LevelInstance/LevelInstanceActor.cpp
--- a/Engine/Source/Runtime/Engine/Private/LevelInstance/LevelInstanceActor.cpp
+++ b/Engine/Source/Runtime/Engine/Private/LevelInstance/LevelInstanceActor.cpp
@@ -170,7 +170,7 @@ void ALevelInstance::PostEditUndo()
 
 void ALevelInstance::PreEditChange(FProperty* PropertyThatWillChange)
 {
-	const bool bWorldAssetChange = PropertyThatWillChange->GetFName() == GET_MEMBER_NAME_CHECKED(ALevelInstance, WorldAsset);
+	const bool bWorldAssetChange = PropertyThatWillChange && (PropertyThatWillChange->GetFName() == GET_MEMBER_NAME_CHECKED(ALevelInstance, WorldAsset));
 	LevelInstanceActorImpl.PreEditChange(PropertyThatWillChange, bWorldAssetChange, [this](FProperty* Property) { Super::PreEditChange(Property); });
 }
 
