Partitioned level instances could flag a main world persistent level as partitioned by error. If the level was OFPA enabled, this would lead into a state where the level wouldn't load its actors.

diff --git a/Engine/Source/Runtime/Engine/Private/WorldPartition/WorldPartition.cpp b/Engine/Source/Runtime/Engine/Private/WorldPartition/WorldPartition.cpp
--- a/Engine/Source/Runtime/Engine/Private/WorldPartition/WorldPartition.cpp
+++ b/Engine/Source/Runtime/Engine/Private/WorldPartition/WorldPartition.cpp
@@ -918,7 +918,7 @@ void UWorldPartition::OnWorldMatchStarting()
 #if WITH_EDITOR
 UWorldPartition* UWorldPartition::CreateOrRepairWorldPartition(AWorldSettings* WorldSettings, TSubclassOf<UWorldPartitionEditorHash> EditorHashClass, TSubclassOf<UWorldPartitionRuntimeHash> RuntimeHashClass)
 {
-	UWorld* World = WorldSettings->GetWorld();
+	UWorld* OuterWorld = WorldSettings->GetTypedOuter<UWorld>();
 	UWorldPartition* WorldPartition = WorldSettings->GetWorldPartition();
 
 	if (!WorldPartition)
@@ -940,14 +940,14 @@ UWorldPartition* UWorldPartition::CreateOrRepairWorldPartition(AWorldSettings* W
 
 		WorldPartition->DefaultHLODLayer = UHLODLayer::GetEngineDefaultHLODLayersSetup();
 
-		AWorldDataLayers* WorldDataLayers = World->GetWorldDataLayers();
+		AWorldDataLayers* WorldDataLayers = OuterWorld->GetWorldDataLayers();
 		if (!WorldDataLayers)
 		{
-			WorldDataLayers = AWorldDataLayers::Create(World);
-			World->SetWorldDataLayers(WorldDataLayers);
+			WorldDataLayers = AWorldDataLayers::Create(OuterWorld);
+			OuterWorld->SetWorldDataLayers(WorldDataLayers);
 		}
 
-		FWorldPartitionMiniMapHelper::GetWorldPartitionMiniMap(World, true);
+		FWorldPartitionMiniMapHelper::GetWorldPartitionMiniMap(OuterWorld, true);
 	}
 
 	if (!WorldPartition->EditorHash)
@@ -972,7 +972,7 @@ UWorldPartition* UWorldPartition::CreateOrRepairWorldPartition(AWorldSettings* W
 		WorldPartition->RuntimeHash->SetDefaultValues();
 	}
 
-	World->PersistentLevel->bIsPartitioned = true;
+	OuterWorld->PersistentLevel->bIsPartitioned = true;
 
 	return WorldPartition;
 }
