Fixup Data Layer asset paths in WorldPartitionActorDescs to properly support renamed Data Layer Assets.

diff --git a/Engine/Source/Runtime/Engine/Private/WorldPartition/DataLayer/WorldDataLayersActorDesc.cpp b/Engine/Source/Runtime/Engine/Private/WorldPartition/DataLayer/WorldDataLayersActorDesc.cpp
--- a/Engine/Source/Runtime/Engine/Private/WorldPartition/DataLayer/WorldDataLayersActorDesc.cpp
+++ b/Engine/Source/Runtime/Engine/Private/WorldPartition/DataLayer/WorldDataLayersActorDesc.cpp
@@ -4,6 +4,7 @@
 
 #if WITH_EDITOR
 
+#include "WorldPartition/WorldPartitionActorDescUtils.h"
 #include "WorldPartition/DataLayer/WorldDataLayers.h"
 #include "WorldPartition/DataLayer/DataLayerInstanceWithAsset.h"
 #include "WorldPartition/WorldPartitionActorContainerID.h"
@@ -32,6 +33,13 @@ FArchive& operator<<(FArchive& Ar, FDataLayerInstanceDesc& Desc)
 			Ar << Desc.bIsRuntime << Desc.ShortName;
 		}
 	}
+
+	// Fixup redirected data layer asset path
+	if (Ar.IsLoading() && Desc.bIsUsingAsset)
+	{
+		FWorldPartitionActorDescUtils::FixupRedirectedAssetPath(Desc.AssetPath);
+	}
+
 	return Ar;
 }
 
diff --git a/Engine/Source/Runtime/Engine/Private/WorldPartition/WorldPartitionActorDesc.cpp b/Engine/Source/Runtime/Engine/Private/WorldPartition/WorldPartitionActorDesc.cpp
--- a/Engine/Source/Runtime/Engine/Private/WorldPartition/WorldPartitionActorDesc.cpp
+++ b/Engine/Source/Runtime/Engine/Private/WorldPartition/WorldPartitionActorDesc.cpp
@@ -4,6 +4,7 @@
 #include "WorldPartition/WorldPartitionActorDesc.h"
 #include "Misc/Paths.h"
 #include "WorldPartition/WorldPartitionActorDescArchive.h"
+#include "WorldPartition/WorldPartitionActorDescUtils.h"
 
 #include "UObject/MetaData.h"
 #include "AssetRegistry/AssetData.h"
@@ -533,6 +534,15 @@ PRAGMA_ENABLE_DEPRECATION_WARNINGS
 			}
 		}
 	}
+
+	// Fixup redirected data layer asset paths
+	if (Ar.IsLoading() && bIsUsingDataLayerAsset)
+	{
+		for (FName& DataLayer : DataLayers)
+		{
+			FWorldPartitionActorDescUtils::FixupRedirectedAssetPath(DataLayer);
+		}
+	}
 }
 
 FBox FWorldPartitionActorDesc::GetEditorBounds() const
diff --git a/Engine/Source/Runtime/Engine/Private/WorldPartition/WorldPartitionActorDescUtils.cpp b/Engine/Source/Runtime/Engine/Private/WorldPartition/WorldPartitionActorDescUtils.cpp
--- a/Engine/Source/Runtime/Engine/Private/WorldPartition/WorldPartitionActorDescUtils.cpp
+++ b/Engine/Source/Runtime/Engine/Private/WorldPartition/WorldPartitionActorDescUtils.cpp
@@ -235,5 +235,56 @@ bool FWorldPartitionActorDescUtils::ValidateActorDescClass(FWorldPartitionActorD
 	}
 
 	return true;
-};
-#endif
\ No newline at end of file
+}
+
+bool FWorldPartitionActorDescUtils::FixupRedirectedAssetPath(FName& InOutAssetPath)
+{
+	const FAssetData* AssetData;
+	IAssetRegistry& AssetRegistry = FModuleManager::LoadModuleChecked<FAssetRegistryModule>(TEXT("AssetRegistry")).Get();
+
+	FTopLevelAssetPath AssetPath = FTopLevelAssetPath(InOutAssetPath.ToString());
+
+	for (;;)
+	{
+		TArray<FAssetData> Assets;							
+		AssetRegistry.ScanFilesSynchronous({ AssetPath.GetPackageName().ToString() }, /*bForceRescan*/false);
+		AssetRegistry.GetAssetsByPackageName(AssetPath.GetPackageName(), Assets, /*bIncludeOnlyOnDiskAssets*/true);
+
+		if (!Assets.Num())
+		{
+			UE_LOG(LogWorldPartition, Warning, TEXT("Failed to find assets for asset path '%s'"), *AssetPath.ToString());
+			return false;
+		}
+
+		AssetData = Assets.FindByPredicate([&AssetPath](const FAssetData& AssetData)
+		{ 
+			return (AssetData.ToSoftObjectPath().GetAssetPath() == AssetPath);
+		});
+
+		if (!AssetData)
+		{
+			UE_LOG(LogWorldPartition, Warning, TEXT("Failed to find asset for asset path '%s'"), *AssetPath.ToString());
+			return false;
+		}
+
+		if (!AssetData->IsRedirector())
+		{
+			break;
+		}
+
+		FString DestinationObjectPath;
+		if (!AssetData->GetTagValue(TEXT("DestinationObject"), DestinationObjectPath))
+		{
+			UE_LOG(LogWorldPartition, Warning, TEXT("Failed to follow redirector for '%s'"), *AssetPath.ToString());
+			return false;
+		}
+
+		// Update asset path
+		AssetPath = FTopLevelAssetPath(DestinationObjectPath);
+	}
+
+	InOutAssetPath = FName(AssetPath.ToString());
+	return true;
+}
+
+#endif
diff --git a/Engine/Source/Runtime/Engine/Public/WorldPartition/WorldPartitionActorDescUtils.h b/Engine/Source/Runtime/Engine/Public/WorldPartition/WorldPartitionActorDescUtils.h
--- a/Engine/Source/Runtime/Engine/Public/WorldPartition/WorldPartitionActorDescUtils.h
+++ b/Engine/Source/Runtime/Engine/Public/WorldPartition/WorldPartitionActorDescUtils.h
@@ -79,5 +79,12 @@ struct ENGINE_API FWorldPartitionActorDescUtils
 	 * @return					Whether the actor descriptor's class is valid or not.
 	 */
 	static bool ValidateActorDescClass(FWorldPartitionActorDesc* InActorDesc);
+
+	/**
+	 * Resolves the provided asset path using asset redirectors.
+	 * @param	InOutAssetPath	Asset path to resolve [In/Out].
+	 * @return					False if an error occured while trying to resolve, else true.
+	 */
+	static bool FixupRedirectedAssetPath(FName& InOutAssetPath);
 };
 #endif
\ No newline at end of file
