Allow the creation of light weight instance managers in a grid set by the cvar `LWI.Editor.GridSize`. This involves changing LWI to use relative locations for the instances and fixing some bugs with construction. Also fix the returned position from ActorInstanceHandles. Future work should move this to use APartitionActors.

diff --git a/Engine/Plugins/LightWeightInstancesEditor/Source/LightWeightInstancesEditor/Private/LightWeightInstancesEditor.cpp b/Engine/Plugins/LightWeightInstancesEditor/Source/LightWeightInstancesEditor/Private/LightWeightInstancesEditor.cpp
--- a/Engine/Plugins/LightWeightInstancesEditor/Source/LightWeightInstancesEditor/Private/LightWeightInstancesEditor.cpp
+++ b/Engine/Plugins/LightWeightInstancesEditor/Source/LightWeightInstancesEditor/Private/LightWeightInstancesEditor.cpp
@@ -113,7 +113,7 @@ void FLightWeightInstancesEditorModule::ConvertActorsToLWIsUIAction(const TArray
 		TArray<const UDataLayerInstance*> DataLayerInstances = Actor->GetDataLayerInstances();
 		const UDataLayerInstance* DataLayerInstance = DataLayerInstances.Num() > 0 ? DataLayerInstances[0] : nullptr;
 
-		ALightWeightInstanceManager* Manager = FLightWeightInstanceSubsystem::Get().FindOrAddLightWeightInstanceManager(Actor->GetClass(), DataLayerInstance, Actor->GetWorld());
+		ALightWeightInstanceManager* Manager = FLightWeightInstanceSubsystem::Get().FindOrAddLightWeightInstanceManager(Actor->GetClass(), DataLayerInstance, Actor->GetWorld(), Actor->GetActorLocation());
 		check(Manager);
 		UDataLayerEditorSubsystem::Get()->OnActorDataLayersChanged().Broadcast(Manager);
 
diff --git a/Engine/Source/Runtime/Engine/Classes/GameFramework/LightWeightInstanceSubsystem.h b/Engine/Source/Runtime/Engine/Classes/GameFramework/LightWeightInstanceSubsystem.h
--- a/Engine/Source/Runtime/Engine/Classes/GameFramework/LightWeightInstanceSubsystem.h
+++ b/Engine/Source/Runtime/Engine/Classes/GameFramework/LightWeightInstanceSubsystem.h
@@ -2,22 +2,15 @@
 
 #pragma once
 
-#include "CoreMinimal.h"
 #include "UObject/ObjectMacros.h"
-#include "UObject/UObjectBaseUtility.h"
-#include "UObject/Object.h"
-#include "Engine/EngineTypes.h"
-#include "Engine/EngineBaseTypes.h"
 #include "GameFramework/LightWeightInstanceManager.h"
 #include "Templates/SharedPointer.h"
 
 #include "Actor.h"
 
-//#include "LightWeightInstanceSubsystem.generated.h"
-
 ENGINE_API DECLARE_LOG_CATEGORY_EXTERN(LogLightWeightInstance, Log, Warning);
 
-//DECLARE_DYNAMIC_DELEGATE_RetVal_OneParam(FActorInstanceHandle, FOnActorReady, FActorInstanceHandle, InHandle);
+class FAutoConsoleVariableRef;
 
 struct ENGINE_API FLightWeightInstanceSubsystem
 {
@@ -46,11 +39,20 @@ struct ENGINE_API FLightWeightInstanceSubsystem
 	ALightWeightInstanceManager* FindLightWeightInstanceManager(const FActorInstanceHandle& Handle) const;
 
 	// Returns the instance manager that handles actors of type ActorClass in level Level
-	ALightWeightInstanceManager* FindLightWeightInstanceManager(UClass* ActorClass, const UDataLayerInstance* Layer) const;
+	UE_DEPRECATED(5.3, "Use the version that takes in a position.")
+	ALightWeightInstanceManager* FindLightWeightInstanceManager(UClass* ActorClass, const UDataLayerInstance* Layer, UWorld* World) const;
+
+	// Returns the instance manager that handles instances of type Class that live in Level
+	UE_DEPRECATED(5.3, "Use the version that takes in a position.")
+	UFUNCTION(Server, Unreliable)
+	ALightWeightInstanceManager* FindOrAddLightWeightInstanceManager(UClass* ActorClass, const UDataLayerInstance* DataLayer, UWorld* World);
+
+	// Returns the instance manager that handles actors of type ActorClass in level Level
+	ALightWeightInstanceManager* FindLightWeightInstanceManager(UClass& ActorClass, UWorld& World, const FVector& InPos, const UDataLayerInstance* DataLayer = nullptr) const;
 
 	// Returns the instance manager that handles instances of type Class that live in Level
 	UFUNCTION(Server, Unreliable)
-	ALightWeightInstanceManager* FindOrAddLightWeightInstanceManager(UClass* ActorClass, const UDataLayerInstance* Layer, UWorld* World);
+	ALightWeightInstanceManager* FindOrAddLightWeightInstanceManager(UClass& ActorClass, UWorld& World, const FVector& InPos, const UDataLayerInstance* DataLayer = nullptr);
 
 	// Returns the actor specified by Handle. This may require loading and creating the actor object.
 	AActor* FetchActor(const FActorInstanceHandle& Handle);
@@ -100,6 +102,9 @@ struct ENGINE_API FLightWeightInstanceSubsystem
 		return nullptr;
 	}
 
+	// Helper that converts a position (world space) into a coordinate for the LWI grid.
+	static FInt32Vector3 ConvertPositionToCoord(const FVector & InPosition);
+
 protected:
 	// Returns the class of the instance manager best suited to support instances of type ActorClass
 	UClass* FindBestInstanceManagerClass(const UClass* ActorClass);
@@ -118,6 +123,10 @@ private:
 	UPROPERTY()
 	TArray<ALightWeightInstanceManager*> LWInstanceManagers;
 
+	// CVar variable that is the size of the grid managers are placed into.
+	static int32 LWIGridSize;
+	static FAutoConsoleVariableRef CVarLWIGridSize;
+
 #ifdef WITH_EDITOR
 private:
 	FDelegateHandle OnLevelActorAddedHandle;
diff --git a/Engine/Source/Runtime/Engine/Private/Engine/ActorInstanceHandle.cpp b/Engine/Source/Runtime/Engine/Private/Engine/ActorInstanceHandle.cpp
--- a/Engine/Source/Runtime/Engine/Private/Engine/ActorInstanceHandle.cpp
+++ b/Engine/Source/Runtime/Engine/Private/Engine/ActorInstanceHandle.cpp
@@ -28,7 +28,7 @@ FActorInstanceHandle::FActorInstanceHandle(AActor* InActor)
 #else
 		const UDataLayerInstance* DataLayerInstance = nullptr;
 #endif // WITH_EDITOR
-		Manager = FLightWeightInstanceSubsystem::Get().FindLightWeightInstanceManager(InActor->StaticClass(), DataLayerInstance);
+		Manager = FLightWeightInstanceSubsystem::Get().FindLightWeightInstanceManager(*InActor->StaticClass(), *InActor->GetWorld(), InActor->GetActorLocation(), DataLayerInstance);
 	}
 }
 
@@ -117,7 +117,7 @@ FVector FActorInstanceHandle::GetLocation() const
 
 	if (Manager.IsValid())
 	{
-		Manager->GetLocation(*this);
+		return Manager->GetLocation(*this);
 	}
 
 	return FVector();
@@ -132,7 +132,7 @@ FRotator FActorInstanceHandle::GetRotation() const
 
 	if (Manager.IsValid())
 	{
-		Manager->GetRotation(*this);
+		return Manager->GetRotation(*this);
 	}
 
 	return FRotator();
@@ -147,7 +147,7 @@ FTransform FActorInstanceHandle::GetTransform() const
 
 	if (Manager.IsValid())
 	{
-		Manager->GetTransform(*this);
+		return Manager->GetTransform(*this);
 	}
 
 	return FTransform();
@@ -172,7 +172,7 @@ FString FActorInstanceHandle::GetName() const
 
 	if (Manager.IsValid())
 	{
-		Manager->GetName(*this);
+		return Manager->GetName(*this);
 	}
 
 	return FString();
@@ -291,7 +291,7 @@ bool FActorInstanceHandle::operator==(const AActor* OtherActor) const
 	const UDataLayerInstance* DataLayerInstance = nullptr;
 #endif // WITH_EDITOR
 
-	if (ALightWeightInstanceManager* LWIManager = FLightWeightInstanceSubsystem::Get().FindLightWeightInstanceManager(OtherActor->StaticClass(), DataLayerInstance))
+	if (ALightWeightInstanceManager* LWIManager = FLightWeightInstanceSubsystem::Get().FindLightWeightInstanceManager(*OtherActor->StaticClass(), *OtherActor->GetWorld(), OtherActor->GetActorLocation(), DataLayerInstance))
 	{
 		if (Manager.Get() != LWIManager)
 		{
diff --git a/Engine/Source/Runtime/Engine/Private/LightWeightInstanceBlueprintFunctionLibrary.cpp b/Engine/Source/Runtime/Engine/Private/LightWeightInstanceBlueprintFunctionLibrary.cpp
--- a/Engine/Source/Runtime/Engine/Private/LightWeightInstanceBlueprintFunctionLibrary.cpp
+++ b/Engine/Source/Runtime/Engine/Private/LightWeightInstanceBlueprintFunctionLibrary.cpp
@@ -30,7 +30,7 @@ FActorInstanceHandle ULightWeightInstanceBlueprintFunctionLibrary::ConvertActorT
 #else
 	const UDataLayerInstance* DataLayerInstance = nullptr;
 #endif // WITH_EDITOR
-	if (ALightWeightInstanceManager* LWIManager = FLightWeightInstanceSubsystem::Get().FindOrAddLightWeightInstanceManager(InActor->GetClass(), DataLayerInstance, InActor->GetWorld()))
+	if (ALightWeightInstanceManager* LWIManager = FLightWeightInstanceSubsystem::Get().FindOrAddLightWeightInstanceManager(*InActor->GetClass(), *InActor->GetWorld(), InActor->GetActorLocation(), DataLayerInstance))
 	{
 		return LWIManager->ConvertActorToLightWeightInstance(InActor);
 	}
diff --git a/Engine/Source/Runtime/Engine/Private/LightWeightInstanceManager.cpp b/Engine/Source/Runtime/Engine/Private/LightWeightInstanceManager.cpp
--- a/Engine/Source/Runtime/Engine/Private/LightWeightInstanceManager.cpp
+++ b/Engine/Source/Runtime/Engine/Private/LightWeightInstanceManager.cpp
@@ -9,11 +9,13 @@
 #include UE_INLINE_GENERATED_CPP_BY_NAME(LightWeightInstanceManager)
 
 UActorInstanceHandleInterface::UActorInstanceHandleInterface(const FObjectInitializer& ObjectInitializer)
+	:Super(ObjectInitializer)
 {
 	// do nothing
 }
 
 ALightWeightInstanceManager::ALightWeightInstanceManager(const FObjectInitializer& ObjectInitializer)
+	:Super(ObjectInitializer)
 {
 	bReplicates = true;
 	AcceptedClass = AActor::StaticClass();
@@ -61,7 +63,9 @@ AActor* ALightWeightInstanceManager::FetchActorFromHandle(const FActorInstanceHa
 		}
 	}
 
-	ensure(Handle.Actor.IsValid());
+	// Unless we are on the server or the actor has been spawned, this ensure will fail.
+	// Commented out until there is more robust replication.
+	// ensure(Handle.Actor.IsValid());
 	return Handle.Actor.Get();
 }
 
@@ -78,16 +82,20 @@ AActor* ALightWeightInstanceManager::ConvertInstanceToActor(const FActorInstance
 		return nullptr;
 	}
 
-	FActorSpawnParameters SpawnParams;
-	SetSpawnParameters(SpawnParams);
-
-	AActor* NewActor = GetLevel()->GetWorld()->SpawnActor<AActor>(GetActorClassToSpawn(Handle), InstanceTransforms[Handle.GetInstanceIndex()], SpawnParams);
-	check(NewActor);
+	AActor* NewActor = nullptr;
+	// Only spawn actors on the server so they are replicated to the clients. Otherwise we'll end up with multiples.
+	if (HasAuthority())
+	{
+		FActorSpawnParameters SpawnParams;
+		SetSpawnParameters(SpawnParams);
+		NewActor = GetLevel()->GetWorld()->SpawnActor<AActor>(GetActorClassToSpawn(Handle), GetTransform(Handle), SpawnParams);
+		check(NewActor);
 
-	Handle.Actor = NewActor;
-	Actors.Add(Handle.GetInstanceIndex(), NewActor);
+		Handle.Actor = NewActor;
+		Actors.Add(Handle.GetInstanceIndex(), NewActor);
 
-	PostActorSpawn(Handle);
+		PostActorSpawn(Handle);
+	}
 
 	return NewActor;
 }
@@ -222,32 +230,12 @@ bool ALightWeightInstanceManager::FindActorForHandle(const FActorInstanceHandle&
 
 FVector ALightWeightInstanceManager::GetLocation(const FActorInstanceHandle& Handle) const
 {
-	if (FindActorForHandle(Handle))
-	{
-		return Handle.Actor->GetActorLocation();
-	}
-
-	if (ensure(IsIndexValid(Handle.GetInstanceIndex())))
-	{
-		return InstanceTransforms[Handle.GetInstanceIndex()].GetTranslation();
-	}
-
-	return FVector();
+	return GetTransform(Handle).GetTranslation();
 }
 
 FRotator ALightWeightInstanceManager::GetRotation(const FActorInstanceHandle& Handle) const
 {
-	if (FindActorForHandle(Handle))
-	{
-		return Handle.Actor->GetActorRotation();
-	}
-
-	if (ensure(IsIndexValid(Handle.GetInstanceIndex())))
-	{
-		return InstanceTransforms[Handle.GetInstanceIndex()].Rotator();
-	}
-
-	return FRotator();
+	return GetTransform(Handle).Rotator();
 }
 
 FTransform ALightWeightInstanceManager::GetTransform(const FActorInstanceHandle& Handle) const
@@ -259,7 +247,8 @@ FTransform ALightWeightInstanceManager::GetTransform(const FActorInstanceHandle&
 
 	if (ensure(IsIndexValid(Handle.GetInstanceIndex())))
 	{
-		return InstanceTransforms[Handle.GetInstanceIndex()];
+		// Return in world space.
+		return InstanceTransforms[Handle.GetInstanceIndex()] * GetActorTransform();
 	}
 
 	return FTransform();
@@ -335,7 +324,8 @@ void ALightWeightInstanceManager::GrowDataArrays()
 
 void ALightWeightInstanceManager::UpdateDataAtIndex(FLWIData* InData, int32 Index)
 {
-	InstanceTransforms[Index] = InData->Transform;
+	// Convert to local space.
+	InstanceTransforms[Index] = InData->Transform * GetActorTransform().Inverse();
 }
 
 void ALightWeightInstanceManager::RemoveInstance(const int32 Index)
diff --git a/Engine/Source/Runtime/Engine/Private/LightWeightInstanceStaticMeshManager.cpp b/Engine/Source/Runtime/Engine/Private/LightWeightInstanceStaticMeshManager.cpp
--- a/Engine/Source/Runtime/Engine/Private/LightWeightInstanceStaticMeshManager.cpp
+++ b/Engine/Source/Runtime/Engine/Private/LightWeightInstanceStaticMeshManager.cpp
@@ -17,21 +17,18 @@
 #endif // WITH_EDITOR
 
 ALightWeightInstanceStaticMeshManager::ALightWeightInstanceStaticMeshManager(const FObjectInitializer& ObjectInitializer)
+	:Super(ObjectInitializer)
 {
-	InstancedStaticMeshComponent = ObjectInitializer.CreateDefaultSubobject<UHierarchicalInstancedStaticMeshComponent>(this, TEXT("InstancedStaticMeshComponent0"));
+	InstancedStaticMeshComponent = CreateDefaultSubobject<UHierarchicalInstancedStaticMeshComponent>(TEXT("InstancedStaticMeshComponent0"));
+	SetRootComponent(InstancedStaticMeshComponent);
+	AddInstanceComponent(InstancedStaticMeshComponent);
 
 	if (StaticMesh.IsValid())
 	{
 		OnStaticMeshSet();
 	}
-	SetInstancedStaticMeshParams();
 
-	RootComponent = InstancedStaticMeshComponent;
-	AddInstanceComponent(InstancedStaticMeshComponent);
-	if (GetWorld())
-	{
-		InstancedStaticMeshComponent->RegisterComponent();
-	}
+	SetInstancedStaticMeshParams();
 }
 
 void ALightWeightInstanceStaticMeshManager::SetRepresentedClass(UClass* ActorClass)
@@ -58,12 +55,12 @@ void ALightWeightInstanceStaticMeshManager::SetRepresentedClass(UClass* ActorCla
 
 int32 ALightWeightInstanceStaticMeshManager::ConvertCollisionIndexToLightWeightIndex(int32 InIndex) const
 {
-	return RenderingIndicesToDataIndices[InIndex];
+	return InIndex;//RenderingIndicesToDataIndices[InIndex];
 }
 
 int32 ALightWeightInstanceStaticMeshManager::ConvertLightWeightIndexToCollisionIndex(int32 InIndex) const
 {
-	return DataIndicesToRenderingIndices[InIndex];
+	return InIndex;//DataIndicesToRenderingIndices[InIndex];
 }
 
 void ALightWeightInstanceStaticMeshManager::AddNewInstanceAt(FLWIData* InitData, int32 Index)
@@ -86,14 +83,17 @@ void ALightWeightInstanceStaticMeshManager::AddNewInstanceAt(FLWIData* InitData,
 	// Update the HISMC
 	if (InstancedStaticMeshComponent)
 	{
-		ensure(RenderingIdx == InstancedStaticMeshComponent->AddInstance(InitData->Transform, /*bWorldSpace*/true));
+		ensure(RenderingIdx == InstancedStaticMeshComponent->AddInstance(InstanceTransforms[Index], /*bWorldSpace*/false));
 	}
 }
 
 void ALightWeightInstanceStaticMeshManager::RemoveInstance(const int32 Index)
 {
-	RemoveInstanceFromRendering(Index);
+#if WITH_EDITOR
+	Modify();
+#endif
 
+	RemoveInstanceFromRendering(Index);
 	Super::RemoveInstance(Index);
 }
 
@@ -176,7 +176,7 @@ void ALightWeightInstanceStaticMeshManager::OnRep_Transforms()
 
 	if (InstancedStaticMeshComponent)
 	{
-		InstancedStaticMeshComponent->AddInstance(InstanceTransforms.Last(), /*bWorldSpace*/true);
+		InstancedStaticMeshComponent->AddInstance(InstanceTransforms.Last(), /*bWorldSpace*/false);
 	}
 }
 
@@ -212,6 +212,10 @@ void ALightWeightInstanceStaticMeshManager::SetStaticMeshFromActor(AActor* InAct
 
 void ALightWeightInstanceStaticMeshManager::OnStaticMeshSet()
 {
+#if WITH_EDITOR
+	Modify();
+#endif
+	
 	if (InstancedStaticMeshComponent)
 	{
 		EComponentMobility::Type Mobility = InstancedStaticMeshComponent->Mobility;
@@ -238,6 +242,10 @@ void ALightWeightInstanceStaticMeshManager::OnStaticMeshSet()
 
 void ALightWeightInstanceStaticMeshManager::ClearStaticMesh()
 {
+#if WITH_EDITOR
+	Modify();
+#endif
+	
 	StaticMesh = nullptr;
 	OnStaticMeshSet();
 }
@@ -281,7 +289,7 @@ bool ALightWeightInstanceStaticMeshManager::GetSMInstanceTransform(const FSMInst
 	if (RenderingIndicesToDataIndices.IsValidIndex(InstanceId.InstanceIndex))
 	{
 		const int32 DataIndex = RenderingIndicesToDataIndices[InstanceId.InstanceIndex];
-		OutInstanceTransform = bWorldSpace ? InstanceTransforms[DataIndex] : InstanceTransforms[DataIndex].GetRelativeTransform(InstancedStaticMeshComponent->GetComponentTransform());
+		OutInstanceTransform = bWorldSpace ? InstanceTransforms[DataIndex] * GetActorTransform() : InstanceTransforms[DataIndex];
 		return true;
 	}
 	
@@ -295,7 +303,7 @@ bool ALightWeightInstanceStaticMeshManager::SetSMInstanceTransform(const FSMInst
 	if (RenderingIndicesToDataIndices.IsValidIndex(InstanceId.InstanceIndex) && InstancedStaticMeshComponent->UpdateInstanceTransform(InstanceId.InstanceIndex, InstanceTransform, bWorldSpace, bMarkRenderStateDirty, bTeleport))
 	{
 		const int32 DataIndex = RenderingIndicesToDataIndices[InstanceId.InstanceIndex];
-		InstancedStaticMeshComponent->GetInstanceTransform(InstanceId.InstanceIndex, InstanceTransforms[DataIndex], /*bWorldSpace*/true);
+		InstancedStaticMeshComponent->GetInstanceTransform(InstanceId.InstanceIndex, InstanceTransforms[DataIndex], /*bWorldSpace*/false);
 		return true;
 	}
 
diff --git a/Engine/Source/Runtime/Engine/Private/LightWeightInstanceSubsystem.cpp b/Engine/Source/Runtime/Engine/Private/LightWeightInstanceSubsystem.cpp
--- a/Engine/Source/Runtime/Engine/Private/LightWeightInstanceSubsystem.cpp
+++ b/Engine/Source/Runtime/Engine/Private/LightWeightInstanceSubsystem.cpp
@@ -10,6 +10,16 @@ DEFINE_LOG_CATEGORY(LogLightWeightInstance);
 TSharedPtr<FLightWeightInstanceSubsystem> FLightWeightInstanceSubsystem::LWISubsystem;
 FCriticalSection FLightWeightInstanceSubsystem::GetFunctionCS;
 
+// CVar variable that is the size of the grid managers are placed into.
+int32 FLightWeightInstanceSubsystem::LWIGridSize = -1;
+FAutoConsoleVariableRef FLightWeightInstanceSubsystem::CVarLWIGridSize
+(
+TEXT("LWI.Editor.GridSize"),
+FLightWeightInstanceSubsystem::LWIGridSize,
+TEXT("Sets the size of a grid that LWI managers will be generated with."),
+ECVF_Cheat
+);
+
 FLightWeightInstanceSubsystem::FLightWeightInstanceSubsystem()
 {
 #if WITH_EDITOR
@@ -62,6 +72,21 @@ const ALightWeightInstanceManager* FLightWeightInstanceSubsystem::GetManagerAt(i
 	return LWInstanceManagers.IsValidIndex(Index) ? LWInstanceManagers[Index] : nullptr;
 }
 
+FInt32Vector3 FLightWeightInstanceSubsystem::ConvertPositionToCoord(const FVector & InPosition)
+{
+	const int32 GridSize = LWIGridSize;
+	if (GridSize > 0)
+	{
+		return FInt32Vector3(
+			FMath::FloorToInt(InPosition.X / (float)GridSize),
+			FMath::FloorToInt(InPosition.Y / (float)GridSize),
+			FMath::FloorToInt(InPosition.Z / (float)GridSize)
+			);
+	}
+
+	return FInt32Vector3::ZeroValue;
+}
+
 ALightWeightInstanceManager* FLightWeightInstanceSubsystem::FindLightWeightInstanceManager(const FActorInstanceHandle& Handle) const
 {
 	if (Handle.Manager.IsValid())
@@ -71,9 +96,13 @@ ALightWeightInstanceManager* FLightWeightInstanceSubsystem::FindLightWeightInsta
 
 	if (Handle.Actor.IsValid())
 	{
+		const FInt32Vector3 GridCoord = ConvertPositionToCoord(Handle.GetLocation());
+		
+		// see if we already have a match
 		for (ALightWeightInstanceManager* LWInstance : LWInstanceManagers)
 		{
-			if (Handle.Actor->GetClass() == LWInstance->GetRepresentedClass())
+			const FInt32Vector3 ManagerGridCoord = ConvertPositionToCoord(LWInstance->GetActorLocation());
+			if (Handle.Actor->GetClass() == LWInstance->GetRepresentedClass() && ManagerGridCoord == GridCoord)
 			{
 #if WITH_EDITOR
 				// make sure the data layers match
@@ -95,41 +124,29 @@ ALightWeightInstanceManager* FLightWeightInstanceSubsystem::FindLightWeightInsta
 	return nullptr;
 }
 
-ALightWeightInstanceManager* FLightWeightInstanceSubsystem::FindLightWeightInstanceManager(UClass* ActorClass, const UDataLayerInstance* Layer) const
+ALightWeightInstanceManager* FLightWeightInstanceSubsystem::FindLightWeightInstanceManager(UClass* ActorClass, const UDataLayerInstance* DataLayer, UWorld* World) const
 {
-	if (ActorClass == nullptr)
-	{
-		return nullptr;
-	}
+	check(World);
+	return FindLightWeightInstanceManager(*ActorClass, *World, FVector::ZeroVector, DataLayer);
+}
 
-	// see if we already have a match
-	for (ALightWeightInstanceManager* Instance : LWInstanceManagers)
-	{
-		if (Instance->GetRepresentedClass() == ActorClass)
-		{
-#if WITH_EDITOR
-			if (!Layer || (Instance->SupportsDataLayer() && Instance->ContainsDataLayer(Layer)))
-#endif // WITH_EDITOR
-			{
-				return Instance;
-			}
-		}
-	}
 
-	return nullptr;
+ALightWeightInstanceManager* FLightWeightInstanceSubsystem::FindOrAddLightWeightInstanceManager(UClass* ActorClass, const UDataLayerInstance* DataLayer, UWorld* World)
+{
+	check(World);
+	return FindOrAddLightWeightInstanceManager(*ActorClass, *World, FVector::ZeroVector, DataLayer);
 }
 
-ALightWeightInstanceManager* FLightWeightInstanceSubsystem::FindOrAddLightWeightInstanceManager(UClass* ActorClass, const UDataLayerInstance* DataLayer, UWorld* World)
+
+ALightWeightInstanceManager* FLightWeightInstanceSubsystem::FindLightWeightInstanceManager(UClass& ActorClass, UWorld& World, const FVector& InPos, const UDataLayerInstance* DataLayer) const
 {
-	if (ActorClass == nullptr || World == nullptr)
-	{
-		return nullptr;
-	}
-	
-	// see if we already have a match
+	const FInt32Vector3 GridCoord = ConvertPositionToCoord(InPos);
+
 	for (ALightWeightInstanceManager* InstanceManager : LWInstanceManagers)
 	{
-		if (InstanceManager->GetRepresentedClass() == ActorClass)
+		const FInt32Vector3 ManagerGridCoord = ConvertPositionToCoord(InstanceManager->GetActorLocation());
+		
+		if (InstanceManager->GetRepresentedClass() == &ActorClass && ManagerGridCoord == GridCoord)
 		{
 #if WITH_EDITOR
 			if (!DataLayer || (InstanceManager->SupportsDataLayer() && InstanceManager->ContainsDataLayer(DataLayer)))
@@ -140,10 +157,19 @@ ALightWeightInstanceManager* FLightWeightInstanceSubsystem::FindOrAddLightWeight
 		}
 	}
 
-	// we didn't find a match so we should add one.
+	return nullptr;
+}
 
+ALightWeightInstanceManager* FLightWeightInstanceSubsystem::FindOrAddLightWeightInstanceManager(UClass& ActorClass, UWorld& World, const FVector& InPos, const UDataLayerInstance* DataLayer)
+{
+	if (ALightWeightInstanceManager* FoundManager = FindLightWeightInstanceManager(ActorClass, World, InPos, DataLayer))
+	{
+		return FoundManager;
+	}
+	
+	// we didn't find a match so we should add one.
 	// Find the best base class to start from
-	UClass* BestMatchingClass = FindBestInstanceManagerClass(ActorClass);
+	UClass* BestMatchingClass = FindBestInstanceManagerClass(&ActorClass);
 	if (BestMatchingClass == nullptr)
 	{
 		return nullptr;
@@ -153,22 +179,36 @@ ALightWeightInstanceManager* FLightWeightInstanceSubsystem::FindOrAddLightWeight
 	SpawnParams.SpawnCollisionHandlingOverride = ESpawnActorCollisionHandlingMethod::AlwaysSpawn;
 	SpawnParams.ObjectFlags = RF_Transactional;
 
-	ALightWeightInstanceManager* NewInstanceManager = World->SpawnActor<ALightWeightInstanceManager>(BestMatchingClass, FTransform::Identity, SpawnParams);
-	NewInstanceManager->SetRepresentedClass(ActorClass);
+	FTransform ManagerTransform(FTransform::Identity);
+	if (LWIGridSize > 0)
+	{
+		FVector ManagerLocation(ConvertPositionToCoord(InPos));
+		ManagerLocation *= LWIGridSize;
+		ManagerLocation += FVector((float)LWIGridSize * 0.5f);
+		ManagerTransform.SetLocation(ManagerLocation);
+	}
 
-#if WITH_EDITOR
-	// Add the new manager to the DataLayer
-	if (DataLayer)
+	if (ALightWeightInstanceManager* NewInstanceManager = World.SpawnActor<ALightWeightInstanceManager>(BestMatchingClass, ManagerTransform, SpawnParams))
 	{
-		ensure(NewInstanceManager->SupportsDataLayer());
+		NewInstanceManager->SetRepresentedClass(&ActorClass);
 
-		NewInstanceManager->AddDataLayer(DataLayer);
-	}
+#if WITH_EDITOR
+		// Add the new manager to the DataLayer
+		if (DataLayer)
+		{
+			ensure(NewInstanceManager->SupportsDataLayer());
+			NewInstanceManager->AddDataLayer(DataLayer);
+		}
 #endif // WITH_EDITOR
 
-	check(LWInstanceManagers.Find(NewInstanceManager) != INDEX_NONE);
+		check(LWInstanceManagers.Find(NewInstanceManager) != INDEX_NONE);
 
-	return NewInstanceManager;
+		return NewInstanceManager;
+	}
+	else
+	{
+		return nullptr;
+	}
 }
 
 UClass* FLightWeightInstanceSubsystem::FindBestInstanceManagerClass(const UClass* InActorClass)
@@ -332,12 +372,15 @@ bool FLightWeightInstanceSubsystem::IsInLevel(const FActorInstanceHandle& Handle
 FActorInstanceHandle FLightWeightInstanceSubsystem::CreateNewLightWeightInstance(UClass* InActorClass, FLWIData* InitData, UDataLayerInstance* InLayer, UWorld* World)
 {
 	// Get or create a light weight instance for this class and data layer
-	if (ALightWeightInstanceManager* LWIManager = FindOrAddLightWeightInstanceManager(InActorClass, InLayer, World))
+	if (InitData)
 	{
-		// create an instance with the given data
-		int32 InstanceIdx = LWIManager->AddNewInstance(InitData);
-		InstanceIdx = LWIManager->ConvertInternalIndexToHandleIndex(InstanceIdx);
-		return FActorInstanceHandle(LWIManager, InstanceIdx);
+		if (ALightWeightInstanceManager* LWIManager = FindOrAddLightWeightInstanceManager(*InActorClass, *World, InitData->Transform.GetLocation(), InLayer))
+		{
+			// create an instance with the given data
+			int32 InstanceIdx = LWIManager->AddNewInstance(InitData);
+			InstanceIdx = LWIManager->ConvertInternalIndexToHandleIndex(InstanceIdx);
+			return FActorInstanceHandle(LWIManager, InstanceIdx);
+		}
 	}
 
 	return FActorInstanceHandle();
