Fix up WITH_EDITOR and WITH_EDITORONLY_DATA mismatches in headers included in Engine module PCH files. One small step on the path to running LowLevelTest programs with engine content.

diff --git a/Engine/Source/Runtime/Engine/Classes/Animation/AnimData/AnimDataModelNotifyCollector.h b/Engine/Source/Runtime/Engine/Classes/Animation/AnimData/AnimDataModelNotifyCollector.h
--- a/Engine/Source/Runtime/Engine/Classes/Animation/AnimData/AnimDataModelNotifyCollector.h
+++ b/Engine/Source/Runtime/Engine/Classes/Animation/AnimData/AnimDataModelNotifyCollector.h
@@ -11,7 +11,7 @@ enum class EAnimDataModelNotifyType : uint8;
 namespace UE {
 namespace Anim {
 
-#if WITH_EDITOR
+#if WITH_EDITORONLY_DATA
 
 /** Helper structure for keeping track of which notifies of type EAnimDataModelNotifyType are broadcasted
 between top-level EAnimDataModelNotifyType::BracketOpened and EAnimDataModelNotifyType::BracketClosed notifies */
@@ -79,7 +79,7 @@ protected:
 	bool bDataModified;
 };
 
-#endif // WITH_EDITOR
+#endif // WITH_EDITORONLY_DATA
 
 } // namespace Anim	
 
diff --git a/Engine/Source/Runtime/Engine/Classes/Curves/CurveBase.h b/Engine/Source/Runtime/Engine/Classes/Curves/CurveBase.h
--- a/Engine/Source/Runtime/Engine/Classes/Curves/CurveBase.h
+++ b/Engine/Source/Runtime/Engine/Classes/Curves/CurveBase.h
@@ -96,7 +96,10 @@ public:
 	virtual void GetAssetRegistryTags(TArray<FAssetRegistryTag>& OutTags) const override;
 	virtual void PostInitProperties() override;
 	virtual void PostLoad() override;
+
+#if WITH_EDITOR
 	virtual void PostEditChangeProperty(struct FPropertyChangedEvent& PropertyChangedEvent) override;
+#endif
 
 	UPROPERTY(VisibleAnywhere, Instanced, Category=ImportSettings)
 	TObjectPtr<class UAssetImportData> AssetImportData;
diff --git a/Engine/Source/Runtime/Engine/Classes/Engine/PoseWatchRenderData.h b/Engine/Source/Runtime/Engine/Classes/Engine/PoseWatchRenderData.h
--- a/Engine/Source/Runtime/Engine/Classes/Engine/PoseWatchRenderData.h
+++ b/Engine/Source/Runtime/Engine/Classes/Engine/PoseWatchRenderData.h
@@ -2,9 +2,10 @@
 
 #pragma once
 
-#if WITH_EDITOR
-
 #include "CoreMinimal.h"
+
+#if WITH_EDITORONLY_DATA
+
 #include "UObject/ObjectMacros.h"
 #include "UObject/Object.h"
 #include "BoneIndices.h"
diff --git a/Engine/Source/Runtime/Engine/Classes/Engine/World.h b/Engine/Source/Runtime/Engine/Classes/Engine/World.h
--- a/Engine/Source/Runtime/Engine/Classes/Engine/World.h
+++ b/Engine/Source/Runtime/Engine/Classes/Engine/World.h
@@ -1847,7 +1847,7 @@ private:
 	static uint32 CleanupWorldGlobalTag;
 
 public:
-#if WITH_EDITORONLY_DATA
+#if WITH_EDITOR
 	/** List of DDC async requests we need to wait on before we register components. Game thread only. */
 	TArray<TSharedPtr<FAsyncPreRegisterDDCRequest>> AsyncPreRegisterDDCRequests;
 #endif
diff --git a/Engine/Source/Runtime/Engine/Classes/GameFramework/Actor.h b/Engine/Source/Runtime/Engine/Classes/GameFramework/Actor.h
--- a/Engine/Source/Runtime/Engine/Classes/GameFramework/Actor.h
+++ b/Engine/Source/Runtime/Engine/Classes/GameFramework/Actor.h
@@ -999,7 +999,7 @@ private:
 	FString ActorLabel;
 #endif
 
-#if !WITH_EDITOR && ACTOR_HAS_LABELS
+#if !WITH_EDITORONLY_DATA && ACTOR_HAS_LABELS
 	FString ActorLabel;
 #endif
 
diff --git a/Engine/Source/Runtime/Engine/Classes/Materials/Material.h b/Engine/Source/Runtime/Engine/Classes/Materials/Material.h
--- a/Engine/Source/Runtime/Engine/Classes/Materials/Material.h
+++ b/Engine/Source/Runtime/Engine/Classes/Materials/Material.h
@@ -1208,8 +1208,11 @@ public:
 	/** Useful to customize rendering if that case (e.g. hide the object) */
 	ENGINE_API bool IsCompilingOrHadCompileError(ERHIFeatureLevel::Type InFeatureLevel);
 
-#if WITH_EDITOR
+#if WITH_EDITORONLY_DATA
 	ENGINE_API TConstArrayView<TObjectPtr<UMaterialExpression>> GetExpressions() const;
+#endif // WITH_EDITORONLY_DATA
+
+#if WITH_EDITOR
 	ENGINE_API TConstArrayView<TObjectPtr<UMaterialExpressionComment>> GetEditorComments() const;
 	ENGINE_API UMaterialExpressionExecBegin* GetExpressionExecBegin() const;
 	ENGINE_API UMaterialExpressionExecEnd* GetExpressionExecEnd() const;
diff --git a/Engine/Source/Runtime/Engine/Classes/Materials/MaterialFunctionInterface.h b/Engine/Source/Runtime/Engine/Classes/Materials/MaterialFunctionInterface.h
--- a/Engine/Source/Runtime/Engine/Classes/Materials/MaterialFunctionInterface.h
+++ b/Engine/Source/Runtime/Engine/Classes/Materials/MaterialFunctionInterface.h
@@ -183,12 +183,14 @@ public:
 	ENGINE_API UMaterialFunctionInterface* GetBaseFunctionInterface();
 	ENGINE_API const UMaterialFunctionInterface* GetBaseFunctionInterface() const;
 
-#if WITH_EDITOR
+#if WITH_EDITORONLY_DATA
 	ENGINE_API TConstArrayView<TObjectPtr<UMaterialExpression>> GetExpressions() const;
 
 	UE_DEPRECATED(5.1, "Use GetExpressions()")
 	inline TConstArrayView<TObjectPtr<UMaterialExpression>> GetFunctionExpressions() const { return GetExpressions(); }
+#endif // WITH_EDITORONLY_DATA
 
+#if WITH_EDITOR
 	ENGINE_API const FString& GetDescription() const;
 	ENGINE_API bool GetReentrantFlag() const;
 	ENGINE_API void SetReentrantFlag(bool bIsReentrant);
diff --git a/Engine/Source/Runtime/Engine/Classes/Materials/MaterialLayersFunctions.h b/Engine/Source/Runtime/Engine/Classes/Materials/MaterialLayersFunctions.h
--- a/Engine/Source/Runtime/Engine/Classes/Materials/MaterialLayersFunctions.h
+++ b/Engine/Source/Runtime/Engine/Classes/Materials/MaterialLayersFunctions.h
@@ -85,7 +85,7 @@ struct FMaterialLayersFunctionsEditorOnlyData
 	UPROPERTY(EditAnywhere, Category = MaterialLayers)
 	TArray<FGuid> DeletedParentLayerGuids;
 
-#if WITH_EDITOR
+#if WITH_EDITORONLY_DATA
 	FORCEINLINE bool operator==(const FMaterialLayersFunctionsEditorOnlyData& Other) const
 	{
 		if (LayerStates != Other.LayerStates ||
@@ -101,7 +101,9 @@ struct FMaterialLayersFunctionsEditorOnlyData
 	{
 		return !operator==(Other);
 	}
+#endif // WITH_EDITORONLY_DATA
 
+#if WITH_EDITOR
 	void Empty()
 	{
 		LayerStates.Empty();
diff --git a/Engine/Source/Runtime/Engine/Classes/Sound/SoundEffectSource.h b/Engine/Source/Runtime/Engine/Classes/Sound/SoundEffectSource.h
--- a/Engine/Source/Runtime/Engine/Classes/Sound/SoundEffectSource.h
+++ b/Engine/Source/Runtime/Engine/Classes/Sound/SoundEffectSource.h
@@ -58,7 +58,7 @@ public:
 
 protected:
 
-#if WITH_EDITORONLY_DATA
+#if WITH_EDITOR
 	virtual void PostEditChangeProperty(FPropertyChangedEvent& PropertyChangedEvent) override;
 #endif
 
diff --git a/Engine/Source/Runtime/Engine/Private/Curves/CurveBase.cpp b/Engine/Source/Runtime/Engine/Private/Curves/CurveBase.cpp
--- a/Engine/Source/Runtime/Engine/Private/Curves/CurveBase.cpp
+++ b/Engine/Source/Runtime/Engine/Private/Curves/CurveBase.cpp
@@ -262,6 +262,9 @@ void UCurveBase::PostLoad()
 		AssetImportData->SourceData = MoveTemp(Info);
 	}
 }
+#endif //WITH_EDITORONLY_DATA
+
+#if WITH_EDITOR
 
 void UCurveBase::PostEditChangeProperty(FPropertyChangedEvent& PropertyChangedEvent)
 {
@@ -269,5 +272,5 @@ void UCurveBase::PostEditChangeProperty(FPropertyChangedEvent& PropertyChangedEv
 
 	OnUpdateCurve.Broadcast(this, PropertyChangedEvent.ChangeType);
 }
-#endif //WITH_EDITORONLY_DATA
 
+#endif // WITH_EDITOR
diff --git a/Engine/Source/Runtime/Engine/Private/Engine/PoseWatchRenderData.cpp b/Engine/Source/Runtime/Engine/Private/Engine/PoseWatchRenderData.cpp
--- a/Engine/Source/Runtime/Engine/Private/Engine/PoseWatchRenderData.cpp
+++ b/Engine/Source/Runtime/Engine/Private/Engine/PoseWatchRenderData.cpp
@@ -1,8 +1,9 @@
 // Copyright Epic Games, Inc. All Rights Reserved.
 
-#if WITH_EDITOR
-
 #include "Engine/PoseWatchRenderData.h"
+
+#if WITH_EDITORONLY_DATA
+
 #include "Animation/BlendProfile.h"
 #include "Engine/PoseWatch.h"
 
@@ -93,4 +94,4 @@ const TArray<int32>& FAnimNodePoseWatch::GetParentIndices() const
 	return ParentIndices; 
 }
 
-#endif // WITH_EDITOR
\ No newline at end of file
+#endif // WITH_EDITORONLY_DATA
\ No newline at end of file
diff --git a/Engine/Source/Runtime/Engine/Private/Materials/Material.cpp b/Engine/Source/Runtime/Engine/Private/Materials/Material.cpp
--- a/Engine/Source/Runtime/Engine/Private/Materials/Material.cpp
+++ b/Engine/Source/Runtime/Engine/Private/Materials/Material.cpp
@@ -1386,12 +1386,14 @@ bool UMaterial::IsCompilingOrHadCompileError(ERHIFeatureLevel::Type InFeatureLev
 	return Res == nullptr || Res->GetGameThreadShaderMap() == nullptr;
 }
 
-#if WITH_EDITOR
+#if WITH_EDITORONLY_DATA
 TConstArrayView<TObjectPtr<UMaterialExpression>> UMaterial::GetExpressions() const
 {
 	return GetEditorOnlyData()->ExpressionCollection.Expressions;
 }
+#endif
 
+#if WITH_EDITOR
 TConstArrayView<TObjectPtr<UMaterialExpressionComment>> UMaterial::GetEditorComments() const
 {
 	return GetEditorOnlyData()->ExpressionCollection.EditorComments;
diff --git a/Engine/Source/Runtime/Engine/Private/Materials/MaterialExpressions.cpp b/Engine/Source/Runtime/Engine/Private/Materials/MaterialExpressions.cpp
--- a/Engine/Source/Runtime/Engine/Private/Materials/MaterialExpressions.cpp
+++ b/Engine/Source/Runtime/Engine/Private/Materials/MaterialExpressions.cpp
@@ -13870,7 +13870,7 @@ const UMaterialFunctionInterface* UMaterialFunctionInterface::GetBaseFunctionInt
 	return GetBaseFunction();
 }
 
-#if WITH_EDITOR
+#if WITH_EDITORONLY_DATA
 
 TConstArrayView<TObjectPtr<UMaterialExpression>> UMaterialFunctionInterface::GetExpressions() const
 {
@@ -13882,6 +13882,10 @@ TConstArrayView<TObjectPtr<UMaterialExpression>> UMaterialFunctionInterface::Get
 	return TConstArrayView<TObjectPtr<UMaterialExpression>>();
 }
 
+#endif // WITH_EDITORONLY_DATA
+
+#if WITH_EDITOR
+
 const FString& UMaterialFunctionInterface::GetDescription() const
 {
 	const UMaterialFunction* BaseFunction = GetBaseFunction();
