Support merging collisions for contact pairs that appear with matching solver time.  Removes cases where multiple hits for a contact pair would be received in a single frame with identical velocity data (since delta-V is updated per-frame).  Includes information on multiple contact pairs that hit in the same solver frame to allow users to still use the same data.

diff --git a/Engine/Source/Runtime/Engine/Private/PhysicsEngine/Experimental/PhysScene_Chaos.cpp b/Engine/Source/Runtime/Engine/Private/PhysicsEngine/Experimental/PhysScene_Chaos.cpp
--- a/Engine/Source/Runtime/Engine/Private/PhysicsEngine/Experimental/PhysScene_Chaos.cpp
+++ b/Engine/Source/Runtime/Engine/Private/PhysicsEngine/Experimental/PhysScene_Chaos.cpp
@@ -713,20 +713,25 @@ const FBodyInstance* FPhysScene_Chaos::GetBodyInstanceFromProxyAndShape(IPhysics
 	return ProxyInstance;
 }
 
-FCollisionNotifyInfo& FPhysScene_Chaos::GetPendingCollisionForContactPair(const void* P0, const void* P1, bool& bNewEntry)
+FCollisionNotifyInfo& FPhysScene_Chaos::GetPendingCollisionForContactPair(const void* P0, const void* P1, Chaos::FReal SolverTime, bool& bNewEntry)
 {
 	const FUniqueContactPairKey Key = { P0, P1 };
-	const int32* PendingNotifyIdx = ContactPairToPendingNotifyMap.Find(Key);
-	if (PendingNotifyIdx)
+	TArray<int32> ExistingPairs;
+	ContactPairToPendingNotifyMap.MultiFind(Key, ExistingPairs);
+	for(int32 ExistingIdx : ExistingPairs)
 	{
-		// we already have one for this pair
-		bNewEntry = false;
-		return PendingCollisionNotifies[*PendingNotifyIdx];
+		if(FMath::IsNearlyEqual(PendingCollisionNotifies[ExistingIdx].SolverTime, SolverTime))
+		{
+			// we already have one for this pair
+			bNewEntry = false;
+			return PendingCollisionNotifies[ExistingIdx];	
+		}
 	}
 
 	// make a new entry
 	bNewEntry = true;
 	int32 NewIdx = PendingCollisionNotifies.AddZeroed();
+	ContactPairToPendingNotifyMap.Add(Key, NewIdx);
 	return PendingCollisionNotifies[NewIdx];
 }
 
@@ -774,7 +779,7 @@ void FPhysScene_Chaos::HandleCollisionEvents(const Chaos::FCollisionEventData& E
 							}
 
 							bool bNewEntry = false;
-							FCollisionNotifyInfo& NotifyInfo = GetPendingCollisionForContactPair(PhysicsProxy0, PhysicsProxy1, bNewEntry);
+							FCollisionNotifyInfo& NotifyInfo = GetPendingCollisionForContactPair(PhysicsProxy0, PhysicsProxy1, CollisionDataItem.SolverTime, bNewEntry);
 
 							// #note: we only notify on the first contact, though we will still accumulate the impulse data from subsequent contacts
 							const FVector NormalImpulse = FVector::DotProduct(CollisionDataItem.AccumulatedImpulse, CollisionDataItem.Normal) * CollisionDataItem.Normal;	// project impulse along normal
@@ -782,19 +787,42 @@ void FPhysScene_Chaos::HandleCollisionEvents(const Chaos::FCollisionEventData& E
 							NotifyInfo.RigidCollisionData.TotalNormalImpulse += NormalImpulse;
 							NotifyInfo.RigidCollisionData.TotalFrictionImpulse += FrictionImpulse;
 
+							// Get swapped velocity deltas
+							const FVector& DeltaVelocity1 = bSwapOrder
+								? CollisionDataItem.DeltaVelocity2
+								: CollisionDataItem.DeltaVelocity1;
+							const FVector& DeltaVelocity2 = bSwapOrder
+								? CollisionDataItem.DeltaVelocity1
+								: CollisionDataItem.DeltaVelocity2;
+
+							// Populate additional contact information for additional hits
+							FRigidBodyContactInfo& NewContact = NotifyInfo.RigidCollisionData.ContactInfos.AddZeroed_GetRef();
+							NewContact.ContactNormal = CollisionDataItem.Normal;
+							NewContact.ContactPosition = CollisionDataItem.Location;
+							NewContact.ContactPenetration = CollisionDataItem.PenetrationDepth;
+							NewContact.bContactProbe = CollisionDataItem.bProbe;
+							NotifyInfo.RigidCollisionData.bIsVelocityDeltaUnderThreshold =
+								DeltaVelocity1.IsNearlyZero(MinDeltaVelocityThreshold) &&
+								DeltaVelocity2.IsNearlyZero(MinDeltaVelocityThreshold);
+								
+							const Chaos::FChaosPhysicsMaterial* InternalMat1 = CollisionDataItem.Mat1.Get();
+							const Chaos::FChaosPhysicsMaterial* InternalMat2 = CollisionDataItem.Mat2.Get();
+								
+							NewContact.PhysMaterial[0] = InternalMat1 ? FPhysicsUserData::Get<UPhysicalMaterial>(InternalMat1->UserData) : nullptr;
+							NewContact.PhysMaterial[1] = InternalMat2 ? FPhysicsUserData::Get<UPhysicalMaterial>(InternalMat2->UserData) : nullptr;
+
+							if(bSwapOrder)
+							{
+								NewContact.SwapOrder();
+							}
+
 							if (bNewEntry)
 							{
 								// fill in legacy contact data
 								NotifyInfo.bCallEvent0 = true;
 								// if Comp1 wants this event too, it will get its own pending collision entry, so we leave it false
 
-								// Get swapped velocity deltas
-								const FVector& DeltaVelocity1 = bSwapOrder
-									? CollisionDataItem.DeltaVelocity2
-									: CollisionDataItem.DeltaVelocity1;
-								const FVector& DeltaVelocity2 = bSwapOrder
-									? CollisionDataItem.DeltaVelocity1
-									: CollisionDataItem.DeltaVelocity2;
+								NotifyInfo.SolverTime = CollisionDataItem.SolverTime;
 									
 								NotifyInfo.Info0.SetFrom(GetBodyInstanceFromProxyAndShape(PhysicsProxy0, CollisionDataItem.ShapeIndex1), DeltaVelocity1);
 								NotifyInfo.Info1.SetFrom(GetBodyInstanceFromProxyAndShape(PhysicsProxy1, CollisionDataItem.ShapeIndex2), DeltaVelocity2);
@@ -810,26 +838,6 @@ void FPhysScene_Chaos::HandleCollisionEvents(const Chaos::FCollisionEventData& E
 									NotifyInfo.Info1.Component = GetOwningComponent<UPrimitiveComponent>(PhysicsProxy1);
 									NotifyInfo.Info1.Actor = (NotifyInfo.Info1.Component != nullptr) ? NotifyInfo.Info1.Component->GetOwner() : nullptr;
 								}
-
-								FRigidBodyContactInfo& NewContact = NotifyInfo.RigidCollisionData.ContactInfos.AddZeroed_GetRef();
-								NewContact.ContactNormal = CollisionDataItem.Normal;
-								NewContact.ContactPosition = CollisionDataItem.Location;
-								NewContact.ContactPenetration = CollisionDataItem.PenetrationDepth;
-								NewContact.bContactProbe = CollisionDataItem.bProbe;
-								NotifyInfo.RigidCollisionData.bIsVelocityDeltaUnderThreshold =
-									DeltaVelocity1.IsNearlyZero(MinDeltaVelocityThreshold) &&
-									DeltaVelocity2.IsNearlyZero(MinDeltaVelocityThreshold);
-								
-								Chaos::FChaosPhysicsMaterial* InternalMat1 = CollisionDataItem.Mat1.Get();
-								Chaos::FChaosPhysicsMaterial* InternalMat2 = CollisionDataItem.Mat2.Get();
-								
-								NewContact.PhysMaterial[0] = InternalMat1 ? FPhysicsUserData::Get<UPhysicalMaterial>(InternalMat1->UserData) : nullptr;
-								NewContact.PhysMaterial[1] = InternalMat2 ? FPhysicsUserData::Get<UPhysicalMaterial>(InternalMat2->UserData) : nullptr;
-
-								if(bSwapOrder)
-								{
-									NotifyInfo.RigidCollisionData.SwapContactOrders();
-								}
 							}
 						}
 					}
diff --git a/Engine/Source/Runtime/Engine/Public/Physics/Experimental/PhysScene_Chaos.h b/Engine/Source/Runtime/Engine/Public/Physics/Experimental/PhysScene_Chaos.h
--- a/Engine/Source/Runtime/Engine/Public/Physics/Experimental/PhysScene_Chaos.h
+++ b/Engine/Source/Runtime/Engine/Public/Physics/Experimental/PhysScene_Chaos.h
@@ -263,9 +263,9 @@ private:
 		}
 	};
 
-	FCollisionNotifyInfo& GetPendingCollisionForContactPair(const void* P0, const void* P1, bool& bNewEntry);
+	FCollisionNotifyInfo& GetPendingCollisionForContactPair(const void* P0, const void* P1, Chaos::FReal SolverTime, bool& bNewEntry);
 	/** Key is the unique pair, value is index into PendingNotifies array */
-	TMap<FUniqueContactPairKey, int32> ContactPairToPendingNotifyMap;
+	TMultiMap<FUniqueContactPairKey, int32> ContactPairToPendingNotifyMap;
 
 	/** Holds the list of pending legacy notifies that are to be processed */
 	TArray<FCollisionNotifyInfo> PendingCollisionNotifies;
diff --git a/Engine/Source/Runtime/Engine/Public/PhysicsPublic.h b/Engine/Source/Runtime/Engine/Public/PhysicsPublic.h
--- a/Engine/Source/Runtime/Engine/Public/PhysicsPublic.h
+++ b/Engine/Source/Runtime/Engine/Public/PhysicsPublic.h
@@ -109,9 +109,12 @@ struct ENGINE_API FCollisionNotifyInfo
 	/** Information about the collision itself */
 	FCollisionImpactData			RigidCollisionData;
 
+	Chaos::FReal					SolverTime;
+
 	FCollisionNotifyInfo() :
 		bCallEvent0(false),
-		bCallEvent1(false)
+		bCallEvent1(false),
+		SolverTime((Chaos::FReal)0.0)
 	{}
 
 	/** Check that is is valid to call a notification for this entry. Looks at the IsValid() flags on both Actors. */
diff --git a/Engine/Source/Runtime/Experimental/Chaos/Private/EventDefaults.cpp b/Engine/Source/Runtime/Experimental/Chaos/Private/EventDefaults.cpp
--- a/Engine/Source/Runtime/Experimental/Chaos/Private/EventDefaults.cpp
+++ b/Engine/Source/Runtime/Experimental/Chaos/Private/EventDefaults.cpp
@@ -204,6 +204,7 @@ namespace Chaos
 								const FGeometryParticleHandle* Particle1 = Constraint.GetParticle1();
 
 								FCollidingData Data;
+								Data.SolverTime = Solver->MTime;
 								Data.Location = Constraint.CalculateWorldContactLocation();
 								Data.AccumulatedImpulse = Constraint.AccumulatedImpulse;
 								Data.Normal = Constraint.CalculateWorldContactNormal();
diff --git a/Engine/Source/Runtime/Experimental/Chaos/Public/Chaos/ExternalCollisionData.h b/Engine/Source/Runtime/Experimental/Chaos/Public/Chaos/ExternalCollisionData.h
--- a/Engine/Source/Runtime/Experimental/Chaos/Public/Chaos/ExternalCollisionData.h
+++ b/Engine/Source/Runtime/Experimental/Chaos/Public/Chaos/ExternalCollisionData.h
@@ -37,11 +37,12 @@ namespace Chaos
 			, Proxy2(nullptr)
 			, ShapeIndex1(INDEX_NONE)
 			, ShapeIndex2(INDEX_NONE)
+			, SolverTime((FReal)0.0)
 		{}
 
 		FCollidingData(FVec3 InLocation, FVec3 InAccumulatedImpulse, FVec3 InNormal, FVec3 InVelocity1, FVec3 InVelocity2, FVec3 InDeltaVelocity1, FVec3 InDeltaVelocity2
 			, FVec3 InAngularVelocity1, FVec3 InAngularVelocity2, FReal InMass1,FReal InMass2,  FReal InPenetrationDepth, IPhysicsProxyBase* InProxy1, IPhysicsProxyBase* InProxy2
-			, int32 InShapeIndex1, int32 InShapeIndex2)
+			, int32 InShapeIndex1, int32 InShapeIndex2, FReal InSolverTime)
 			: Location(InLocation)
 			, AccumulatedImpulse(InAccumulatedImpulse)
 			, Normal(InNormal)
@@ -61,6 +62,7 @@ namespace Chaos
 			, Proxy2(InProxy2)
 			, ShapeIndex1(InShapeIndex1)
 			, ShapeIndex2(InShapeIndex2)
+			, SolverTime(InSolverTime)
 		{}
 
 		FVec3 Location;
@@ -90,6 +92,8 @@ namespace Chaos
 
 		int32 ShapeIndex1;
 		int32 ShapeIndex2;
+
+		FReal SolverTime;
 	};
 
 	/*
