#UE Lower verbosity about failing to find a class during a world partition actor redirector, for now

diff --git a/Engine/Source/Runtime/Engine/Private/WorldPartition/WorldPartitionActorDescUtils.cpp b/Engine/Source/Runtime/Engine/Private/WorldPartition/WorldPartitionActorDescUtils.cpp
--- a/Engine/Source/Runtime/Engine/Private/WorldPartition/WorldPartitionActorDescUtils.cpp
+++ b/Engine/Source/Runtime/Engine/Private/WorldPartition/WorldPartitionActorDescUtils.cpp
@@ -214,7 +214,7 @@ bool FWorldPartitionActorDescUtils::ValidateActorDescClass(FWorldPartitionActorD
 
 		if (!ClassData)
 		{
-			UE_LOG(LogWorldPartition, Warning, TEXT("Failed to find class asset '%s' for actor '%s"), *AssetClassPath.ToString(), *InActorDesc->GetActorSoftPath().ToString());
+			UE_LOG(LogWorldPartition, Display, TEXT("Failed to find class asset '%s' for actor '%s"), *AssetClassPath.ToString(), *InActorDesc->GetActorSoftPath().ToString());
 			return false;
 		}
 
