Add option to force enable Nanite on all materials with r.Nanite.ForceEnableMeshes

diff --git a/Engine/Source/Runtime/Engine/Private/Materials/MaterialShared.cpp b/Engine/Source/Runtime/Engine/Private/Materials/MaterialShared.cpp
--- a/Engine/Source/Runtime/Engine/Private/Materials/MaterialShared.cpp
+++ b/Engine/Source/Runtime/Engine/Private/Materials/MaterialShared.cpp
@@ -1619,7 +1619,9 @@ bool FMaterialResource::IsUsedWithAPEXCloth() const
 
 bool FMaterialResource::IsUsedWithNanite() const
 {
-	return Material->bUsedWithNanite;
+	static auto NaniteForceEnableMeshesCvar = IConsoleManager::Get().FindTConsoleVariableDataInt(TEXT("r.Nanite.ForceEnableMeshes"));
+
+	return Material->bUsedWithNanite || (NaniteForceEnableMeshesCvar && NaniteForceEnableMeshesCvar->GetValueOnAnyThread() != 0);
 }
 
 bool FMaterialResource::IsUsedWithVolumetricCloud() const
