World Partition: fix perf regression from UpdateStreamingState.

diff --git a/Engine/Source/Runtime/Engine/Private/WorldPartition/WorldPartitionRuntimeCell.cpp b/Engine/Source/Runtime/Engine/Private/WorldPartition/WorldPartitionRuntimeCell.cpp
--- a/Engine/Source/Runtime/Engine/Private/WorldPartition/WorldPartitionRuntimeCell.cpp
+++ b/Engine/Source/Runtime/Engine/Private/WorldPartition/WorldPartitionRuntimeCell.cpp
@@ -114,21 +114,26 @@ FLinearColor UWorldPartitionRuntimeCell::GetDebugStreamingPriorityColor() const
 	return FLinearColor::Transparent;
 }
 
+UDataLayerManager* UWorldPartitionRuntimeCell::GetDataLayerManager() const
+{
+	return GetOuterWorld()->GetWorldPartition()->GetDataLayerManager();
+}
+
 bool UWorldPartitionRuntimeCell::HasAnyDataLayerInEffectiveRuntimeState(EDataLayerRuntimeState InState) const
 {
-	const UDataLayerManager* DataLayerManager = HasDataLayers() ? UDataLayerManager::GetDataLayerManager(this) : nullptr;
+	const UDataLayerManager* DataLayerManager = !DataLayers.IsEmpty() ? GetDataLayerManager() : nullptr;
 	return DataLayerManager ? DataLayerManager->IsAnyDataLayerInEffectiveRuntimeState(GetDataLayers(), InState) : false;
 }
 
 TArray<const UDataLayerInstance*> UWorldPartitionRuntimeCell::GetDataLayerInstances() const
 {
-	const UDataLayerManager* DataLayerManager = HasDataLayers() ? UDataLayerManager::GetDataLayerManager(this) : nullptr;
+	const UDataLayerManager* DataLayerManager = !DataLayers.IsEmpty() ? GetDataLayerManager() : nullptr;
 	return DataLayerManager ? DataLayerManager->GetDataLayerInstances(GetDataLayers()) : TArray<const UDataLayerInstance*>();
 }
 
 bool UWorldPartitionRuntimeCell::ContainsDataLayer(const UDataLayerAsset* DataLayerAsset) const
 {
-	const UDataLayerManager* DataLayerManager = HasDataLayers() ? UDataLayerManager::GetDataLayerManager(this) : nullptr;
+	const UDataLayerManager* DataLayerManager = !DataLayers.IsEmpty() ? GetDataLayerManager() : nullptr;
 	const UDataLayerInstance* DataLayerInstance = DataLayerManager ? DataLayerManager->GetDataLayerInstance(DataLayerAsset) : nullptr;
 	return DataLayerInstance ? ContainsDataLayer(DataLayerInstance) : false;
 }
diff --git a/Engine/Source/Runtime/Engine/Public/WorldPartition/WorldPartitionRuntimeCell.h b/Engine/Source/Runtime/Engine/Public/WorldPartition/WorldPartitionRuntimeCell.h
--- a/Engine/Source/Runtime/Engine/Public/WorldPartition/WorldPartitionRuntimeCell.h
+++ b/Engine/Source/Runtime/Engine/Public/WorldPartition/WorldPartitionRuntimeCell.h
@@ -18,6 +18,7 @@ class UActorContainer;
 class UDataLayerAsset;
 class UDataLayerInstance;
 class UWorldPartition;
+class UDataLayerManager;
 struct FHierarchicalLogArchive;
 
 enum class EWorldPartitionRuntimeCellVisualizeMode
@@ -203,6 +204,7 @@ class ENGINE_API UWorldPartitionRuntimeCell : public UObject, public IWorldParti
 	virtual FString GetDebugName() const override;
 	//~End IWorldPartitionCell Interface
 
+	UDataLayerManager* GetDataLayerManager() const;
 	bool HasAnyDataLayerInEffectiveRuntimeState(EDataLayerRuntimeState InState) const;
 
 	bool GetBlockOnSlowLoading() const { return bBlockOnSlowLoading; }
