Fix deleting subobjects when flushing dormancy
Controlled by net.TrackFlushedDormantObjects, enabled by default

diff --git a/Engine/Source/Runtime/Engine/Classes/Engine/NetConnection.h b/Engine/Source/Runtime/Engine/Classes/Engine/NetConnection.h
--- a/Engine/Source/Runtime/Engine/Classes/Engine/NetConnection.h
+++ b/Engine/Source/Runtime/Engine/Classes/Engine/NetConnection.h
@@ -1806,6 +1806,23 @@ protected:
 	TOptional<FNetworkCongestionControl> NetworkCongestionControl;
 
 	void InitChannelData();
+
+public:
+	/**
+	 * Retrieve stored set of replicated sub-objects of the given actor at the time of the last dormancy flush
+	 * This data is cleared when the actor is processed by ReplicateActor
+	 *
+	 * @param Actor		The actor to retrieve the object map for
+	 * @return A map of network guids to weak object pointers 
+	 */
+	UE::Net::FDormantObjectMap* GetDormantFlushedObjectsForActor(AActor* Actor);
+
+	/**
+	 * Clear stored flushed replicated sub-objects for a given actor, generally after replication or when the actor is destroyed
+	 *
+	 * @param Actor		The actor to clear the flushed object data for
+	 */
+	void ClearDormantFlushedObjectsForActor(AActor* Actor);
 };
 
 struct FScopedRepContext
diff --git a/Engine/Source/Runtime/Engine/Private/DataChannel.cpp b/Engine/Source/Runtime/Engine/Private/DataChannel.cpp
--- a/Engine/Source/Runtime/Engine/Private/DataChannel.cpp
+++ b/Engine/Source/Runtime/Engine/Private/DataChannel.cpp
@@ -3625,17 +3625,17 @@ bool UActorChannel::UpdateDeletedSubObjects(FOutBunch& Bunch)
 
 	bool bWroteSomethingImportant = false;
 
-	auto DeleteSubObject = [this, &bWroteSomethingImportant, &Bunch](const TSharedRef<FObjectReplicator>& SubObjectReplicator, const TWeakObjectPtr<UObject>& ObjectPtrToRemove, ESubObjectDeleteFlag DeleteFlag)
+	auto DeleteSubObject = [this, &bWroteSomethingImportant, &Bunch](FNetworkGUID ObjectNetGUID, const TWeakObjectPtr<UObject>& ObjectPtrToRemove, ESubObjectDeleteFlag DeleteFlag)
 	{
-		if (SubObjectReplicator->ObjectNetGUID.IsValid())
+		if (ObjectNetGUID.IsValid())
 		{
 			UE_NET_TRACE_SCOPE(ContentBlockForSubObjectDelete, Bunch, GetTraceCollector(Bunch), ENetTraceVerbosity::Trace);
-			UE_NET_TRACE_OBJECT_SCOPE(SubObjectReplicator->ObjectNetGUID, Bunch, GetTraceCollector(Bunch), ENetTraceVerbosity::Trace);
+			UE_NET_TRACE_OBJECT_SCOPE(ObjectNetGUID, Bunch, GetTraceCollector(Bunch), ENetTraceVerbosity::Trace);
 
-			UE_LOG(LogNetSubObject, Verbose, TEXT("NetSubObject: Sending request to %s %s::%s (0x%p) NetGUID %s"), ToString(DeleteFlag), *Actor->GetName(), *GetNameSafe(ObjectPtrToRemove.GetEvenIfUnreachable()), ObjectPtrToRemove.GetEvenIfUnreachable(), *SubObjectReplicator->ObjectNetGUID.ToString());
+			UE_LOG(LogNetSubObject, Verbose, TEXT("NetSubObject: Sending request to %s %s::%s (0x%p) NetGUID %s"), ToString(DeleteFlag), *Actor->GetName(), *GetNameSafe(ObjectPtrToRemove.GetEvenIfUnreachable()), ObjectPtrToRemove.GetEvenIfUnreachable(), *ObjectNetGUID.ToString());
 
 			// Write a deletion content header:
-			WriteContentBlockForSubObjectDelete(Bunch, SubObjectReplicator->ObjectNetGUID, DeleteFlag);
+			WriteContentBlockForSubObjectDelete(Bunch, ObjectNetGUID, DeleteFlag);
 
 			bWroteSomethingImportant = true;
 			Bunch.bReliable = true;
@@ -3644,10 +3644,19 @@ bool UActorChannel::UpdateDeletedSubObjects(FOutBunch& Bunch)
 		{
 			UE_LOG(LogNetTraffic, Error, TEXT("Unable to write subobject delete for %s::%s (0x%p), object replicator has invalid NetGUID"), *GetPathNameSafe(Actor), *GetNameSafe(ObjectPtrToRemove.GetEvenIfUnreachable()), ObjectPtrToRemove.GetEvenIfUnreachable());
 		}
-
-		SubObjectReplicator->CleanUp();
 	};
 
+	UE::Net::FDormantObjectMap* DormantObjects = Connection->GetDormantFlushedObjectsForActor(Actor);
+
+	if (DormantObjects)
+	{
+		// no need to track these if they're in the replication map
+		for (auto RepComp = ReplicationMap.CreateIterator(); RepComp; ++RepComp)
+		{
+			DormantObjects->Remove(RepComp.Value()->ObjectNetGUID);
+		}
+	}
+
 #if UE_REPLICATED_OBJECT_REFCOUNTING
 	TArray< TWeakObjectPtr<UObject>, TInlineAllocator<16> > SubObjectsRemoved;
 
@@ -3670,7 +3679,8 @@ bool UActorChannel::UpdateDeletedSubObjects(FOutBunch& Bunch)
 				{
 					const ESubObjectDeleteFlag DeleteFlag = SubObjectRef.IsTearOff() ? ESubObjectDeleteFlag::TearOff : ESubObjectDeleteFlag::ForceDelete;
 					SubObjectsRemoved.Add(SubObjectRef.SubObjectPtr);
-					DeleteSubObject(*SubObjectReplicator, SubObjectRef.SubObjectPtr, DeleteFlag);
+					DeleteSubObject((*SubObjectReplicator)->ObjectNetGUID, SubObjectRef.SubObjectPtr, DeleteFlag);
+					(*SubObjectReplicator)->CleanUp();
 					ReplicationMap.Remove(ObjectToRemove);
 				}
 			}
@@ -3704,11 +3714,29 @@ bool UActorChannel::UpdateDeletedSubObjects(FOutBunch& Bunch)
 			SubObjectsRemoved.Add(WeakObjPtr);
 #endif
 
-			DeleteSubObject(SubObjectReplicator, WeakObjPtr, ESubObjectDeleteFlag::Destroyed);
+			DeleteSubObject(SubObjectReplicator->ObjectNetGUID, WeakObjPtr, ESubObjectDeleteFlag::Destroyed);
+			SubObjectReplicator->CleanUp();
+
 			RepComp.RemoveCurrent();
 		}
 	}
 
+	if (DormantObjects)
+	{
+		for (auto DormComp = DormantObjects->CreateConstIterator(); DormComp; ++DormComp)
+		{
+			const FNetworkGUID& NetGuid = DormComp.Key();
+			const TWeakObjectPtr<UObject>& WeakObjPtr = DormComp.Value();
+
+			if (!WeakObjPtr.IsValid())
+			{
+				DeleteSubObject(NetGuid, WeakObjPtr, ESubObjectDeleteFlag::Destroyed);
+			}
+		}
+
+		Connection->ClearDormantFlushedObjectsForActor(Actor);
+	}
+
 #if UE_REPLICATED_OBJECT_REFCOUNTING
 	if (!SubObjectsRemoved.IsEmpty())
 	{
diff --git a/Engine/Source/Runtime/Engine/Private/DemoNetDriver.cpp b/Engine/Source/Runtime/Engine/Private/DemoNetDriver.cpp
--- a/Engine/Source/Runtime/Engine/Private/DemoNetDriver.cpp
+++ b/Engine/Source/Runtime/Engine/Private/DemoNetDriver.cpp
@@ -5107,8 +5107,6 @@ void UDemoNetDriver::NotifyActorDestroyed(AActor* Actor, bool IsSeamlessTravel)
 		{
 			// This is a stripped down version of UNetDriver::NotifyActorDestroy and UActorChannel::Close
 			// combined, and should be kept up to date with those methods.
-
-
 			if (UNetConnection* Connection = ClientConnections[0])
 			{
 				if (Actor->bNetTemporary)
@@ -5124,7 +5122,7 @@ void UDemoNetDriver::NotifyActorDestroyed(AActor* Actor, bool IsSeamlessTravel)
 					Channel->ReleaseReferences(false);
 				}
 
-				Connection->CleanupDormantReplicatorsForActor(Actor);
+				Connection->NotifyActorDestroyed(Actor, IsSeamlessTravel);
 			}
 
 			GetNetworkObjectList().Remove(Actor);
diff --git a/Engine/Source/Runtime/Engine/Private/Net/NetDormantHolder.cpp b/Engine/Source/Runtime/Engine/Private/Net/NetDormantHolder.cpp
--- a/Engine/Source/Runtime/Engine/Private/Net/NetDormantHolder.cpp
+++ b/Engine/Source/Runtime/Engine/Private/Net/NetDormantHolder.cpp
@@ -197,7 +197,28 @@ void FDormantReplicatorHolder::CountBytes(FArchive& Ar) const
 	{
 		ActorReplicators.CountBytes(Ar);
 	}
+
+	FlushedObjectMap.CountBytes(Ar);
+
+	for (auto FlushedIt = FlushedObjectMap.CreateConstIterator(); FlushedIt; ++FlushedIt)
+	{
+		FlushedIt.Value().CountBytes(Ar);
+	}
+}
+
+UE::Net::FDormantObjectMap* FDormantReplicatorHolder::FindFlushedObjectsForActor(AActor* Actor)
+{
+	return FlushedObjectMap.Find(Actor);
+}
+
+UE::Net::FDormantObjectMap& FDormantReplicatorHolder::FindOrAddFlushedObjectsForActor(AActor* Actor)
+{
+	return FlushedObjectMap.FindOrAdd(Actor);
 }
 
+void FDormantReplicatorHolder::ClearFlushedObjectsForActor(AActor* Actor)
+{
+	FlushedObjectMap.Remove(Actor);
+}
 
 } //end namespace UE::Net::Private
\ No newline at end of file
diff --git a/Engine/Source/Runtime/Engine/Private/NetConnection.cpp b/Engine/Source/Runtime/Engine/Private/NetConnection.cpp
--- a/Engine/Source/Runtime/Engine/Private/NetConnection.cpp
+++ b/Engine/Source/Runtime/Engine/Private/NetConnection.cpp
@@ -114,7 +114,7 @@ DECLARE_CYCLE_STAT(TEXT("NetConnection Tick"), Stat_NetConnectionTick, STATGROUP
 DECLARE_CYCLE_STAT(TEXT("NetConnection ReceivedNak"), Stat_NetConnectionReceivedNak, STATGROUP_Net);
 DECLARE_CYCLE_STAT(TEXT("NetConnection NetConnectionReceivedAcks"), Stat_NetConnectionReceivedAcks, STATGROUP_Net);
 
-namespace UE_NetConnectionPrivate
+namespace UE::Net::Connection::Private
 {
 	struct FValidateLevelVisibilityResult
 	{
@@ -200,6 +200,9 @@ namespace UE_NetConnectionPrivate
 		Result = static_cast<int32>(Result64);
 		return false;
 	}
+
+	int32 bTrackFlushedDormantObjects = true;
+	FAutoConsoleVariableRef CVarNetTrackFlushedDormantObjects(TEXT("net.TrackFlushedDormantObjects"), bTrackFlushedDormantObjects, TEXT("If enabled, track dormant subobjects when dormancy is flushed, so they can be properly deleted if destroyed prior to the next ReplicateActor."));
 }
 
 // ChannelRecord Implementation
@@ -1521,7 +1524,7 @@ void UNetConnection::UpdateLevelVisibilityInternal(const FUpdateLevelVisibilityL
 {
 	QUICK_SCOPE_CYCLE_COUNTER(STAT_NetConnection_UpdateLevelVisibilityInternal)
 
-	using namespace UE_NetConnectionPrivate;
+	using namespace UE::Net::Connection::Private;
 
 	GNumClientUpdateLevelVisibility++;
 
@@ -2171,7 +2174,7 @@ void UNetConnection::FlushNet(bool bIgnoreSimulation)
 		if (!IsReplay())
 		{
 			int32 NewQueuedBits = 0;
-			const bool bWouldOverflow = UE_NetConnectionPrivate::Add_DetectOverflow_Clamp(QueuedBits, PacketBytes * 8, NewQueuedBits);
+			const bool bWouldOverflow = UE::Net::Connection::Private::Add_DetectOverflow_Clamp(QueuedBits, PacketBytes * 8, NewQueuedBits);
 
 			if (bWouldOverflow && !bLoggedFlushNetQueuedBitsOverflow)
 			{
@@ -2509,20 +2512,20 @@ void UNetConnection::WriteFinalPacketInfo(FBitWriter& Writer, const double Packe
 		uint32 ClockTimeMilliseconds = 0;
 
 		// If the delta is over our max precision, we send MAX value and jitter will be ignored by the receiver.
-		if (DeltaSendTimeInMS >= UE_NetConnectionPrivate::MaxJitterPrecisionInMS)
+		if (DeltaSendTimeInMS >= UE::Net::Connection::Private::MaxJitterPrecisionInMS)
 		{
-			ClockTimeMilliseconds = UE_NetConnectionPrivate::MaxJitterClockTimeValue;
+			ClockTimeMilliseconds = UE::Net::Connection::Private::MaxJitterClockTimeValue;
 		}
 		else
 		{
 			// Get the fractional part (milliseconds) of the clock time
-			ClockTimeMilliseconds = UE_NetConnectionPrivate::GetClockTimeMilliseconds(PacketSentTimeInS);
+			ClockTimeMilliseconds = UE::Net::Connection::Private::GetClockTimeMilliseconds(PacketSentTimeInS);
 
 			// Ensure we don't overflow
-			ClockTimeMilliseconds &= UE_NetConnectionPrivate::MaxJitterClockTimeValue;
+			ClockTimeMilliseconds &= UE::Net::Connection::Private::MaxJitterClockTimeValue;
 		}
 
-		Writer.SerializeInt(ClockTimeMilliseconds, UE_NetConnectionPrivate::MaxJitterClockTimeValue + 1);
+		Writer.SerializeInt(ClockTimeMilliseconds, UE::Net::Connection::Private::MaxJitterClockTimeValue + 1);
 
 #if !UE_BUILD_SHIPPING
 		checkf(Writer.GetNumBits() - BitsWrittenPreJitterClock == NetConnectionHelper::NumBitsForJitterClockTimeInHeader, TEXT("WriteFinalPacketInfo did not write the expected nb of bits: Wrote %d, Expected %d"),
@@ -2777,7 +2780,7 @@ void UNetConnection::ReceivedPacket( FBitReader& Reader, bool bIsReinjectedPacke
 #endif
 				// Read jitter clock time from the packet header
 				uint32 PacketJitterClockTimeMS = 0;
-				Reader.SerializeInt(PacketJitterClockTimeMS, UE_NetConnectionPrivate::MaxJitterClockTimeValue + 1);
+				Reader.SerializeInt(PacketJitterClockTimeMS, UE::Net::Connection::Private::MaxJitterClockTimeValue + 1);
 
 #if !UE_BUILD_SHIPPING
 				static double LastJitterLogTime = 0.0;
@@ -4534,7 +4537,7 @@ void UNetConnection::Tick(float DeltaSeconds)
 
 		int32 NewQueuedBits = 0;
 		const int64 DeltaQueuedBits = -FMath::TruncToInt(DeltaBits);
-		const bool bWouldOverflow = UE_NetConnectionPrivate::Add_DetectOverflow_Clamp(QueuedBits, DeltaQueuedBits, NewQueuedBits);
+		const bool bWouldOverflow = UE::Net::Connection::Private::Add_DetectOverflow_Clamp(QueuedBits, DeltaQueuedBits, NewQueuedBits);
 
 		ensureMsgf(!bWouldOverflow, TEXT("UNetConnection::Tick: QueuedBits overflow detected! QueuedBits: %d, change: %lld, BandwidthDeltaTime: %.4f, DesiredTickRate: %.2f"),
 			QueuedBits, DeltaQueuedBits, BandwidthDeltaTime, DesiredTickRate);
@@ -4546,7 +4549,7 @@ void UNetConnection::Tick(float DeltaSeconds)
 		if (QueuedBits < AllowedLag)
 		{
 			int32 NewAllowedLag = 0;
-			const bool bAllowedLagOverflow = UE_NetConnectionPrivate::Add_DetectOverflow_Clamp(0, AllowedLag, NewAllowedLag);
+			const bool bAllowedLagOverflow = UE::Net::Connection::Private::Add_DetectOverflow_Clamp(0, AllowedLag, NewAllowedLag);
 
 			ensureMsgf(!bAllowedLagOverflow, TEXT("UNetConnection::Tick: AllowedLag overflow detected! AllowedLag: %lld, BandwidthDeltaTime: %.4f, DesiredTickRate: %.2f"),
 				AllowedLag, BandwidthDeltaTime, DesiredTickRate);
@@ -4860,14 +4863,39 @@ void UNetConnection::FlushDormancy(AActor* Actor)
 	{
 		FlushDormancyForObject( Actor, Actor );
 
+		const TArray<UActorComponent*>& ReplicatedComponents = Actor->GetReplicatedComponents();
+
+		UE::Net::FDormantObjectMap FlushedGuids;
+		if (UE::Net::Connection::Private::bTrackFlushedDormantObjects)
+		{
+			FlushedGuids.Reserve(ReplicatedComponents.Num());
+		}
+
 		// TODO: Is this set of objects sufficient? Should we query the dormancy map from the connection instead?
-		for ( UActorComponent* ActorComp : Actor->GetReplicatedComponents() )
+		for (UActorComponent* ActorComp : ReplicatedComponents)
 		{
-			if ( ActorComp && ActorComp->GetIsReplicated() )
+			if (ActorComp && ActorComp->GetIsReplicated())
 			{
-				FlushDormancyForObject(Actor, ActorComp );
+				FlushDormancyForObject(Actor, ActorComp);
+
+				if (UE::Net::Connection::Private::bTrackFlushedDormantObjects)
+				{
+					// Searching for the guid because the value on the replicator built in FlushDormancyForObject can be invalid until the next replication
+					// We can then safely ignore any object that still has no guid, since it won't have ever been replicated
+					FNetworkGUID ObjectNetGUID = Driver->GuidCache->GetNetGUID(ActorComp);
+					if (ObjectNetGUID.IsValid())
+					{
+						FlushedGuids.Add(ObjectNetGUID, ActorComp);
+					}
+				}
 			}
 		}
+
+		if (UE::Net::Connection::Private::bTrackFlushedDormantObjects && !FlushedGuids.IsEmpty())
+		{
+			UE::Net::FDormantObjectMap& DormantObjects = DormantReplicatorSet.FindOrAddFlushedObjectsForActor(Actor);
+			DormantObjects.Append(FlushedGuids);
+		}
 	}
 
 	// If channel is pending dormancy, cancel it
@@ -5218,6 +5246,16 @@ void UNetConnection::CleanupStaleDormantReplicators()
 	}
 }
 
+UE::Net::FDormantObjectMap* UNetConnection::GetDormantFlushedObjectsForActor(AActor* Actor)
+{
+	return DormantReplicatorSet.FindFlushedObjectsForActor(Actor);
+}
+
+void UNetConnection::ClearDormantFlushedObjectsForActor(AActor* Actor)
+{
+	DormantReplicatorSet.ClearFlushedObjectsForActor(Actor);
+}
+
 void UNetConnection::SetPendingCloseDueToReplicationFailure()
 {
 	bConnectionPendingCloseDueToReplicationFailure = true;
@@ -5300,13 +5338,13 @@ void UNetConnection::TrackReplicationForAnalytics(const bool bWasSaturated)
 
 void UNetConnection::ProcessJitter(uint32 PacketJitterClockTimeMS)
 {
-	if (PacketJitterClockTimeMS >= UE_NetConnectionPrivate::MaxJitterClockTimeValue)
+	if (PacketJitterClockTimeMS >= UE::Net::Connection::Private::MaxJitterClockTimeValue)
 	{
 		// Ignore this packet for jitter calculations
 		return;
 	}
 
-	const int32 CurrentClockTimeMilliseconds = UE_NetConnectionPrivate::GetClockTimeMilliseconds(LastReceiveRealtime);
+	const int32 CurrentClockTimeMilliseconds = UE::Net::Connection::Private::GetClockTimeMilliseconds(LastReceiveRealtime);
 
 	// Get the delta between the sent and receive clock time.
 	int32 CurrentJitterTimeDelta = CurrentClockTimeMilliseconds - (int32)PacketJitterClockTimeMS;
@@ -5314,7 +5352,7 @@ void UNetConnection::ProcessJitter(uint32 PacketJitterClockTimeMS)
 	if (CurrentJitterTimeDelta < 0)
 	{
 		// Account for wrap-around
-		CurrentJitterTimeDelta += UE_NetConnectionPrivate::MaxJitterPrecisionInMS;
+		CurrentJitterTimeDelta += UE::Net::Connection::Private::MaxJitterPrecisionInMS;
 	}
 
 	// Get the difference between the last two deltas
@@ -5323,7 +5361,7 @@ void UNetConnection::ProcessJitter(uint32 PacketJitterClockTimeMS)
 
 	// Add the value to the average and smooth it out to reduce divergences
 	const float PreviousJitterInMS = AverageJitterInMS;
-	AverageJitterInMS = PreviousJitterInMS + ((CurrentJitter - PreviousJitterInMS) / UE_NetConnectionPrivate::JitterNoiseReduction);
+	AverageJitterInMS = PreviousJitterInMS + ((CurrentJitter - PreviousJitterInMS) / UE::Net::Connection::Private::JitterNoiseReduction);
 
 	//UE_LOG(LogNetTraffic, Verbose, TEXT("Avg Jitter: %f | CurrentJitter: %f | CurrentTimeDelta: %d | PreviousTimeDelta: %d | ReceivedClockMilli %d | LocalClockMilli %d"), AverageJitterInMS, CurrentJitter, CurrentJitterTimeDelta, PreviousJitterTimeDelta, PacketJitterClockTimeMS, CurrentClockTimeMilliseconds);
 
@@ -5386,6 +5424,7 @@ void UNetConnection::NotifyActorDestroyed(AActor* Actor, bool IsSeamlessTravel /
 {
 	// Remove it from any dormancy lists
 	CleanupDormantReplicatorsForActor(Actor);
+	ClearDormantFlushedObjectsForActor(Actor);
 }
 
 void UNetConnection::NotifyActorChannelCleanedUp(UActorChannel* Channel, EChannelCloseReason CloseReason)
diff --git a/Engine/Source/Runtime/Engine/Public/Net/NetDormantHolder.h b/Engine/Source/Runtime/Engine/Public/Net/NetDormantHolder.h
--- a/Engine/Source/Runtime/Engine/Public/Net/NetDormantHolder.h
+++ b/Engine/Source/Runtime/Engine/Public/Net/NetDormantHolder.h
@@ -5,6 +5,7 @@
 #include "UObject/ObjectKey.h"
 #include "Templates/SharedPointer.h"
 #include "Serialization/Archive.h"
+#include "Misc/NetworkGuid.h"
 
 class AActor;
 class FObjectReplicator;
@@ -14,6 +15,8 @@ namespace UE::Net
 {
 	/** Function definition allowed in the ForEach functions */
 	typedef TFunctionRef<void(FObjectKey OwnerActorKey, FObjectKey ObjectKey, const TSharedRef<FObjectReplicator>& ReplicatorRef)> FExecuteForEachDormantReplicator;
+
+	typedef TMap<FNetworkGUID, TWeakObjectPtr<UObject>> FDormantObjectMap;
 }
 
 namespace UE::Net::Private
@@ -176,6 +179,34 @@ struct FDormantReplicatorHolder
 	typedef TSet<FActorDormantReplicators, FActorDormantReplicatorsKeyFuncs> FActorReplicatorSet;
 	/** The TSet indexed by Actor that stores all their dormant object replicators */
 	FActorReplicatorSet ActorReplicatorSet;
+
+	/**
+	 * Retrieve stored set of replicated sub-objects of the given actor at the time of the last dormancy flush
+	 * This data is cleared when the actor is processed by ReplicateActor
+	 *
+	 * @param Actor		The actor to retrieve the object map for
+	 * @return A map of network guids to weak object pointers, or nullptr
+	 */
+	UE::Net::FDormantObjectMap* FindFlushedObjectsForActor(AActor* Actor);
+
+	/**
+	 * Retrieve stored set of replicated sub-objects of the given actor at the time of the last dormancy flush
+	 * This data is cleared when the actor is processed by ReplicateActor
+	 *
+	 * @param Actor		The actor to retrieve the object map for
+	 * @return A map of network guids to weak object pointers
+	 */
+	UE::Net::FDormantObjectMap& FindOrAddFlushedObjectsForActor(AActor* Actor);
+
+	/**
+	 * Clear stored flushed replicated sub-objects for a given actor, generally after replication or when the actor is destroyed
+	 *
+	 * @param Actor		The actor to clear the flushed object data for
+	 */
+	void ClearFlushedObjectsForActor(AActor* Actor);
+
+	/** Map of actors to guid/weakptr pairs from the last dormancy flush */
+	TMap<FObjectKey, UE::Net::FDormantObjectMap> FlushedObjectMap;
 };
 
 }//end namespace UE::Net::Private
