World Partition: code cleanup following 24280786.

diff --git a/Engine/Source/Runtime/Engine/Private/WorldPartition/WorldPartitionRuntimeCell.cpp b/Engine/Source/Runtime/Engine/Private/WorldPartition/WorldPartitionRuntimeCell.cpp
--- a/Engine/Source/Runtime/Engine/Private/WorldPartition/WorldPartitionRuntimeCell.cpp
+++ b/Engine/Source/Runtime/Engine/Private/WorldPartition/WorldPartitionRuntimeCell.cpp
@@ -121,19 +121,19 @@ UDataLayerManager* UWorldPartitionRuntimeCell::GetDataLayerManager() const
 
 bool UWorldPartitionRuntimeCell::HasAnyDataLayerInEffectiveRuntimeState(EDataLayerRuntimeState InState) const
 {
-	const UDataLayerManager* DataLayerManager = !DataLayers.IsEmpty() ? GetDataLayerManager() : nullptr;
+	const UDataLayerManager* DataLayerManager = HasDataLayers() ? GetDataLayerManager() : nullptr;
 	return DataLayerManager ? DataLayerManager->IsAnyDataLayerInEffectiveRuntimeState(GetDataLayers(), InState) : false;
 }
 
 TArray<const UDataLayerInstance*> UWorldPartitionRuntimeCell::GetDataLayerInstances() const
 {
-	const UDataLayerManager* DataLayerManager = !DataLayers.IsEmpty() ? GetDataLayerManager() : nullptr;
+	const UDataLayerManager* DataLayerManager = HasDataLayers() ? GetDataLayerManager() : nullptr;
 	return DataLayerManager ? DataLayerManager->GetDataLayerInstances(GetDataLayers()) : TArray<const UDataLayerInstance*>();
 }
 
 bool UWorldPartitionRuntimeCell::ContainsDataLayer(const UDataLayerAsset* DataLayerAsset) const
 {
-	const UDataLayerManager* DataLayerManager = !DataLayers.IsEmpty() ? GetDataLayerManager() : nullptr;
+	const UDataLayerManager* DataLayerManager = HasDataLayers() ? GetDataLayerManager() : nullptr;
 	const UDataLayerInstance* DataLayerInstance = DataLayerManager ? DataLayerManager->GetDataLayerInstance(DataLayerAsset) : nullptr;
 	return DataLayerInstance ? ContainsDataLayer(DataLayerInstance) : false;
 }
diff --git a/Engine/Source/Runtime/Engine/Public/WorldPartition/WorldPartitionRuntimeCell.h b/Engine/Source/Runtime/Engine/Public/WorldPartition/WorldPartitionRuntimeCell.h
--- a/Engine/Source/Runtime/Engine/Public/WorldPartition/WorldPartitionRuntimeCell.h
+++ b/Engine/Source/Runtime/Engine/Public/WorldPartition/WorldPartitionRuntimeCell.h
@@ -194,7 +194,6 @@ class ENGINE_API UWorldPartitionRuntimeCell : public UObject, public IWorldParti
 	virtual TArray<const UDataLayerInstance*> GetDataLayerInstances() const override;
 	virtual bool ContainsDataLayer(const UDataLayerAsset* DataLayerAsset) const override;
 	virtual bool ContainsDataLayer(const UDataLayerInstance* DataLayerInstance) const override;
-	virtual bool HasDataLayers() const override { return !DataLayers.IsEmpty(); }
 	virtual const TArray<FName>& GetDataLayers() const override  { return DataLayers; }
 	virtual bool HasAnyDataLayer(const TSet<FName>& InDataLayers) const override
 	{
@@ -204,6 +203,8 @@ class ENGINE_API UWorldPartitionRuntimeCell : public UObject, public IWorldParti
 	virtual FString GetDebugName() const override;
 	//~End IWorldPartitionCell Interface
 
+	inline bool HasDataLayers() const { return !DataLayers.IsEmpty(); }
+
 	UDataLayerManager* GetDataLayerManager() const;
 	bool HasAnyDataLayerInEffectiveRuntimeState(EDataLayerRuntimeState InState) const;
 
diff --git a/Engine/Source/Runtime/Engine/Public/WorldPartition/WorldPartitionRuntimeCellInterface.h b/Engine/Source/Runtime/Engine/Public/WorldPartition/WorldPartitionRuntimeCellInterface.h
--- a/Engine/Source/Runtime/Engine/Public/WorldPartition/WorldPartitionRuntimeCellInterface.h
+++ b/Engine/Source/Runtime/Engine/Public/WorldPartition/WorldPartitionRuntimeCellInterface.h
@@ -26,7 +26,7 @@ public:
 	/** Returns whether the cell data layers referenced the provided data layer instance or not. */
 	virtual bool ContainsDataLayer(const UDataLayerInstance* DataLayerInstance) const = 0;
 	/** Returns whether the cell content is associated with data layers or not. */
-	virtual bool HasDataLayers() const = 0;
+	bool HasDataLayers() const { return !GetDataLayers().IsEmpty(); }
 	/** Returns the cell content associated data layers. */
 	virtual const TArray<FName>& GetDataLayers() const = 0;
 	/** Returns whether the cell data layers referenced any of the provided data layer or not. */
