When Unregistering a SceneCaptureComponent, make sure it's not in the pending update map. If we don't do this, it will get its ViewStates members reinitialized which causes issues on destruction.
Adding better locks around the SceneCaptureComponent update map.

diff --git a/Engine/Source/Runtime/Engine/Private/Components/SceneCaptureComponent.cpp b/Engine/Source/Runtime/Engine/Private/Components/SceneCaptureComponent.cpp
--- a/Engine/Source/Runtime/Engine/Private/Components/SceneCaptureComponent.cpp
+++ b/Engine/Source/Runtime/Engine/Private/Components/SceneCaptureComponent.cpp
@@ -36,6 +36,7 @@
 #define LOCTEXT_NAMESPACE "SceneCaptureComponent"
 
 static TMultiMap<TWeakObjectPtr<UWorld>, TWeakObjectPtr<USceneCaptureComponent> > SceneCapturesToUpdateMap;
+static FCriticalSection SceneCapturesToUpdateMapCS;
 
 static TAutoConsoleVariable<bool> CVarSCOverrideOrthographicTilingValues(
 	TEXT("r.SceneCapture.OverrideOrthographicTilingValues"),
@@ -467,6 +468,7 @@ void USceneCaptureComponent::Serialize(FArchive& Ar)
 
 void USceneCaptureComponent::UpdateDeferredCaptures(FSceneInterface* Scene)
 {
+	FScopeLock ScopeLock(&SceneCapturesToUpdateMapCS);
 	UWorld* World = Scene->GetWorld();
 	if (!World || SceneCapturesToUpdateMap.Num() == 0)
 	{
@@ -504,6 +506,11 @@ void USceneCaptureComponent::UpdateDeferredCaptures(FSceneInterface* Scene)
 
 void USceneCaptureComponent::OnUnregister()
 {
+	// Make sure this component isn't still in the update map before we fully unregister
+	{
+		FScopeLock ScopeLock(&SceneCapturesToUpdateMapCS);
+		SceneCapturesToUpdateMap.Remove(GetWorld(), this);
+	}
 	for (int32 ViewIndex = 0; ViewIndex < ViewStates.Num(); ViewIndex++)
 	{
 		ViewStates[ViewIndex].Destroy();
@@ -668,8 +675,7 @@ void USceneCaptureComponent2D::CaptureSceneDeferred()
 	{
 		// Defer until after updates finish
 		// Needs some CS because of parallel updates.
-		static FCriticalSection CriticalSection;
-		FScopeLock ScopeLock(&CriticalSection);
+		FScopeLock ScopeLock(&SceneCapturesToUpdateMapCS);
 		SceneCapturesToUpdateMap.AddUnique(World, this);
 	}	
 }
@@ -1247,9 +1253,8 @@ void USceneCaptureComponentCube::CaptureSceneDeferred()
 	{
 		// Defer until after updates finish
 		// Needs some CS because of parallel updates.
-		static FCriticalSection CriticalSection;
-		FScopeLock ScopeLock(&CriticalSection);
-		SceneCapturesToUpdateMap.AddUnique( World, this );
+		FScopeLock ScopeLock(&SceneCapturesToUpdateMapCS);
+		SceneCapturesToUpdateMap.AddUnique(World, this);
 	}	
 }
 
