Added BlueprintDisplayName to BlueprintGeneratedClass editor tags to permit aliased class names in cooked editor.

diff --git a/Engine/Source/Runtime/Engine/Private/BlueprintGeneratedClass.cpp b/Engine/Source/Runtime/Engine/Private/BlueprintGeneratedClass.cpp
--- a/Engine/Source/Runtime/Engine/Private/BlueprintGeneratedClass.cpp
+++ b/Engine/Source/Runtime/Engine/Private/BlueprintGeneratedClass.cpp
@@ -2447,6 +2447,7 @@ void UBlueprintGeneratedClass::GetEditorTags(FEditorTags& Tags) const
 		};
 
 		AddEditorTag(FBlueprintTags::BlueprintType, GET_MEMBER_NAME_CHECKED(UBlueprint, BlueprintType));
+		AddEditorTag(FBlueprintTags::BlueprintDisplayName, GET_MEMBER_NAME_CHECKED(UBlueprint, BlueprintDisplayName));
 
 		{
 			// Clear the FBPInterfaceDescription Graphs because they are irrelevant to the BPGC
