[Engine]
* Fixed funny syntax fail where ; was in the wrong place :-)

diff --git a/Engine/Source/Runtime/Engine/Classes/Sound/SoundWave.h b/Engine/Source/Runtime/Engine/Classes/Sound/SoundWave.h
--- a/Engine/Source/Runtime/Engine/Classes/Sound/SoundWave.h
+++ b/Engine/Source/Runtime/Engine/Classes/Sound/SoundWave.h
@@ -1418,7 +1418,7 @@ public:
 	float GetSampleRate() const;
 
 	int32 GetNumFrames() const;
-	uint32 GetNumChunks(); const
+	uint32 GetNumChunks() const;
 	uint32 GetNumChannels() const;
 	uint32 GetSizeOfChunk(uint32 ChunkIndex) const;
 	const TArray<FSoundWaveCuePoint>& GetCuePoints() const;
