Make sure that changing a texture on a material instance adds the new texture to the material's GC cluster

diff --git a/Engine/Source/Runtime/Engine/Private/Materials/MaterialInstance.cpp b/Engine/Source/Runtime/Engine/Private/Materials/MaterialInstance.cpp
--- a/Engine/Source/Runtime/Engine/Private/Materials/MaterialInstance.cpp
+++ b/Engine/Source/Runtime/Engine/Private/Materials/MaterialInstance.cpp
@@ -3577,6 +3577,7 @@ void UMaterialInstance::SetTextureParameterValueInternal(const FMaterialParamete
 		if (Value && ensureMsgf(Value->IsA(UTexture::StaticClass()), TEXT("Expecting a UTexture! Value='%s' class='%s'"), *Value->GetName(), *Value->GetClass()->GetName()))
 		{
 			ParameterValue->ParameterValue = Value;
+			Value->AddToCluster(this, true);
 			// Update the material instance data in the rendering thread.
 			GameThread_UpdateMIParameter(this, *ParameterValue);
 
@@ -3610,6 +3611,7 @@ void UMaterialInstance::SetRuntimeVirtualTextureParameterValueInternal(const FMa
 		if (Value && ensureMsgf(Value->IsA(URuntimeVirtualTexture::StaticClass()), TEXT("Expecting a URuntimeVirtualTexture! Value='%s' class='%s'"), *Value->GetName(), *Value->GetClass()->GetName()))
 		{
 			ParameterValue->ParameterValue = Value;
+			Value->AddToCluster(this, true);
 			// Update the material instance data in the rendering thread.
 			GameThread_UpdateMIParameter(this, *ParameterValue);
 
@@ -3676,6 +3678,10 @@ void UMaterialInstance::SetFontParameterValueInternal(const FMaterialParameterIn
 	{
 		ParameterValue->FontValue = FontValue;
 		ParameterValue->FontPage = FontPage;
+		if (FontValue)
+		{
+			FontValue->AddToCluster(this, true);
+		}
 		// Update the material instance data in the rendering thread.
 		GameThread_UpdateMIParameter(this, *ParameterValue);
 	}
