- Clamp MinLOD to GetCurrentFirstLODIdx_RenderThread during setup of dynamic draw elements to take mesh streaming into account as well
- Add cvar to force end cull distance for HISM instances (default 0 disabled)

#rnx

diff --git a/Engine/Source/Runtime/Engine/Private/HierarchicalInstancedStaticMesh.cpp b/Engine/Source/Runtime/Engine/Private/HierarchicalInstancedStaticMesh.cpp
--- a/Engine/Source/Runtime/Engine/Private/HierarchicalInstancedStaticMesh.cpp
+++ b/Engine/Source/Runtime/Engine/Private/HierarchicalInstancedStaticMesh.cpp
@@ -85,6 +85,11 @@ TAutoConsoleVariable<float> CVarFoliageMinimumScreenSize(
 	TEXT("This controls the screen size at which we cull foliage instances entirely."),
 	ECVF_Scalability);
 
+TAutoConsoleVariable<int32> CVarFoliageMaxEndCullDistance(
+	TEXT("foliage.MaxEndCullDistance"),
+	0,
+	TEXT("Max distance for end culling (0 disabled)."));
+
 TAutoConsoleVariable<float> CVarFoliageLODDistanceScale(
 	TEXT("foliage.LODDistanceScale"),
 	1.0f,
@@ -1642,6 +1647,7 @@ void FHierarchicalStaticMeshSceneProxy::GetDynamicMeshElements(const TArray<cons
 
 				float MinSize = bIsOrtho ? 0.0f : CVarFoliageMinimumScreenSize.GetValueOnRenderThread();
 				float LODScale = CVarFoliageLODDistanceScale.GetValueOnRenderThread();
+				int MaxEndCullDistance = CVarFoliageMaxEndCullDistance.GetValueOnRenderThread();
 				float LODRandom = CVarRandomLODRange.GetValueOnRenderThread();
 				float MaxDrawDistanceScale = GetCachedScalabilityCVars().ViewDistanceScale;
 				
@@ -1658,9 +1664,21 @@ void FHierarchicalStaticMeshSceneProxy::GetDynamicMeshElements(const TArray<cons
 				{
 					FinalCull = FMath::Min(FinalCull, View->SceneViewInitOptions.OverrideFarClippingPlaneDistance * MaxDrawDistanceScale);
 				}
-				if (UserData_AllInstances.EndCullDistance > 0.0f)
+				int32 EndCullDistance = UserData_AllInstances.EndCullDistance * MaxDrawDistanceScale;
+				if (MaxEndCullDistance > 0)
+				{
+					if (EndCullDistance > 0)
+					{
+						EndCullDistance = FMath::Min(MaxEndCullDistance, EndCullDistance);
+					}
+					else
+					{
+						EndCullDistance = MaxEndCullDistance;
+					}
+				}
+				if (EndCullDistance > 0.0f)
 				{
-					FinalCull = FMath::Min(FinalCull, UserData_AllInstances.EndCullDistance * MaxDrawDistanceScale);
+					FinalCull = FMath::Min(FinalCull, EndCullDistance);
 				}
 				ElementParams.FinalCullDistance = FinalCull;
 
@@ -1755,6 +1773,10 @@ void FHierarchicalStaticMeshSceneProxy::GetDynamicMeshElements(const TArray<cons
 						UseMaxLOD = FMath::Clamp(Force, 0, InstanceParams.LODs - 1);
 					}
 
+					// Clamp the min LOD to available LOD taking mesh streaming into account as well
+					const int8 CurFirstLODIdx = GetCurrentFirstLODIdx_RenderThread();
+					UseMinLOD = FMath::Max(UseMinLOD, CurFirstLODIdx);
+
 					if (CVarCullAll.GetValueOnRenderThread() < 1)
 					{
 						if (bUseVectorCull)
@@ -1824,6 +1846,7 @@ void FHierarchicalStaticMeshSceneProxy::GetDynamicMeshElements(const TArray<cons
 						const bool bIsOrtho = !View->ViewMatrices.IsPerspectiveProjection();
 						const float MinSize = bIsOrtho ? 0.0f : CVarFoliageMinimumScreenSize.GetValueOnRenderThread();
 						const float LODScale = CVarFoliageLODDistanceScale.GetValueOnRenderThread();
+						int MaxEndCullDistance = CVarFoliageMaxEndCullDistance.GetValueOnRenderThread();
 						const float LODRandom = CVarRandomLODRange.GetValueOnRenderThread();
 						const float MaxDrawDistanceScale = GetCachedScalabilityCVars().ViewDistanceScale;
 						const float SphereRadius = RenderData->Bounds.SphereRadius;
@@ -1839,9 +1862,21 @@ void FHierarchicalStaticMeshSceneProxy::GetDynamicMeshElements(const TArray<cons
 						{
 							FinalCull = FMath::Min(FinalCull, View->SceneViewInitOptions.OverrideFarClippingPlaneDistance * MaxDrawDistanceScale);
 						}
-						if (UserData_AllInstances.EndCullDistance > 0.0f)
+						int32 EndCullDistance = UserData_AllInstances.EndCullDistance * MaxDrawDistanceScale;
+						if (MaxEndCullDistance > 0)
 						{
-							FinalCull = FMath::Min(FinalCull, UserData_AllInstances.EndCullDistance * MaxDrawDistanceScale);
+							if (EndCullDistance > 0)
+							{
+								EndCullDistance = FMath::Min(MaxEndCullDistance, EndCullDistance);
+							}
+							else
+							{
+								EndCullDistance = MaxEndCullDistance;
+							}
+						}
+						if (EndCullDistance > 0.0f)
+						{
+							FinalCull = FMath::Min(FinalCull, EndCullDistance);
 						}
 						ElementParams.FinalCullDistance = FinalCull;
 
@@ -1859,6 +1894,11 @@ void FHierarchicalStaticMeshSceneProxy::GetDynamicMeshElements(const TArray<cons
 						// calculate runs
 						int32 MinLOD = ClampedMinLOD;
 						int32 MaxLOD = NumLODs;
+
+						// Clamp the min LOD to available LOD taking mesh streaming into account as well
+						const int8 CurFirstLODIdx = GetCurrentFirstLODIdx_RenderThread();
+						MinLOD = FMath::Max(MinLOD, CurFirstLODIdx);
+
 						CalcLOD(MinLOD, MaxLOD, UnbuiltBounds[0].Min, UnbuiltBounds[0].Max, ViewOriginInLocalZero, ViewOriginInLocalOne, LODPlanesMin, LODPlanesMax);
 						int32 FirstIndexInRun = 0;
 						for (int32 Index = 1; Index < UnbuiltInstanceCount; ++Index)
