Fix for Instanced Foliage paint brush taking a very long time when the Radius option is used.

With settings:
brush size: 4000
brush paint density: 1
density: 400
radius: 100000

Original Time:  20 minutes 37 seconds
New Time: 112 ms

diff --git a/Engine/Source/Editor/FoliageEdit/Private/FoliageEdMode.cpp b/Engine/Source/Editor/FoliageEdit/Private/FoliageEdMode.cpp
--- a/Engine/Source/Editor/FoliageEdit/Private/FoliageEdMode.cpp
+++ b/Engine/Source/Editor/FoliageEdit/Private/FoliageEdMode.cpp
@@ -1063,6 +1063,8 @@ static bool CheckForOverlappingSphere(UWorld* InWorld, const UFoliageType* Setti
 
 static bool CheckLocationForPotentialInstance(UWorld* InWorld, const UFoliageType* Settings, const bool bSingleInstanceMode, const FVector& Location, const FVector& Normal, TArray<FVector>& PotentialInstanceLocations, FFoliageInstanceHash& PotentialInstanceHash)
 {
+	TRACE_CPUPROFILER_EVENT_SCOPE(CheckLocationForPotentialInstance);
+
 	if (CheckLocationForPotentialInstance_ThreadSafe(Settings, Location, Normal) == false)
 	{
 		return false;
@@ -1080,14 +1082,10 @@ static bool CheckLocationForPotentialInstance(UWorld* InWorld, const UFoliageTyp
 		}
 
 		// Check with other potential instances we're about to add.
-		const float RadiusSquared = FMath::Square(SettingsRadius);
-		auto TempInstances = PotentialInstanceHash.GetInstancesOverlappingBox(FBox::BuildAABB(Location, FVector(SettingsRadius)));
-		for (int32 Idx : TempInstances)
+		TArrayView<const FVector> InstanceLocationsView = PotentialInstanceLocations;
+		if (PotentialInstanceHash.IsAnyInstanceInSphere([InstanceLocationsView](int32 Index) -> FVector { return InstanceLocationsView[Index]; }, Location, SettingsRadius))
 		{
-			if ((PotentialInstanceLocations[Idx] - Location).SizeSquared() < RadiusSquared)
-			{
-				return false;
-			}
+			return false;
 		}
 	}
 
diff --git a/Engine/Source/Runtime/Engine/Public/Instances/InstancedPlacementHash.h b/Engine/Source/Runtime/Engine/Public/Instances/InstancedPlacementHash.h
--- a/Engine/Source/Runtime/Engine/Public/Instances/InstancedPlacementHash.h
+++ b/Engine/Source/Runtime/Engine/Public/Instances/InstancedPlacementHash.h
@@ -26,6 +26,18 @@ private:
 		FKey(int64 InX, int64 InY, int64 InZ)
 			: X(InX), Y(InY), Z(InZ) {}
 
+		// returns the minimum distance from the given position to the cell described by this key, squared
+		double GetMinDistSquared(const FVector& Position, const int32 LocalHashCellBits)
+		{
+			// compute distance to the cell cube
+			int32 CellSize = (1 << LocalHashCellBits);
+			int32 HalfCellSize = CellSize >> 1;
+			FVector CellCenter(X * CellSize + HalfCellSize, Y * CellSize + HalfCellSize, Z * CellSize + HalfCellSize);
+			FVector AbsRelativePosition = (Position - CellCenter).GetAbs();			
+			FVector CubeDeltaVector = (AbsRelativePosition - HalfCellSize).ComponentMax(FVector::ZeroVector);
+			return CubeDeltaVector.SizeSquared();
+		}
+
 		bool operator==(const FKey& Other) const
 		{
 			return (X == Other.X) && (Y == Other.Y) && (Z == Other.Z);
@@ -95,21 +107,118 @@ public:
 		InsertInstance(InstanceLocation, NewIndex);
 	}
 
-	void GetInstancesOverlappingBox(const FBox& InBox, TArray<int32>& OutInstanceIndices) const
+	template<typename FunctionType>
+	bool IsAnyInstanceInSphere(FunctionType InstanceLocationGetter, const FVector& SphereCenter, double SphereRadius)
 	{
-		FKey MinKey = MakeKey(InBox.Min);
-		FKey MaxKey = MakeKey(InBox.Max);
+		TRACE_CPUPROFILER_EVENT_SCOPE(IsAnyInstanceInSphere);
 
+		double SphereRadiusSquared = SphereRadius * SphereRadius;
+
+		// if there are more potential cells within range than the number of actually populated cells...
+		double SphereRadiusCells = SphereRadius / (1 << HashCellBits);
+		constexpr double SphereVolumeConstant = 4.0 / 3.0 * UE_PI;
+		double ApproxCellsToCheck = SphereVolumeConstant * (SphereRadiusCells * SphereRadiusCells * SphereRadiusCells);
+		if (ApproxCellsToCheck > CellMap.Num())
+		{
+			// then it's probably faster to just check all populated cells and be done
+			for (auto& Pair : CellMap)
+			{
+				if (Pair.Key.GetMinDistSquared(SphereCenter, HashCellBits) <= SphereRadiusSquared)
+				{
+					for (int32 InstanceIndex : Pair.Value)
+					{
+						double DistSquared = (InstanceLocationGetter(InstanceIndex) - SphereCenter).SizeSquared();
+						if (DistSquared < SphereRadiusSquared)
+						{
+							return true;
+						}
+					}
+				}
+			}
+			return false;
+		}
+
+		// otherwise we check all potential cells within range
+		FBox SphereBox = FBox::BuildAABB(SphereCenter, FVector(SphereRadius));
+		FKey MinKey = MakeKey(SphereBox.Min);
+		FKey MaxKey = MakeKey(SphereBox.Max);
 		for (int64 z = MinKey.Z; z <= MaxKey.Z; ++z)
 		{
 			for (int64 y = MinKey.Y; y <= MaxKey.Y; y++)
 			{
 				for (int64 x = MinKey.X; x <= MaxKey.X; x++)
 				{
-					auto* SetPtr = CellMap.Find(FKey(x, y, z));
+					FKey Key(x, y, z);
+					auto* SetPtr = CellMap.Find(Key);
 					if (SetPtr)
 					{
-						OutInstanceIndices.Append(SetPtr->Array());
+						for (int32 InstanceIndex : *SetPtr)
+						{
+							float DistSquared = (InstanceLocationGetter(InstanceIndex) - SphereCenter).SizeSquared();
+							if (DistSquared < SphereRadiusSquared)
+							{
+								return true;
+							}
+						}
+					}
+				}
+			}
+		}
+
+		return false;
+	}
+
+	void GetInstancesOverlappingBox(const FBox& InBox, TArray<int32>& OutInstanceIndices) const
+	{
+		FKey MinKey = MakeKey(InBox.Min);
+		FKey MaxKey = MakeKey(InBox.Max);
+
+		int64 CellCount = (MaxKey.Z - MinKey.Z + 1) * (MaxKey.Y - MinKey.Y + 1) * (MaxKey.X - MinKey.X + 1);
+
+		// The idea here is to decide when it is faster to just check every populated cell.
+		// The actual ideal threshold will depend on the exact state of the cells, but we just pick
+		// some arbitrary ratio where we switch to checking populated cells over potential cells within range.
+		// In practice the exact threshold is not that important; we just want to know if there are different orders of magnitude involved.
+		// (i.e. it is faster to check 32 populated cells than 3,000,000,000 potential cells...)
+		constexpr int32 RelativeCostOfCheckingPopulatedCells = 2;
+		if (CellCount > RelativeCostOfCheckingPopulatedCells * CellMap.Num())
+		{
+			// check every populated cell
+			for (auto& Pair : CellMap)
+			{
+				if ((Pair.Key.X >= MinKey.X) &&
+					(Pair.Key.X <= MaxKey.X) &&
+					(Pair.Key.Y >= MinKey.Y) &&
+					(Pair.Key.Y <= MaxKey.Y) &&
+					(Pair.Key.Z >= MinKey.Z) &&
+					(Pair.Key.Z <= MaxKey.Z))
+				{
+					OutInstanceIndices.Reserve(OutInstanceIndices.Num() + Pair.Value.Num());
+					for (int32 InstanceIndex : Pair.Value)
+					{
+						OutInstanceIndices.Add(InstanceIndex);
+					}
+				}
+			}
+		}
+		else
+		{
+			// check all potential cells within range
+			for (int64 z = MinKey.Z; z <= MaxKey.Z; ++z)
+			{
+				for (int64 y = MinKey.Y; y <= MaxKey.Y; y++)
+				{
+					for (int64 x = MinKey.X; x <= MaxKey.X; x++)
+					{
+						auto* SetPtr = CellMap.Find(FKey(x, y, z));
+						if (SetPtr)
+						{
+							OutInstanceIndices.Reserve(OutInstanceIndices.Num() + SetPtr->Num());
+							for (int32 InstanceIndex : *SetPtr)
+							{
+								OutInstanceIndices.Add(InstanceIndex);
+							}
+						}
 					}
 				}
 			}
diff --git a/Engine/Source/Runtime/Foliage/Private/InstancedFoliage.cpp b/Engine/Source/Runtime/Foliage/Private/InstancedFoliage.cpp
--- a/Engine/Source/Runtime/Foliage/Private/InstancedFoliage.cpp
+++ b/Engine/Source/Runtime/Foliage/Private/InstancedFoliage.cpp
@@ -2706,14 +2706,9 @@ void FFoliageInfo::GetInstanceAtLocation(const FVector& Location, int32& OutInst
 // Returns whether or not there is are any instances overlapping the sphere specified
 bool FFoliageInfo::CheckForOverlappingSphere(const FSphere& Sphere) const
 {
-	auto TempInstances = InstanceHash->GetInstancesOverlappingBox(FBox::BuildAABB(Sphere.Center, FVector(Sphere.W)));
-	for (int32 Idx : TempInstances)
-	{
-		if (FSphere(Instances[Idx].Location, 0.f).IsInside(Sphere))
-		{
-			return true;
-		}
-	}
+	TArrayView<const FFoliageInstance> InstancesView = Instances;
+	return InstanceHash->IsAnyInstanceInSphere([InstancesView](int32 Index) -> FVector { return InstancesView[Index].Location; }, Sphere.Center, Sphere.W);
+
 	return false;
 }
 
