World Partition: failsafe for duplicated SpatialHashRuntimeGridInfo actors.

diff --git a/Engine/Source/Runtime/Engine/Private/WorldPartition/WorldPartitionRuntimeSpatialHash.cpp b/Engine/Source/Runtime/Engine/Private/WorldPartition/WorldPartitionRuntimeSpatialHash.cpp
--- a/Engine/Source/Runtime/Engine/Private/WorldPartition/WorldPartitionRuntimeSpatialHash.cpp
+++ b/Engine/Source/Runtime/Engine/Private/WorldPartition/WorldPartitionRuntimeSpatialHash.cpp
@@ -1119,7 +1119,16 @@ bool UWorldPartitionRuntimeSpatialHash::GenerateStreaming(UWorldPartitionStreami
 	{
 		FWorldPartitionReference Ref(WorldPartition, SpatialHashRuntimeGridInfo->GetGuid());
 		ASpatialHashRuntimeGridInfo* RuntimeGridActor = CastChecked<ASpatialHashRuntimeGridInfo>(SpatialHashRuntimeGridInfo->GetActor());
-		AllGrids.Add(RuntimeGridActor->GridSettings);
+
+		if (const FSpatialHashRuntimeGrid* ExistingRuntimeGrid = AllGrids.FindByPredicate([RuntimeGridActor](const FSpatialHashRuntimeGrid& RuntimeGrid) { return RuntimeGrid.GridName == RuntimeGridActor->GridSettings.GridName; }))
+		{
+			UE_LOG(LogWorldPartition, Error, TEXT("Got duplicated runtime grid actor '%s' for grid '%s'"), *SpatialHashRuntimeGridInfo->GetActorLabelOrName().ToString(), *RuntimeGridActor->GridSettings.GridName.ToString());
+			check(*ExistingRuntimeGrid == RuntimeGridActor->GridSettings);
+		}
+		else
+		{
+			AllGrids.Add(RuntimeGridActor->GridSettings);
+		}
 	}
 
 	TMap<FName, int32> GridsMapping;
