Added macroses with no param but the condition to evaluate for all UTEST
Added unit tests for these macros

diff --git a/Engine/Source/Runtime/Core/Public/Misc/AutomationTest.h b/Engine/Source/Runtime/Core/Public/Misc/AutomationTest.h
--- a/Engine/Source/Runtime/Core/Public/Misc/AutomationTest.h
+++ b/Engine/Source/Runtime/Core/Public/Misc/AutomationTest.h
@@ -3638,72 +3638,150 @@ public: \
 		return false;\
 	}
 
+#define UTEST_EQUAL_EXPR(Actual, Expected)\
+	if (!TestEqual(TEXT(#Actual), Actual, Expected))\
+	{\
+		return false;\
+	}
+
 #define UTEST_EQUAL_TOLERANCE(What, Actual, Expected, Tolerance)\
 	if (!TestEqual(What, Actual, Expected, Tolerance))\
 	{\
 		return false;\
 	}
 
+#define UTEST_EQUAL_TOLERANCE_EXPR(Actual, Expected, Tolerance)\
+	if (!TestEqual(TEXT(#Actual), Actual, Expected, Tolerance))\
+	{\
+		return false;\
+	}
+
+#define UTEST_NEARLY_EQUAL_EXPR(Actual, Expected, Tolerance)\
+	if (!TestNearlyEqual(TEXT(#Actual), Actual, Expected, Tolerance))\
+	{\
+		return false;\
+	}
+
 #define UTEST_EQUAL_INSENSITIVE(What, Actual, Expected)\
 	if (!TestEqualInsensitive(What, Actual, Expected))\
 	{\
 		return false;\
 	}
 
+#define UTEST_EQUAL_INSENSITIVE_EXPR(Actual, Expected)\
+	if (!TestEqualInsensitive(TEXT(#Actual), Actual, Expected))\
+	{\
+		return false;\
+	}
+
 #define UTEST_NOT_EQUAL(What, Actual, Expected)\
 	if (!TestNotEqual(What, Actual, Expected))\
 	{\
 		return false;\
 	}
 
+#define UTEST_NOT_EQUAL_EXPR(Actual, Expected)\
+	if (!TestNotEqual(FString::Printf(TEXT("%s != %s"), TEXT(#Actual), TEXT(#Expected)), Actual, Expected))\
+	{\
+		return false;\
+	}
+
 #define UTEST_SAME(What, Actual, Expected)\
 	if (!TestSame(What, Actual, Expected))\
 	{\
 		return false;\
 	}
 
+#define UTEST_SAME_EXPR(Actual, Expected)\
+	if (!TestSame(FString::Printf(TEXT("%s == %s"), TEXT(#Actual), TEXT(#Expected)), Actual, Expected))\
+	{\
+		return false;\
+	}
+
 #define UTEST_NOT_SAME(What, Actual, Expected)\
 	if (!TestNotSame(What, Actual, Expected))\
 	{\
 		return false;\
 	}
 
+#define UTEST_NOT_SAME_EXPR(Actual, Expected)\
+	if (!TestNotSame(FString::Printf(TEXT("%s != %s"), TEXT(#Actual), TEXT(#Expected)), Actual, Expected))\
+	{\
+		return false;\
+	}
+
 #define UTEST_TRUE(What, Value)\
 	if (!TestTrue(What, Value))\
 	{\
 		return false;\
 	}
 
+#define UTEST_TRUE_EXPR(Expression)\
+	if (!TestTrue(TEXT(#Expression), Expression))\
+	{\
+		return false;\
+	}
+
 #define UTEST_FALSE(What, Value)\
 	if (!TestFalse(What, Value))\
 	{\
 		return false;\
 	}
 
+#define UTEST_FALSE_EXPR(Expression)\
+	if (!TestFalse(TEXT(#Expression), Expression))\
+	{\
+		return false;\
+	}
+
 #define UTEST_VALID(What, SharedPointer)\
 	if (!TestValid(What, SharedPointer))\
 	{\
 		return false;\
 	}
 
+#define UTEST_VALID_EXPR(SharedPointer)\
+	if (!TestValid(TEXT(#SharedPointer), SharedPointer))\
+	{\
+		return false;\
+	}
+
 #define UTEST_INVALID(What, SharedPointer)\
 	if (!TestInvalid(What, SharedPointer))\
 	{\
 		return false;\
 	}
 
+#define UTEST_INVALID_EXPR(SharedPointer)\
+	if (!TestInvalid(TEXT(#SharedPointer), SharedPointer))\
+	{\
+		return false;\
+	}
+
 #define UTEST_NULL(What, Pointer)\
 	if (!TestNull(What, Pointer))\
 	{\
 		return false;\
 	}
 
+#define UTEST_NULL_EXPR(Pointer)\
+	if (!TestNull(TEXT(#Pointer), Pointer))\
+	{\
+		return false;\
+	}
+
 #define UTEST_NOT_NULL(What, Pointer)\
 	if (!TestNotNull(What, Pointer))\
 	{\
 		return false;\
 	}
 
+#define UTEST_NOT_NULL_EXPR(Pointer)\
+	if (!TestNotNull(TEXT(#Pointer), Pointer))\
+	{\
+		return false;\
+	}
+
 //////////////////////////////////////////////////
 // Basic Latent Commands
 
diff --git a/Engine/Source/Runtime/Engine/Private/Tests/EngineAutomationTests.cpp b/Engine/Source/Runtime/Engine/Private/Tests/EngineAutomationTests.cpp
--- a/Engine/Source/Runtime/Engine/Private/Tests/EngineAutomationTests.cpp
+++ b/Engine/Source/Runtime/Engine/Private/Tests/EngineAutomationTests.cpp
@@ -834,6 +834,92 @@ bool FAutomationTestNearlyEqualFTransformNegative::RunTest(const FString& Parame
 	return true;
 }
 
+class FAutomationUTestMacrosExpr : public FAutomationTestBase
+{
+public:
+	FAutomationUTestMacrosExpr(const FString& InName, const bool bInComplexTask)
+		: FAutomationTestBase(InName, bInComplexTask)
+	{
+	}
+
+protected:
+
+	static const float PositiveToleranceFloat;
+	static const float ActualFloatValue;
+	static const float ExpectedFloatValue;
+	static const float WrongFloatValue;
+	static const float ExpectedFloatValueOutOfTolerance;
+	static const FString ActualFStringValue;
+	static const FString ExpectedFStringValueLowerCase;
+};
+
+const float FAutomationUTestMacrosExpr::PositiveToleranceFloat(1.e-4f);
+const float FAutomationUTestMacrosExpr::ActualFloatValue(0.f);
+const float FAutomationUTestMacrosExpr::ExpectedFloatValue(ActualFloatValue);
+const float FAutomationUTestMacrosExpr::WrongFloatValue(ActualFloatValue + 1.f);
+const float FAutomationUTestMacrosExpr::ExpectedFloatValueOutOfTolerance(ActualFloatValue + PositiveToleranceFloat);
+const FString FAutomationUTestMacrosExpr::ActualFStringValue(TEXT("EQUALS"));
+const FString FAutomationUTestMacrosExpr::ExpectedFStringValueLowerCase(TEXT("equals"));
+
+IMPLEMENT_CUSTOM_SIMPLE_AUTOMATION_TEST(FAutomationEqualEXPR, FAutomationUTestMacrosExpr, "System.Automation.Validation.UTestEqual", EAutomationTestFlags::ApplicationContextMask | EAutomationTestFlags::EngineFilter);
+bool FAutomationEqualEXPR::RunTest(const FString& Parameters)
+{
+
+	UTEST_EQUAL_EXPR(ActualFloatValue, ExpectedFloatValue);
+	UTEST_NEARLY_EQUAL_EXPR(ActualFloatValue, ExpectedFloatValueOutOfTolerance, PositiveToleranceFloat);
+	UTEST_EQUAL_TOLERANCE_EXPR(ActualFloatValue, ExpectedFloatValueOutOfTolerance, PositiveToleranceFloat);
+	UTEST_NOT_EQUAL_EXPR(ActualFloatValue, WrongFloatValue);
+	UTEST_EQUAL_INSENSITIVE_EXPR(*ActualFStringValue, *ExpectedFStringValueLowerCase);
+
+	return true;
+}
+
+IMPLEMENT_CUSTOM_SIMPLE_AUTOMATION_TEST(FAutomationSameNotSameEXPR, FAutomationUTestMacrosExpr, "System.Automation.Validation.UTestSameNotSame", EAutomationTestFlags::ApplicationContextMask | EAutomationTestFlags::EngineFilter);
+bool FAutomationSameNotSameEXPR::RunTest(const FString& Parameters)
+{
+
+	UTEST_SAME_EXPR(ActualFStringValue, ActualFStringValue);
+	UTEST_NOT_SAME_EXPR(ActualFStringValue, ExpectedFStringValueLowerCase);
+
+	return true;
+}
+
+IMPLEMENT_CUSTOM_SIMPLE_AUTOMATION_TEST(FAutomationTrueFalseEXPR, FAutomationUTestMacrosExpr, "System.Automation.Validation.UTestTrueFalse", EAutomationTestFlags::ApplicationContextMask | EAutomationTestFlags::EngineFilter);
+bool FAutomationTrueFalseEXPR::RunTest(const FString& Parameters)
+{
+
+	UTEST_TRUE_EXPR(ActualFloatValue == ExpectedFloatValue);
+	UTEST_FALSE_EXPR(ActualFloatValue > ExpectedFloatValue);
+
+	return true;
+}
+
+IMPLEMENT_CUSTOM_SIMPLE_AUTOMATION_TEST(FAutomationValidInvalidEXPR, FAutomationUTestMacrosExpr, "System.Automation.Validation.UTestValidInvalid", EAutomationTestFlags::ApplicationContextMask | EAutomationTestFlags::EngineFilter);
+bool FAutomationValidInvalidEXPR::RunTest(const FString& Parameters)
+{
+	struct HandleResumeContext
+	{
+		const char* PackageIdentifier = nullptr;
+	};
+	TSharedPtr<HandleResumeContext, ESPMode::ThreadSafe> Context = MakeShared<HandleResumeContext, ESPMode::ThreadSafe>();
+	TSharedPtr< UObject > Entry = nullptr;
+	UTEST_VALID_EXPR(Context);
+	UTEST_INVALID_EXPR(Entry);
+
+	return true;
+}
+
+IMPLEMENT_CUSTOM_SIMPLE_AUTOMATION_TEST(FAutomationNullNotNullPtrEXPR, FAutomationUTestMacrosExpr, "System.Automation.Validation.UTestNullNotNull", EAutomationTestFlags::ApplicationContextMask | EAutomationTestFlags::EngineFilter);
+bool FAutomationNullNotNullPtrEXPR::RunTest(const FString& Parameters)
+{
+	UWorld* World = UWorld::CreateWorld(EWorldType::Game, false);
+
+	UTEST_NULL_EXPR(nullptr);
+	UTEST_NOT_NULL_EXPR(World);
+
+	return true;
+}
+
 IMPLEMENT_SIMPLE_AUTOMATION_TEST(FAutomationAttachment, "System.Engine.Attachment", EAutomationTestFlags::EditorContext | EAutomationTestFlags::EngineFilter)
 
 #define DUMP_EXPECTED_TRANSFORMS 0
