World Partition Editor Loading: added custom loading adapters filtering and moved all Data Layers specific code into Data Layer subsystem, using the new filtering function.

diff --git a/Engine/Source/Editor/DataLayerEditor/Private/DataLayer/DataLayerEditorSubsystem.cpp b/Engine/Source/Editor/DataLayerEditor/Private/DataLayer/DataLayerEditorSubsystem.cpp
--- a/Engine/Source/Editor/DataLayerEditor/Private/DataLayer/DataLayerEditorSubsystem.cpp
+++ b/Engine/Source/Editor/DataLayerEditor/Private/DataLayer/DataLayerEditorSubsystem.cpp
@@ -48,6 +48,7 @@
 #include "WorldPartition/DataLayer/DataLayerManager.h"
 #include "WorldPartition/DataLayer/DeprecatedDataLayerInstance.h"
 #include "WorldPartition/DataLayer/WorldDataLayers.h"
+#include "WorldPartition/WorldPartitionActorDescViewProxy.h"
 #include "WorldPartition/WorldPartition.h"
 
 class SWidget;
@@ -198,6 +199,9 @@ void UDataLayerEditorSubsystem::Initialize(FSubsystemCollectionBase& Collection)
 
 	// Register the engine broadcast bridge
 	OnActorDataLayersEditorLoadingStateChangedEngineBridgeHandle = DataLayerEditorLoadingStateChanged.AddStatic(&FDataLayersEditorBroadcast::StaticOnActorDataLayersEditorLoadingStateChanged);
+
+	// Register actor descriptor loading filter
+	IWorldPartitionActorLoaderInterface::RegisterActorDescFilter([this](UWorld* World, const FWorldPartitionHandle& ActorHandle) { return PassDataLayersFilter(World, ActorHandle); });
 }
 
 void UDataLayerEditorSubsystem::Deinitialize()
@@ -1132,6 +1136,35 @@ bool UDataLayerEditorSubsystem::HasDeprecatedDataLayers() const
 	return false;
 }
 
+bool UDataLayerEditorSubsystem::PassDataLayersFilter(UWorld* World, const FWorldPartitionHandle& ActorHandle)
+{
+	UWorld* OwningWorld = World->PersistentLevel->GetWorld();
+
+	if (UDataLayerManager* DataLayerManager = UDataLayerManager::GetDataLayerManager(OwningWorld))
+	{
+		FWorldPartitionActorViewProxy ActorDescProxy(*ActorHandle);
+
+		if (IsRunningCookCommandlet())
+		{
+			// When running cook commandlet, dont allow loading of actors with runtime loaded data layers
+			for (const FName& DataLayerInstanceName : ActorDescProxy.GetDataLayerInstanceNames())
+			{
+				const UDataLayerInstance* DataLayerInstance = DataLayerManager->GetDataLayerInstance(DataLayerInstanceName);
+				if (DataLayerInstance && DataLayerInstance->IsRuntime())
+				{
+					return false;
+				}
+			}
+
+			return true;
+		}
+
+		return DataLayerManager->ResolveIsLoadedInEditor(ActorDescProxy.GetDataLayerInstanceNames());
+	}
+
+	return true;
+}
+
 UDataLayerInstance* UDataLayerEditorSubsystem::GetDataLayerInstance(const FName& DataLayerInstanceName) const
 {
 	if (UDataLayerManager* DataLayerManager = UDataLayerManager::GetDataLayerManager(GetWorld()))
diff --git a/Engine/Source/Editor/DataLayerEditor/Public/DataLayer/DataLayerEditorSubsystem.h b/Engine/Source/Editor/DataLayerEditor/Public/DataLayer/DataLayerEditorSubsystem.h
--- a/Engine/Source/Editor/DataLayerEditor/Public/DataLayer/DataLayerEditorSubsystem.h
+++ b/Engine/Source/Editor/DataLayerEditor/Public/DataLayer/DataLayerEditorSubsystem.h
@@ -22,6 +22,7 @@
 #include "UObject/UnrealNames.h"
 #include "UObject/WeakObjectPtr.h"
 #include "UObject/WeakObjectPtrTemplates.h"
+#include "WorldPartition/WorldPartitionHandle.h"
 #include "WorldPartition/DataLayer/ActorDataLayer.h"
 #include "WorldPartition/DataLayer/DataLayerInstance.h"
 #include "WorldPartition/DataLayer/WorldDataLayers.h"
@@ -713,6 +714,8 @@ private:
 
 	bool SetDataLayerIsLoadedInEditorInternal(UDataLayerInstance* DataLayer, const bool bIsLoadedInEditor, const bool bIsFromUserChange);
 
+	bool PassDataLayersFilter(UWorld* World, const FWorldPartitionHandle& ActorHandle);
+
 	void BroadcastActorDataLayersChanged(const TWeakObjectPtr<AActor>& ChangedActor);
 	void BroadcastDataLayerChanged(const EDataLayerAction Action, const TWeakObjectPtr<const UDataLayerInstance>& ChangedDataLayer, const FName& ChangedProperty);
 	void BroadcastDataLayerEditorLoadingStateChanged(bool bIsFromUserChange);
diff --git a/Engine/Source/Runtime/Engine/Private/WorldPartition/DataLayer/DataLayerManager.cpp b/Engine/Source/Runtime/Engine/Private/WorldPartition/DataLayer/DataLayerManager.cpp
--- a/Engine/Source/Runtime/Engine/Private/WorldPartition/DataLayer/DataLayerManager.cpp
+++ b/Engine/Source/Runtime/Engine/Private/WorldPartition/DataLayer/DataLayerManager.cpp
@@ -19,6 +19,7 @@
 #include "WorldPartition/DataLayer/IDataLayerEditorModule.h"
 #include "WorldPartition/ActorDescContainer.h"
 #include "WorldPartition/WorldPartitionEditorPerProjectUserSettings.h"
+#include "WorldPartition/WorldPartitionActorLoaderInterface.h"
 #include "WorldPartition/ActorDescList.h"
 #include "LevelInstance/LevelInstanceSubsystem.h"
 #include "Modules/ModuleManager.h"
@@ -784,6 +785,7 @@ FDataLayersEditorBroadcast& FDataLayersEditorBroadcast::Get()
 void FDataLayersEditorBroadcast::StaticOnActorDataLayersEditorLoadingStateChanged(bool bIsFromUserChange)
 {
 	Get().DataLayerEditorLoadingStateChanged.Broadcast(bIsFromUserChange);
+	IWorldPartitionActorLoaderInterface::RefreshLoadedState(bIsFromUserChange);
 }
 
 #endif
\ No newline at end of file
diff --git a/Engine/Source/Runtime/Engine/Private/WorldPartition/WorldPartitionActorLoaderInterface.cpp b/Engine/Source/Runtime/Engine/Private/WorldPartition/WorldPartitionActorLoaderInterface.cpp
--- a/Engine/Source/Runtime/Engine/Private/WorldPartition/WorldPartitionActorLoaderInterface.cpp
+++ b/Engine/Source/Runtime/Engine/Private/WorldPartition/WorldPartitionActorLoaderInterface.cpp
@@ -9,9 +9,6 @@
 #include "Engine/Level.h"
 #include "Misc/ScopedSlowTask.h"
 #include "WorldPartition/WorldPartition.h"
-#include "WorldPartition/WorldPartitionActorDescViewProxy.h"
-#include "WorldPartition/DataLayer/WorldDataLayers.h"
-#include "WorldPartition/DataLayer/DataLayerManager.h"
 #endif
 
 #define LOCTEXT_NAMESPACE "WorldPartition"
@@ -21,6 +18,14 @@ UWorldPartitionActorLoaderInterface::UWorldPartitionActorLoaderInterface(const F
 {}
 
 #if WITH_EDITOR
+TArray<IWorldPartitionActorLoaderInterface::FActorDescFilter> IWorldPartitionActorLoaderInterface::ActorDescFilters;
+IWorldPartitionActorLoaderInterface::FOnActorLoaderInterfaceRefreshState IWorldPartitionActorLoaderInterface::ActorLoaderInterfaceRefreshState;
+
+void IWorldPartitionActorLoaderInterface::RegisterActorDescFilter(const FActorDescFilter& InActorDescFilter)
+{
+	ActorDescFilters.Add(InActorDescFilter);
+}
+
 IWorldPartitionActorLoaderInterface::ILoaderAdapter::ILoaderAdapter(UWorld* InWorld)
 	: World(InWorld)
 	, bLoaded(false)
@@ -94,8 +99,9 @@ bool IWorldPartitionActorLoaderInterface::ILoaderAdapter::IsLoaded() const
 
 void IWorldPartitionActorLoaderInterface::ILoaderAdapter::RegisterDelegates()
 {
-	FDataLayersEditorBroadcast::Get().OnActorDataLayersEditorLoadingStateChanged().AddRaw(this, &IWorldPartitionActorLoaderInterface::ILoaderAdapter::OnActorDataLayersEditorLoadingStateChanged);
-	check(World->GetWorldPartition() != nullptr);
+	ActorLoaderInterfaceRefreshState.AddRaw(this, &IWorldPartitionActorLoaderInterface::ILoaderAdapter::OnRefreshLoadedState);
+
+	check(World->GetWorldPartition());
 	if (UWorldPartition* WorldPartition = World->GetWorldPartition())
 	{
 		WorldPartition->OnActorDescContainerRegistered.AddRaw(this, &IWorldPartitionActorLoaderInterface::ILoaderAdapter::ILoaderAdapter::OnActorDescContainerInitialize);
@@ -105,8 +111,9 @@ void IWorldPartitionActorLoaderInterface::ILoaderAdapter::RegisterDelegates()
 
 void IWorldPartitionActorLoaderInterface::ILoaderAdapter::UnregisterDelegates()
 {
-	FDataLayersEditorBroadcast::Get().OnActorDataLayersEditorLoadingStateChanged().RemoveAll(this);
-	check(World->GetWorldPartition() != nullptr);
+	ActorLoaderInterfaceRefreshState.RemoveAll(this);
+
+	check(World->GetWorldPartition());
 	if (UWorldPartition* WorldPartition = World->GetWorldPartition())
 	{
 		WorldPartition->OnActorDescContainerRegistered.RemoveAll(this);
@@ -119,42 +126,6 @@ bool IWorldPartitionActorLoaderInterface::ILoaderAdapter::PassActorDescFilter(co
 	return Actor->IsEditorRelevant();
 }
 
-// Helper added to since the nested class IWorldPartitionActorLoaderInterface::ILoaderAdapter can't be friend of UDataLayerManager.
-bool IWorldPartitionActorLoaderInterface::PassDataLayersFilter(UWorld* InWorld, const TArray<FName>& InDataLayerInstanceNames)
-{
-	check(!IsRunningCookCommandlet());
-	UDataLayerManager* DataLayerManager = UDataLayerManager::GetDataLayerManager(InWorld);
-	return DataLayerManager ? DataLayerManager->ResolveIsLoadedInEditor(InDataLayerInstanceNames) : true;
-}
-
-bool IWorldPartitionActorLoaderInterface::ILoaderAdapter::PassDataLayersFilter(const FWorldPartitionHandle& Actor) const
-{
-	UWorld* OwningWorld = World->PersistentLevel->GetWorld();
-	if (UDataLayerManager* DataLayerManager = UDataLayerManager::GetDataLayerManager(OwningWorld))
-	{
-		FWorldPartitionActorViewProxy ActorDescProxy(*Actor);
-
-		if (IsRunningCookCommandlet())
-		{
-			// When running cook commandlet, dont allow loading of actors with runtime loaded data layers
-			for (const FName& DataLayerInstanceName : ActorDescProxy.GetDataLayerInstanceNames())
-			{
-				const UDataLayerInstance* DataLayerInstance = DataLayerManager->GetDataLayerInstance(DataLayerInstanceName);
-				if (DataLayerInstance && DataLayerInstance->IsRuntime())
-				{
-					return false;
-				}
-			}
-		}
-		else
-		{
-			return IWorldPartitionActorLoaderInterface::PassDataLayersFilter(OwningWorld, ActorDescProxy.GetDataLayerInstanceNames());
-		}
-	}
-
-	return true;
-}
-
 void IWorldPartitionActorLoaderInterface::ILoaderAdapter::RefreshLoadedState()
 {
 	if (bLoaded)
@@ -261,7 +232,21 @@ void IWorldPartitionActorLoaderInterface::ILoaderAdapter::OnActorDescContainerUn
 bool IWorldPartitionActorLoaderInterface::ILoaderAdapter::ShouldActorBeLoaded(const FWorldPartitionHandle& Actor) const
 {
 	check(Actor.IsValid());
-	return PassActorDescFilter(Actor) && PassDataLayersFilter(Actor);
+
+	if (!PassActorDescFilter(Actor))
+	{
+		return false;
+	}
+
+	for (auto& ActorDescFilter : ActorDescFilters)
+	{
+		if (!ActorDescFilter(World, Actor))
+		{
+			return false;
+		}
+	}
+
+	return true;
 };
 
 void IWorldPartitionActorLoaderInterface::ILoaderAdapter::PostLoadedStateChanged(int32 NumLoads, int32 NumUnloads)
@@ -318,10 +303,15 @@ void IWorldPartitionActorLoaderInterface::ILoaderAdapter::RemoveReferenceToActor
 	ActorReferences.Remove(ActorHandle->GetGuid());
 }
 
-void IWorldPartitionActorLoaderInterface::ILoaderAdapter::OnActorDataLayersEditorLoadingStateChanged(bool bFromUserOperation)
+void IWorldPartitionActorLoaderInterface::ILoaderAdapter::OnRefreshLoadedState(bool bFromUserOperation)
 {
 	RefreshLoadedState();
 }
+
+void IWorldPartitionActorLoaderInterface::RefreshLoadedState(bool bIsFromUserChange)
+{
+	ActorLoaderInterfaceRefreshState.Broadcast(bIsFromUserChange);
+}
 #endif
 
 #undef LOCTEXT_NAMESPACE
diff --git a/Engine/Source/Runtime/Engine/Public/WorldPartition/WorldPartitionActorDescViewProxy.h b/Engine/Source/Runtime/Engine/Public/WorldPartition/WorldPartitionActorDescViewProxy.h
--- a/Engine/Source/Runtime/Engine/Public/WorldPartition/WorldPartitionActorDescViewProxy.h
+++ b/Engine/Source/Runtime/Engine/Public/WorldPartition/WorldPartitionActorDescViewProxy.h
@@ -7,7 +7,7 @@
 /**
  * A view on top of an actor desc, which reverts to the actor if its loaded and dirty.
  */
-class FWorldPartitionActorViewProxy : public FWorldPartitionActorDescView
+class ENGINE_API FWorldPartitionActorViewProxy : public FWorldPartitionActorDescView
 {
 public:
 	FWorldPartitionActorViewProxy(const FWorldPartitionActorDesc* InActorDesc);
diff --git a/Engine/Source/Runtime/Engine/Public/WorldPartition/WorldPartitionActorLoaderInterface.h b/Engine/Source/Runtime/Engine/Public/WorldPartition/WorldPartitionActorLoaderInterface.h
--- a/Engine/Source/Runtime/Engine/Public/WorldPartition/WorldPartitionActorLoaderInterface.h
+++ b/Engine/Source/Runtime/Engine/Public/WorldPartition/WorldPartitionActorLoaderInterface.h
@@ -55,13 +55,12 @@ public:
 
 		// Actors filtering
 		virtual bool PassActorDescFilter(const FWorldPartitionHandle& Actor) const;
-		virtual bool PassDataLayersFilter(const FWorldPartitionHandle& Actor) const;
 		void RefreshLoadedState();
 
 		void PostLoadedStateChanged(int32 NumLoads, int32 NumUnloads);
 		void AddReferenceToActor(FWorldPartitionHandle& Actor);
 		void RemoveReferenceToActor(FWorldPartitionHandle& Actor);
-		void OnActorDataLayersEditorLoadingStateChanged(bool bFromUserOperation);		
+		void OnRefreshLoadedState(bool bFromUserOperation);		
 
 	private:
 		UWorld* World;
@@ -74,8 +73,16 @@ public:
 
 	virtual ILoaderAdapter* GetLoaderAdapter() =0;
 
+	using FActorDescFilter = TFunction<bool(class UWorld*, const FWorldPartitionHandle&)>;
+	static void RegisterActorDescFilter(const FActorDescFilter& InActorDescFilter);
+
+	static void RefreshLoadedState(bool bIsFromUserChange);
+
 private:
-	static bool PassDataLayersFilter(class UWorld* InWorld, const TArray<FName>& InDataLayerInstanceNames);
+	DECLARE_EVENT_OneParam(IWorldPartitionActorLoaderInterface, FOnActorLoaderInterfaceRefreshState, bool /*bIsFromUserChange*/);
+	static FOnActorLoaderInterfaceRefreshState ActorLoaderInterfaceRefreshState;
+
+	static TArray<FActorDescFilter> ActorDescFilters;
 	friend class ILoaderAdapter;
 #endif
 };
\ No newline at end of file
