Refresh render proxies when NaniteSettings.bAllowMaskedMaterials changes.

diff --git a/Engine/Source/Runtime/Engine/Private/WorldSettings.cpp b/Engine/Source/Runtime/Engine/Private/WorldSettings.cpp
--- a/Engine/Source/Runtime/Engine/Private/WorldSettings.cpp
+++ b/Engine/Source/Runtime/Engine/Private/WorldSettings.cpp
@@ -29,6 +29,7 @@
 #include "Engine/BookMark.h"
 #include "WorldSettingsCustomVersion.h"
 #include "Materials/Material.h"
+#include "ComponentRecreateRenderStateContext.h"
 
 #if WITH_EDITOR
 #include "Editor.h"
@@ -814,6 +815,11 @@ void AWorldSettings::InternalPostPropertyChanged(FName PropertyName)
 				}
 			}
 		}
+		else if (PropertyName == GET_MEMBER_NAME_CHECKED(FNaniteSettings, bAllowMaskedMaterials))
+		{
+			// Need to recreate scene proxies when this flag changes.
+			FGlobalComponentRecreateRenderStateContext Context;
+		}
 	}
 }
 
