No longer need to log invalid bodies - we have seen that the log hits in smoke tests which means invalid bodies are being updated in the acceleration structure via deferred kinematic update. The logs are creating spam in tests.

#localization none
#tests preflight:https://horde.devtools.epicgames.com/job/63f3f46990198dffbaff19cb

diff --git a/Engine/Source/Runtime/Engine/Private/PhysicsEngine/Experimental/PhysScene_Chaos.cpp b/Engine/Source/Runtime/Engine/Private/PhysicsEngine/Experimental/PhysScene_Chaos.cpp
--- a/Engine/Source/Runtime/Engine/Private/PhysicsEngine/Experimental/PhysScene_Chaos.cpp
+++ b/Engine/Source/Runtime/Engine/Private/PhysicsEngine/Experimental/PhysScene_Chaos.cpp
@@ -50,7 +50,7 @@ FAutoConsoleVariableRef CVar_KinematicDeferralCheckValidBodies(TEXT("p.Kinematic
 
 bool GKinematicDeferralUpdateExternalAccelerationStructure = true;
 FAutoConsoleVariableRef CVar_KinematicDeferralUpdateExternalAccelerationStructure(TEXT("p.KinematicDeferralUpdateExternalAccelerationStructure"), GKinematicDeferralUpdateExternalAccelerationStructure, TEXT("If true, process any operations in PendingSpatialOperations_External before doing deferred kinematic updates."));
-bool GKinematicDeferralLogInvalidBodies = true;
+bool GKinematicDeferralLogInvalidBodies = false;
 FAutoConsoleVariableRef CVar_KinematicDeferralLogInvalidBodies(TEXT("p.KinematicDeferralLogInvalidBodies"), GKinematicDeferralLogInvalidBodies, TEXT("If true and p.KinematicDeferralCheckValidBodies is true, log when an invalid body is found on kinematic update."));
 
 DECLARE_CYCLE_STAT(TEXT("Update Kinematics On Deferred SkelMeshes"), STAT_UpdateKinematicsOnDeferredSkelMeshesChaos, STATGROUP_Physics);
