World Partition: added a way to insert a proxy error handler during streaming generation.

diff --git a/Engine/Source/Runtime/Engine/Private/WorldPartition/ErrorHandling/WorldPartitionStreamingGenerationProxyErrorHandler.cpp b/Engine/Source/Runtime/Engine/Private/WorldPartition/ErrorHandling/WorldPartitionStreamingGenerationProxyErrorHandler.cpp
new file mode 100644
--- /dev/null
+++ b/Engine/Source/Runtime/Engine/Private/WorldPartition/ErrorHandling/WorldPartitionStreamingGenerationProxyErrorHandler.cpp
@@ -0,0 +1,65 @@
+// Copyright Epic Games, Inc. All Rights Reserved.
+
+#if WITH_EDITOR
+#include "WorldPartition/ErrorHandling/WorldPartitionStreamingGenerationProxyErrorHandler.h"
+
+void FStreamingGenerationProxyErrorHandler::OnInvalidRuntimeGrid(const FWorldPartitionActorDescView& ActorDescView, FName GridName)
+{
+	InnerErrorHandler->OnInvalidRuntimeGrid(ActorDescView, GridName);
+}
+
+void FStreamingGenerationProxyErrorHandler::OnInvalidReference(const FWorldPartitionActorDescView& ActorDescView, const FGuid& ReferenceGuid, FWorldPartitionActorDescView* ReferenceActorDescView)
+{
+	InnerErrorHandler->OnInvalidReference(ActorDescView, ReferenceGuid, ReferenceActorDescView);
+}
+
+void FStreamingGenerationProxyErrorHandler::OnInvalidReferenceGridPlacement(const FWorldPartitionActorDescView& ActorDescView, const FWorldPartitionActorDescView& ReferenceActorDescView)
+{
+	InnerErrorHandler->OnInvalidReferenceGridPlacement(ActorDescView, ReferenceActorDescView);
+}
+
+void FStreamingGenerationProxyErrorHandler::OnInvalidReferenceDataLayers(const FWorldPartitionActorDescView& ActorDescView, const FWorldPartitionActorDescView& ReferenceActorDescView)
+{
+	InnerErrorHandler->OnInvalidReferenceDataLayers(ActorDescView, ReferenceActorDescView);
+}
+
+void FStreamingGenerationProxyErrorHandler::OnInvalidReferenceRuntimeGrid(const FWorldPartitionActorDescView& ActorDescView, const FWorldPartitionActorDescView& ReferenceActorDescView)
+{
+	InnerErrorHandler->OnInvalidReferenceRuntimeGrid(ActorDescView, ReferenceActorDescView);
+}
+
+void FStreamingGenerationProxyErrorHandler::OnInvalidReferenceLevelScriptStreamed(const FWorldPartitionActorDescView& ActorDescView)
+{
+	InnerErrorHandler->OnInvalidReferenceLevelScriptStreamed(ActorDescView);
+}
+
+void FStreamingGenerationProxyErrorHandler::OnInvalidReferenceLevelScriptDataLayers(const FWorldPartitionActorDescView& ActorDescView)
+{
+	InnerErrorHandler->OnInvalidReferenceLevelScriptDataLayers(ActorDescView);
+}
+
+void FStreamingGenerationProxyErrorHandler::OnInvalidReferenceDataLayerAsset(const UDataLayerInstanceWithAsset* DataLayerInstance)
+{
+	InnerErrorHandler->OnInvalidReferenceDataLayerAsset(DataLayerInstance);
+}
+
+void FStreamingGenerationProxyErrorHandler::OnDataLayerHierarchyTypeMismatch(const UDataLayerInstance* DataLayerInstance, const UDataLayerInstance* Parent)
+{
+	InnerErrorHandler->OnDataLayerHierarchyTypeMismatch(DataLayerInstance, Parent);
+}
+
+void FStreamingGenerationProxyErrorHandler::OnDataLayerAssetConflict(const UDataLayerInstanceWithAsset* DataLayerInstance, const UDataLayerInstanceWithAsset* ConflictingDataLayerInstance)
+{
+	InnerErrorHandler->OnDataLayerAssetConflict(DataLayerInstance, ConflictingDataLayerInstance);
+}
+
+void FStreamingGenerationProxyErrorHandler::OnActorNeedsResave(const FWorldPartitionActorDescView& ActorDescView)
+{
+	InnerErrorHandler->OnActorNeedsResave(ActorDescView);
+}
+
+void FStreamingGenerationProxyErrorHandler::OnLevelInstanceInvalidWorldAsset(const FWorldPartitionActorDescView& ActorDescView, FName WorldAsset, ELevelInstanceInvalidReason Reason)
+{
+	InnerErrorHandler->OnLevelInstanceInvalidWorldAsset(ActorDescView, WorldAsset, Reason);
+}
+#endif
diff --git a/Engine/Source/Runtime/Engine/Private/WorldPartition/WorldPartitionActorDescView.cpp b/Engine/Source/Runtime/Engine/Private/WorldPartition/WorldPartitionActorDescView.cpp
--- a/Engine/Source/Runtime/Engine/Private/WorldPartition/WorldPartitionActorDescView.cpp
+++ b/Engine/Source/Runtime/Engine/Private/WorldPartition/WorldPartitionActorDescView.cpp
@@ -243,4 +243,13 @@ AActor* FWorldPartitionActorDescView::GetActor() const
 {
 	return ActorDesc->GetActor();
 }
+
+bool FWorldPartitionActorDescView::GetProperty(FName PropertyName, FName* PropertyValue) const
+{
+	return ActorDesc->GetProperty(PropertyName, PropertyValue);
+}
+bool FWorldPartitionActorDescView::HasProperty(FName PropertyName) const
+{
+	return ActorDesc->HasProperty(PropertyName);
+}
 #endif
diff --git a/Engine/Source/Runtime/Engine/Private/WorldPartition/WorldPartitionStreamingGeneration.cpp b/Engine/Source/Runtime/Engine/Private/WorldPartition/WorldPartitionStreamingGeneration.cpp
--- a/Engine/Source/Runtime/Engine/Private/WorldPartition/WorldPartitionStreamingGeneration.cpp
+++ b/Engine/Source/Runtime/Engine/Private/WorldPartition/WorldPartitionStreamingGeneration.cpp
@@ -1034,7 +1034,7 @@ bool UWorldPartition::GenerateContainerStreaming(const UActorDescContainer* InAc
 	StreamingGeneratorParams.WorldPartitionContext = this;
 	StreamingGeneratorParams.ModifiedActorsDescList = ModifiedActorsDescList;
 	StreamingGeneratorParams.IsValidGrid = [this](FName GridName) { return RuntimeHash->IsValidGrid(GridName); };
-	StreamingGeneratorParams.ErrorHandler = ErrorHandler;
+	StreamingGeneratorParams.ErrorHandler = StreamingGenerationErrorHandlerOverride ? (*StreamingGenerationErrorHandlerOverride)(ErrorHandler) : ErrorHandler;
 	StreamingGeneratorParams.bEnableStreaming = IsStreamingEnabled();
 
 	FWorldPartitionStreamingGenerator StreamingGenerator(StreamingGeneratorParams);
@@ -1077,7 +1077,7 @@ void UWorldPartition::GenerateHLOD(ISourceControlHelper* SourceControlHelper, bo
 		FStreamingGenerationLogErrorHandler LogErrorHandler;
 		FWorldPartitionStreamingGenerator::FWorldPartitionStreamingGeneratorParams StreamingGeneratorParams;
 		StreamingGeneratorParams.WorldPartitionContext = this;
-		StreamingGeneratorParams.ErrorHandler = &LogErrorHandler;
+		StreamingGeneratorParams.ErrorHandler = StreamingGenerationErrorHandlerOverride ? (*StreamingGenerationErrorHandlerOverride)(&LogErrorHandler) : &LogErrorHandler;
 		StreamingGeneratorParams.bEnableStreaming = IsStreamingEnabled();
 		StreamingGeneratorParams.FilteredClasses.Add(AWorldPartitionHLOD::StaticClass());
 		StreamingGeneratorParams.IsValidGrid = [this](FName GridName) { return RuntimeHash->IsValidGrid(GridName); };
@@ -1100,7 +1100,7 @@ void UWorldPartition::CheckForErrors(IStreamingGenerationErrorHandler* ErrorHand
 	FWorldPartitionStreamingGenerator::FWorldPartitionStreamingGeneratorParams StreamingGeneratorParams;
 	StreamingGeneratorParams.WorldPartitionContext = this;
 	StreamingGeneratorParams.ModifiedActorsDescList = &ModifiedActorDescList;
-	StreamingGeneratorParams.ErrorHandler = ErrorHandler;
+	StreamingGeneratorParams.ErrorHandler = StreamingGenerationErrorHandlerOverride ? (*StreamingGenerationErrorHandlerOverride)(ErrorHandler) : ErrorHandler;
 	StreamingGeneratorParams.IsValidGrid = [this](FName GridName) { return RuntimeHash->IsValidGrid(GridName); };
 	StreamingGeneratorParams.bEnableStreaming = IsStreamingEnabled();
 
@@ -1139,7 +1139,7 @@ void UWorldPartition::CheckForErrors(const FCheckForErrorsParams& Params)
 	FWorldPartitionStreamingGenerator::FWorldPartitionStreamingGeneratorParams StreamingGeneratorParams;
 	StreamingGeneratorParams.WorldPartitionContext = Params.ActorDescContainer->GetWorldPartition();
 	StreamingGeneratorParams.ModifiedActorsDescList = !Params.ActorDescContainer->IsTemplateContainer() ? &ModifiedActorDescList : nullptr;
-	StreamingGeneratorParams.ErrorHandler = Params.ErrorHandler;
+	StreamingGeneratorParams.ErrorHandler = StreamingGenerationErrorHandlerOverride ? (*StreamingGenerationErrorHandlerOverride)(Params.ErrorHandler) : Params.ErrorHandler;
 	StreamingGeneratorParams.IsValidGrid = [](FName GridName) { return true; };
 	StreamingGeneratorParams.bEnableStreaming = Params.bEnableStreaming;
 	StreamingGeneratorParams.ActorGuidsToContainerMap = Params.ActorGuidsToContainerMap;
diff --git a/Engine/Source/Runtime/Engine/Public/WorldPartition/ErrorHandling/WorldPartitionStreamingGenerationProxyErrorHandler.h b/Engine/Source/Runtime/Engine/Public/WorldPartition/ErrorHandling/WorldPartitionStreamingGenerationProxyErrorHandler.h
new file mode 100644
--- /dev/null
+++ b/Engine/Source/Runtime/Engine/Public/WorldPartition/ErrorHandling/WorldPartitionStreamingGenerationProxyErrorHandler.h
@@ -0,0 +1,30 @@
+// Copyright Epic Games, Inc. All Rights Reserved.
+#pragma once
+
+#if WITH_EDITOR
+#include "WorldPartition/ErrorHandling/WorldPartitionStreamingGenerationErrorHandler.h"
+
+class ENGINE_API FStreamingGenerationProxyErrorHandler : public IStreamingGenerationErrorHandler
+{
+public:
+	FStreamingGenerationProxyErrorHandler(IStreamingGenerationErrorHandler* InInnerErrorHandler)
+		: InnerErrorHandler(InInnerErrorHandler)
+	{}
+
+	virtual void OnInvalidRuntimeGrid(const FWorldPartitionActorDescView& ActorDescView, FName GridName) override;
+	virtual void OnInvalidReference(const FWorldPartitionActorDescView& ActorDescView, const FGuid& ReferenceGuid, FWorldPartitionActorDescView* ReferenceActorDescView) override;
+	virtual void OnInvalidReferenceGridPlacement(const FWorldPartitionActorDescView& ActorDescView, const FWorldPartitionActorDescView& ReferenceActorDescView) override;
+	virtual void OnInvalidReferenceDataLayers(const FWorldPartitionActorDescView& ActorDescView, const FWorldPartitionActorDescView& ReferenceActorDescView) override;
+	virtual void OnInvalidReferenceLevelScriptStreamed(const FWorldPartitionActorDescView& ActorDescView) override;
+	virtual void OnInvalidReferenceLevelScriptDataLayers(const FWorldPartitionActorDescView& ActorDescView) override;
+	virtual void OnInvalidReferenceRuntimeGrid(const FWorldPartitionActorDescView& ActorDescView, const FWorldPartitionActorDescView& ReferenceActorDescView) override;
+	virtual void OnInvalidReferenceDataLayerAsset(const UDataLayerInstanceWithAsset* DataLayerInstance) override;
+	virtual void OnDataLayerHierarchyTypeMismatch(const UDataLayerInstance* DataLayerInstance, const UDataLayerInstance* Parent) override;;
+	virtual void OnDataLayerAssetConflict(const UDataLayerInstanceWithAsset* DataLayerInstance, const UDataLayerInstanceWithAsset* ConflictingDataLayerInstance) override;;
+	virtual void OnActorNeedsResave(const FWorldPartitionActorDescView& ActorDescView) override;;
+	virtual void OnLevelInstanceInvalidWorldAsset(const FWorldPartitionActorDescView& ActorDescView, FName WorldAsset, ELevelInstanceInvalidReason Reason) override;;
+
+protected:
+	IStreamingGenerationErrorHandler* InnerErrorHandler;
+};
+#endif
diff --git a/Engine/Source/Runtime/Engine/Public/WorldPartition/WorldPartition.h b/Engine/Source/Runtime/Engine/Public/WorldPartition/WorldPartition.h
--- a/Engine/Source/Runtime/Engine/Public/WorldPartition/WorldPartition.h
+++ b/Engine/Source/Runtime/Engine/Public/WorldPartition/WorldPartition.h
@@ -216,6 +216,9 @@ public:
 
 	static void CheckForErrors(const FCheckForErrorsParams& Params);
 
+	using FStreamingGenerationErrorHandlerOverride = TFunction<IStreamingGenerationErrorHandler*(IStreamingGenerationErrorHandler* InErrorHandler)>;
+	inline static TOptional<FStreamingGenerationErrorHandlerOverride> StreamingGenerationErrorHandlerOverride;
+
 	void AppendAssetRegistryTags(TArray<FAssetRegistryTag>& OutTags) const;
 
 	DECLARE_MULTICAST_DELEGATE_OneParam(FActorDescContainerRegistrationDelegate, UActorDescContainer*);
diff --git a/Engine/Source/Runtime/Engine/Public/WorldPartition/WorldPartitionActorDescView.h b/Engine/Source/Runtime/Engine/Public/WorldPartition/WorldPartitionActorDescView.h
--- a/Engine/Source/Runtime/Engine/Public/WorldPartition/WorldPartitionActorDescView.h
+++ b/Engine/Source/Runtime/Engine/Public/WorldPartition/WorldPartitionActorDescView.h
@@ -95,6 +95,9 @@ PRAGMA_ENABLE_DEPRECATION_WARNINGS
 
 	const FWorldPartitionActorDesc* GetActorDesc() const { return ActorDesc; }
 
+	bool GetProperty(FName PropertyName, FName* PropertyValue) const;
+	bool HasProperty(FName PropertyName) const;
+
 protected:
 	const FWorldPartitionActorDesc* ActorDesc;
 	bool bIsForcedNonSpatiallyLoaded;
