Compute & force fixup ULevel::WorldPartitionRuntimeCell when cooking content bundles
Order on PopulateGeneratedPackageForCook for each PackageToCook parameters is not deterministic
This can lead to ULevel::WorldPartitionRuntimeCell targeting cells in /Game/ instead of their /plugin/ externalobject

diff --git a/Engine/Source/Runtime/Engine/Private/WorldPartition/ContentBundle/ContentBundleEditor.cpp b/Engine/Source/Runtime/Engine/Private/WorldPartition/ContentBundle/ContentBundleEditor.cpp
--- a/Engine/Source/Runtime/Engine/Private/WorldPartition/ContentBundle/ContentBundleEditor.cpp
+++ b/Engine/Source/Runtime/Engine/Private/WorldPartition/ContentBundle/ContentBundleEditor.cpp
@@ -18,6 +18,7 @@
 #include "WorldPartition/ContentBundle/ContentBundlePaths.h"
 #include "WorldPartition/Cook/WorldPartitionCookPackageContextInterface.h"
 #include "WorldPartition/Cook/WorldPartitionCookPackage.h"
+#include "WorldPartition/WorldPartitionLevelHelper.h"
 
 FContentBundleEditor::FContentBundleEditor(TSharedPtr<FContentBundleClient>& InClient, UWorld* InWorld)
 	: FContentBundleBase(InClient, InWorld)
@@ -437,12 +438,33 @@ bool FContentBundleEditor::PopulateGeneratedPackageForCook(class IWorldPartition
 				if (!Cell->IsAlwaysLoaded())
 				{
 					TArray<UPackage*> ModifiedPackages;
-					if (!Cell->PopulateGeneratedPackageForCook(PackageToCook.GetPackage(), OutModifiedPackages))
+					if (Cell->PopulateGeneratedPackageForCook(PackageToCook.GetPackage(), OutModifiedPackages))
+					{
+						UWorld* CellWorld = FindObject<UWorld>(PackageToCook.GetPackage(), *GetInjectedWorld()->GetName());
+						if (CellWorld != nullptr)
+						{
+							// @todo_ow : Order on PopulateGeneratedPackageForCook for each PackageToCook parameters is not deterministic.
+							// PopulateGeneratedPackageForCook can be called on some cell(s) before it was called for the ExternalStreamingObject.
+							// This lead to a wrong path in ULevel::WorldPartitionRuntimeCell because the ExternalStreamingObject could not yet have been renamed.
+							// Compute & set ULevel::WorldPartitionRuntimeCell by hand in FWorldPartitionLevelHelper::RemapLevelCellPathInContentBundle
+							// We already know that the cell will reside in the ExternalStreamingObject package at runtime.
+							if (!FWorldPartitionLevelHelper::RemapLevelCellPathInContentBundle(CellWorld->PersistentLevel, this, Cell))
+							{
+								UE_LOG(LogContentBundle, Warning, TEXT("%s[Cook] Failed to Remap cell in level for package %s."), *ContentBundle::Log::MakeDebugInfoString(*this), *PackageToCook.RelativePath);
+								bIsSuccess = false;
+							}
+						}
+						else
+						{
+							UE_LOG(LogContentBundle, Warning, TEXT("%s[Cook] Failed to retrieve cell world from package %s."), *ContentBundle::Log::MakeDebugInfoString(*this), *PackageToCook.RelativePath);
+							bIsSuccess = false;
+						}
+					}
+					else
 					{
 						UE_LOG(LogContentBundle, Error, TEXT("%s[Cook] Failed to populate cell package %s."), *ContentBundle::Log::MakeDebugInfoString(*this), *PackageToCook.RelativePath);
 						bIsSuccess = false;
 					}
-					
 				}
 				else
 				{
diff --git a/Engine/Source/Runtime/Engine/Private/WorldPartition/WorldPartitionLevelHelper.cpp b/Engine/Source/Runtime/Engine/Private/WorldPartition/WorldPartitionLevelHelper.cpp
--- a/Engine/Source/Runtime/Engine/Private/WorldPartition/WorldPartitionLevelHelper.cpp
+++ b/Engine/Source/Runtime/Engine/Private/WorldPartition/WorldPartitionLevelHelper.cpp
@@ -17,6 +17,7 @@
 #include "UnrealEngine.h"
 #include "WorldPartition/WorldPartition.h"
 #include "WorldPartition/WorldPartitionPackageHelper.h"
+#include "WorldPartition/ContentBundle/ContentBundleEditor.h"
 #include "LevelUtils.h"
 #include "ActorFolder.h"
 #include "Containers/StringFwd.h"
@@ -241,6 +242,19 @@ bool FWorldPartitionLevelHelper::RemapActorPath(const FActorContainerID& InConta
 	return false;
 }
 
+bool FWorldPartitionLevelHelper::RemapLevelCellPathInContentBundle(ULevel* Level, const class FContentBundleEditor* ContentBundleEditor, const UWorldPartitionRuntimeCell* Cell)
+{
+	FString CellPath = ContentBundleEditor->GetExternalStreamingObjectPackageName();
+	CellPath += TEXT(".");
+	CellPath += ContentBundleEditor->GetExternalStreamingObjectName();
+	CellPath += TEXT(".");
+	CellPath += Cell->GetName();
+	
+	Level->WorldPartitionRuntimeCell = FSoftObjectPath(CellPath);
+
+	return Level->WorldPartitionRuntimeCell.IsValid();
+}
+
 /**
  * Creates an empty Level used in World Partition
  */
diff --git a/Engine/Source/Runtime/Engine/Public/WorldPartition/WorldPartitionLevelHelper.h b/Engine/Source/Runtime/Engine/Public/WorldPartition/WorldPartitionLevelHelper.h
--- a/Engine/Source/Runtime/Engine/Public/WorldPartition/WorldPartitionLevelHelper.h
+++ b/Engine/Source/Runtime/Engine/Public/WorldPartition/WorldPartitionLevelHelper.h
@@ -50,6 +50,10 @@ private:
 	static FString GetContainerPackage(const FActorContainerID& InContainerID, const FString& InPackageName, const FString& InDestLevelPackageName = FString());
 	static UWorld::InitializationValues GetWorldInitializationValues();
 
+
+	friend class FContentBundleEditor;
+	static bool RemapLevelCellPathInContentBundle(ULevel* Level, const class FContentBundleEditor* ContentBundleEditor, const UWorldPartitionRuntimeCell* Cell);
+
 	struct FPackageReference
 	{
 		TSet<FPackageReferencer*> Referencers;
