Test softobjectpath uniqueid ID instead of softobjectptr to avoid false positive when the externalstreamingobject has already been renamed.

diff --git a/Engine/Source/Runtime/Engine/Private/WorldPartition/WorldPartitionLevelHelper.cpp b/Engine/Source/Runtime/Engine/Private/WorldPartition/WorldPartitionLevelHelper.cpp
--- a/Engine/Source/Runtime/Engine/Private/WorldPartition/WorldPartitionLevelHelper.cpp
+++ b/Engine/Source/Runtime/Engine/Private/WorldPartition/WorldPartitionLevelHelper.cpp
@@ -252,7 +252,7 @@ bool FWorldPartitionLevelHelper::RemapLevelCellPathInContentBundle(ULevel* Level
 	
 	Level->WorldPartitionRuntimeCell = FSoftObjectPath(CellPath);
 
-	return Level->WorldPartitionRuntimeCell.IsValid();
+	return !Level->WorldPartitionRuntimeCell.GetUniqueID().IsNull();
 }
 
 /**
