Add optional linked layer persistence to shared linked layer module.
Classes marked as persistent will have one instance remain alive even when unlinked so that linking / unlinking doesn't incur the cost of creating and initializing the instance
Move lifetime and ownership management of shared layers to linked anim layer module
Stop marking shared layer as garbage when we stop using them as this shouldn't be required. Can be reverted with a console command if this causes issues

diff --git a/Engine/Source/Runtime/Engine/Classes/Components/SkeletalMeshComponent.h b/Engine/Source/Runtime/Engine/Classes/Components/SkeletalMeshComponent.h
--- a/Engine/Source/Runtime/Engine/Classes/Components/SkeletalMeshComponent.h
+++ b/Engine/Source/Runtime/Engine/Classes/Components/SkeletalMeshComponent.h
@@ -311,7 +311,8 @@ class ENGINE_API USkeletalMeshComponent : public USkinnedMeshComponent, public I
 	friend struct FAnimNode_LinkedAnimGraph;
 	friend struct FAnimNode_LinkedAnimLayer;
 	friend struct FLinkedInstancesAdapter;
-
+	friend struct FLinkedAnimLayerClassData;
+	
 #if WITH_EDITORONLY_DATA
 private:
 	/** The skeletal mesh used by this component. */
diff --git a/Engine/Source/Runtime/Engine/Private/Animation/AnimInstance.cpp b/Engine/Source/Runtime/Engine/Private/Animation/AnimInstance.cpp
--- a/Engine/Source/Runtime/Engine/Private/Animation/AnimInstance.cpp
+++ b/Engine/Source/Runtime/Engine/Private/Animation/AnimInstance.cpp
@@ -3078,40 +3078,17 @@ void UAnimInstance::PerformLinkedLayerOverlayOperation(TSubclassOf<UAnimInstance
 							// Skip setting if the class is the same
 							if (TargetInstance == nullptr || ClassToSet != TargetInstance->GetClass())
 							{
-								FLinkedAnimLayerClassData& ClassesData = SharedLinkedAnimLayers->FindOrAddClassData(ClassToSet);
-
-								const FName FunctionToLink = LayerNode->GetDynamicLinkFunctionName();
-
-								// Try to find an existing instance whose function is available for linking
-								FLinkedAnimLayerInstanceData* LayerAnimInstanceData = ClassesData.FindInstanceForLinking(FunctionToLink);
-
 								bool bIsNewInstance;
-								UAnimInstance* LinkedInstance = nullptr;
-								// Use existing instance if possible
-								if (LayerAnimInstanceData)
-								{
-									bIsNewInstance = false;
-									LinkedInstance = LayerAnimInstanceData->Instance;
-								}
-								// Create the new linked instance if needed
-								else
+ 								const FName FunctionToLink = LayerNode->GetDynamicLinkFunctionName();
+								UAnimInstance* LinkedInstance = SharedLinkedAnimLayers->AddLinkedFunction(this, ClassToSet, FunctionToLink, bIsNewInstance);
+
+								if (bIsNewInstance)
 								{
-									bIsNewInstance = true;
-									
-									// Init new instance
-									LinkedInstance = NewObject<UAnimInstance>(MeshComp, ClassToSet);
-									LinkedInstance->bCreatedByLinkedAnimGraph = true;
-									LinkedInstance->InitializeAnimation();
-									
-									// Insert new instance data
-									LayerAnimInstanceData = &ClassesData.AddInstance(LinkedInstance);
-									
 									// Unlink any layer nodes in the new linked instance, as they may have been hooked up to self in InitializeAnimation above.
 									UnlinkLayerNodesInInstance(LinkedInstance, LayerPair.Value);
 								}
 
 								// Mark function as linked
-								LayerAnimInstanceData->AddLinkedFunction(FunctionToLink, this);
 								LayerNode->SetLinkedLayerInstance(this, LinkedInstance);
 
 								// Propagate notify flags. If any nodes have this set then we need to propagate to the group.
@@ -3126,11 +3103,6 @@ void UAnimInstance::PerformLinkedLayerOverlayOperation(TSubclassOf<UAnimInstance
 									// Initialize the correct parts of the linked instance
 									InitializeAndCacheBonesForLinkedRoot(LayerNode, ThisProxy, LinkedInstance, LinkedProxy);
 								}
-
-								if (bIsNewInstance)
-								{
-									MeshComp->GetLinkedAnimInstances().Add(LinkedInstance);
-								}
 							}
 						}
 						else
diff --git a/Engine/Source/Runtime/Engine/Private/Animation/AnimNode_LinkedAnimLayer.cpp b/Engine/Source/Runtime/Engine/Private/Animation/AnimNode_LinkedAnimLayer.cpp
--- a/Engine/Source/Runtime/Engine/Private/Animation/AnimNode_LinkedAnimLayer.cpp
+++ b/Engine/Source/Runtime/Engine/Private/Animation/AnimNode_LinkedAnimLayer.cpp
@@ -60,15 +60,19 @@ void FAnimNode_LinkedAnimLayer::OnUninitializeAnimInstance(UAnimInstance* InOwni
 	// Owning anim instance is being destroyed, clean dynamic layer data 
 	if (UAnimInstance* CurrentTarget = GetTargetInstance<UAnimInstance>())
 	{
-		USkeletalMeshComponent* MeshComp = InOwningAnimInstance->GetSkelMeshComponent();
-		if (FAnimSubsystem_SharedLinkedAnimLayers* SharedLinkedAnimLayers = FAnimSubsystem_SharedLinkedAnimLayers::GetFromMesh(MeshComp))
+		// Skip self layers
+		if (CurrentTarget != InOwningAnimInstance)
 		{
-			// If target instance is a shared instance, unlink it when owner uninitialize
-			if (FLinkedAnimLayerInstanceData* PreviousLinkedLayerInstanceData = SharedLinkedAnimLayers->FindInstanceData(CurrentTarget))
+			USkeletalMeshComponent* MeshComp = InOwningAnimInstance->GetSkelMeshComponent();
+			if (FAnimSubsystem_SharedLinkedAnimLayers* SharedLinkedAnimLayers = FAnimSubsystem_SharedLinkedAnimLayers::GetFromMesh(MeshComp))
 			{
-				DynamicUnlink(InOwningAnimInstance);
-				SetTargetInstance(nullptr);
-				CleanupSharedLinkedLayersData(InOwningAnimInstance, CurrentTarget);
+				// If target instance is a shared instance, unlink it when owner uninitialize
+				if (SharedLinkedAnimLayers->IsSharedInstance(CurrentTarget))
+				{
+					DynamicUnlink(InOwningAnimInstance);
+					SetTargetInstance(nullptr);
+					CleanupSharedLinkedLayersData(InOwningAnimInstance, CurrentTarget);
+				}
 			}
 		}
 	}
@@ -187,7 +191,7 @@ bool FAnimNode_LinkedAnimLayer::CanTeardownLinkedInstance(const UAnimInstance* L
 	USkeletalMeshComponent* MeshComp = LinkedInstance->GetSkelMeshComponent();
 	if (FAnimSubsystem_SharedLinkedAnimLayers* SharedLinkedAnimLayers = FAnimSubsystem_SharedLinkedAnimLayers::GetFromMesh(MeshComp))
 	{
-		return SharedLinkedAnimLayers->FindInstanceData(LinkedInstance) == nullptr;
+		return !SharedLinkedAnimLayers->IsSharedInstance(LinkedInstance);
 	}
 	return true; 
 }
@@ -199,23 +203,9 @@ void FAnimNode_LinkedAnimLayer::CleanupSharedLinkedLayersData(const UAnimInstanc
 		USkeletalMeshComponent* MeshComp = InPreviousTargetInstance->GetSkelMeshComponent();
 		if (FAnimSubsystem_SharedLinkedAnimLayers* SharedLinkedAnimLayers = FAnimSubsystem_SharedLinkedAnimLayers::GetFromMesh(MeshComp))
 		{
-			if (FLinkedAnimLayerInstanceData* PreviousLinkedLayerInstanceData = SharedLinkedAnimLayers->FindInstanceData(InPreviousTargetInstance))
+			if (SharedLinkedAnimLayers->IsSharedInstance(InPreviousTargetInstance))
 			{
-				const FName FunctionToLink = GetDynamicLinkFunctionName();
-
-				// The owning anim instance should never be part of the dynamic linked instance data
-				check(InOwningAnimInstance != InPreviousTargetInstance);
-
-				PreviousLinkedLayerInstanceData->RemoveLinkedFunction(FunctionToLink);
-				// Instance isn't used anymore, delete it
-				if (PreviousLinkedLayerInstanceData->GetLinkedFunctions().Num() == 0)
-				{
-					SharedLinkedAnimLayers->RemoveInstance(InPreviousTargetInstance);
-					MeshComp->GetLinkedAnimInstances().Remove(InPreviousTargetInstance);
-					// Only call UninitializeAnimation if we are not the owning anim instance
-					InPreviousTargetInstance->UninitializeAnimation();
-					InPreviousTargetInstance->MarkAsGarbage();
-				}
+				SharedLinkedAnimLayers->RemoveLinkedFunction(InPreviousTargetInstance, GetDynamicLinkFunctionName());
 			}
 		}
 	}
diff --git a/Engine/Source/Runtime/Engine/Private/Animation/AnimSubsystem_SharedLinkedAnimLayers.cpp b/Engine/Source/Runtime/Engine/Private/Animation/AnimSubsystem_SharedLinkedAnimLayers.cpp
--- a/Engine/Source/Runtime/Engine/Private/Animation/AnimSubsystem_SharedLinkedAnimLayers.cpp
+++ b/Engine/Source/Runtime/Engine/Private/Animation/AnimSubsystem_SharedLinkedAnimLayers.cpp
@@ -6,6 +6,15 @@
 
 #include UE_INLINE_GENERATED_CPP_BY_NAME(AnimSubsystem_SharedLinkedAnimLayers)
 
+namespace SharedLinkedAnimLayersConsoleCommands
+{
+	static int32 MarkLayerAsGarbageOnUninitialize = 0;
+	static FAutoConsoleVariableRef CVarMarkLayerAsGarbageOnUninitialize(
+		TEXT("a.MarkLayerAsGarbageOnUninitialize"), MarkLayerAsGarbageOnUninitialize,
+		TEXT("Whether to mark the layers as garbage after uinitializing them."),
+		ECVF_Default);
+}
+
 #if LINKEDANIMLAYERSDATA_INTEGRITYCHECKS
 #include "Animation/AnimClassInterface.h"
 #include "Animation/AnimNode_LinkedAnimLayer.h"
@@ -70,46 +79,121 @@ void FLinkedAnimLayerInstanceData::RemoveLinkedFunction(FName Function)
 	check(LinkedFunctions.Contains(Function));
 	LinkedFunctions.Remove(Function);
 }
+
 FLinkedAnimLayerInstanceData* FLinkedAnimLayerClassData::FindInstanceData(const UAnimInstance* AnimInstance)
 {
 	return InstancesData.FindByPredicate([AnimInstance](FLinkedAnimLayerInstanceData& InstanceData) {return InstanceData.Instance == AnimInstance; });
 }
 
-FLinkedAnimLayerInstanceData* FLinkedAnimLayerClassData::FindInstanceForLinking(FName Function)
+UAnimInstance* FLinkedAnimLayerClassData::FindOrAddInstanceForLinking(UAnimInstance* OwningInstance, FName Function, bool& bIsNewInstance)
 {
 	for (FLinkedAnimLayerInstanceData& LayerInstanceData : InstancesData)
 	{
 		// Check if function is already linked
 		if (!LayerInstanceData.GetLinkedFunctions().Contains(Function))
 		{
+			bIsNewInstance = false;
+			LayerInstanceData.AddLinkedFunction(Function, LayerInstanceData.Instance);
 			// not linked, use this instance
-			return &LayerInstanceData;
+			return LayerInstanceData.Instance;
 		}
 	}
-	return nullptr;
+
+	// Create a new object
+	bIsNewInstance = true;
+	UAnimInstance* NewAnimInstance = NewObject<UAnimInstance>(OwningInstance->GetSkelMeshComponent(), Class);
+	NewAnimInstance->bCreatedByLinkedAnimGraph = true;
+	NewAnimInstance->InitializeAnimation();
+	FLinkedAnimLayerInstanceData& NewInstanceData = AddInstance(NewAnimInstance);
+	NewInstanceData.AddLinkedFunction(Function, NewAnimInstance);
+	OwningInstance->GetSkelMeshComponent()->GetLinkedAnimInstances().Add(NewAnimInstance);
+
+	return NewAnimInstance;
 }
 
 FLinkedAnimLayerInstanceData& FLinkedAnimLayerClassData::AddInstance(UAnimInstance* AnimInstance)
 {
-	InstancesData.Push(FLinkedAnimLayerInstanceData(AnimInstance));
+	// First instance we create for persistent layer is marked has persistent
+	InstancesData.Push(FLinkedAnimLayerInstanceData(AnimInstance, bIsPersistent && InstancesData.Num() == 0));
 	return InstancesData.Last();
 }
 
-void FLinkedAnimLayerClassData::RemoveInstance(const UAnimInstance* AnimInstance)
+void FLinkedAnimLayerClassData::RemoveLinkedFunction(UAnimInstance* AnimInstance, FName Function)
+{
+	if (FLinkedAnimLayerInstanceData* InstanceData = FindInstanceData(AnimInstance))
+	{
+		InstanceData->RemoveLinkedFunction(Function);
+		if (InstanceData->GetLinkedFunctions().Num() == 0)
+		{
+			if (USkeletalMeshComponent* Mesh = InstanceData->Instance->GetSkelMeshComponent())
+			{
+				Mesh->GetLinkedAnimInstances().Remove(AnimInstance);
+			}
+			if (!InstanceData->IsPersistent())
+			{
+				RemoveInstance(AnimInstance);
+			}
+		}
+	}
+}
+
+void FLinkedAnimLayerClassData::RemoveInstance(UAnimInstance* AnimInstance)
 {
 	for (int i = 0; i < InstancesData.Num(); ++i)
 	{
 		if (InstancesData[i].Instance == AnimInstance)
 		{
+			const FLinkedAnimLayerInstanceData& InstanceData = InstancesData[i];
+
+			// If we have no function linked, the instance should never be part of the skeletal mesh component
+			check(!InstanceData.Instance->GetSkelMeshComponent() || !InstanceData.Instance->GetSkelMeshComponent()->GetLinkedAnimInstances().Contains(InstanceData.Instance));
 			// Make sure no functions are still linked
-			check(InstancesData[i].GetLinkedFunctions().Num() == 0);
+			check(InstanceData.GetLinkedFunctions().Num() == 0);
+			check(!InstanceData.IsPersistent());
+
+			// Since UninitializeAnimation can make calls to the shared layer system via self layer nodes, cleanup before Uninitializing to prevent unnecessary checks
 			InstancesData.RemoveAt(i);
+
+			AnimInstance->UninitializeAnimation();
+			if (SharedLinkedAnimLayersConsoleCommands::MarkLayerAsGarbageOnUninitialize)
+			{
+				AnimInstance->MarkAsGarbage();
+			}
 			return;
 		}
 	}
 	check(0);
 }
 
+void FLinkedAnimLayerClassData::SetPersistence(bool bInIsPersistent)
+{
+	if (bInIsPersistent != IsPersistent())
+	{
+		bIsPersistent = bInIsPersistent;
+		// If persistence is added we mark the current linked instance, if there is one, as persistent
+		if (bIsPersistent)
+		{
+			if (InstancesData.Num())
+			{
+				InstancesData[0].SetPersistence(true);
+			}
+		}
+		// If we remove persistence however make sure any currently unlinked instance is correctly deleted
+		else // if (!bIsPersistent)
+		{
+			int32 InstanceIndex = InstancesData.IndexOfByPredicate([](const FLinkedAnimLayerInstanceData& InstanceData) {return InstanceData.IsPersistent(); });
+			if (InstanceIndex != INDEX_NONE)
+			{
+				FLinkedAnimLayerInstanceData& InstanceData = InstancesData[InstanceIndex];
+				InstanceData.SetPersistence(false);
+				if (InstanceData.GetLinkedFunctions().Num() == 0)
+				{
+					RemoveInstance(InstanceData.Instance);
+				}
+			}
+		}
+	}
+}
 FAnimSubsystem_SharedLinkedAnimLayers* FAnimSubsystem_SharedLinkedAnimLayers::GetFromMesh(USkeletalMeshComponent* SkelMesh)
 {
 #if WITH_EDITOR
@@ -125,6 +209,7 @@ FAnimSubsystem_SharedLinkedAnimLayers* FAnimSubsystem_SharedLinkedAnimLayers::Ge
 
 void FAnimSubsystem_SharedLinkedAnimLayers::Reset()
 {
+	PersistentClasses.Empty();
 	ClassesData.Empty(ClassesData.Num());
 }
 
@@ -137,7 +222,7 @@ FLinkedAnimLayerInstanceData* FAnimSubsystem_SharedLinkedAnimLayers::FindInstanc
 	return nullptr;
 }
 
-FLinkedAnimLayerClassData* FAnimSubsystem_SharedLinkedAnimLayers::FindClassData(TSubclassOf<UAnimInstance> AnimClass)
+FLinkedAnimLayerClassData* FAnimSubsystem_SharedLinkedAnimLayers::FindClassData(TSubclassOf<UAnimInstance> AnimClass) 
 {
 	return ClassesData.FindByPredicate([AnimClass](const FLinkedAnimLayerClassData& ClassData) {return ClassData.GetClass() == AnimClass; });
 }
@@ -148,26 +233,40 @@ FLinkedAnimLayerClassData& FAnimSubsystem_SharedLinkedAnimLayers::FindOrAddClass
 	{
 		return *Result;
 	}
-	ClassesData.Add(FLinkedAnimLayerClassData(AnimClass));
+	bool bIsPersistent = PersistentClasses.Find(AnimClass) != INDEX_NONE;
+	ClassesData.Add(FLinkedAnimLayerClassData(AnimClass, bIsPersistent));
 	return ClassesData.Last();
 };
 
-void FAnimSubsystem_SharedLinkedAnimLayers::RemoveInstance(const UAnimInstance* AnimInstance)
+void FAnimSubsystem_SharedLinkedAnimLayers::RemovePersistentAnimLayerClass(TSubclassOf<UAnimInstance> AnimInstanceClass)
+{
+	// When a class loses its persistency, make sure to clean it up if it's already unlinked
+	int32 ClassIndex = ClassesData.IndexOfByPredicate([AnimInstanceClass](const FLinkedAnimLayerClassData& ClassData) {return ClassData.GetClass() == AnimInstanceClass; });
+	if (ClassIndex != INDEX_NONE)
+	{
+		FLinkedAnimLayerClassData& ClassData = ClassesData[ClassIndex];
+		ClassData.SetPersistence(false);
+	}
+	PersistentClasses.Remove(AnimInstanceClass);
+
+}
+
+UAnimInstance* FAnimSubsystem_SharedLinkedAnimLayers::AddLinkedFunction(UAnimInstance* OwningInstance, TSubclassOf<UAnimInstance> AnimClass, FName Function, bool& bIsNewInstance)
+{
+	FLinkedAnimLayerClassData& ClassData = FindOrAddClassData(AnimClass);
+	return ClassData.FindOrAddInstanceForLinking(OwningInstance, Function, bIsNewInstance);
+}
+
+void FAnimSubsystem_SharedLinkedAnimLayers::RemoveLinkedFunction(UAnimInstance* AnimInstance, FName Function)
 {
 	int32 ClassIndex = ClassesData.IndexOfByPredicate([AnimInstance](const FLinkedAnimLayerClassData& ClassData) {return ClassData.GetClass() == AnimInstance->GetClass(); });
 	if (ClassIndex != INDEX_NONE)
 	{
 		FLinkedAnimLayerClassData& ClassData = ClassesData[ClassIndex];
-		ClassData.RemoveInstance(AnimInstance);
-		// Remove class data when all instances are removed
+		ClassData.RemoveLinkedFunction(AnimInstance, Function);
 		if (ClassData.GetInstancesData().Num() == 0)
 		{
-			// Remove class
 			ClassesData.RemoveAt(ClassIndex);
 		}
 	}
-	else
-	{
-		check(0);
-	}
 }
diff --git a/Engine/Source/Runtime/Engine/Public/Animation/AnimSubsystem_SharedLinkedAnimLayers.h b/Engine/Source/Runtime/Engine/Public/Animation/AnimSubsystem_SharedLinkedAnimLayers.h
--- a/Engine/Source/Runtime/Engine/Public/Animation/AnimSubsystem_SharedLinkedAnimLayers.h
+++ b/Engine/Source/Runtime/Engine/Public/Animation/AnimSubsystem_SharedLinkedAnimLayers.h
@@ -19,8 +19,11 @@ struct FLinkedAnimLayerInstanceData
 	GENERATED_USTRUCT_BODY()
 
 	FLinkedAnimLayerInstanceData() {}
-	FLinkedAnimLayerInstanceData(UAnimInstance* AnimInstance) { Instance = AnimInstance; }
+	FLinkedAnimLayerInstanceData(UAnimInstance* AnimInstance, bool bInIsPersistent) { Instance = AnimInstance; bIsPersistent = bInIsPersistent; }
 
+	bool IsPersistent() const { return bIsPersistent; }
+	void SetPersistence(bool bInIsPersistent) { bIsPersistent = bInIsPersistent; }
+	
 	// Mark a function as linked
 	void AddLinkedFunction(FName Function, UAnimInstance* AnimInstance);
 	// Unmark a function as linked
@@ -33,8 +36,12 @@ struct FLinkedAnimLayerInstanceData
 	TObjectPtr<UAnimInstance> Instance;
 
 private:
+	
 	UPROPERTY(Transient)
 	TMap<FName, TWeakObjectPtr<UAnimInstance>> LinkedFunctions;
+
+	// True if this is a persistent instance that remains alive when unlinked
+	bool bIsPersistent = false;
 };
 
 // Linked layer class info
@@ -44,27 +51,35 @@ struct FLinkedAnimLayerClassData
 	GENERATED_BODY()
 
 	FLinkedAnimLayerClassData() {}
-	FLinkedAnimLayerClassData(TSubclassOf<UAnimInstance> AnimClass) { Class = AnimClass; }
+	FLinkedAnimLayerClassData(TSubclassOf<UAnimInstance> AnimClass, bool bInIsPersistent) { Class = AnimClass; bIsPersistent = bInIsPersistent; }
+
+	bool IsPersistent() const { return bIsPersistent; }
+	void SetPersistence(bool bInIsPersistent);
 
 	// Find instance data for given instance if it exists, returns nullptr otherwise
 	FLinkedAnimLayerInstanceData* FindInstanceData(const UAnimInstance* AnimInstance);
 
-	// Find an existing instance to link given function if one is available, returns nullptr otherwise
-	FLinkedAnimLayerInstanceData* FindInstanceForLinking(FName Function);
+	// Find an existing instance to link given function if one is available, creates one otherwise
+	UAnimInstance* FindOrAddInstanceForLinking(UAnimInstance* OwningInstance, FName Function, bool& IsNewInstance);
 	
-	// Add given instance
-	FLinkedAnimLayerInstanceData& AddInstance(UAnimInstance* AnimInstance);
-	// Remove given instance
-	void RemoveInstance(const UAnimInstance* AnimInstance);
-
 	TSubclassOf<UAnimInstance> GetClass() const {return Class;}
 	const TArray<FLinkedAnimLayerInstanceData>& GetInstancesData() const {return InstancesData;}
 
+	void RemoveLinkedFunction(UAnimInstance* AnimInstance, FName Function);
+
 private:
+	// Add given instance
+	FLinkedAnimLayerInstanceData& AddInstance(UAnimInstance* AnimInstance);
+	// Remove given instance
+	void RemoveInstance(UAnimInstance* AnimInstance);
+
 	TSubclassOf<UAnimInstance> Class;
 
 	UPROPERTY(Transient)
 	TArray<FLinkedAnimLayerInstanceData> InstancesData;
+
+	// If true, one instance will be kept alive when unlinked so that we don't recreate it from scratch next time it's linked
+	bool bIsPersistent;
 };
 
 // Data for shared linked anim instances module
@@ -79,6 +94,24 @@ struct ENGINE_API FAnimSubsystem_SharedLinkedAnimLayers : public FAnimSubsystemI
 	// Clear all linked layers data
 	void Reset();
 
+	// Check if a given anim instance is shared
+	bool IsSharedInstance(const UAnimInstance* AnimInstance) { return FindInstanceData(AnimInstance) != nullptr; }
+
+	const TArray<FLinkedAnimLayerClassData>& GetClassesData() const { return ClassesData; }
+
+	// Add a class to the persistent class array, insuring one instance is kept alive even when unlinked
+	void AddPersistentAnimLayerClass(TSubclassOf<UAnimInstance> AnimInstanceClass) { PersistentClasses.AddUnique(AnimInstanceClass); }
+
+	// Remove a class from the persistent class array. 
+	void RemovePersistentAnimLayerClass(TSubclassOf<UAnimInstance> AnimInstanceClass);
+
+	// Add a linked function to be managed by a shared anim instance. Returns the instance to bind to the function
+	UAnimInstance* AddLinkedFunction(UAnimInstance* OwningInstance, TSubclassOf<UAnimInstance> AnimClass, FName Function, bool& bIsNewInstance);
+
+	// Remove a linked function, cleaning the instance if it's not used anymore and not persistent
+	void RemoveLinkedFunction(UAnimInstance* AnimInstance, FName Function);
+
+private:
 	// Find instance data for given instance if it exists, returns nullptr otherwise
 	FLinkedAnimLayerInstanceData* FindInstanceData(const UAnimInstance* AnimInstance);
 
@@ -88,11 +121,10 @@ struct ENGINE_API FAnimSubsystem_SharedLinkedAnimLayers : public FAnimSubsystemI
 	// Find or add ClassData for given AnimClass
 	FLinkedAnimLayerClassData& FindOrAddClassData(TSubclassOf<UAnimInstance> AnimClass);
 
-	// Remove given instance
-	void RemoveInstance(const UAnimInstance* AnimInstance);
-
-	const TArray<FLinkedAnimLayerClassData>& GetClassesData() const { return ClassesData; }
-private:
 	UPROPERTY(Transient)
 	TArray<FLinkedAnimLayerClassData> ClassesData;
+
+	// Anim instance classes that should be kept alive even when unlinked
+	UPROPERTY(Transient)
+	TArray<TSubclassOf<UAnimInstance>> PersistentClasses;
 };
\ No newline at end of file
