Non unity CIS fix due to missing include

#rnx

diff --git a/Engine/Source/Runtime/Engine/Public/Animation/AnimSubsystem_SharedLinkedAnimLayers.h b/Engine/Source/Runtime/Engine/Public/Animation/AnimSubsystem_SharedLinkedAnimLayers.h
--- a/Engine/Source/Runtime/Engine/Public/Animation/AnimSubsystem_SharedLinkedAnimLayers.h
+++ b/Engine/Source/Runtime/Engine/Public/Animation/AnimSubsystem_SharedLinkedAnimLayers.h
@@ -2,6 +2,7 @@
 #pragma once
 
 #include "CoreMinimal.h"
+#include "Animation/AnimInstance.h"
 #include "Templates/SubclassOf.h"
 #include "Animation/AnimSubsystemInstance.h"
 #include "AnimSubsystem_SharedLinkedAnimLayers.generated.h"
