PR #10133: Fixed comment for FGraphNodeCreator (Contributed by [REDACTED_PERSON])

diff --git a/Engine/Source/Runtime/Engine/Classes/EdGraph/EdGraph.h b/Engine/Source/Runtime/Engine/Classes/EdGraph/EdGraph.h
--- a/Engine/Source/Runtime/Engine/Classes/EdGraph/EdGraph.h
+++ b/Engine/Source/Runtime/Engine/Classes/EdGraph/EdGraph.h
@@ -260,11 +260,11 @@ private:
  * Helper object to ensure a graph node is correctly constructed
  *
  * Typical use pattern is:
- * FNodeGraphNodeCreate<NodeType> NodeCreator(Graph);
+ * FGraphNodeCreator<NodeType> NodeCreator(Graph);
  * NodeType* Node = NodeCreator.CreateNode();
  * // calls to build out node 
  * Node->MemberVar = ...
- * NodeCreator.Finalize
+ * NodeCreator.Finalize();
  */
 template <typename NodeType>
 struct FGraphNodeCreator
