Fix false-positive ensure when a component still using ReplicateSubObjects was used by an actor that migrated to the new method.

diff --git a/Engine/Source/Runtime/Engine/Private/ActorReplication.cpp b/Engine/Source/Runtime/Engine/Private/ActorReplication.cpp
--- a/Engine/Source/Runtime/Engine/Private/ActorReplication.cpp
+++ b/Engine/Source/Runtime/Engine/Private/ActorReplication.cpp
@@ -497,16 +497,25 @@ bool AActor::ReplicateSubobjects(UActorChannel *Channel, FOutBunch *Bunch, FRepl
 	{
 		if (ActorComp && ActorComp->GetReplicationCondition() != COND_Never)
 		{
+#if SUBOBJECT_TRANSITION_VALIDATION
+			if (UActorChannel::CanIgnoreDeprecatedReplicateSubObjects() && !ActorComp->IsUsingRegisteredSubObjectList())
+			{
+				// Ignore components that only work with the old virtual method
+				continue;
+			}
+#endif
+
 			UActorChannel::SetCurrentSubObjectOwner(ActorComp);
 			WroteSomething |= ActorComp->ReplicateSubobjects(Channel, Bunch, RepFlags);		// Lets the component add subobjects before replicating its own properties.
 
 #if SUBOBJECT_TRANSITION_VALIDATION
 			if (UActorChannel::CanIgnoreDeprecatedReplicateSubObjects())
 			{
-				// Don't replicate the component itself when testing since we know this function is not voluntarily implemented
+				// Don't replicate the component below when testing since we know this function is not voluntarily implemented
 				continue;
 			}
 #endif
+
 			UActorChannel::SetCurrentSubObjectOwner(this);
 			WroteSomething |= Channel->ReplicateSubobject(ActorComp, *Bunch, *RepFlags);	// (this makes those subobjects 'supported', and from here on those objects may have reference replicated)		
 		}
