Adding/Removing actor to/from a datalayer is now done via UDataLayerInstance interface
Deprecate add/removal of data layer via DataLayerAsset  in AActor.
Modifying AActor::DataLayersAssets is now restricted to DataLayerInstances
Add CanAdd/Remove(Actor) to DataLayerInstance to filter actors before trying to add them

diff --git a/Engine/Source/Editor/UnrealEd/Private/Commandlets/WorldPartitionDataLayerToAssetCommandlet.cpp b/Engine/Source/Editor/UnrealEd/Private/Commandlets/WorldPartitionDataLayerToAssetCommandlet.cpp
--- a/Engine/Source/Editor/UnrealEd/Private/Commandlets/WorldPartitionDataLayerToAssetCommandlet.cpp
+++ b/Engine/Source/Editor/UnrealEd/Private/Commandlets/WorldPartitionDataLayerToAssetCommandlet.cpp
@@ -581,16 +581,7 @@ uint32 UDataLayerToAssetCommandlet::RemapDataLayersAssetsFromPreviousConversions
 		{
 			if (DataLayerConversionInfo->IsAPreviousConversion())
 			{
-				if (Actor->AddDataLayer(DataLayerConversionInfo->GetCurrentConversion()->DataLayerInstance))
-				{
-					if (!Actor->RemoveDataLayer(DataLayerConversionInfo->DataLayerAsset))
-					{
-						ErrorCount++;
-						UE_LOG(LogDataLayerToAssetCommandlet, Error, TEXT("Failed to remove deprecated data layer asset %s in Actor %s"),
-							*DataLayerConversionInfo->DataLayerAsset->GetFullName(), *Actor->GetName());
-					}
-				}
-				else
+				if (!DataLayerConversionInfo->GetCurrentConversion()->DataLayerInstance->AddActor(Actor))
 				{
 					ErrorCount++;
 					UE_LOG(LogDataLayerToAssetCommandlet, Error, TEXT("Failed to replace data layer asset %s in Actor %s"),
diff --git a/Engine/Source/Runtime/Engine/Classes/GameFramework/Actor.h b/Engine/Source/Runtime/Engine/Classes/GameFramework/Actor.h
--- a/Engine/Source/Runtime/Engine/Classes/GameFramework/Actor.h
+++ b/Engine/Source/Runtime/Engine/Classes/GameFramework/Actor.h
@@ -1238,9 +1238,7 @@ private:
 	virtual bool ActorTypeSupportsDataLayer() const { return true; }
 	virtual bool ActorTypeShouldSkipFromLevelInstance() const { return false; }
 public:
-	bool AddDataLayer(const UDataLayerAsset* DataLayerAsset);
 	bool AddDataLayer(const UDataLayerInstance* DataLayerInstance);
-	bool RemoveDataLayer(const UDataLayerAsset* DataLayerAsset);
 	bool RemoveDataLayer(const UDataLayerInstance* DataLayerInstance);
 	bool RemoveAllDataLayers();
 	bool SupportsDataLayer() const;
@@ -1266,6 +1264,12 @@ public:
 
 	PRAGMA_ENABLE_DEPRECATION_WARNINGS
 
+	UE_DEPRECATED(5.3, "Use UDataLayerInstance::AddActor(AActor*) instead.")
+	bool AddDataLayer(const UDataLayerAsset* DataLayerAsset) { return false; }
+
+	UE_DEPRECATED(5.3, "Use UDataLayerInstance::RemoveDataLayer(AActor*) instead.")
+	bool RemoveDataLayer(const UDataLayerAsset* DataLayerAsset) { return false; }
+
 	UE_DEPRECATED(5.1, "Convert DataLayer using UDataLayerToAssetCommandlet and use AddDataLayer(UDataLayerAsset*)")
 	bool AddDataLayer(const FActorDataLayer& ActorDataLayer);
 
@@ -3954,6 +3958,7 @@ private:
 	friend struct FSetActorHiddenInSceneOutliner;
 	friend struct FSetActorGuid;
 	friend struct FSetActorContentBundleGuid;
+	friend struct FAssignActorDataLayer;
 	friend struct FSetActorSelectable;
 #endif
 
@@ -4219,6 +4224,16 @@ private:
 	}
 	friend class FContentBundleEditor;
 };
+
+struct FAssignActorDataLayer
+{
+private:
+	static bool AddDataLayerAsset(AActor* InActor, const UDataLayerAsset* InDataLayerAsset);
+	static bool RemoveDataLayerAsset(AActor* InActor, const UDataLayerAsset* InDataLayerAsset);
+
+	friend class UDataLayerInstanceWithAsset;
+	friend class ULevelInstanceSubsystem;
+};
 #endif
 
 /** Helper function for executing tick functions based on the normal conditions previous found in UActorComponent::ConditionalTick */
diff --git a/Engine/Source/Runtime/Engine/Private/ActorEditor.cpp b/Engine/Source/Runtime/Engine/Private/ActorEditor.cpp
--- a/Engine/Source/Runtime/Engine/Private/ActorEditor.cpp
+++ b/Engine/Source/Runtime/Engine/Private/ActorEditor.cpp
@@ -1600,46 +1600,14 @@ EDataValidationResult AActor::IsDataValid(TArray<FText>& ValidationErrors)
 
 bool AActor::AddDataLayer(const UDataLayerInstance* DataLayerInstance)
 {
-	if (!SupportsDataLayer())
-	{
-		return false;
-	}
 	return DataLayerInstance->AddActor(this);
 }
 
-bool AActor::AddDataLayer(const UDataLayerAsset* DataLayerAsset)
-{
-	bool bActorWasModified = false;
-	if (SupportsDataLayer() && DataLayerAsset && !DataLayerAssets.Contains(DataLayerAsset))
-	{
-		Modify();
-		bActorWasModified = true;
-		DataLayerAssets.Add(DataLayerAsset);
-	}
-	return bActorWasModified;
-}
-
 bool AActor::RemoveDataLayer(const UDataLayerInstance* DataLayerInstance)
 {
 	return DataLayerInstance->RemoveActor(this);
 }
 
-bool AActor::RemoveDataLayer(const UDataLayerAsset* DataLayerAsset)
-{
-	bool bActorWasModified = false;
-	if (DataLayerAsset && DataLayerAssets.Contains(DataLayerAsset))
-	{
-		if (!bActorWasModified)
-		{
-			Modify();
-			bActorWasModified = true;
-		}
-
-		DataLayerAssets.Remove(DataLayerAsset);
-	}
-	return bActorWasModified;
-}
-
 bool AActor::RemoveAllDataLayers()
 {
 	if ((DataLayerAssets.Num() > 0) || (DataLayers.Num() > 0))
@@ -1784,6 +1752,34 @@ bool AActor::SupportsDataLayer() const
 			!GetClass()->GetDefaultObject<AActor>()->bHiddenEd);
 }
 
+bool FAssignActorDataLayer::AddDataLayerAsset(AActor* InActor, const UDataLayerAsset* InDataLayerAsset)
+{
+	check(InDataLayerAsset != nullptr);
+
+	if (!InActor->ContainsDataLayer(InDataLayerAsset))
+	{
+		InActor->Modify();
+		InActor->DataLayerAssets.Add(InDataLayerAsset);
+		return true;
+	}
+
+	return false;
+}
+
+bool FAssignActorDataLayer::RemoveDataLayerAsset(AActor* InActor, const UDataLayerAsset* InDataLayerAsset)
+{
+	check(InDataLayerAsset != nullptr);
+	
+	if (InActor->DataLayerAssets.Contains(InDataLayerAsset))
+	{
+		InActor->Modify();
+		InActor->DataLayerAssets.Remove(InDataLayerAsset);
+		return true;
+	}
+
+	return false;
+}
+
 //~ Begin Deprecated
 
 PRAGMA_DISABLE_DEPRECATION_WARNINGS
diff --git a/Engine/Source/Runtime/Engine/Private/LevelInstance/LevelInstanceSubsystem.cpp b/Engine/Source/Runtime/Engine/Private/LevelInstance/LevelInstanceSubsystem.cpp
--- a/Engine/Source/Runtime/Engine/Private/LevelInstance/LevelInstanceSubsystem.cpp
+++ b/Engine/Source/Runtime/Engine/Private/LevelInstance/LevelInstanceSubsystem.cpp
@@ -1326,7 +1326,8 @@ void ULevelInstanceSubsystem::BreakLevelInstance_Impl(ILevelInstanceInterface* L
 							{
 								// For DataLayerInstanceAsset we add the Asset to the actor instead of the DataLayerInstance because Actor hasn't moved yet and this
 								// will fail on a UDataLayerInstanceWithAsset when comparing the Actor's outer to the DataLayerInstance outer.
-								Actor->AddDataLayer(DataLayerInstanceWithAsset->GetAsset());
+								// todo_ow : Assign the actor to the DataLayerInstance in the destination level via UDataLayerInstance::AddActor to go through the DataLayerInstance process.
+								FAssignActorDataLayer::AddDataLayerAsset(Actor, DataLayerInstanceWithAsset->GetAsset());
 							}
 							else
 							{
diff --git a/Engine/Source/Runtime/Engine/Private/WorldPartition/DataLayer/DataLayerInstance.cpp b/Engine/Source/Runtime/Engine/Private/WorldPartition/DataLayer/DataLayerInstance.cpp
--- a/Engine/Source/Runtime/Engine/Private/WorldPartition/DataLayer/DataLayerInstance.cpp
+++ b/Engine/Source/Runtime/Engine/Private/WorldPartition/DataLayer/DataLayerInstance.cpp
@@ -124,6 +124,36 @@ const TCHAR* UDataLayerInstance::GetDataLayerIconName() const
 	return FDataLayerUtils::GetDataLayerIconName(GetType());
 }
 
+bool UDataLayerInstance::CanAddActor(AActor* InActor) const
+{
+	return InActor != nullptr && InActor->SupportsDataLayer() && !InActor->ContainsDataLayer(this);
+}
+
+bool UDataLayerInstance::AddActor(AActor* InActor) const
+{
+	if (CanAddActor(InActor))
+	{
+		return PerformAddActor(InActor);
+	}
+
+	return false;
+}
+
+bool UDataLayerInstance::CanRemoveActor(AActor* InActor) const
+{
+	return InActor->ContainsDataLayer(this);
+}
+
+bool UDataLayerInstance::RemoveActor(AActor* InActor) const
+{
+	if (CanRemoveActor(InActor))
+	{
+		return PerformRemoveActor(InActor);
+	}
+
+	return false;
+}
+
 bool UDataLayerInstance::CanParent(const UDataLayerInstance* InParent) const
 {
 	return (this != InParent)
diff --git a/Engine/Source/Runtime/Engine/Private/WorldPartition/DataLayer/DataLayerInstanceWithAsset.cpp b/Engine/Source/Runtime/Engine/Private/WorldPartition/DataLayer/DataLayerInstanceWithAsset.cpp
--- a/Engine/Source/Runtime/Engine/Private/WorldPartition/DataLayer/DataLayerInstanceWithAsset.cpp
+++ b/Engine/Source/Runtime/Engine/Private/WorldPartition/DataLayer/DataLayerInstanceWithAsset.cpp
@@ -52,16 +52,26 @@ bool UDataLayerInstanceWithAsset::IsLocked() const
 	return IsReadOnly();
 }
 
-bool UDataLayerInstanceWithAsset::AddActor(AActor* Actor) const
+bool UDataLayerInstanceWithAsset::CanAddActor(AActor* InActor) const
+{
+	return DataLayerAsset != nullptr && Super::CanAddActor(InActor);
+}
+
+bool UDataLayerInstanceWithAsset::PerformAddActor(AActor* InActor) const
 {	
-	check(GetTypedOuter<ULevel>() == Actor->GetLevel()); // Make sure the instance is part of the same world as the actor.
-	check(UDataLayerManager::GetDataLayerManager(Actor)->GetDataLayerInstance(DataLayerAsset) != nullptr); // Make sure the DataLayerInstance exists for this level
-	return Actor->AddDataLayer(DataLayerAsset);
+	check(GetTypedOuter<ULevel>() == InActor->GetLevel()); // Make sure the instance is part of the same world as the actor.
+	check(UDataLayerManager::GetDataLayerManager(InActor)->GetDataLayerInstance(DataLayerAsset) != nullptr); // Make sure the DataLayerInstance exists for this level
+	return FAssignActorDataLayer::AddDataLayerAsset(InActor, DataLayerAsset);
+}
+
+bool UDataLayerInstanceWithAsset::CanRemoveActor(AActor* InActor) const
+{
+	return DataLayerAsset != nullptr && Super::CanRemoveActor(InActor);
 }
 
-bool UDataLayerInstanceWithAsset::RemoveActor(AActor* Actor) const
+bool UDataLayerInstanceWithAsset::PerformRemoveActor(AActor* InActor) const
 {
-	return Actor->RemoveDataLayer(DataLayerAsset);
+	return FAssignActorDataLayer::RemoveDataLayerAsset(InActor, DataLayerAsset);
 }
 
 bool UDataLayerInstanceWithAsset::Validate(IStreamingGenerationErrorHandler* ErrorHandler) const
diff --git a/Engine/Source/Runtime/Engine/Private/WorldPartition/DataLayer/DeprecatedDataLayerInstance.cpp b/Engine/Source/Runtime/Engine/Private/WorldPartition/DataLayer/DeprecatedDataLayerInstance.cpp
--- a/Engine/Source/Runtime/Engine/Private/WorldPartition/DataLayer/DeprecatedDataLayerInstance.cpp
+++ b/Engine/Source/Runtime/Engine/Private/WorldPartition/DataLayer/DeprecatedDataLayerInstance.cpp
@@ -72,14 +72,14 @@ void UDeprecatedDataLayerInstance::OnCreated(const UDEPRECATED_DataLayer* Deprec
 	InitialRuntimeState = DeprecatedDataLayer->InitialRuntimeState;
 }
 
-bool UDeprecatedDataLayerInstance::AddActor(AActor* Actor) const
+bool UDeprecatedDataLayerInstance::PerformAddActor(AActor* InActor) const
 {
-	return Actor->AddDataLayer(FActorDataLayer(GetDataLayerFName()));
+	return InActor->AddDataLayer(FActorDataLayer(GetDataLayerFName()));
 }
 
-bool UDeprecatedDataLayerInstance::RemoveActor(AActor* Actor) const
+bool UDeprecatedDataLayerInstance::PerformRemoveActor(AActor* InActor) const
 {
-	return Actor->RemoveDataLayer(FActorDataLayer(GetDataLayerFName()));
+	return InActor->RemoveDataLayer(FActorDataLayer(GetDataLayerFName()));
 }
 
 bool UDeprecatedDataLayerInstance::RelabelDataLayer(FName InDataLayerLabel)
diff --git a/Engine/Source/Runtime/Engine/Public/WorldPartition/DataLayer/DataLayerInstance.h b/Engine/Source/Runtime/Engine/Public/WorldPartition/DataLayer/DataLayerInstance.h
--- a/Engine/Source/Runtime/Engine/Public/WorldPartition/DataLayer/DataLayerInstance.h
+++ b/Engine/Source/Runtime/Engine/Public/WorldPartition/DataLayer/DataLayerInstance.h
@@ -75,8 +75,11 @@ public:
 	virtual bool IsLocked() const;
 	virtual bool IsReadOnly() const;
 	virtual bool CanEditChange(const FProperty* InProperty) const;
-	virtual bool AddActor(AActor* Actor) const { return false; }
-	virtual bool RemoveActor(AActor* Actor) const { return false; }
+	
+	virtual bool CanAddActor(AActor* Actor) const;
+	virtual bool AddActor(AActor* Actor) const;
+	virtual bool CanRemoveActor(AActor* Actor) const;
+	virtual bool RemoveActor(AActor* Actor) const;
 
 	bool IsInActorEditorContext() const;
 	bool AddToActorEditorContext();
@@ -143,6 +146,11 @@ private:
 #endif
 
 protected:
+#if WITH_EDITOR
+	virtual bool PerformAddActor(AActor* InActor) const { return false; }
+	virtual bool PerformRemoveActor(AActor* InActor) const { return false;  }
+#endif
+
 #if WITH_EDITORONLY_DATA
 	/** Whether actors associated with the DataLayer are visible in the viewport */
 	UPROPERTY(Transient)
diff --git a/Engine/Source/Runtime/Engine/Public/WorldPartition/DataLayer/DataLayerInstanceWithAsset.h b/Engine/Source/Runtime/Engine/Public/WorldPartition/DataLayer/DataLayerInstanceWithAsset.h
--- a/Engine/Source/Runtime/Engine/Public/WorldPartition/DataLayer/DataLayerInstanceWithAsset.h
+++ b/Engine/Source/Runtime/Engine/Public/WorldPartition/DataLayer/DataLayerInstanceWithAsset.h
@@ -22,8 +22,8 @@ public:
 	virtual void PostEditUndo() override;
 	virtual bool IsLocked() const override;
 	virtual bool IsReadOnly() const override;
-	virtual bool AddActor(AActor* Actor) const override;
-	virtual bool RemoveActor(AActor* Actor) const override;
+	virtual bool CanAddActor(AActor* InActor) const override;
+	virtual bool CanRemoveActor(AActor* InActor) const override;
 
 	virtual bool Validate(IStreamingGenerationErrorHandler* ErrorHandler) const override;
 	virtual void PostEditChangeProperty(FPropertyChangedEvent& PropertyChangedEvent) override;
@@ -40,6 +40,12 @@ public:
 	virtual FString GetDataLayerShortName() const override { return DataLayerAsset != nullptr ? DataLayerAsset->GetName() : GetDataLayerFName().ToString(); }
 	virtual FString GetDataLayerFullName() const override { return DataLayerAsset != nullptr ? DataLayerAsset->GetPathName() : GetDataLayerFName().ToString(); }
 
+protected:
+#if WITH_EDITOR
+	virtual bool PerformAddActor(AActor* InActor) const;
+	virtual bool PerformRemoveActor(AActor* InActor) const;
+#endif
+
 private:
 	UPROPERTY(Category = "Data Layer", EditAnywhere)
 	TObjectPtr<const UDataLayerAsset> DataLayerAsset;
diff --git a/Engine/Source/Runtime/Engine/Public/WorldPartition/DataLayer/DeprecatedDataLayerInstance.h b/Engine/Source/Runtime/Engine/Public/WorldPartition/DataLayer/DeprecatedDataLayerInstance.h
--- a/Engine/Source/Runtime/Engine/Public/WorldPartition/DataLayer/DeprecatedDataLayerInstance.h
+++ b/Engine/Source/Runtime/Engine/Public/WorldPartition/DataLayer/DeprecatedDataLayerInstance.h
@@ -45,8 +45,8 @@ public:
 
 private:
 #if WITH_EDITOR
-	virtual bool AddActor(AActor* Actor) const override;
-	virtual bool RemoveActor(AActor* Actor) const override;
+	virtual bool PerformAddActor(AActor* InActor) const override;
+	virtual bool PerformRemoveActor(AActor* InActor) const override;
 
 	virtual bool SupportRelabeling() const override { return true; }
 	virtual bool RelabelDataLayer(FName NewDataLayerLabel) override;
