UDataLayerInstance can be outered to a Uobject other than AWorldDataLayers
Provide GetWorldDataLayer to Replace GetAWorldDataLayers
UObject might not have a level in its outerchain.
UDataLayerInstance::GetTypedOuter overloaded so compilation fails if done on UWorld, ULevel or AWorldDataLayer

diff --git a/Engine/Source/Editor/DataLayerEditor/Private/DataLayer/DataLayerEditorSubsystem.cpp b/Engine/Source/Editor/DataLayerEditor/Private/DataLayer/DataLayerEditorSubsystem.cpp
--- a/Engine/Source/Editor/DataLayerEditor/Private/DataLayer/DataLayerEditorSubsystem.cpp
+++ b/Engine/Source/Editor/DataLayerEditor/Private/DataLayer/DataLayerEditorSubsystem.cpp
@@ -496,9 +496,10 @@ bool UDataLayerEditorSubsystem::AddActorsToDataLayers(const TArray<AActor*>& Act
 			{
 				if (const UDataLayerInstanceWithAsset* DataLayerInstanceWithAsset = Cast<UDataLayerInstanceWithAsset>(DataLayerInstance))
 				{
-					// If actor's level doesn't match this DataLayerInstance outer level, 
-					// Make sure that a DataLayer Instance for this Data Layer Asset exists in the Actor's level.
-					if (Actor->GetLevel() != DataLayerInstance->GetTypedOuter<ULevel>())
+					// If actor's level WorldDataLayers doesn't match this DataLayerInstance outer WorldDataLayers, 
+					// Make sure that a DataLayer Instance for this Data Layer Asset exists in the Actor's level WorldDataLayers.
+					AWorldDataLayers* TargetWorldDataLayers = Actor->GetLevel()->GetWorldDataLayers();
+					if (TargetWorldDataLayers != DataLayerInstance->GetOuterWorldDataLayers())
 					{
 						UDataLayerManager* DataLayerManager = UDataLayerManager::GetDataLayerManager(Actor);
 						if (ensureMsgf(DataLayerManager, TEXT("No DataLayerManager found for Actor %s, can't add actors to data layers."), *Actor->GetName()))
@@ -508,7 +509,7 @@ bool UDataLayerEditorSubsystem::AddActorsToDataLayers(const TArray<AActor*>& Act
 							bool bDataLayerInstanceExistsInActorLevel = DataLayerInstance != nullptr;
 							if (!bDataLayerInstanceExistsInActorLevel)
 							{
-								DataLayerInstance = CreateDataLayerInstance<UDataLayerInstanceWithAsset>(Actor->GetLevel()->GetWorldDataLayers(), DataLayerInstanceWithAsset->GetAsset());
+								DataLayerInstance = CreateDataLayerInstance<UDataLayerInstanceWithAsset>(TargetWorldDataLayers, DataLayerInstanceWithAsset->GetAsset());
 							}
 						}
 					}
@@ -1240,7 +1241,7 @@ void UDataLayerEditorSubsystem::DeleteDataLayers(const TArray<UDataLayerInstance
 	TArray<UDataLayerInstance*> DeletedDataLayerInstances;
 	for (UDataLayerInstance* DataLayerToDelete : DataLayersToDelete)
 	{
-		if (AWorldDataLayers* OuterWorldDataLayers = DataLayerToDelete ? DataLayerToDelete->GetOuterAWorldDataLayers() : nullptr)
+		if (AWorldDataLayers* OuterWorldDataLayers = DataLayerToDelete ? DataLayerToDelete->GetOuterWorldDataLayers() : nullptr)
 		{
 			if (OuterWorldDataLayers->RemoveDataLayer(DataLayerToDelete))
 			{
@@ -1257,7 +1258,7 @@ void UDataLayerEditorSubsystem::DeleteDataLayers(const TArray<UDataLayerInstance
 void UDataLayerEditorSubsystem::DeleteDataLayer(UDataLayerInstance* DataLayerToDelete)
 {
 	UE_CLOG(!GetWorld(), LogDataLayerEditorSubsystem, Error, TEXT("%s - Failed because world in null."), ANSI_TO_TCHAR(__FUNCTION__));
-	if (AWorldDataLayers* OuterWorldDataLayers = DataLayerToDelete ? DataLayerToDelete->GetOuterAWorldDataLayers() : nullptr)
+	if (AWorldDataLayers* OuterWorldDataLayers = DataLayerToDelete ? DataLayerToDelete->GetOuterWorldDataLayers() : nullptr)
 	{
 		if (OuterWorldDataLayers->RemoveDataLayer(DataLayerToDelete))
 		{
diff --git a/Engine/Source/Editor/DataLayerEditor/Private/DataLayer/DataLayerHierarchy.cpp b/Engine/Source/Editor/DataLayerEditor/Private/DataLayer/DataLayerHierarchy.cpp
--- a/Engine/Source/Editor/DataLayerEditor/Private/DataLayer/DataLayerHierarchy.cpp
+++ b/Engine/Source/Editor/DataLayerEditor/Private/DataLayer/DataLayerHierarchy.cpp
@@ -317,7 +317,7 @@ FSceneOutlinerTreeItemPtr FDataLayerHierarchy::FindOrCreateParentItem(const ISce
 			else
 			{
 				// Parent WorldDataLayers
-				if (AWorldDataLayers* OuterWorldDataLayers = DataLayerInstance->GetOuterAWorldDataLayers())
+				if (AWorldDataLayers* OuterWorldDataLayers = DataLayerInstance->GetOuterWorldDataLayers())
 				{
 					if (const FSceneOutlinerTreeItemPtr* ParentItem = Items.Find(OuterWorldDataLayers))
 					{
diff --git a/Engine/Source/Editor/DataLayerEditor/Private/DataLayer/DataLayerMode.cpp b/Engine/Source/Editor/DataLayerEditor/Private/DataLayer/DataLayerMode.cpp
--- a/Engine/Source/Editor/DataLayerEditor/Private/DataLayer/DataLayerMode.cpp
+++ b/Engine/Source/Editor/DataLayerEditor/Private/DataLayer/DataLayerMode.cpp
@@ -658,7 +658,7 @@ FSceneOutlinerDragValidationInfo FDataLayerMode::ValidateDrop(const ISceneOutlin
 						return FSceneOutlinerDragValidationInfo(ESceneOutlinerDropCompatibility::IncompatibleGeneric, LOCTEXT("CantMoveToSameDataLayer", "Can't move Data Layer to same Data Layer"));
 					}
 
-					if (ParentDataLayer->GetOuterAWorldDataLayers() != DataLayerInstance->GetOuterAWorldDataLayers())
+					if (ParentDataLayer->GetOuterWorldDataLayers() != DataLayerInstance->GetOuterWorldDataLayers())
 					{
 						return FSceneOutlinerDragValidationInfo(ESceneOutlinerDropCompatibility::IncompatibleGeneric, LOCTEXT("CantMoveToDifferentOuterWorldDataLayer", "Can't move Data Layer to a different World Data Layer"));
 					}
@@ -1269,7 +1269,7 @@ void FDataLayerMode::RegisterContextMenu()
 
 					FDataLayerCreationParameters CreationParams;
 					CreationParams.DataLayerAsset = const_cast<UDataLayerAsset*>(InDataLayerAsset);
-					CreationParams.WorldDataLayers = InParentDataLayer ? InParentDataLayer->GetOuterAWorldDataLayers() : (SelectedWorldDataLayers.Num() == 1 ? SelectedWorldDataLayers[0] : Mode->GetOwningWorldAWorldDataLayers());
+					CreationParams.WorldDataLayers = InParentDataLayer ? InParentDataLayer->GetOuterWorldDataLayers() : (SelectedWorldDataLayers.Num() == 1 ? SelectedWorldDataLayers[0] : Mode->GetOwningWorldAWorldDataLayers());
 					if (UDataLayerInstance* NewDataLayerInstance = UDataLayerEditorSubsystem::Get()->CreateDataLayerInstance(CreationParams))
 					{
 						UDataLayerEditorSubsystem::Get()->SetParentDataLayer(NewDataLayerInstance, InParentDataLayer);
diff --git a/Engine/Source/Runtime/Engine/Private/WorldPartition/DataLayer/DataLayerInstance.cpp b/Engine/Source/Runtime/Engine/Private/WorldPartition/DataLayer/DataLayerInstance.cpp
--- a/Engine/Source/Runtime/Engine/Private/WorldPartition/DataLayer/DataLayerInstance.cpp
+++ b/Engine/Source/Runtime/Engine/Private/WorldPartition/DataLayer/DataLayerInstance.cpp
@@ -42,6 +42,16 @@ void UDataLayerInstance::PostLoad()
 	}
 }
 
+UWorld* UDataLayerInstance::GetOuterWorld() const
+{
+	return UObject::GetTypedOuter<UWorld>();
+}
+
+AWorldDataLayers* UDataLayerInstance::GetOuterWorldDataLayers() const
+{
+	return GetOuterWorld()->GetWorldDataLayers();
+}
+
 #if WITH_EDITOR
 void UDataLayerInstance::SetVisible(bool bInIsVisible)
 {
@@ -111,7 +121,7 @@ bool UDataLayerInstance::IsLocked() const
 		return true;
 	}
 
-	return IsRuntime() && !GetOuterAWorldDataLayers()->GetAllowRuntimeDataLayerEditing();
+	return IsRuntime() && !GetOuterWorldDataLayers()->GetAllowRuntimeDataLayerEditing();
 }
 
 bool UDataLayerInstance::IsReadOnly() const
@@ -158,7 +168,7 @@ bool UDataLayerInstance::CanParent(const UDataLayerInstance* InParent) const
 {
 	return (this != InParent)
 		&& (Parent != InParent) 
-		&& (InParent == nullptr || (IsDataLayerTypeValidToParent(InParent->GetType()) && (InParent->GetOuterAWorldDataLayers() == GetOuterAWorldDataLayers())));
+		&& (InParent == nullptr || (IsDataLayerTypeValidToParent(InParent->GetType()) && (InParent->GetOuterWorldDataLayers() == GetOuterWorldDataLayers())));
 }
 
 bool UDataLayerInstance::IsDataLayerTypeValidToParent(EDataLayerType ParentDataLayerType) const
@@ -242,17 +252,17 @@ bool UDataLayerInstance::Validate(IStreamingGenerationErrorHandler* ErrorHandler
 
 bool UDataLayerInstance::IsInActorEditorContext() const
 {
-	return GetOuterAWorldDataLayers()->IsInActorEditorContext(this);
+	return GetOuterWorldDataLayers()->IsInActorEditorContext(this);
 }
 
 bool UDataLayerInstance::AddToActorEditorContext()
 {
-	return GetOuterAWorldDataLayers()->AddToActorEditorContext(this);
+	return GetOuterWorldDataLayers()->AddToActorEditorContext(this);
 }
 
 bool UDataLayerInstance::RemoveFromActorEditorContext()
 {
-	return GetOuterAWorldDataLayers()->RemoveFromActorEditorContext(this);
+	return GetOuterWorldDataLayers()->RemoveFromActorEditorContext(this);
 }
 
 #endif
@@ -304,7 +314,7 @@ void UDataLayerInstance::ForEachChild(TFunctionRef<bool(const UDataLayerInstance
 
 void UDataLayerInstance::AddChild(UDataLayerInstance* InDataLayer)
 {
-	check(InDataLayer->GetOuterAWorldDataLayers() == GetOuterAWorldDataLayers())
+	check(InDataLayer->GetOuterWorldDataLayers() == GetOuterWorldDataLayers())
 	Modify();
 	checkSlow(!Children.Contains(InDataLayer));
 	Children.Add(InDataLayer);
@@ -312,19 +322,19 @@ void UDataLayerInstance::AddChild(UDataLayerInstance* InDataLayer)
 
 EDataLayerRuntimeState UDataLayerInstance::GetRuntimeState() const
 {
-	return GetOuterAWorldDataLayers()->GetDataLayerRuntimeStateByName(GetDataLayerFName());
+	return GetOuterWorldDataLayers()->GetDataLayerRuntimeStateByName(GetDataLayerFName());
 }
 
 EDataLayerRuntimeState UDataLayerInstance::GetEffectiveRuntimeState() const
 {
-	return GetOuterAWorldDataLayers()->GetDataLayerEffectiveRuntimeStateByName(GetDataLayerFName());
+	return GetOuterWorldDataLayers()->GetDataLayerEffectiveRuntimeStateByName(GetDataLayerFName());
 }
 
 bool UDataLayerInstance::SetRuntimeState(EDataLayerRuntimeState InState, bool bInIsRecursive) const
 {
-	if (GetOuterAWorldDataLayers()->HasAuthority())
+	if (GetOuterWorldDataLayers()->HasAuthority())
 	{
-		GetOuterAWorldDataLayers()->SetDataLayerRuntimeState(this, InState, bInIsRecursive);
+		GetOuterWorldDataLayers()->SetDataLayerRuntimeState(this, InState, bInIsRecursive);
 		return true;
 	}
 	UE_LOG(LogWorldPartition, Error, TEXT("SetDataLayerRuntimeState can only execute on authority"));
diff --git a/Engine/Source/Runtime/Engine/Private/WorldPartition/DataLayer/DataLayerInstanceWithAsset.cpp b/Engine/Source/Runtime/Engine/Private/WorldPartition/DataLayer/DataLayerInstanceWithAsset.cpp
--- a/Engine/Source/Runtime/Engine/Private/WorldPartition/DataLayer/DataLayerInstanceWithAsset.cpp
+++ b/Engine/Source/Runtime/Engine/Private/WorldPartition/DataLayer/DataLayerInstanceWithAsset.cpp
@@ -59,7 +59,7 @@ bool UDataLayerInstanceWithAsset::CanAddActor(AActor* InActor) const
 
 bool UDataLayerInstanceWithAsset::PerformAddActor(AActor* InActor) const
 {	
-	check(GetTypedOuter<ULevel>() == InActor->GetLevel()); // Make sure the instance is part of the same world as the actor.
+	check(GetOuterWorldDataLayers() == InActor->GetLevel()->GetWorldDataLayers()); // Make sure the instance is part of the same WorldDataLayers as the actor's level WorldDataLayer.
 	check(UDataLayerManager::GetDataLayerManager(InActor)->GetDataLayerInstance(DataLayerAsset) != nullptr); // Make sure the DataLayerInstance exists for this level
 	return FAssignActorDataLayer::AddDataLayerAsset(InActor, DataLayerAsset);
 }
diff --git a/Engine/Source/Runtime/Engine/Public/WorldPartition/DataLayer/DataLayerInstance.h b/Engine/Source/Runtime/Engine/Public/WorldPartition/DataLayer/DataLayerInstance.h
--- a/Engine/Source/Runtime/Engine/Public/WorldPartition/DataLayer/DataLayerInstance.h
+++ b/Engine/Source/Runtime/Engine/Public/WorldPartition/DataLayer/DataLayerInstance.h
@@ -39,7 +39,7 @@ const inline TCHAR* GetDataLayerRuntimeStateName(EDataLayerRuntimeState State)
 	return TEXT("Invalid");
 }
 
-UCLASS(Config = Engine, PerObjectConfig, Within = WorldDataLayers, BlueprintType, AutoCollapseCategories = ("Data Layer|Advanced"), AutoExpandCategories = ("Data Layer|Editor", "Data Layer|Advanced|Runtime"))
+UCLASS(Config = Engine, PerObjectConfig, BlueprintType, AutoCollapseCategories = ("Data Layer|Advanced"), AutoExpandCategories = ("Data Layer|Editor", "Data Layer|Advanced|Runtime"))
 class ENGINE_API UDataLayerInstance : public UObject
 {
 	GENERATED_UCLASS_BODY()
@@ -49,6 +49,18 @@ class ENGINE_API UDataLayerInstance : public UObject
 
 public:
 	virtual void PostLoad() override;
+	
+	template<class T>
+	T* GetTypedOuter() const
+	{
+		static_assert(!std::is_same<T, UWorld>::value, "Use GetOuterWorld instead");
+		static_assert(!std::is_same<T, ULevel>::value, "Use GetOuterWorld()->PersistentLevel instead");
+		static_assert(!std::is_same<T, AWorldDataLayers>::value, "Use GetOuterWorldDataLayers instead");
+		return Super::GetTypedOuter<T>();
+	}
+
+	UWorld* GetOuterWorld() const;
+	AWorldDataLayers* GetOuterWorldDataLayers() const;
 
 #if WITH_EDITOR
 	void SetVisible(bool bIsVisible);
