Added missing forward declare

diff --git a/Engine/Source/Runtime/Engine/Public/WorldPartition/DataLayer/DataLayerInstance.h b/Engine/Source/Runtime/Engine/Public/WorldPartition/DataLayer/DataLayerInstance.h
--- a/Engine/Source/Runtime/Engine/Public/WorldPartition/DataLayer/DataLayerInstance.h
+++ b/Engine/Source/Runtime/Engine/Public/WorldPartition/DataLayer/DataLayerInstance.h
@@ -12,7 +12,9 @@
 
 #define DATALAYER_TO_INSTANCE_RUNTIME_CONVERSION_ENABLED 1
 
+class AWorldDataLayers;
 class IStreamingGenerationErrorHandler;
+class UWorld;
 
 UENUM(BlueprintType)
 enum class EDataLayerRuntimeState : uint8
