Adding back in the call to SetupParameterCollectionInstances in UWorld::PostLoad() to fix bugs with parameters not updating on static meshes.

- We need to initialize all loaded material parameter collections up-front so static meshes can grab the correct instance when caching mesh draw commands.
- Without this call, static meshes will build their draw commands with the Default MPC, and then sequencer will create a _new_ MPC instance and the static mesh is never made aware of this.
- While this is an up-front performance cost to update these MPCs on load, it is better than the alternative which is to notify _every_ static mesh that a new MPC instance was created and all the static meshes would need to re-cache all their draw commands.  Re-caching draw commands is an expensive operation.

diff --git a/Engine/Source/Runtime/Engine/Private/World.cpp b/Engine/Source/Runtime/Engine/Private/World.cpp
--- a/Engine/Source/Runtime/Engine/Private/World.cpp
+++ b/Engine/Source/Runtime/Engine/Private/World.cpp
@@ -1225,6 +1225,9 @@ void UWorld::PostLoad()
 	// Add the garbage collection callbacks
 	FLevelStreamingGCHelper::AddGarbageCollectorCallback();
 
+	// Initially set up the parameter collection list. This may be run again in UWorld::InitWorld but it's required here for some editor and streaming cases
+	SetupParameterCollectionInstances();
+
 #if WITH_EDITOR
 	if (GIsEditor)
 	{
