Fix for static analysis complaint (265952) introduced by 24316939

diff --git a/Engine/Source/Runtime/Engine/Public/Instances/InstancedPlacementHash.h b/Engine/Source/Runtime/Engine/Public/Instances/InstancedPlacementHash.h
--- a/Engine/Source/Runtime/Engine/Public/Instances/InstancedPlacementHash.h
+++ b/Engine/Source/Runtime/Engine/Public/Instances/InstancedPlacementHash.h
@@ -115,7 +115,8 @@ public:
 		double SphereRadiusSquared = SphereRadius * SphereRadius;
 
 		// if there are more potential cells within range than the number of actually populated cells...
-		double SphereRadiusCells = SphereRadius / (1 << HashCellBits);
+		int32 CellSize = (1 << HashCellBits);
+		double SphereRadiusCells = SphereRadius / CellSize;
 		constexpr double SphereVolumeConstant = 4.0 / 3.0 * UE_PI;
 		double ApproxCellsToCheck = SphereVolumeConstant * (SphereRadiusCells * SphereRadiusCells * SphereRadiusCells);
 		if (ApproxCellsToCheck > CellMap.Num())
