Compile time assert if trying to modify replication conditions for properties that aren't part of the class hierarchy.

#rnx

diff --git a/Engine/Source/Runtime/Engine/Public/Net/UnrealNetwork.h b/Engine/Source/Runtime/Engine/Public/Net/UnrealNetwork.h
--- a/Engine/Source/Runtime/Engine/Public/Net/UnrealNetwork.h
+++ b/Engine/Source/Runtime/Engine/Public/Net/UnrealNetwork.h
@@ -19,6 +19,7 @@
 #include "Net/Core/Connection/NetResult.h"
 #include "Net/Core/PropertyConditions/PropertyConditions.h"
 #include "Net/ReplayResult.h"
+#include <type_traits>
 
 class AActor;
 
@@ -210,23 +211,15 @@ namespace NetworkingPrivate
 	Replication.
 -----------------------------------------------------------------------------*/
 
-static bool ValidateReplicatedClassInheritance(const UClass* CallingClass, const UClass* PropClass, const TCHAR* PropertyName)
+template<class BaseClass, class DerivedClass>
+constexpr bool ValidateReplicatedClassInheritance()
 {
-#if !(UE_BUILD_SHIPPING || UE_BUILD_TEST)
-	if (!CallingClass->IsChildOf(PropClass))
-	{
-		UE_LOG(LogNet, Fatal, TEXT("Attempt to replicate property '%s.%s' in C++ but class '%s' is not a child of '%s'"), *PropClass->GetName(), PropertyName, *CallingClass->GetName(), *PropClass->GetName());
-	}
-#endif
-
-	return true;
+	return std::is_base_of_v<BaseClass, DerivedClass>;
 }
 
 /** wrapper to find replicated properties that also makes sure they're valid */
-static FProperty* GetReplicatedProperty(const UClass* CallingClass, const UClass* PropClass, const FName& PropName)
+inline FProperty* GetReplicatedProperty(const UClass* CallingClass, const UClass* PropClass, const FName& PropName)
 {
-	ValidateReplicatedClassInheritance(CallingClass, PropClass, *PropName.ToString());
-
 	FProperty* TheProperty = FindFieldChecked<FProperty>(PropClass, PropName);
 #if !(UE_BUILD_SHIPPING || UE_BUILD_TEST)
 	if (!(TheProperty->PropertyFlags & CPF_Net))
@@ -239,7 +232,7 @@ static FProperty* GetReplicatedProperty(const UClass* CallingClass, const UClass
 
 #define DOREPLIFETIME_WITH_PARAMS_FAST(c,v,params) \
 { \
-	static const bool bIsValid_##c_##v = ValidateReplicatedClassInheritance(StaticClass(), c::StaticClass(), TEXT(#v)); \
+	static_assert(ValidateReplicatedClassInheritance<c, ThisClass>(), #c "." #v " is not accessible from this class."); \
 	const TCHAR* DoRepPropertyName_##c_##v(TEXT(#v)); \
 	const NetworkingPrivate::FRepPropertyDescriptor PropertyDescriptor_##c_##v(DoRepPropertyName_##c_##v, (int32)c::ENetFields_Private::v, 1); \
 \
@@ -250,7 +243,7 @@ static FProperty* GetReplicatedProperty(const UClass* CallingClass, const UClass
 
 #define DOREPLIFETIME_WITH_PARAMS_FAST_STATIC_ARRAY(c,v,params) \
 { \
-	static const bool bIsValid_##c_##v = ValidateReplicatedClassInheritance(StaticClass(), c::StaticClass(), TEXT(#v)); \
+	static_assert(ValidateReplicatedClassInheritance<c, ThisClass>(), #c "." #v " is not accessible from this class."); \
 	const TCHAR* DoRepPropertyName_##c_##v(TEXT(#v)); \
 	const NetworkingPrivate::FRepPropertyDescriptor PropertyDescriptor_##c_##v(DoRepPropertyName_##c_##v, (int32)c::ENetFields_Private::v##_STATIC_ARRAY, (int32)c::EArrayDims_Private::v); \
 	RegisterReplicatedLifetimeProperty(PropertyDescriptor_##c_##v, OutLifetimeProps, params); \
@@ -258,7 +251,8 @@ static FProperty* GetReplicatedProperty(const UClass* CallingClass, const UClass
 
 #define DOREPLIFETIME_WITH_PARAMS(c,v,params) \
 { \
-	FProperty* ReplicatedProperty = GetReplicatedProperty(StaticClass(), c::StaticClass(),GET_MEMBER_NAME_CHECKED(c,v)); \
+	static_assert(ValidateReplicatedClassInheritance<c, ThisClass>(), #c "." #v " is not accessible from this class."); \
+	FProperty* ReplicatedProperty = GetReplicatedProperty(StaticClass(), c::StaticClass(), GET_MEMBER_NAME_CHECKED(c,v)); \
 	PRAGMA_DISABLE_DEPRECATION_WARNINGS \
 	RegisterReplicatedLifetimeProperty(ReplicatedProperty, OutLifetimeProps, FixupParams<decltype(c::v)>(params)); \
 	PRAGMA_ENABLE_DEPRECATION_WARNINGS \
@@ -302,13 +296,13 @@ static FProperty* GetReplicatedProperty(const UClass* CallingClass, const UClass
 
 #define DOREPLIFETIME_ACTIVE_OVERRIDE_FAST(c,v,active) \
 { \
-	static const bool bIsValid_##c_##v = ValidateReplicatedClassInheritance(StaticClass(), c::StaticClass(), TEXT(#v)); \
+	static_assert(ValidateReplicatedClassInheritance<c, ThisClass>(), #c "." #v " is not accessible from this class."); \
 	UE::Net::Private::FNetPropertyConditionManager::SetPropertyActiveOverride(ChangedPropertyTracker, this, (int32)c::ENetFields_Private::v, active); \
 }
 
 #define DOREPLIFETIME_ACTIVE_OVERRIDE_FAST_STATIC_ARRAY(c,v,active) \
 { \
-	static const bool bIsValid_##c_##v = ValidateReplicatedClassInheritance(StaticClass(), c::StaticClass(), TEXT(#v)); \
+	static_assert(ValidateReplicatedClassInheritance<c, ThisClass>(), #c "." #v " is not accessible from this class."); \
 	for (int32 i = 0; i < (int32)c::EArrayDims_Private::v; ++i) \
 	{ \
 		UE::Net::Private::FNetPropertyConditionManager::SetPropertyActiveOverride(ChangedPropertyTracker, this, (int32)c::ENetFields_Private::v##_STATIC_ARRAY + i, active); \
@@ -317,6 +311,7 @@ static FProperty* GetReplicatedProperty(const UClass* CallingClass, const UClass
 
 #define DOREPLIFETIME_ACTIVE_OVERRIDE(c,v,active) \
 { \
+	static_assert(ValidateReplicatedClassInheritance<c, ThisClass>(), #c "." #v " is not accessible from this class."); \
 	static FProperty* sp##v = GetReplicatedProperty(StaticClass(), c::StaticClass(),GET_MEMBER_NAME_CHECKED(c,v)); \
 	for (int32 i = 0; i < sp##v->ArrayDim; i++) \
 	{ \
@@ -326,7 +321,7 @@ static FProperty* GetReplicatedProperty(const UClass* CallingClass, const UClass
 
 #define DOREPCUSTOMCONDITION_ACTIVE_FAST(c,v,active) \
 { \
-	static const bool bIsValid_##c_##v = ValidateReplicatedClassInheritance(StaticClass(), c::StaticClass(), TEXT(#v)); \
+	static_assert(ValidateReplicatedClassInheritance<c, ThisClass>(), #c "." #v " is not accessible from this class."); \
 	OutActiveState.SetActiveState((uint16)c::ENetFields_Private::v, active); \
 }
 
@@ -347,7 +342,7 @@ static FProperty* GetReplicatedProperty(const UClass* CallingClass, const UClass
 /** To be used in a UObject's GetReplicatedCustomConditionState for example. */
 #define DOREPDYNAMICCONDITION_INITCONDITION_FAST(c,v,cond) \
 { \
-	static const bool bIsValid_##c_##v = ValidateReplicatedClassInheritance(StaticClass(), c::StaticClass(), TEXT(#v)); \
+	static_assert(ValidateReplicatedClassInheritance<c, ThisClass>(), #c "." #v " is not accessible from this class."); \
 	OutActiveState.SetDynamicCondition((uint16)c::ENetFields_Private::v, cond); \
 }
 
@@ -420,7 +415,7 @@ DisableReplicatedLifetimeProperty(StaticClass(), c::StaticClass(), FName(TEXT(#v
 
 #define DISABLE_REPLICATED_PROPERTY_FAST(c,v) \
 { \
-	static const bool bIsValid_##c_##v = ValidateReplicatedClassInheritance(StaticClass(), c::StaticClass(), TEXT(#v)); \
+	static_assert(ValidateReplicatedClassInheritance<c, ThisClass>(), #c "." #v " is not accessible from this class."); \
 	const TCHAR* DoRepPropertyName_##c_##v(TEXT(#v)); \
 	const NetworkingPrivate::FRepPropertyDescriptor PropertyDescriptor_##c_##v(DoRepPropertyName_##c_##v, (int32)c::ENetFields_Private::v, 1); \
 	DisableReplicatedLifetimeProperty(PropertyDescriptor_##c_##v, OutLifetimeProps); \
@@ -428,7 +423,7 @@ DisableReplicatedLifetimeProperty(StaticClass(), c::StaticClass(), FName(TEXT(#v
 
 #define DISABLE_REPLICATED_PROPERTY_FAST_STATIC_ARRAY(c,v) \
 { \
-	static const bool bIsValid_##c_##v = ValidateReplicatedClassInheritance(StaticClass(), c::StaticClass(), TEXT(#v)); \
+	static_assert(ValidateReplicatedClassInheritance<c, ThisClass>(), #c "." #v " is not accessible from this class."); \
 	const TCHAR* DoRepPropertyName_##c_##v(TEXT(#v)); \
 	const NetworkingPrivate::FRepPropertyDescriptor PropertyDescriptor_##c_##v(DoRepPropertyName_##c_##v, (int32)c::ENetFields_Private::v##_STATIC_ARRAY, (int32)c::EArrayDims_Private::v); \
 	DisableReplicatedLifetimeProperty(PropertyDescriptor_##c_##v, OutLifetimeProps); \
@@ -442,7 +437,7 @@ DisableAllReplicatedPropertiesOfClass(StaticClass(), c::StaticClass(), SuperClas
 
 #define DISABLE_ALL_CLASS_REPLICATED_PROPERTIES_FAST(c, SuperClassBehavior) \
 { \
-	static const bool bIsValid_##c_##v = ValidateReplicatedClassInheritance(StaticClass(), c::StaticClass(), TEXT("DISABLE_ALL_CLASS_REPLICATED_PROPERTIES")); \
+	static_assert(ValidateReplicatedClassInheritance<c, ThisClass>(), "DISABLE_ALL_CLASS_REPLICATED_PROPERTIES"); \
 	const TCHAR* DoRepClassName_##c(TEXT(#c)); \
 	const NetworkingPrivate::FRepClassDescriptor ClassDescriptor_##c(DoRepClassName_##c, (int32)c::ENetFields_Private::NETFIELD_REP_START, (int32)c::ENetFields_Private::NETFIELD_REP_END); \
 	DisableAllReplicatedPropertiesOfClass(ClassDescriptor_##c, SuperClassBehavior, OutLifetimeProps); \
@@ -465,7 +460,7 @@ ENGINE_API void DisableAllReplicatedPropertiesOfClass(const NetworkingPrivate::F
 
 #define RESET_REPLIFETIME_CONDITION_FAST(c,v,cond) \
 { \
-	static const bool bIsValid_##c_##v = ValidateReplicatedClassInheritance(StaticClass(), c::StaticClass(), TEXT(#v)); \
+	static_assert(ValidateReplicatedClassInheritance<c, ThisClass>(), #c "." #v " is not accessible from this class."); \
 	const TCHAR* DoRepPropertyName_##c_##v(TEXT(#v)); \
 	const NetworkingPrivate::FRepPropertyDescriptor PropertyDescriptor_##c_##v(DoRepPropertyName_##c_##v, (int32)c::ENetFields_Private::v, 1); \
 	ResetReplicatedLifetimeProperty(PropertyDescriptor_##c_##v, cond, OutLifetimeProps); \
@@ -473,7 +468,7 @@ ENGINE_API void DisableAllReplicatedPropertiesOfClass(const NetworkingPrivate::F
 
 #define RESET_REPLIFETIME_CONDITION_FAST_STATIC_ARRAY(c,v,cond) \
 { \
-	static const bool bIsValid_##c_##v = ValidateReplicatedClassInheritance(StaticClass(), c::StaticClass(), TEXT(#v)); \
+	static_assert(ValidateReplicatedClassInheritance<c, ThisClass>(), #c "." #v " is not accessible from this class."); \
 	const TCHAR* DoRepPropertyName_##c_##v(TEXT(#v)); \
 	const NetworkingPrivate::FRepPropertyDescriptor PropertyDescriptor_##c_##v(DoRepPropertyName_##c_##v, (int32)c::ENetFields_Private::v##_STATIC_ARRAY, (int32)c::EArrayDims_Private::v); \
 	ResetReplicatedLifetimeProperty(PropertyDescriptor_##c_##v, cond, OutLifetimeProps); \
@@ -486,7 +481,7 @@ ENGINE_API void DisableAllReplicatedPropertiesOfClass(const NetworkingPrivate::F
 
 #define RESET_REPLIFETIME_FAST_WITH_PARAMS(c,v,params) \
 { \
-	static const bool bIsValid_##c_##v = ValidateReplicatedClassInheritance(StaticClass(), c::StaticClass(), TEXT(#v)); \
+	static_assert(ValidateReplicatedClassInheritance<c, ThisClass>(), #c "." #v " is not accessible from this class."); \
 	const TCHAR* DoRepPropertyName_##c_##v(TEXT(#v)); \
 	const NetworkingPrivate::FRepPropertyDescriptor PropertyDescriptor_##c_##v(DoRepPropertyName_##c_##v, (int32)c::ENetFields_Private::v, 1); \
 	ResetReplicatedLifetimeProperty(PropertyDescriptor_##c_##v, params, OutLifetimeProps); \
