Rename all user/editor-facing Strata name into Substrate.

This CL address only user-facing part, i.e. Property field, view modes, tool tips, ...) mentioning Strata. All internal (host code, shaders code, CVars) still contains the 'Strata' name.

diff --git a/Engine/Source/Editor/CommonMenuExtensions/Private/StrataVisualizationMenuCommands.cpp b/Engine/Source/Editor/CommonMenuExtensions/Private/StrataVisualizationMenuCommands.cpp
--- a/Engine/Source/Editor/CommonMenuExtensions/Private/StrataVisualizationMenuCommands.cpp
+++ b/Engine/Source/Editor/CommonMenuExtensions/Private/StrataVisualizationMenuCommands.cpp
@@ -16,13 +16,13 @@
 #include "Styling/AppStyle.h"
 #include "UObject/UnrealNames.h"
 
-#define LOCTEXT_NAMESPACE "StrataVisualizationMenuCommands"
+#define LOCTEXT_NAMESPACE "SubstrateVisualizationMenuCommands"
 
 FStrataVisualizationMenuCommands::FStrataVisualizationMenuCommands()
 	: TCommands<FStrataVisualizationMenuCommands>
 	(
-		TEXT("StrataVisualizationMenu"), // Context name for fast lookup
-		NSLOCTEXT("Contexts", "StrataVisualizationMenu", "Strata"), // Localized context name for displaying
+		TEXT("SubstrateVisualizationMenu"), // Context name for fast lookup
+		NSLOCTEXT("Contexts", "SubstrateVisualizationMenu", "Substrate"), // Localized context name for displaying
 		NAME_None, // Parent context name.  
 		FAppStyle::GetAppStyleSetName() // Icon Style Set
 	),
@@ -73,7 +73,7 @@ void FStrataVisualizationMenuCommands::BuildVisualisationSubMenu(FMenuBuilder& M
 	const FStrataVisualizationMenuCommands& Commands = FStrataVisualizationMenuCommands::Get();
 	if (Commands.IsPopulated())
 	{
-		Menu.BeginSection("LevelViewportStrataVisualizationMode", LOCTEXT("StrataVisualizationHeader", "Strata Visualization Mode"));
+		Menu.BeginSection("LevelViewportSubstrateVisualizationMode", LOCTEXT("SubstrateVisualizationHeader", "Substrate Visualization Mode"));
 
 		Commands.AddCommandTypeToMenu(Menu, FStrataVisualizationType::MaterialProperties);
 		Commands.AddCommandTypeToMenu(Menu, FStrataVisualizationType::MaterialCount);
diff --git a/Engine/Source/Editor/DetailCustomizations/Private/DebugCameraControllerSettingsCustomization.cpp b/Engine/Source/Editor/DetailCustomizations/Private/DebugCameraControllerSettingsCustomization.cpp
--- a/Engine/Source/Editor/DetailCustomizations/Private/DebugCameraControllerSettingsCustomization.cpp
+++ b/Engine/Source/Editor/DetailCustomizations/Private/DebugCameraControllerSettingsCustomization.cpp
@@ -45,7 +45,7 @@ void FDebugCameraControllerSettingsViewModeIndexCustomization::CustomizeHeader(T
 	EnumRestriction->AddHiddenValue(ViewModeIndexEnum->GetNameStringByValue((uint8)EViewModeIndex::VMI_VisualizeBuffer));
 	EnumRestriction->AddHiddenValue(ViewModeIndexEnum->GetNameStringByValue((uint8)EViewModeIndex::VMI_VisualizeNanite));
 	EnumRestriction->AddHiddenValue(ViewModeIndexEnum->GetNameStringByValue((uint8)EViewModeIndex::VMI_VisualizeLumen));
-	EnumRestriction->AddHiddenValue(ViewModeIndexEnum->GetNameStringByValue((uint8)EViewModeIndex::VMI_VisualizeStrata));
+	EnumRestriction->AddHiddenValue(ViewModeIndexEnum->GetNameStringByValue((uint8)EViewModeIndex::VMI_VisualizeSubstrate));
 	EnumRestriction->AddHiddenValue(ViewModeIndexEnum->GetNameStringByValue((uint8)EViewModeIndex::VMI_VisualizeGroom));
 	EnumRestriction->AddHiddenValue(ViewModeIndexEnum->GetNameStringByValue((uint8)EViewModeIndex::VMI_VisualizeVirtualShadowMap));
 	EnumRestriction->AddHiddenValue(ViewModeIndexEnum->GetNameStringByValue((uint8)EViewModeIndex::VMI_VisualizeGPUSkinCache));
diff --git a/Engine/Source/Editor/EditorStyle/Private/StarshipStyle.cpp b/Engine/Source/Editor/EditorStyle/Private/StarshipStyle.cpp
--- a/Engine/Source/Editor/EditorStyle/Private/StarshipStyle.cpp
+++ b/Engine/Source/Editor/EditorStyle/Private/StarshipStyle.cpp
@@ -4168,7 +4168,7 @@ void FStarshipEditorStyle::FStyle::SetupLevelEditorStyle()
 		Set( "EditorViewport.VisualizeBufferMode",    new IMAGE_BRUSH_SVG("Starship/Common/BufferVisualization", Icon16x16 ) );
 		Set( "EditorViewport.VisualizeNaniteMode",    new IMAGE_BRUSH_SVG("Starship/Common/BufferVisualization", Icon16x16 ) );
 		Set( "EditorViewport.VisualizeLumenMode",     new IMAGE_BRUSH_SVG("Starship/Common/BufferVisualization", Icon16x16 ) );
-		Set( "EditorViewport.VisualizeStrataMode",    new IMAGE_BRUSH_SVG("Starship/Common/BufferVisualization", Icon16x16 ) );
+		Set( "EditorViewport.VisualizeSubstrateMode", new IMAGE_BRUSH_SVG("Starship/Common/BufferVisualization", Icon16x16 ) );
 		Set( "EditorViewport.VisualizeGroomMode",     new IMAGE_BRUSH_SVG("Starship/Common/BufferVisualization", Icon16x16 ) );
 		Set( "EditorViewport.VisualizeVirtualShadowMapMode", new IMAGE_BRUSH_SVG("Starship/Common/BufferVisualization", Icon16x16 ) );
 		Set( "EditorViewport.CollisionPawn",          new IMAGE_BRUSH_SVG("Starship/Common/PlayerCollision", Icon16x16 ) );
diff --git a/Engine/Source/Editor/LevelEditor/Private/SLevelViewportToolBar.cpp b/Engine/Source/Editor/LevelEditor/Private/SLevelViewportToolBar.cpp
--- a/Engine/Source/Editor/LevelEditor/Private/SLevelViewportToolBar.cpp
+++ b/Engine/Source/Editor/LevelEditor/Private/SLevelViewportToolBar.cpp
@@ -1621,9 +1621,9 @@ void SLevelViewportToolBar::FillViewMenu(UToolMenu* Menu)
 	{
 		FToolMenuSection& Section = Menu->FindOrAddSection("ViewMode");
 		Section.AddSubMenu(
-			"VisualizeStrataViewMode",
-			LOCTEXT("VisualizeStrataViewModeDisplayName", "Strata"),
-			LOCTEXT("StrataVisualizationMenu_ToolTip", "Select a mode for Strata visualization"),
+			"VisualizeSubstrateViewMode",
+			LOCTEXT("VisualizeSubstrateViewModeDisplayName", "Substrate"),
+			LOCTEXT("SubstrateVisualizationMenu_ToolTip", "Select a mode for Substrate visualization"),
 			FNewMenuDelegate::CreateStatic(&FStrataVisualizationMenuCommands::BuildVisualisationSubMenu),
 			FUIAction(
 				FExecuteAction(),
@@ -1633,12 +1633,12 @@ void SLevelViewportToolBar::FillViewMenu(UToolMenu* Menu)
 						const TSharedRef<SEditorViewport> ViewportRef = Viewport.Pin().ToSharedRef();
 						const TSharedPtr<FEditorViewportClient> ViewportClient = ViewportRef->GetViewportClient();
 						check(ViewportClient.IsValid());
-						return ViewportClient->IsViewModeEnabled(VMI_VisualizeStrata);
+						return ViewportClient->IsViewModeEnabled(VMI_VisualizeSubstrate);
 					})
 			),
 			EUserInterfaceActionType::RadioButton,
 						/* bInOpenSubMenuOnClick = */ false,
-						FSlateIcon(FAppStyle::GetAppStyleSetName(), "EditorViewport.VisualizeStrataMode")
+						FSlateIcon(FAppStyle::GetAppStyleSetName(), "EditorViewport.VisualizeSubstrateMode")
 						);
 	}
 
diff --git a/Engine/Source/Editor/MaterialEditor/Private/MaterialEditor.cpp b/Engine/Source/Editor/MaterialEditor/Private/MaterialEditor.cpp
--- a/Engine/Source/Editor/MaterialEditor/Private/MaterialEditor.cpp
+++ b/Engine/Source/Editor/MaterialEditor/Private/MaterialEditor.cpp
@@ -445,7 +445,7 @@ void FMaterialEditor::RegisterToolbarTab(const TSharedRef<class FTabManager>& In
 		.SetIcon(FSlateIcon(FAppStyle::GetAppStyleSetName(), "LevelEditor.Tabs.Layers"));
 
 	InTabManager->RegisterTabSpawner(FMaterialEditorTabs::StrataTabId, FOnSpawnTab::CreateSP(this, &FMaterialEditor::SpawnTab_Strata))
-		.SetDisplayName(LOCTEXT("StrataTab", "Strata"))
+		.SetDisplayName(LOCTEXT("SubstrateTab", "Substrate"))
 		.SetGroup(WorkspaceMenuCategoryRef)
 		.SetIcon(FSlateIcon(FAppStyle::GetAppStyleSetName(), "Kismet.Tabs.Palette"));	// STRATA_TODO a strata icon
 
@@ -5185,7 +5185,7 @@ TSharedRef<SDockTab> FMaterialEditor::SpawnTab_Strata(const FSpawnTabArgs& Args)
 	check(Args.GetTabId() == FMaterialEditorTabs::StrataTabId);
 
 	TSharedRef<SDockTab> StrataTab = SNew(SDockTab)
-		.Label(LOCTEXT("MaterialStrataTabTitle", "Strata"))
+		.Label(LOCTEXT("MaterialSubstrateTabTitle", "Substrate"))
 		[
 			SNew(SBox)
 			[
diff --git a/Engine/Source/Editor/MaterialEditor/Private/SMaterialEditorStrataWidget.cpp b/Engine/Source/Editor/MaterialEditor/Private/SMaterialEditorStrataWidget.cpp
--- a/Engine/Source/Editor/MaterialEditor/Private/SMaterialEditorStrataWidget.cpp
+++ b/Engine/Source/Editor/MaterialEditor/Private/SMaterialEditorStrataWidget.cpp
@@ -101,7 +101,7 @@ void SMaterialEditorStrataWidget::Construct(const FArguments& InArgs, TWeakPtr<F
 					.ColorAndOpacity(FLinearColor::Yellow)
 					.ShadowColorAndOpacity(FLinearColor::Black)
 					.ShadowOffset(FVector2D::UnitVector)
-					.Text(LOCTEXT("StrataWidgetNotEnable", "Strata is not enabled."))
+					.Text(LOCTEXT("SubstrateWidgetNotEnable", "Substrate is not enabled."))
 				]
 			]
 		];
diff --git a/Engine/Source/Editor/UnrealEd/Private/EditorViewportClient.cpp b/Engine/Source/Editor/UnrealEd/Private/EditorViewportClient.cpp
--- a/Engine/Source/Editor/UnrealEd/Private/EditorViewportClient.cpp
+++ b/Engine/Source/Editor/UnrealEd/Private/EditorViewportClient.cpp
@@ -2689,18 +2689,18 @@ FText FEditorViewportClient::GetCurrentVirtualShadowMapVisualizationModeDisplayN
 
 void FEditorViewportClient::ChangeStrataVisualizationMode(FName InName)
 {
-	SetViewMode(VMI_VisualizeStrata);
+	SetViewMode(VMI_VisualizeSubstrate);
 	CurrentStrataVisualizationMode = InName;
 }
 
 bool FEditorViewportClient::IsStrataVisualizationModeSelected(FName InName) const
 {
-	return IsViewModeEnabled(VMI_VisualizeStrata) && CurrentStrataVisualizationMode == InName;
+	return IsViewModeEnabled(VMI_VisualizeSubstrate) && CurrentStrataVisualizationMode == InName;
 }
 
 FText FEditorViewportClient::GetCurrentStrataVisualizationModeDisplayName() const
 {
-	checkf(IsViewModeEnabled(VMI_VisualizeStrata), TEXT("In order to call GetCurrentStrataVisualizationMode(), first you must set ViewMode to VMI_VisualizeStrata."));
+	checkf(IsViewModeEnabled(VMI_VisualizeSubstrate), TEXT("In order to call GetCurrentSubstrateVisualizationMode(), first you must set ViewMode to VMI_VisualizeSubstrate."));
 	return GetStrataVisualizationData().GetModeDisplayName(CurrentStrataVisualizationMode);
 }
 
diff --git a/Engine/Source/Editor/UnrealEd/Private/EditorViewportCommands.cpp b/Engine/Source/Editor/UnrealEd/Private/EditorViewportCommands.cpp
--- a/Engine/Source/Editor/UnrealEd/Private/EditorViewportCommands.cpp
+++ b/Engine/Source/Editor/UnrealEd/Private/EditorViewportCommands.cpp
@@ -120,7 +120,7 @@ void FEditorViewportCommands::RegisterCommands()
 	UI_COMMAND( VisualizeBufferMode, "Buffer Visualization View Mode", "Renders a set of selected post process materials, which visualize various intermediate render buffers (material attributes)", EUserInterfaceActionType::RadioButton, FInputChord());
 	UI_COMMAND( VisualizeNaniteMode, "Nanite Visualization View Mode", "Visualizes various rendering aspects of Nanite.", EUserInterfaceActionType::RadioButton, FInputChord());
 	UI_COMMAND( VisualizeLumenMode, "Lumen Visualization View Mode", "Visualizes Lumen debug views.", EUserInterfaceActionType::RadioButton, FInputChord());
-	UI_COMMAND( VisualizeStrataMode, "Strata Visualization View Mode", "Visualizes Strata debug views.", EUserInterfaceActionType::RadioButton, FInputChord());
+	UI_COMMAND( VisualizeSubstrateMode, "Substrate Visualization View Mode", "Visualizes Substrate debug views.", EUserInterfaceActionType::RadioButton, FInputChord());
 	UI_COMMAND( VisualizeGroomMode, "Groom Visualization View Mode", "Visualizes Groom debug views.", EUserInterfaceActionType::RadioButton, FInputChord());
 	UI_COMMAND( VisualizeVirtualShadowMapMode, "Virtual Shadow Map Visualization View Mode", "Visualizes various rendering aspects of virtual shadow maps.", EUserInterfaceActionType::RadioButton, FInputChord());
 
diff --git a/Engine/Source/Editor/UnrealEd/Private/Lightmass/Lightmass.cpp b/Engine/Source/Editor/UnrealEd/Private/Lightmass/Lightmass.cpp
--- a/Engine/Source/Editor/UnrealEd/Private/Lightmass/Lightmass.cpp
+++ b/Engine/Source/Editor/UnrealEd/Private/Lightmass/Lightmass.cpp
@@ -57,7 +57,7 @@
 
 extern FSwarmDebugOptions GSwarmDebugOptions;
 
-bool Lightmass_IsStrataEnabled();
+bool Lightmass_IsSubstrateEnabled();
 
 DEFINE_LOG_CATEGORY_STATIC(LogLightmassSolver, Warning, All);
 /**
@@ -1357,10 +1357,10 @@ void FLightmassExporter::GetMaterialHash(const UMaterialInterface* Material, FSH
 		}
 	}
 
-	if (Lightmass_IsStrataEnabled())
+	if (Lightmass_IsSubstrateEnabled())
 	{
-		uint32 LightmassStrataVersion = 0XB6A0D99F; // This can be change when the code/logic for converting material to strata for lightmap has changed.
-		HashState.Update((const uint8*)&LightmassStrataVersion, sizeof(LightmassStrataVersion));
+		uint32 LightmassSubstrateVersion = 0XB6A0D99F; // This can be change when the code/logic for converting material to Substrate for lightmap has changed.
+		HashState.Update((const uint8*)&LightmassSubstrateVersion, sizeof(LightmassSubstrateVersion));
 	}
 
 	HashState.Final();
diff --git a/Engine/Source/Editor/UnrealEd/Private/Lightmass/LightmassRender.cpp b/Engine/Source/Editor/UnrealEd/Private/Lightmass/LightmassRender.cpp
--- a/Engine/Source/Editor/UnrealEd/Private/Lightmass/LightmassRender.cpp
+++ b/Engine/Source/Editor/UnrealEd/Private/Lightmass/LightmassRender.cpp
@@ -32,9 +32,9 @@
 
 DEFINE_LOG_CATEGORY_STATIC(LogLightmassRender, Error, All);
 
-extern bool Lightmass_IsStrataEnabled()
+extern bool Lightmass_IsSubstrateEnabled()
 {
-	static const auto CVar = IConsoleManager::Get().FindTConsoleVariableDataInt(TEXT("r.Strata"));
+	static const auto CVar = IConsoleManager::Get().FindTConsoleVariableDataInt(TEXT("r.Substrate"));
 	return CVar && CVar->GetValueOnAnyThread() > 0;
 }
 
@@ -255,7 +255,7 @@ public:
 	{
 		if (InMaterialInterface)
 		{
-			bStrataEnabled = Lightmass_IsStrataEnabled();
+			bStrataEnabled = Lightmass_IsSubstrateEnabled();
 			MaterialInterface = InMaterialInterface;
 			Material = MaterialInterface ? MaterialInterface->GetMaterial() : NULL;
 			PropertyToCompile = InPropertyToCompile;
@@ -1048,7 +1048,7 @@ void FLightmassMaterialRenderer::BeginGenerateMaterialData(
 	if (BaseMaterial)
 	{
 		check(!MaterialExportData.Contains(InMaterial));
-		const bool bStrataEnabled = Lightmass_IsStrataEnabled();
+		const bool bStrataEnabled = Lightmass_IsSubstrateEnabled();
 
 		FMaterialExportDataEntry& MaterialData = MaterialExportData.Add(InMaterial, FMaterialExportDataEntry(ChannelName));
 
@@ -1109,7 +1109,7 @@ bool FLightmassMaterialRenderer::GenerateMaterialData(
 	check(BaseMaterial);
 
 	EBlendMode BlendMode = InMaterial.GetBlendMode();
-	const bool bStrataEnabled = Lightmass_IsStrataEnabled();
+	const bool bStrataEnabled = Lightmass_IsSubstrateEnabled();
 
 	FMaterialShadingModelField ShadingModels = InMaterial.GetShadingModels();
  	if (!bStrataEnabled &&		// Shading models are irrelevant when using Strata
diff --git a/Engine/Source/Editor/UnrealEd/Private/MaterialGraphNode.cpp b/Engine/Source/Editor/UnrealEd/Private/MaterialGraphNode.cpp
--- a/Engine/Source/Editor/UnrealEd/Private/MaterialGraphNode.cpp
+++ b/Engine/Source/Editor/UnrealEd/Private/MaterialGraphNode.cpp
@@ -144,7 +144,7 @@ bool UMaterialGraphNode::CanPasteHere(const UEdGraph* TargetGraph) const
 		if (MaterialExpression->IsA(UMaterialExpressionStrataLegacyConversion::StaticClass()))
 		{
 			// We could have used CanDuplicateNode() returning false to prevent the copy but it is nicer to have a notification about why the copy is not happening.
-			FText NotificationText = NSLOCTEXT("UMaterialGraphNode", "SkippedStrataLegacyConversion", "StrataLegacyConversion node cannot be copied! It is only used to convert legacy material to Strata.");
+			FText NotificationText = NSLOCTEXT("UMaterialGraphNode", "SkippedSubstrateLegacyConversion", "SubstrateLegacyConversion node cannot be copied! It is only used to convert legacy material to Substrate.");
 			FNotificationInfo Info(NotificationText);
 			Info.ExpireDuration = 5.0f;
 			Info.bUseLargeFont = false;
diff --git a/Engine/Source/Editor/UnrealEd/Private/MaterialGraphSchema.cpp b/Engine/Source/Editor/UnrealEd/Private/MaterialGraphSchema.cpp
--- a/Engine/Source/Editor/UnrealEd/Private/MaterialGraphSchema.cpp
+++ b/Engine/Source/Editor/UnrealEd/Private/MaterialGraphSchema.cpp
@@ -123,7 +123,7 @@ void FMaterialGraphSchemaAction_NewNode::SetFunctionInputType(UMaterialExpressio
 		FunctionInput->InputType = FunctionInput_MaterialAttributes;
 		break;
 	case MCT_Strata:
-		FunctionInput->InputType = FunctionInput_Strata;
+		FunctionInput->InputType = FunctionInput_Substrate;
 		break;
 	default:
 		break;
@@ -1007,7 +1007,7 @@ void UMaterialGraphSchema::GetMaterialFunctionActions(FGraphActionMenuBuilder& A
 					bool bSkipMaterialFunction = false;
 					for (const FText& Category : LibraryCategoriesText)
 					{
-						if (Category.ToString().Contains("Strata"))
+						if (Category.ToString().Contains("Substrate"))
 						{
 							bSkipMaterialFunction = !Strata::IsStrataEnabled();
 							break;
diff --git a/Engine/Source/Editor/UnrealEd/Private/SEditorViewport.cpp b/Engine/Source/Editor/UnrealEd/Private/SEditorViewport.cpp
--- a/Engine/Source/Editor/UnrealEd/Private/SEditorViewport.cpp
+++ b/Engine/Source/Editor/UnrealEd/Private/SEditorViewport.cpp
@@ -429,7 +429,7 @@ void SEditorViewport::BindCommands()
 	MAP_VIEWMODE_ACTION( Commands.VisualizeBufferMode, VMI_VisualizeBuffer );
 	MAP_VIEWMODE_ACTION( Commands.VisualizeNaniteMode, VMI_VisualizeNanite );
 	MAP_VIEWMODE_ACTION( Commands.VisualizeLumenMode, VMI_VisualizeLumen );
-	MAP_VIEWMODE_ACTION( Commands.VisualizeStrataMode, VMI_VisualizeStrata);
+	MAP_VIEWMODE_ACTION( Commands.VisualizeSubstrateMode, VMI_VisualizeSubstrate);
 	MAP_VIEWMODE_ACTION( Commands.VisualizeGroomMode, VMI_VisualizeGroom);
 	MAP_VIEWMODE_ACTION( Commands.VisualizeVirtualShadowMapMode, VMI_VisualizeVirtualShadowMap );
 	MAP_VIEWMODE_ACTION( Commands.CollisionPawn, VMI_CollisionPawn);
diff --git a/Engine/Source/Editor/UnrealEd/Private/SEditorViewportViewMenu.cpp b/Engine/Source/Editor/UnrealEd/Private/SEditorViewportViewMenu.cpp
--- a/Engine/Source/Editor/UnrealEd/Private/SEditorViewportViewMenu.cpp
+++ b/Engine/Source/Editor/UnrealEd/Private/SEditorViewportViewMenu.cpp
@@ -56,7 +56,7 @@ FText SEditorViewportViewMenu::GetViewMenuLabel() const
 		{
 			Label = ViewportClient->GetCurrentLumenVisualizationModeDisplayName();
 		}
-		else if (ViewMode == VMI_VisualizeStrata)
+		else if (ViewMode == VMI_VisualizeSubstrate)
 		{
 			Label = ViewportClient->GetCurrentStrataVisualizationModeDisplayName();
 		}
diff --git a/Engine/Source/Editor/UnrealEd/Public/EditorViewportCommands.h b/Engine/Source/Editor/UnrealEd/Public/EditorViewportCommands.h
--- a/Engine/Source/Editor/UnrealEd/Public/EditorViewportCommands.h
+++ b/Engine/Source/Editor/UnrealEd/Public/EditorViewportCommands.h
@@ -136,8 +136,8 @@ public:
 	/** Changes the viewport to visualize Lumen debug views */
 	TSharedPtr< FUICommandInfo > VisualizeLumenMode;
 
-	/** Changes the viewport to visualize Strata debug views */
-	TSharedPtr< FUICommandInfo > VisualizeStrataMode;
+	/** Changes the viewport to visualize Substrate debug views */
+	TSharedPtr< FUICommandInfo > VisualizeSubstrateMode;
 
 	/** Changes the viewport to visualize Groom debug views */
 	TSharedPtr< FUICommandInfo > VisualizeGroomMode;
diff --git a/Engine/Source/Runtime/Engine/Classes/Engine/EngineBaseTypes.h b/Engine/Source/Runtime/Engine/Classes/Engine/EngineBaseTypes.h
--- a/Engine/Source/Runtime/Engine/Classes/Engine/EngineBaseTypes.h
+++ b/Engine/Source/Runtime/Engine/Classes/Engine/EngineBaseTypes.h
@@ -921,8 +921,8 @@ enum EViewModeIndex : int
 	/** Visualize Skin Cache. */
 	VMI_VisualizeGPUSkinCache = 33 UMETA(DisplayName = "GPU Skin Cache Visualization"),
 
-	/** Visualize Strata debug views */
-	VMI_VisualizeStrata = 34 UMETA(DisplayName = "Strata Visualization"),
+	/** Visualize Substrate debug views */
+	VMI_VisualizeSubstrate = 34 UMETA(DisplayName = "Substrate Visualization"),
 
 	/** Visualize Groom debug views */
 	VMI_VisualizeGroom = 35 UMETA(DisplayName = "Groom Visualization"),
diff --git a/Engine/Source/Runtime/Engine/Classes/Engine/EngineTypes.h b/Engine/Source/Runtime/Engine/Classes/Engine/EngineTypes.h
--- a/Engine/Source/Runtime/Engine/Classes/Engine/EngineTypes.h
+++ b/Engine/Source/Runtime/Engine/Classes/Engine/EngineTypes.h
@@ -254,11 +254,11 @@ enum EBlendMode : int
 	BLEND_Modulate UMETA(DisplayName="Modulate"),
 	BLEND_AlphaComposite UMETA(DisplayName = "AlphaComposite (Premultiplied Alpha)"),
 	BLEND_AlphaHoldout UMETA(DisplayName = "AlphaHoldout"),
-	BLEND_TranslucentColoredTransmittance UMETA(DisplayName = "STRATA_ONLY - Translucent - Colored Transmittance"), /*Strata only */
+	BLEND_TranslucentColoredTransmittance UMETA(DisplayName = "SUBSTRATE_ONLY - Translucent - Colored Transmittance"), /*Substrate only */
 	BLEND_MAX UMETA(Hidden),
 	// Renamed blend modes. These blend modes are remapped onto legacy ones and kept hidden for not confusing users in legacy mode, while allowing to use the new blend mode names into code.
-	BLEND_TranslucentGreyTransmittance = BLEND_Translucent UMETA(Hidden, DisplayName = "Translucent - Grey Transmittance"), /*Strata only */
-	BLEND_ColoredTransmittanceOnly = BLEND_Modulate UMETA(Hidden, DisplayName = "Colored Transmittance Only"), /*Strata only */
+	BLEND_TranslucentGreyTransmittance = BLEND_Translucent UMETA(Hidden, DisplayName = "Translucent - Grey Transmittance"), /*Substrate only */
+	BLEND_ColoredTransmittanceOnly = BLEND_Modulate UMETA(Hidden, DisplayName = "Colored Transmittance Only"), /*Substrate only */
 };
 
 class FMaterial;
@@ -564,13 +564,13 @@ namespace EGBufferFormat
 {
 	enum Type : int
 	{
-		/** Forces all GBuffers to 8 bits per channel. Intended as profiling for best performance. (Strata: Octahedral encoding as 2x11bits for simple and single materials, 2x16bits for complex materials) */
+		/** Forces all GBuffers to 8 bits per channel. Intended as profiling for best performance. (Substrate: Octahedral encoding as 2x11bits for simple and single materials, 2x16bits for complex materials) */
 		Force8BitsPerChannel = 0 UMETA(DisplayName = "Force 8 Bits Per Channel"),
-		/** See GBuffer allocation function for layout details. (Strata: Octahedral encoding as 2x11bits for simple and single material, 2x16bits for complex materials) */
+		/** See GBuffer allocation function for layout details. (Substrate: Octahedral encoding as 2x11bits for simple and single material, 2x16bits for complex materials) */
 		Default = 1,
-		/** Same as Default except normals are encoded at 16 bits per channel. (Strata: Octahedral encoding as 2x16bits for all materials.) */
+		/** Same as Default except normals are encoded at 16 bits per channel. (Substrate: Octahedral encoding as 2x16bits for all materials.) */
 		HighPrecisionNormals = 3,
-		/** Forces all GBuffers to 16 bits per channel. Intended as profiling for best quality. (Strata: Octahedral encoding as 2x16bits for all materials.) */
+		/** Forces all GBuffers to 16 bits per channel. Intended as profiling for best quality. (Substrate: Octahedral encoding as 2x16bits for all materials.) */
 		Force16BitsPerChannel = 5 UMETA(DisplayName = "Force 16 Bits Per Channel"),
 	};
 }
@@ -614,7 +614,7 @@ enum EMaterialShadingModel : int
 	MSM_Eye						UMETA(DisplayName="Eye"),
 	MSM_SingleLayerWater		UMETA(DisplayName="SingleLayerWater"),
 	MSM_ThinTranslucent			UMETA(DisplayName="Thin Translucent"),
-	MSM_Strata					UMETA(DisplayName="Strata", Hidden),
+	MSM_Strata					UMETA(DisplayName="Substrate", Hidden),
 	/** Number of unique shading models. */
 	MSM_NUM						UMETA(Hidden),
 	/** Shading model will be determined by the Material Expression Graph,
@@ -671,7 +671,7 @@ private:
 };
 
 /**
- * Specifies the Strata runtime shading model summarized from the material graph
+ * Specifies the Substrate runtime shading model summarized from the material graph
  */
 UENUM()
 enum EStrataShadingModel : int
@@ -692,7 +692,7 @@ enum EStrataShadingModel : int
 };
 static_assert(SSM_NUM <= 16, "Do not exceed 16 shading models without expanding FStrataMaterialShadingModelField to support uint32 instead of uint16!");
 
-/** Gather information from the Strata material graph to setup material for runtime. */
+/** Gather information from the Substrate material graph to setup material for runtime. */
 USTRUCT()
 struct ENGINE_API FStrataMaterialInfo
 {
diff --git a/Engine/Source/Runtime/Engine/Classes/Engine/RendererSettings.h b/Engine/Source/Runtime/Engine/Classes/Engine/RendererSettings.h
--- a/Engine/Source/Runtime/Engine/Classes/Engine/RendererSettings.h
+++ b/Engine/Source/Runtime/Engine/Classes/Engine/RendererSettings.h
@@ -788,7 +788,7 @@ class ENGINE_API URendererSettings : public UDeveloperSettings
 
 	UPROPERTY(config, EditAnywhere, Category = Optimizations, meta = (
 		ConsoleVariable = "r.GBufferFormat", DisplayName = "GBuffer Format",
-		ToolTip = "Selects which GBuffer format should be used. Affects performance primarily via how much GPU memory bandwidth used. This also controls Strata normal quality and, in this case, a restart is required.",
+		ToolTip = "Selects which GBuffer format should be used. Affects performance primarily via how much GPU memory bandwidth used. This also controls Substrate normal quality and, in this case, a restart is required.",
 		ConfigRestartRequired = true))
 	TEnumAsByte<EGBufferFormat::Type> GBufferFormat;
 
@@ -922,29 +922,29 @@ class ENGINE_API URendererSettings : public UDeveloperSettings
 		uint32 bSupportCloudShadowOnSingleLayerWater : 1;
 
 	/**
-	"Enable Strata materials (Beta)."
+	"Enable Substrate materials (Beta)."
 	*/
-	UPROPERTY(config, EditAnywhere, Category = Strata, meta = (
-		ConsoleVariable = "r.Strata", DisplayName = "Strata materials (Experimental)",
-		ToolTip = "Enable Strata materials (Experimental).",
+	UPROPERTY(config, EditAnywhere, Category = Substrate, meta = (
+		ConsoleVariable = "r.Substrate", DisplayName = "Substrate materials (Experimental)",
+		ToolTip = "Enable Substrate materials (Experimental).",
 		ConfigRestartRequired = true))
 		uint32 bEnableStrata : 1;
 
 	/**
-	"Enable Strata opaque material rough refractions effect from top layers over layers below."
+	"Enable Substrate opaque material rough refractions effect from top layers over layers below."
 	*/
-	UPROPERTY(config, EditAnywhere, Category = Strata, meta = (
-		ConsoleVariable = "r.Strata.OpaqueMaterialRoughRefraction", DisplayName = "Strata opaque material rough refraction",
-		ToolTip = "Enable Strata opaque material rough refractions effect from top layers over layers below.",
+	UPROPERTY(config, EditAnywhere, Category = Substrate, meta = (
+		ConsoleVariable = "r.Substrate.OpaqueMaterialRoughRefraction", DisplayName = "Substrate opaque material rough refraction",
+		ToolTip = "Enable Substrate opaque material rough refractions effect from top layers over layers below.",
 		ConfigRestartRequired = true))
 		uint32 StrataOpaqueMaterialRoughRefraction : 1;
 
 	/**
-	"Enable advanced strata material debug visualization shaders. Base pas shaders can output such advanced data."
+	"Enable advanced Substrate material debug visualization shaders. Base pas shaders can output such advanced data."
 	*/
-	UPROPERTY(config, EditAnywhere, Category = Strata, meta = (
-		ConsoleVariable = "r.Strata.Debug.AdvancedVisualizationShaders", DisplayName = "Strata advanced visualization shaders",
-		ToolTip = "Enable advanced strata material debug visualization shaders. Base pass shaders can output such advanced data.",
+	UPROPERTY(config, EditAnywhere, Category = Substrate, meta = (
+		ConsoleVariable = "r.Substrate.Debug.AdvancedVisualizationShaders", DisplayName = "Substrate advanced visualization shaders",
+		ToolTip = "Enable advanced Substrate material debug visualization shaders. Base pass shaders can output such advanced data.",
 		ConfigRestartRequired = true))
 		uint32 StrataDebugAdvancedVisualizationShaders : 1;
 
@@ -962,7 +962,7 @@ class ENGINE_API URendererSettings : public UDeveloperSettings
 	*/
 	UPROPERTY(config, EditAnywhere, Category = Materials, meta = (
 		ConsoleVariable = "r.Material.EnergyConservation", DisplayName = "Enable Energy Conservation on Material",
-		ToolTip = "Enable Energy Conservation on Material. Please note that when Strata is enabled, energy conservation is forced to enabled.",
+		ToolTip = "Enable Energy Conservation on Material. Please note that when Substrate is enabled, energy conservation is forced to enabled.",
 		ConfigRestartRequired = true))
 		uint32 bMaterialEnergyConservation : 1;
 
diff --git a/Engine/Source/Runtime/Engine/Classes/Materials/Material.h b/Engine/Source/Runtime/Engine/Classes/Materials/Material.h
--- a/Engine/Source/Runtime/Engine/Classes/Materials/Material.h
+++ b/Engine/Source/Runtime/Engine/Classes/Materials/Material.h
@@ -523,7 +523,7 @@ public:
 	UPROPERTY(EditAnywhere, Category=Material)
 	uint8 TwoSided : 1;
 
-	/** Indicates that the material should be rendered as a thin surface (i.e., without inner volume). Only used by Strata materials. Enabling ThinSurface will disable subsurface profiles. */
+	/** Indicates that the material should be rendered as a thin surface (i.e., without inner volume). Only used by Substrate materials. Enabling ThinSurface will disable subsurface profiles. */
 	UPROPERTY(EditAnywhere, Category = Material)
 	uint8 bIsThinSurface : 1;
 
diff --git a/Engine/Source/Runtime/Engine/Classes/Materials/MaterialExpressionFunctionInput.h b/Engine/Source/Runtime/Engine/Classes/Materials/MaterialExpressionFunctionInput.h
--- a/Engine/Source/Runtime/Engine/Classes/Materials/MaterialExpressionFunctionInput.h
+++ b/Engine/Source/Runtime/Engine/Classes/Materials/MaterialExpressionFunctionInput.h
@@ -29,7 +29,7 @@ enum EFunctionInputType : int
 	FunctionInput_MaterialAttributes,
 	FunctionInput_TextureExternal,
 	FunctionInput_Bool,
-	FunctionInput_Strata,
+	FunctionInput_Substrate,
 	FunctionInput_MAX,
 };
 
diff --git a/Engine/Source/Runtime/Engine/Classes/Materials/MaterialExpressionStrata.h b/Engine/Source/Runtime/Engine/Classes/Materials/MaterialExpressionStrata.h
--- a/Engine/Source/Runtime/Engine/Classes/Materials/MaterialExpressionStrata.h
+++ b/Engine/Source/Runtime/Engine/Classes/Materials/MaterialExpressionStrata.h
@@ -10,7 +10,7 @@
 class FMaterialCompiler;
 
 /**
- * Compile a special blend function for strata when blending material attribute
+ * Compile a special blend function for Substrate when blending material attribute
  *
  * @param Compiler				The compiler to add code to
  * @param Foreground			Entry A, has a bigger impact when Alpha is close to 0
@@ -24,8 +24,8 @@ extern int32 CompileStrataBlendFunction(FMaterialCompiler* Compiler, const int32
 ///////////////////////////////////////////////////////////////////////////////
 // BSDF nodes
 
-// UMaterialExpressionStrataBSDF can only be used for Strata nodes ouputing StrataData that would need a preview,
-UCLASS(MinimalAPI, collapsecategories, hidecategories = Object, Abstract, DisplayName = "Strata Expression")
+// UMaterialExpressionStrataBSDF can only be used for Substrate nodes ouputing StrataData that would need a preview,
+UCLASS(MinimalAPI, collapsecategories, hidecategories = Object, Abstract, DisplayName = "Substrate Expression")
 class UMaterialExpressionStrataBSDF : public UMaterialExpression
 {
 	GENERATED_UCLASS_BODY()
@@ -41,7 +41,7 @@ class UMaterialExpressionStrataBSDF : public UMaterialExpression
 };
 
 
-UCLASS(MinimalAPI, collapsecategories, hidecategories = Object, DisplayName = "Strata Legacy Conversion")
+UCLASS(MinimalAPI, collapsecategories, hidecategories = Object, DisplayName = "Substrate Legacy Conversion")
 class UMaterialExpressionStrataLegacyConversion : public UMaterialExpressionStrataBSDF
 {
 	GENERATED_UCLASS_BODY()
@@ -166,7 +166,7 @@ class UMaterialExpressionStrataLegacyConversion : public UMaterialExpressionStra
 	UPROPERTY()
 	FShadingModelMaterialInput ShadingModel;
 	
-	/** SubsurfaceProfile, for Screen Space Subsurface Scattering. The profile needs to be set up on both the Strata diffuse node, and the material node at the moment. */
+	/** SubsurfaceProfile, for Screen Space Subsurface Scattering. The profile needs to be set up on both the Substrate diffuse node, and the material node at the moment. */
 	UPROPERTY(EditAnywhere, BlueprintReadOnly, Category = Material, meta = (DisplayName = "Subsurface Profile"))
 	TObjectPtr<class USubsurfaceProfile> SubsurfaceProfile;
 
@@ -195,7 +195,7 @@ class UMaterialExpressionStrataLegacyConversion : public UMaterialExpressionStra
 	//~ End UMaterialExpression Interface
 };
 
-UCLASS(MinimalAPI, collapsecategories, hidecategories = Object, DisplayName = "Strata Slab")
+UCLASS(MinimalAPI, collapsecategories, hidecategories = Object, DisplayName = "Substrate Slab")
 class UMaterialExpressionStrataSlabBSDF : public UMaterialExpressionStrataBSDF
 {
 	GENERATED_UCLASS_BODY()
@@ -296,7 +296,7 @@ class UMaterialExpressionStrataSlabBSDF : public UMaterialExpressionStrataBSDF
 	UPROPERTY()
 	FExpressionInput FuzzColor;
 
-	/** SubsurfaceProfile, for Screen Space Subsurface Scattering. The profile needs to be set up on both the Strata diffuse node, and the material node at the moment. */
+	/** SubsurfaceProfile, for Screen Space Subsurface Scattering. The profile needs to be set up on both the Substrate diffuse node, and the material node at the moment. */
 	UPROPERTY(EditAnywhere, BlueprintReadOnly, Category = Material, meta = (DisplayName = "Subsurface Profile"))
 	TObjectPtr<class USubsurfaceProfile> SubsurfaceProfile;
 
@@ -329,7 +329,7 @@ class UMaterialExpressionStrataSlabBSDF : public UMaterialExpressionStrataBSDF
 	//~ End UMaterialExpression Interface
 };
 
-UCLASS(MinimalAPI, collapsecategories, hidecategories = Object, DisplayName = "Strata Simple Clear Coat")
+UCLASS(MinimalAPI, collapsecategories, hidecategories = Object, DisplayName = "Substrate Simple Clear Coat")
 class UMaterialExpressionStrataSimpleClearCoatBSDF : public UMaterialExpressionStrataBSDF
 {
 	GENERATED_UCLASS_BODY()
@@ -391,7 +391,7 @@ class UMaterialExpressionStrataSimpleClearCoatBSDF : public UMaterialExpressionS
 	//~ End UMaterialExpression Interface
 };
 
-UCLASS(MinimalAPI, collapsecategories, hidecategories = Object, DisplayName = "Strata Volumetric-Fog-Cloud BSDF")
+UCLASS(MinimalAPI, collapsecategories, hidecategories = Object, DisplayName = "Substrate Volumetric-Fog-Cloud BSDF")
 class UMaterialExpressionStrataVolumetricFogCloudBSDF : public UMaterialExpressionStrataBSDF
 {
 	GENERATED_UCLASS_BODY()
@@ -433,7 +433,7 @@ class UMaterialExpressionStrataVolumetricFogCloudBSDF : public UMaterialExpressi
 	//~ End UMaterialExpression Interface
 };
 
-UCLASS(MinimalAPI, collapsecategories, hidecategories = Object, DisplayName = "Strata Unlit BSDF")
+UCLASS(MinimalAPI, collapsecategories, hidecategories = Object, DisplayName = "Substrate Unlit BSDF")
 class UMaterialExpressionStrataUnlitBSDF : public UMaterialExpressionStrataBSDF
 {
 	GENERATED_UCLASS_BODY()
@@ -463,7 +463,7 @@ class UMaterialExpressionStrataUnlitBSDF : public UMaterialExpressionStrataBSDF
 	//~ End UMaterialExpression Interface
 };
 
-UCLASS(MinimalAPI, collapsecategories, hidecategories = Object, DisplayName = "Strata Hair BSDF")
+UCLASS(MinimalAPI, collapsecategories, hidecategories = Object, DisplayName = "Substrate Hair BSDF")
 class UMaterialExpressionStrataHairBSDF : public UMaterialExpressionStrataBSDF
 {
 	GENERATED_UCLASS_BODY()
@@ -523,7 +523,7 @@ class UMaterialExpressionStrataHairBSDF : public UMaterialExpressionStrataBSDF
 	//~ End UMaterialExpression Interface
 };
 
-UCLASS(MinimalAPI, collapsecategories, hidecategories = Object, DisplayName = "Strata Eye BSDF")
+UCLASS(MinimalAPI, collapsecategories, hidecategories = Object, DisplayName = "Substrate Eye BSDF")
 class UMaterialExpressionStrataEyeBSDF : public UMaterialExpressionStrataBSDF
 {
 	GENERATED_UCLASS_BODY()
@@ -593,7 +593,7 @@ class UMaterialExpressionStrataEyeBSDF : public UMaterialExpressionStrataBSDF
 	//~ End UMaterialExpression Interface
 };
 
-UCLASS(MinimalAPI, collapsecategories, hidecategories = Object, DisplayName = "Strata Single Layer Water BSDF")
+UCLASS(MinimalAPI, collapsecategories, hidecategories = Object, DisplayName = "Substrate Single Layer Water BSDF")
 class UMaterialExpressionStrataSingleLayerWaterBSDF : public UMaterialExpressionStrataBSDF
 {
 	GENERATED_UCLASS_BODY()
@@ -677,7 +677,7 @@ class UMaterialExpressionStrataSingleLayerWaterBSDF : public UMaterialExpression
 	//~ End UMaterialExpression Interface
 };
 
-UCLASS(MinimalAPI, collapsecategories, hidecategories = Object, DisplayName = "Strata Light Function")
+UCLASS(MinimalAPI, collapsecategories, hidecategories = Object, DisplayName = "Substrate Light Function")
 class UMaterialExpressionStrataLightFunction : public UMaterialExpressionStrataBSDF
 {
 	GENERATED_UCLASS_BODY()
@@ -701,13 +701,13 @@ class UMaterialExpressionStrataLightFunction : public UMaterialExpressionStrataB
 	//~ End UMaterialExpression Interface
 };
 
-UCLASS(MinimalAPI, collapsecategories, hidecategories = Object, DisplayName = "Strata Post Process")
+UCLASS(MinimalAPI, collapsecategories, hidecategories = Object, DisplayName = "Substrate Post Process")
 class UMaterialExpressionStrataPostProcess : public UMaterialExpressionStrataBSDF
 {
 	GENERATED_UCLASS_BODY()
 
 	/**
-	 * The output color of the post process: it represents a color added over the back buffer, or a color multiplied if the Strata blend mode is TransmittanceOnly.
+	 * The output color of the post process: it represents a color added over the back buffer, or a color multiplied if the Substrate blend mode is TransmittanceOnly.
 	 */
 	UPROPERTY()
 	FExpressionInput Color;
@@ -731,7 +731,7 @@ class UMaterialExpressionStrataPostProcess : public UMaterialExpressionStrataBSD
 	//~ End UMaterialExpression Interface
 };
 
-UCLASS(MinimalAPI, collapsecategories, hidecategories = Object, DisplayName = "Strata UI")
+UCLASS(MinimalAPI, collapsecategories, hidecategories = Object, DisplayName = "Substrate UI")
 class UMaterialExpressionStrataUI : public UMaterialExpressionStrataBSDF
 {
 	GENERATED_UCLASS_BODY()
@@ -761,13 +761,13 @@ class UMaterialExpressionStrataUI : public UMaterialExpressionStrataBSDF
 	//~ End UMaterialExpression Interface
 };
 
-UCLASS(MinimalAPI, collapsecategories, hidecategories = Object, DisplayName = "Strata Convert To Decal")
+UCLASS(MinimalAPI, collapsecategories, hidecategories = Object, DisplayName = "Substrate Convert To Decal")
 class UMaterialExpressionStrataConvertToDecal : public UMaterialExpressionStrataBSDF
 {
 	GENERATED_UCLASS_BODY()
 
 	/**
-	 * The Strata material to convert to a decal.
+	 * The Substrate material to convert to a decal.
 	 */
 	UPROPERTY()
 	FExpressionInput DecalMaterial;
@@ -796,19 +796,19 @@ class UMaterialExpressionStrataConvertToDecal : public UMaterialExpressionStrata
 ///////////////////////////////////////////////////////////////////////////////
 // Operator nodes
 
-UCLASS(MinimalAPI, collapsecategories, hidecategories = Object, DisplayName = "Strata Horizontal Blend")
+UCLASS(MinimalAPI, collapsecategories, hidecategories = Object, DisplayName = "Substrate Horizontal Blend")
 class UMaterialExpressionStrataHorizontalMixing : public UMaterialExpressionStrataBSDF
 {
 	GENERATED_UCLASS_BODY()
 		
 	/**
-	 * Strata material
+	 * Substrate material
 	 */
 	UPROPERTY()
 	FExpressionInput Background;
 
 	/**
-	 * Strata material
+	 * Substrate material
 	 */
 	UPROPERTY()
 	FExpressionInput Foreground;
@@ -838,19 +838,19 @@ class UMaterialExpressionStrataHorizontalMixing : public UMaterialExpressionStra
 	//~ End UMaterialExpression Interface
 };
 
-UCLASS(MinimalAPI, collapsecategories, hidecategories = Object, DisplayName = "Strata Vertical Layer")
+UCLASS(MinimalAPI, collapsecategories, hidecategories = Object, DisplayName = "Substrate Vertical Layer")
 class UMaterialExpressionStrataVerticalLayering : public UMaterialExpressionStrataBSDF
 {
 	GENERATED_UCLASS_BODY()
 
 	/**
-	 * Strata material layer on top of the Base material layer
+	 * Substrate material layer on top of the Base material layer
 	 */
 	UPROPERTY()
 	FExpressionInput Top;
 	
 	/**
-	 * Strata material layer below the Top material layer
+	 * Substrate material layer below the Top material layer
 	 */
 	UPROPERTY()
 	FExpressionInput Base;
@@ -881,19 +881,19 @@ class UMaterialExpressionStrataVerticalLayering : public UMaterialExpressionStra
 	//~ End UMaterialExpression Interface
 };
 
-UCLASS(MinimalAPI, collapsecategories, hidecategories = Object, DisplayName = "Strata Add")
+UCLASS(MinimalAPI, collapsecategories, hidecategories = Object, DisplayName = "Substrate Add")
 class UMaterialExpressionStrataAdd : public UMaterialExpressionStrataBSDF
 {
 	GENERATED_UCLASS_BODY()
 
 	/**
-	 * Strata material
+	 * Substrate material
 	 */
 	UPROPERTY()
 	FExpressionInput A;
 	
 	/**
-	 * Strata material
+	 * Substrate material
 	 */
 	UPROPERTY()
 	FExpressionInput B;
@@ -917,19 +917,19 @@ class UMaterialExpressionStrataAdd : public UMaterialExpressionStrataBSDF
 	//~ End UMaterialExpression Interface
 };
 
-UCLASS(MinimalAPI, collapsecategories, hidecategories = Object, DisplayName = "Strata Coverage Weight")
+UCLASS(MinimalAPI, collapsecategories, hidecategories = Object, DisplayName = "Substrate Coverage Weight")
 class UMaterialExpressionStrataWeight : public UMaterialExpressionStrataBSDF
 {
 	GENERATED_UCLASS_BODY()
 
 	/**
-	 * Strata material
+	 * Substrate material
 	 */
 	UPROPERTY()
 	FExpressionInput A;
 	
 	/**
-	 * Weight to apply to the strata material BSDFs
+	 * Weight to apply to the Substrate material BSDFs
 	 */
 	UPROPERTY()
 	FExpressionInput Weight;
@@ -951,13 +951,13 @@ class UMaterialExpressionStrataWeight : public UMaterialExpressionStrataBSDF
 ///////////////////////////////////////////////////////////////////////////////
 // Utilities
 
-UCLASS(MinimalAPI, collapsecategories, hidecategories = Object, Abstract, DisplayName = "Strata Utility Base Class")
+UCLASS(MinimalAPI, collapsecategories, hidecategories = Object, Abstract, DisplayName = "Substrate Utility Base Class")
 class UMaterialExpressionStrataUtilityBase : public UMaterialExpression
 {
 	GENERATED_UCLASS_BODY()
 };
 
-UCLASS(MinimalAPI, collapsecategories, hidecategories = Object, DisplayName = "Strata Transmittance-To-MeanFreePath")
+UCLASS(MinimalAPI, collapsecategories, hidecategories = Object, DisplayName = "Substrate Transmittance-To-MeanFreePath")
 class UMaterialExpressionStrataTransmittanceToMFP : public UMaterialExpressionStrataUtilityBase
 {
 	GENERATED_UCLASS_BODY()
@@ -987,7 +987,7 @@ class UMaterialExpressionStrataTransmittanceToMFP : public UMaterialExpressionSt
 	//~ End UMaterialExpression Interface
 };
 
-UCLASS(MinimalAPI, collapsecategories, hidecategories = Object, DisplayName = "Strata Metalness-To-DiffuseColorF0")
+UCLASS(MinimalAPI, collapsecategories, hidecategories = Object, DisplayName = "Substrate Metalness-To-DiffuseColorF0")
 class UMaterialExpressionStrataMetalnessToDiffuseAlbedoF0 : public UMaterialExpressionStrataUtilityBase
 {
 	GENERATED_UCLASS_BODY()
@@ -1022,7 +1022,7 @@ class UMaterialExpressionStrataMetalnessToDiffuseAlbedoF0 : public UMaterialExpr
 	//~ End UMaterialExpression Interface
 };
 
-UCLASS(MinimalAPI, collapsecategories, hidecategories = Object, DisplayName = "Strata Haziness-To-Secondary-Roughness")
+UCLASS(MinimalAPI, collapsecategories, hidecategories = Object, DisplayName = "Substrate Haziness-To-Secondary-Roughness")
 class UMaterialExpressionStrataHazinessToSecondaryRoughness : public UMaterialExpressionStrataUtilityBase
 {
 	GENERATED_UCLASS_BODY()
@@ -1051,7 +1051,7 @@ class UMaterialExpressionStrataHazinessToSecondaryRoughness : public UMaterialEx
 	//~ End UMaterialExpression Interface
 };
 
-UCLASS(MinimalAPI, collapsecategories, hidecategories = Object, DisplayName = "Strata Thin-Film")
+UCLASS(MinimalAPI, collapsecategories, hidecategories = Object, DisplayName = "Substrate Thin-Film")
 class UMaterialExpressionStrataThinFilm : public UMaterialExpressionStrataUtilityBase
 {
 	GENERATED_UCLASS_BODY()
diff --git a/Engine/Source/Runtime/Engine/Private/Materials/HLSLMaterialTranslator.cpp b/Engine/Source/Runtime/Engine/Private/Materials/HLSLMaterialTranslator.cpp
--- a/Engine/Source/Runtime/Engine/Private/Materials/HLSLMaterialTranslator.cpp
+++ b/Engine/Source/Runtime/Engine/Private/Materials/HLSLMaterialTranslator.cpp
@@ -808,7 +808,7 @@ bool FHLSLMaterialTranslator::Translate()
 
 			if (!StrataGenerateDerivedMaterialOperatorData())
 			{
-				Errorf(TEXT("Strata material errors encountered."));
+				Errorf(TEXT("Substrate material errors encountered."));
 			}
 		}
 
@@ -2047,8 +2047,8 @@ void FHLSLMaterialTranslator::GetMaterialEnvironment(EShaderPlatform InPlatform,
 
 			FString StrataMaterialDescription;
 
-			StrataMaterialDescription += FString::Printf(TEXT("----- STRATA -----\r\n"));
-			StrataMaterialDescription += FString::Printf(TEXT("StrataCompilationInfo -\r\n"));
+			StrataMaterialDescription += FString::Printf(TEXT("----- SUBSTRATE -----\r\n"));
+			StrataMaterialDescription += FString::Printf(TEXT("Substrate Compilation Info -\r\n"));
 			StrataMaterialDescription += FString::Printf(TEXT(" - Byte Per Pixel Budget                      %u\r\n"), StrataBytePerPixel);
 			StrataMaterialDescription += FString::Printf(TEXT(" - Requested Byte Size before simplification  %u (%d UINT32)\r\n"), StrataSimplificationStatus.OriginalRequestedByteSize, StrataSimplificationStatus.OriginalRequestedByteSize / 4);
 			StrataMaterialDescription += FString::Printf(TEXT(" - Requested Byte Size after simplification   %u (%d UINT32)\r\n"), StrataMaterialRequestedSizeByte, StrataMaterialRequestedSizeByte / 4);
@@ -2072,7 +2072,7 @@ void FHLSLMaterialTranslator::GetMaterialEnvironment(EShaderPlatform InPlatform,
 				StrataMaterialDescription += FString::Printf(TEXT("     - %s - SharedLocalBasisIndexMacro = %s \r\n"), *GetStrataBSDFName(BSDFOperator.BSDFType), *GetStrataSharedLocalBasisIndexMacro(BSDFOperator.BSDFRegisteredSharedLocalBasis));
 			}
 
-			StrataMaterialDescription += FString::Printf(TEXT("----------- STRATA TREE -----------\r\n"));
+			StrataMaterialDescription += FString::Printf(TEXT("----------- SUBSTRATE TREE -----------\r\n"));
 			StrataMaterialDescription += FString::Printf(TEXT("Graph maximum distance to leaves %u\r\n"), RootMaximumDistanceToLeaves);
 			// Debug print operators according to depth from root.
 			{
@@ -10266,7 +10266,7 @@ FStrataOperator* FHLSLMaterialTranslator::StrataCompilationGetOperatorFromIndex(
 	int32 OperatorCount = StrataMaterialExpressionRegisteredOperators.Num();
 	if (OperatorIndex < 0 || OperatorIndex >= OperatorCount)
 	{
-		Errorf(TEXT("StrataCompilationGetOperatorFromIndex - OperatorIndex out of range %s (asset: %s).\r\n"), *Material->GetDebugName(), *Material->GetAssetPath().ToString());
+		Errorf(TEXT("SubstrateCompilationGetOperatorFromIndex - OperatorIndex out of range %s (asset: %s).\r\n"), *Material->GetDebugName(), *Material->GetAssetPath().ToString());
 		return nullptr;
 	};
 	return &StrataMaterialExpressionRegisteredOperators[OperatorIndex];
@@ -10902,7 +10902,7 @@ bool FHLSLMaterialTranslator::StrataGenerateDerivedMaterialOperatorData()
 					}
 					case STRATA_BSDF_TYPE_SINGLELAYERWATER:
 					{
-						Errorf(TEXT("Strata error: single layer water should go through the its dedicated fast path in %s (asset: %s).\r\n"), *Material->GetDebugName(), *Material->GetAssetPath().ToString());
+						Errorf(TEXT("Substrate error: single layer water should go through the its dedicated fast path in %s (asset: %s).\r\n"), *Material->GetDebugName(), *Material->GetAssetPath().ToString());
 						break;
 					}
 					case STRATA_BSDF_TYPE_UNLIT:
@@ -11630,12 +11630,12 @@ int32 FHLSLMaterialTranslator::StrataThinFilm(int32 NormalCodeChunk, int32 Specu
 	if (NormalCodeChunk == INDEX_NONE || SpecularColorCodeChunk == INDEX_NONE || EdgeSpecularColorCodeChunk == INDEX_NONE
 		|| ThicknessCodeChunk == INDEX_NONE || IORCodeChunk == INDEX_NONE)
 	{
-		Errorf(TEXT("Strata error: one of StrataThinFilm node input is invalid (asset: %s).\r\n"), *Material->GetDebugName(), *Material->GetAssetPath().ToString());
+		Errorf(TEXT("Substrate error: one of SubstrateThinFilm node input is invalid (asset: %s).\r\n"), *Material->GetDebugName(), *Material->GetAssetPath().ToString());
 		return INDEX_NONE;
 	}
 	if (OutputIndex < 0 || OutputIndex > 1)
 	{
-		Errorf(TEXT("Strata error: StrataThinFilm output index is invalid (asset: %s).\r\n"), *Material->GetDebugName(), *Material->GetAssetPath().ToString());
+		Errorf(TEXT("Substrate error: SubstrateThinFilm output index is invalid (asset: %s).\r\n"), *Material->GetDebugName(), *Material->GetAssetPath().ToString());
 		return INDEX_NONE;
 	}
 	return AddCodeChunk(
diff --git a/Engine/Source/Runtime/Engine/Private/Materials/MaterialExpressions.cpp b/Engine/Source/Runtime/Engine/Private/Materials/MaterialExpressions.cpp
--- a/Engine/Source/Runtime/Engine/Private/Materials/MaterialExpressions.cpp
+++ b/Engine/Source/Runtime/Engine/Private/Materials/MaterialExpressions.cpp
@@ -436,7 +436,7 @@ void GetMaterialValueTypeDescriptions(const uint32 MaterialValueType, TArray<FTe
 	if (MaterialValueType & MCT_ShadingModel)
 		OutDescriptions.Add(LOCTEXT("ShadingModel", "Shading Model"));
 	if (MaterialValueType & MCT_Strata)
-		OutDescriptions.Add(LOCTEXT("Strata", "Strata Material"));
+		OutDescriptions.Add(LOCTEXT("Substrate", "Substrate Material"));
 	if (MaterialValueType & MCT_Unknown)
 		OutDescriptions.Add(LOCTEXT("Unknown", "Unknown"));
 }
@@ -16784,7 +16784,7 @@ int32 UMaterialExpressionFunctionInput::CompilePreviewValue(FMaterialCompiler* C
 		case FunctionInput_TextureExternal:
 		case FunctionInput_StaticBool:
 		case FunctionInput_Bool:
-		case FunctionInput_Strata:
+		case FunctionInput_Substrate:
 			return Compiler->Errorf(TEXT("Missing Preview connection for function input '%s'"), *InputName.ToString());
 		default:
 			return Compiler->Errorf(TEXT("Unknown input type"));
@@ -16941,7 +16941,7 @@ bool UMaterialExpressionFunctionInput::IsResultMaterialAttributes(int32 OutputIn
 
 bool UMaterialExpressionFunctionInput::IsResultStrataMaterial(int32 OutputIndex)
 {
-	if (FunctionInput_Strata == InputType)
+	if (FunctionInput_Substrate == InputType)
 	{
 		return true;
 	}
@@ -16953,7 +16953,7 @@ bool UMaterialExpressionFunctionInput::IsResultStrataMaterial(int32 OutputIndex)
 
 void UMaterialExpressionFunctionInput::GatherStrataMaterialInfo(FStrataMaterialInfo& StrataMaterialInfo, int32 OutputIndex)
 {
-	if (FunctionInput_Strata == InputType)
+	if (FunctionInput_Substrate == InputType)
 	{
 		FExpressionInput EffectivePreviewDuringCompileTracedInput = EffectivePreviewDuringCompile.GetTracedInput();
 		int32 ExpressionResult = INDEX_NONE;
@@ -16966,7 +16966,7 @@ void UMaterialExpressionFunctionInput::GatherStrataMaterialInfo(FStrataMaterialI
 
 FStrataOperator* UMaterialExpressionFunctionInput::StrataGenerateMaterialTopologyTree(class FMaterialCompiler* Compiler, class UMaterialExpression* Parent, int32 OutputIndex)
 {
-	if (FunctionInput_Strata == InputType)
+	if (FunctionInput_Substrate == InputType)
 	{
 		FExpressionInput EffectivePreviewDuringCompileTracedInput = EffectivePreviewDuringCompile.GetTracedInput();
 		int32 ExpressionResult = INDEX_NONE;
@@ -17012,7 +17012,7 @@ uint32 UMaterialExpressionFunctionInput::GetInputType(int32 InputIndex)
 		return MCT_Bool;
 	case FunctionInput_MaterialAttributes:
 		return MCT_MaterialAttributes;
-	case FunctionInput_Strata:
+	case FunctionInput_Substrate:
 		return MCT_Strata;
 	default:
 		return MCT_Unknown;
@@ -17623,7 +17623,7 @@ void UMaterialExpressionShaderStageSwitch::GatherStrataMaterialInfo(FStrataMater
 
 FStrataOperator* UMaterialExpressionShaderStageSwitch::StrataGenerateMaterialTopologyTree(class FMaterialCompiler* Compiler, class UMaterialExpression* Parent, int32 OutputIndex)
 {
-	Compiler->Errorf(TEXT("Strata materials are only supported in pixel shaders: ShaderStageSwitch thus should not be plugged to convey Strata material informations."));
+	Compiler->Errorf(TEXT("Substrate materials are only supported in pixel shaders: ShaderStageSwitch thus should not be plugged to convey Strata material informations."));
 	return nullptr;
 }
 
@@ -17829,7 +17829,7 @@ void UMaterialExpressionRayTracingQualitySwitch::GatherStrataMaterialInfo(FStrat
 
 FStrataOperator* UMaterialExpressionRayTracingQualitySwitch::StrataGenerateMaterialTopologyTree(class FMaterialCompiler* Compiler, class UMaterialExpression* Parent, int32 OutputIndex)
 {
-	Compiler->Errorf(TEXT("Strata material topology must be statically define. We do not support topology update via dynamic evaluation such as `is raytracing or not`. Only input to BSDFs or Operators can be controled this way."));
+	Compiler->Errorf(TEXT("Substrate material topology must be statically define. We do not support topology update via dynamic evaluation such as `is raytracing or not`. Only input to BSDFs or Operators can be controled this way."));
 	return nullptr;
 }
 #endif // WITH_EDITOR
@@ -17921,7 +17921,7 @@ void UMaterialExpressionPathTracingQualitySwitch::GatherStrataMaterialInfo(FStra
 
 FStrataOperator* UMaterialExpressionPathTracingQualitySwitch::StrataGenerateMaterialTopologyTree(class FMaterialCompiler* Compiler, class UMaterialExpression* Parent, int32 OutputIndex)
 {
-	Compiler->Errorf(TEXT("Strata material topology must be statically define. We do not support topology update via dynamic evaluation such as `is pathtracing or not`. Only input to BSDFs or Operators can be controled this way."));
+	Compiler->Errorf(TEXT("Substrate material topology must be statically define. We do not support topology update via dynamic evaluation such as `is pathtracing or not`. Only input to BSDFs or Operators can be controled this way."));
 	return nullptr;
 }
 #endif // WITH_EDITOR
@@ -22274,7 +22274,7 @@ void UMaterialExpressionReflectionCapturePassSwitch::GatherStrataMaterialInfo(FS
 
 FStrataOperator* UMaterialExpressionReflectionCapturePassSwitch::StrataGenerateMaterialTopologyTree(class FMaterialCompiler* Compiler, class UMaterialExpression* Parent, int32 OutputIndex)
 {
-	Compiler->Errorf(TEXT("Strata material topology must be statically define. We do not support topology update via dynamic evaluation such as `is reflection or not`. Only input to BSDFs or Operators can be controled this way."));
+	Compiler->Errorf(TEXT("Substrate material topology must be statically define. We do not support topology update via dynamic evaluation such as `is reflection or not`. Only input to BSDFs or Operators can be controled this way."));
 	return nullptr;
 }
 #endif // WITH_EDITOR
@@ -22695,7 +22695,7 @@ UMaterialExpressionStrataLegacyConversion::UMaterialExpressionStrataLegacyConver
 	struct FConstructorStatics
 	{
 		FText NAME_Strata;
-		FConstructorStatics() : NAME_Strata(LOCTEXT("Strata Conversion", "Strata Conversion")) { }
+		FConstructorStatics() : NAME_Strata(LOCTEXT("Substrate Conversion", "Substrate Conversion")) { }
 	};
 	static FConstructorStatics ConstructorStatics;
 #if WITH_EDITORONLY_DATA
@@ -22859,7 +22859,7 @@ void UMaterialExpressionStrataLegacyConversion::PostEditChangeProperty(FProperty
 
 void UMaterialExpressionStrataLegacyConversion::GetCaption(TArray<FString>& OutCaptions) const
 {
-	OutCaptions.Add(TEXT("Strata Legacy Conversion"));
+	OutCaptions.Add(TEXT("Substrate Legacy Conversion"));
 }
 
 uint32 UMaterialExpressionStrataLegacyConversion::GetOutputType(int32 OutputIndex)
@@ -23068,7 +23068,7 @@ UMaterialExpressionStrataSlabBSDF::UMaterialExpressionStrataSlabBSDF(const FObje
 	struct FConstructorStatics
 	{
 		FText NAME_Strata;
-		FConstructorStatics() : NAME_Strata(LOCTEXT("Strata BSDFs", "Strata BSDFs")) { }
+		FConstructorStatics() : NAME_Strata(LOCTEXT("Substrate BSDFs", "Substrate BSDFs")) { }
 	};
 	static FConstructorStatics ConstructorStatics;
 #if WITH_EDITORONLY_DATA
@@ -23173,7 +23173,7 @@ const TArray<FExpressionInput*> UMaterialExpressionStrataSlabBSDF::GetInputs()
 
 void UMaterialExpressionStrataSlabBSDF::GetCaption(TArray<FString>& OutCaptions) const
 {
-	OutCaptions.Add(TEXT("Strata Slab BSDF"));
+	OutCaptions.Add(TEXT("Substrate Slab BSDF"));
 }
 
 uint32 UMaterialExpressionStrataSlabBSDF::GetOutputType(int32 OutputIndex)
@@ -23428,7 +23428,7 @@ UMaterialExpressionStrataSimpleClearCoatBSDF::UMaterialExpressionStrataSimpleCle
 	struct FConstructorStatics
 	{
 		FText NAME_Strata;
-		FConstructorStatics() : NAME_Strata(LOCTEXT("Strata BSDFs", "Strata BSDFs")) { }
+		FConstructorStatics() : NAME_Strata(LOCTEXT("Substrate BSDFs", "Substrate BSDFs")) { }
 	};
 	static FConstructorStatics ConstructorStatics;
 #if WITH_EDITORONLY_DATA
@@ -23497,7 +23497,7 @@ const TArray<FExpressionInput*> UMaterialExpressionStrataSimpleClearCoatBSDF::Ge
 
 void UMaterialExpressionStrataSimpleClearCoatBSDF::GetCaption(TArray<FString>& OutCaptions) const
 {
-	OutCaptions.Add(TEXT("Strata Simple Clear Coat"));
+	OutCaptions.Add(TEXT("Substrate Simple Clear Coat"));
 }
 
 uint32 UMaterialExpressionStrataSimpleClearCoatBSDF::GetOutputType(int32 OutputIndex)
@@ -23615,7 +23615,7 @@ UMaterialExpressionStrataVolumetricFogCloudBSDF::UMaterialExpressionStrataVolume
 	struct FConstructorStatics
 	{
 		FText NAME_Strata;
-		FConstructorStatics() : NAME_Strata(LOCTEXT("Strata BSDFs", "Strata BSDFs")) { }
+		FConstructorStatics() : NAME_Strata(LOCTEXT("Substrate BSDFs", "Substrate BSDFs")) { }
 	};
 	static FConstructorStatics ConstructorStatics;
 #if WITH_EDITORONLY_DATA
@@ -23637,7 +23637,7 @@ int32 UMaterialExpressionStrataVolumetricFogCloudBSDF::Compile(class FMaterialCo
 
 void UMaterialExpressionStrataVolumetricFogCloudBSDF::GetCaption(TArray<FString>& OutCaptions) const
 {
-	OutCaptions.Add(TEXT("Strata Volumetric-Fog-Cloud BSDF"));
+	OutCaptions.Add(TEXT("Substrate Volumetric-Fog-Cloud BSDF"));
 }
 
 uint32 UMaterialExpressionStrataVolumetricFogCloudBSDF::GetOutputType(int32 OutputIndex)
@@ -23694,7 +23694,7 @@ UMaterialExpressionStrataLightFunction::UMaterialExpressionStrataLightFunction(c
 	struct FConstructorStatics
 	{
 		FText NAME_Strata;
-		FConstructorStatics() : NAME_Strata(LOCTEXT("Strata Extras", "Strata Extras")) { }
+		FConstructorStatics() : NAME_Strata(LOCTEXT("Substrate Extras", "Substrate Extras")) { }
 	};
 	static FConstructorStatics ConstructorStatics;
 #if WITH_EDITORONLY_DATA
@@ -23714,7 +23714,7 @@ int32 UMaterialExpressionStrataLightFunction::Compile(class FMaterialCompiler* C
 
 void UMaterialExpressionStrataLightFunction::GetCaption(TArray<FString>& OutCaptions) const
 {
-	OutCaptions.Add(TEXT("Strata Light Function"));
+	OutCaptions.Add(TEXT("Substrate Light Function"));
 }
 
 uint32 UMaterialExpressionStrataLightFunction::GetOutputType(int32 OutputIndex)
@@ -23762,7 +23762,7 @@ UMaterialExpressionStrataPostProcess::UMaterialExpressionStrataPostProcess(const
 	struct FConstructorStatics
 	{
 		FText NAME_Strata;
-		FConstructorStatics() : NAME_Strata(LOCTEXT("Strata Extras", "Strata Extras")) { }
+		FConstructorStatics() : NAME_Strata(LOCTEXT("Substrate Extras", "Substrate Extras")) { }
 	};
 	static FConstructorStatics ConstructorStatics;
 #if WITH_EDITORONLY_DATA
@@ -23785,7 +23785,7 @@ int32 UMaterialExpressionStrataPostProcess::Compile(class FMaterialCompiler* Com
 
 void UMaterialExpressionStrataPostProcess::GetCaption(TArray<FString>& OutCaptions) const
 {
-	OutCaptions.Add(TEXT("Strata Post Process"));
+	OutCaptions.Add(TEXT("Substrate Post Process"));
 }
 
 uint32 UMaterialExpressionStrataPostProcess::GetOutputType(int32 OutputIndex)
@@ -23836,7 +23836,7 @@ UMaterialExpressionStrataUI::UMaterialExpressionStrataUI(const FObjectInitialize
 	struct FConstructorStatics
 	{
 		FText NAME_Strata;
-		FConstructorStatics() : NAME_Strata(LOCTEXT("Strata Extras", "Strata Extras")) { }
+		FConstructorStatics() : NAME_Strata(LOCTEXT("Substrate Extras", "Substrate Extras")) { }
 	};
 	static FConstructorStatics ConstructorStatics;
 #if WITH_EDITORONLY_DATA
@@ -23859,7 +23859,7 @@ int32 UMaterialExpressionStrataUI::Compile(class FMaterialCompiler* Compiler, in
 
 void UMaterialExpressionStrataUI::GetCaption(TArray<FString>& OutCaptions) const
 {
-	OutCaptions.Add(TEXT("Strata UI"));
+	OutCaptions.Add(TEXT("Substrate UI"));
 }
 
 uint32 UMaterialExpressionStrataUI::GetOutputType(int32 OutputIndex)
@@ -23910,7 +23910,7 @@ UMaterialExpressionStrataConvertToDecal::UMaterialExpressionStrataConvertToDecal
 	struct FConstructorStatics
 	{
 		FText NAME_Strata;
-		FConstructorStatics() : NAME_Strata(LOCTEXT("Strata Extras", "Strata Extras")) { }
+		FConstructorStatics() : NAME_Strata(LOCTEXT("Substrate Extras", "Substrate Extras")) { }
 	};
 	static FConstructorStatics ConstructorStatics;
 #if WITH_EDITORONLY_DATA
@@ -23935,11 +23935,11 @@ int32 UMaterialExpressionStrataConvertToDecal::Compile(class FMaterialCompiler*
 	FStrataOperator& StrataOperator = Compiler->StrataCompilationGetOperator(Compiler->StrataTreeStackGetPathUniqueId());
 	if (!StrataOperator.bUseParameterBlending)
 	{
-		return Compiler->Errorf(TEXT("Strata Convert To Decal node must recveive StrataData a parameter blended strata material sub tree."));
+		return Compiler->Errorf(TEXT("Substrate Convert To Decal node must receive SubstrateData a parameter blended Substrate material sub tree."));
 	}
 	if (!StrataOperator.bRootOfParameterBlendingSubTree)
 	{
-		return Compiler->Errorf(TEXT("Strata Convert To Decal node must be the root of a parameter blending sub tree: no more strata operations can be applied a over its output."));
+		return Compiler->Errorf(TEXT("Substrate Convert To Decal node must be the root of a parameter blending sub tree: no more Substrate operations can be applied a over its output."));
 	}
 
 	// Propagate the parameter blended normal
@@ -23955,7 +23955,7 @@ int32 UMaterialExpressionStrataConvertToDecal::Compile(class FMaterialCompiler*
 
 void UMaterialExpressionStrataConvertToDecal::GetCaption(TArray<FString>& OutCaptions) const
 {
-	OutCaptions.Add(TEXT("Strata Convert To Decal"));
+	OutCaptions.Add(TEXT("Substrate Convert To Decal"));
 }
 
 uint32 UMaterialExpressionStrataConvertToDecal::GetOutputType(int32 OutputIndex)
@@ -24021,7 +24021,7 @@ UMaterialExpressionStrataUnlitBSDF::UMaterialExpressionStrataUnlitBSDF(const FOb
 	struct FConstructorStatics
 	{
 		FText NAME_Strata;
-		FConstructorStatics() : NAME_Strata(LOCTEXT("Strata BSDFs", "Strata BSDFs")) { }
+		FConstructorStatics() : NAME_Strata(LOCTEXT("Substrate BSDFs", "Substrate BSDFs")) { }
 	};
 	static FConstructorStatics ConstructorStatics;
 #if WITH_EDITORONLY_DATA
@@ -24041,7 +24041,7 @@ int32 UMaterialExpressionStrataUnlitBSDF::Compile(class FMaterialCompiler* Compi
 
 void UMaterialExpressionStrataUnlitBSDF::GetCaption(TArray<FString>& OutCaptions) const
 {
-	OutCaptions.Add(TEXT("Strata Unlit BSDF"));
+	OutCaptions.Add(TEXT("Substrate Unlit BSDF"));
 }
 
 uint32 UMaterialExpressionStrataUnlitBSDF::GetOutputType(int32 OutputIndex)
@@ -24093,7 +24093,7 @@ UMaterialExpressionStrataHairBSDF::UMaterialExpressionStrataHairBSDF(const FObje
 	struct FConstructorStatics
 	{
 		FText NAME_Strata;
-		FConstructorStatics() : NAME_Strata(LOCTEXT("Strata BSDFs", "Strata BSDFs")) { }
+		FConstructorStatics() : NAME_Strata(LOCTEXT("Substrate BSDFs", "Substrate BSDFs")) { }
 	};
 	static FConstructorStatics ConstructorStatics;
 #if WITH_EDITORONLY_DATA
@@ -24113,11 +24113,11 @@ int32 UMaterialExpressionStrataHairBSDF::Compile(class FMaterialCompiler* Compil
 
 	if (StrataOperator.bUseParameterBlending)
 	{
-		return Compiler->Errorf(TEXT("Strata Hair BSDF node cannot be used with parameter blending."));
+		return Compiler->Errorf(TEXT("Substrate Hair BSDF node cannot be used with parameter blending."));
 	}
 	else if (StrataOperator.bRootOfParameterBlendingSubTree)
 	{
-		return Compiler->Errorf(TEXT("Strata Hair BSDF node cannot be the root of a parameter blending sub tree."));
+		return Compiler->Errorf(TEXT("Substrate Hair BSDF node cannot be the root of a parameter blending sub tree."));
 	}
 
 	int32 OutputCodeChunk = Compiler->StrataHairBSDF(
@@ -24136,7 +24136,7 @@ int32 UMaterialExpressionStrataHairBSDF::Compile(class FMaterialCompiler* Compil
 
 void UMaterialExpressionStrataHairBSDF::GetCaption(TArray<FString>& OutCaptions) const
 {
-	OutCaptions.Add(TEXT("Strata Hair BSDF"));
+	OutCaptions.Add(TEXT("Substrate Hair BSDF"));
 }
 
 uint32 UMaterialExpressionStrataHairBSDF::GetOutputType(int32 OutputIndex)
@@ -24207,7 +24207,7 @@ UMaterialExpressionStrataEyeBSDF::UMaterialExpressionStrataEyeBSDF(const FObject
 	struct FConstructorStatics
 	{
 		FText NAME_Strata;
-		FConstructorStatics() : NAME_Strata(LOCTEXT("Strata BSDFs", "Strata BSDFs")) { }
+		FConstructorStatics() : NAME_Strata(LOCTEXT("Substrate BSDFs", "Substrate BSDFs")) { }
 	};
 	static FConstructorStatics ConstructorStatics;
 #if WITH_EDITORONLY_DATA
@@ -24232,11 +24232,11 @@ int32 UMaterialExpressionStrataEyeBSDF::Compile(class FMaterialCompiler* Compile
 
 	if (StrataOperator.bUseParameterBlending)
 	{
-		return Compiler->Errorf(TEXT("Strata Eye BSDF node cannot be used with parameter blending."));
+		return Compiler->Errorf(TEXT("Substrate Eye BSDF node cannot be used with parameter blending."));
 	}
 	else if (StrataOperator.bRootOfParameterBlendingSubTree)
 	{
-		return Compiler->Errorf(TEXT("Strata Eye BSDF node cannot be the root of a parameter blending sub tree."));
+		return Compiler->Errorf(TEXT("Substrate Eye BSDF node cannot be the root of a parameter blending sub tree."));
 	}
 
 	int32 OutputCodeChunk = Compiler->StrataEyeBSDF(
@@ -24257,7 +24257,7 @@ int32 UMaterialExpressionStrataEyeBSDF::Compile(class FMaterialCompiler* Compile
 
 void UMaterialExpressionStrataEyeBSDF::GetCaption(TArray<FString>& OutCaptions) const
 {
-	OutCaptions.Add(TEXT("Strata Eye BSDF"));
+	OutCaptions.Add(TEXT("Substrate Eye BSDF"));
 }
 
 uint32 UMaterialExpressionStrataEyeBSDF::GetOutputType(int32 OutputIndex)
@@ -24322,7 +24322,7 @@ UMaterialExpressionStrataSingleLayerWaterBSDF::UMaterialExpressionStrataSingleLa
 	struct FConstructorStatics
 	{
 		FText NAME_Strata;
-		FConstructorStatics() : NAME_Strata(LOCTEXT("Strata BSDFs", "Strata BSDFs")) { }
+		FConstructorStatics() : NAME_Strata(LOCTEXT("Substrate BSDFs", "Substrate BSDFs")) { }
 	};
 	static FConstructorStatics ConstructorStatics;
 #if WITH_EDITORONLY_DATA
@@ -24341,11 +24341,11 @@ int32 UMaterialExpressionStrataSingleLayerWaterBSDF::Compile(class FMaterialComp
 
 	if (StrataOperator.bUseParameterBlending)
 	{
-		return Compiler->Errorf(TEXT("Strata SingleLayerWater BSDF node cannot be used with parameter blending."));
+		return Compiler->Errorf(TEXT("Substrate SingleLayerWater BSDF node cannot be used with parameter blending."));
 	}
 	else if (StrataOperator.bRootOfParameterBlendingSubTree)
 	{
-		return Compiler->Errorf(TEXT("Strata SingleLayerWater BSDF node cannot be the root of a parameter blending sub tree."));
+		return Compiler->Errorf(TEXT("Substrate SingleLayerWater BSDF node cannot be the root of a parameter blending sub tree."));
 	}
 
 	int32 OutputCodeChunk = Compiler->StrataSingleLayerWaterBSDF(
@@ -24368,7 +24368,7 @@ int32 UMaterialExpressionStrataSingleLayerWaterBSDF::Compile(class FMaterialComp
 
 void UMaterialExpressionStrataSingleLayerWaterBSDF::GetCaption(TArray<FString>& OutCaptions) const
 {
-	OutCaptions.Add(TEXT("Strata Single Layer Water BSDF"));
+	OutCaptions.Add(TEXT("Substrate Single Layer Water BSDF"));
 }
 
 uint32 UMaterialExpressionStrataSingleLayerWaterBSDF::GetOutputType(int32 OutputIndex)
@@ -24455,7 +24455,7 @@ UMaterialExpressionStrataHorizontalMixing::UMaterialExpressionStrataHorizontalMi
 	struct FConstructorStatics
 	{
 		FText NAME_Strata;
-		FConstructorStatics() : NAME_Strata(LOCTEXT("Strata Ops", "Strata Operators")) { }
+		FConstructorStatics() : NAME_Strata(LOCTEXT("Substrate Ops", "Substrate Operators")) { }
 	};
 	static FConstructorStatics ConstructorStatics;
 #if WITH_EDITORONLY_DATA
@@ -24547,11 +24547,11 @@ void UMaterialExpressionStrataHorizontalMixing::GetCaption(TArray<FString>& OutC
 {
 	if (bUseParameterBlending)
 	{
-		OutCaptions.Add(TEXT("Strata Horizontal Blend (Parameter Blend)"));
+		OutCaptions.Add(TEXT("Substrate Horizontal Blend (Parameter Blend)"));
 	}
 	else
 	{
-		OutCaptions.Add(TEXT("Strata Horizontal Blend"));
+		OutCaptions.Add(TEXT("Substrate Horizontal Blend"));
 	}
 }
 
@@ -24619,7 +24619,7 @@ UMaterialExpressionStrataVerticalLayering::UMaterialExpressionStrataVerticalLaye
 	struct FConstructorStatics
 	{
 		FText NAME_Strata;
-		FConstructorStatics() : NAME_Strata(LOCTEXT("Strata Ops", "Strata Operators")) { }
+		FConstructorStatics() : NAME_Strata(LOCTEXT("Substrate Ops", "Substrate Operators")) { }
 	};
 	static FConstructorStatics ConstructorStatics;
 #if WITH_EDITORONLY_DATA
@@ -24708,11 +24708,11 @@ void UMaterialExpressionStrataVerticalLayering::GetCaption(TArray<FString>& OutC
 {
 	if (bUseParameterBlending)
 	{
-		OutCaptions.Add(TEXT("Strata Vertical Layer (Parameter Blend)"));
+		OutCaptions.Add(TEXT("Substrate Vertical Layer (Parameter Blend)"));
 	}
 	else
 	{
-		OutCaptions.Add(TEXT("Strata Vertical Layer"));
+		OutCaptions.Add(TEXT("Substrate Vertical Layer"));
 	}
 }
 
@@ -24807,7 +24807,7 @@ UMaterialExpressionStrataAdd::UMaterialExpressionStrataAdd(const FObjectInitiali
 	struct FConstructorStatics
 	{
 		FText NAME_Strata;
-		FConstructorStatics() : NAME_Strata(LOCTEXT("Strata Ops", "Strata Operators")) { }
+		FConstructorStatics() : NAME_Strata(LOCTEXT("Substrate Ops", "Substrate Operators")) { }
 	};
 	static FConstructorStatics ConstructorStatics;
 #if WITH_EDITORONLY_DATA
@@ -24892,11 +24892,11 @@ void UMaterialExpressionStrataAdd::GetCaption(TArray<FString>& OutCaptions) cons
 {
 	if (bUseParameterBlending)
 	{
-		OutCaptions.Add(TEXT("Strata Add (Parameter Blend)"));
+		OutCaptions.Add(TEXT("Substrate Add (Parameter Blend)"));
 	}
 	else
 	{
-		OutCaptions.Add(TEXT("Strata Add"));
+		OutCaptions.Add(TEXT("Substrate Add"));
 	}
 }
 
@@ -24963,7 +24963,7 @@ UMaterialExpressionStrataWeight::UMaterialExpressionStrataWeight(const FObjectIn
 	struct FConstructorStatics
 	{
 		FText NAME_Strata;
-		FConstructorStatics() : NAME_Strata(LOCTEXT("Strata Ops", "Strata Operators")) { }
+		FConstructorStatics() : NAME_Strata(LOCTEXT("Substrate Ops", "Substrate Operators")) { }
 	};
 	static FConstructorStatics ConstructorStatics;
 #if WITH_EDITORONLY_DATA
@@ -25019,7 +25019,7 @@ int32 UMaterialExpressionStrataWeight::Compile(class FMaterialCompiler* Compiler
 
 void UMaterialExpressionStrataWeight::GetCaption(TArray<FString>& OutCaptions) const
 {
-	OutCaptions.Add(TEXT("Strata Coverage Weight"));
+	OutCaptions.Add(TEXT("Substrate Coverage Weight"));
 }
 
 uint32 UMaterialExpressionStrataWeight::GetOutputType(int32 OutputIndex)
@@ -25079,7 +25079,7 @@ UMaterialExpressionStrataTransmittanceToMFP::UMaterialExpressionStrataTransmitta
 	struct FConstructorStatics
 	{
 		FText NAME_Strata;
-		FConstructorStatics() : NAME_Strata(LOCTEXT("Strata Helpers", "Strata Helpers")) { }
+		FConstructorStatics() : NAME_Strata(LOCTEXT("Substrate Helpers", "Substrate Helpers")) { }
 	};
 	static FConstructorStatics ConstructorStatics;
 #if WITH_EDITORONLY_DATA
@@ -25104,7 +25104,7 @@ int32 UMaterialExpressionStrataTransmittanceToMFP::Compile(class FMaterialCompil
 
 void UMaterialExpressionStrataTransmittanceToMFP::GetCaption(TArray<FString>& OutCaptions) const
 {
-	OutCaptions.Add(TEXT("Strata Transmittance-To-MeanFreePath"));
+	OutCaptions.Add(TEXT("Substrate Transmittance-To-MeanFreePath"));
 }
 
 uint32 UMaterialExpressionStrataTransmittanceToMFP::GetOutputType(int32 OutputIndex)
@@ -25179,7 +25179,7 @@ UMaterialExpressionStrataMetalnessToDiffuseAlbedoF0::UMaterialExpressionStrataMe
 	struct FConstructorStatics
 	{
 		FText NAME_Strata;
-		FConstructorStatics() : NAME_Strata(LOCTEXT("Strata Helpers", "Strata Helpers")) { }
+		FConstructorStatics() : NAME_Strata(LOCTEXT("Substrate Helpers", "Substrate Helpers")) { }
 	};
 	static FConstructorStatics ConstructorStatics;
 #if WITH_EDITORONLY_DATA
@@ -25205,7 +25205,7 @@ int32 UMaterialExpressionStrataMetalnessToDiffuseAlbedoF0::Compile(class FMateri
 
 void UMaterialExpressionStrataMetalnessToDiffuseAlbedoF0::GetCaption(TArray<FString>& OutCaptions) const
 {
-	OutCaptions.Add(TEXT("Strata Metalness-To-DiffuseAlbedo-F0"));
+	OutCaptions.Add(TEXT("Substrate Metalness-To-DiffuseAlbedo-F0"));
 }
 
 uint32 UMaterialExpressionStrataMetalnessToDiffuseAlbedoF0::GetOutputType(int32 OutputIndex)
@@ -25257,7 +25257,7 @@ UMaterialExpressionStrataHazinessToSecondaryRoughness::UMaterialExpressionStrata
 	struct FConstructorStatics
 	{
 		FText NAME_Strata;
-		FConstructorStatics() : NAME_Strata(LOCTEXT("Strata Helpers", "Strata Helpers")) { }
+		FConstructorStatics() : NAME_Strata(LOCTEXT("Substrate Helpers", "Substrate Helpers")) { }
 	};
 	static FConstructorStatics ConstructorStatics;
 #if WITH_EDITORONLY_DATA
@@ -25282,7 +25282,7 @@ int32 UMaterialExpressionStrataHazinessToSecondaryRoughness::Compile(class FMate
 
 void UMaterialExpressionStrataHazinessToSecondaryRoughness::GetCaption(TArray<FString>& OutCaptions) const
 {
-	OutCaptions.Add(TEXT("Strata Haziness-To-Secondary-Roughness"));
+	OutCaptions.Add(TEXT("Substrate Haziness-To-Secondary-Roughness"));
 }
 
 uint32 UMaterialExpressionStrataHazinessToSecondaryRoughness::GetOutputType(int32 OutputIndex)
@@ -25359,7 +25359,7 @@ UMaterialExpressionStrataThinFilm::UMaterialExpressionStrataThinFilm(const FObje
 	struct FConstructorStatics
 	{
 		FText NAME_Strata;
-		FConstructorStatics() : NAME_Strata(LOCTEXT("Strata Helpers", "Strata Helpers")) { }
+		FConstructorStatics() : NAME_Strata(LOCTEXT("Substrate Helpers", "Substrate Helpers")) { }
 	};
 	static FConstructorStatics ConstructorStatics;
 #if WITH_EDITORONLY_DATA
@@ -25389,7 +25389,7 @@ int32 UMaterialExpressionStrataThinFilm::Compile(class FMaterialCompiler* Compil
 
 void UMaterialExpressionStrataThinFilm::GetCaption(TArray<FString>& OutCaptions) const
 {
-	OutCaptions.Add(TEXT("Strata Thin-Film"));
+	OutCaptions.Add(TEXT("Substrate Thin-Film"));
 }
 
 uint32 UMaterialExpressionStrataThinFilm::GetOutputType(int32 OutputIndex)
diff --git a/Engine/Source/Runtime/Engine/Private/RendererSettings.cpp b/Engine/Source/Runtime/Engine/Private/RendererSettings.cpp
--- a/Engine/Source/Runtime/Engine/Private/RendererSettings.cpp
+++ b/Engine/Source/Runtime/Engine/Private/RendererSettings.cpp
@@ -263,7 +263,7 @@ void URendererSettings::PostEditChangeProperty(FPropertyChangedEvent& PropertyCh
 		{
 			if (bEnableStrata)
 			{
-				FMessageDialog::Open(EAppMsgType::Ok, LOCTEXT("Strata Experimental", "Warning: Strata is experimental. Be aware that any materials saved when Strata is enabled won't be rendered correctly if Strata is disabled later on."));
+				FMessageDialog::Open(EAppMsgType::Ok, LOCTEXT("Substrate Experimental", "Warning: Substrate is experimental. Be aware that any materials saved when Substrate is enabled won't be rendered correctly if Substrate is disabled later on."));
 			}
 		}
 
diff --git a/Engine/Source/Runtime/Engine/Private/ShowFlags.cpp b/Engine/Source/Runtime/Engine/Private/ShowFlags.cpp
--- a/Engine/Source/Runtime/Engine/Private/ShowFlags.cpp
+++ b/Engine/Source/Runtime/Engine/Private/ShowFlags.cpp
@@ -346,7 +346,7 @@ void ApplyViewMode(EViewModeIndex ViewModeIndex, bool bPerspective, FEngineShowF
 		case VMI_VisualizeLumen:
 			bPostProcessing = true;
 			break;
-		case VMI_VisualizeStrata:
+		case VMI_VisualizeSubstrate:
 			bPostProcessing = true;
 			break;
 		case VMI_VisualizeGroom:
@@ -389,7 +389,7 @@ void ApplyViewMode(EViewModeIndex ViewModeIndex, bool bPerspective, FEngineShowF
 	EngineShowFlags.SetVisualizeBuffer(ViewModeIndex == VMI_VisualizeBuffer);
 	EngineShowFlags.SetVisualizeNanite(ViewModeIndex == VMI_VisualizeNanite);
 	EngineShowFlags.SetVisualizeLumen(ViewModeIndex == VMI_VisualizeLumen);
-	EngineShowFlags.SetVisualizeStrata(ViewModeIndex == VMI_VisualizeStrata);
+	EngineShowFlags.SetVisualizeSubstrate(ViewModeIndex == VMI_VisualizeSubstrate);
 	EngineShowFlags.SetVisualizeGroom(ViewModeIndex == VMI_VisualizeGroom);
 	EngineShowFlags.SetVisualizeVirtualShadowMap(ViewModeIndex == VMI_VisualizeVirtualShadowMap);
 	EngineShowFlags.SetVisualizeLightCulling(ViewModeIndex == VMI_LightComplexity);
@@ -713,9 +713,9 @@ EViewModeIndex FindViewMode(const FEngineShowFlags& EngineShowFlags)
 	{
 		return VMI_VisualizeLumen;
 	}
-	else if (EngineShowFlags.VisualizeStrata)
+	else if (EngineShowFlags.VisualizeSubstrate)
 	{
-		return VMI_VisualizeStrata;
+		return VMI_VisualizeSubstrate;
 	}
 	else if (EngineShowFlags.VisualizeGroom)
 	{
@@ -863,7 +863,7 @@ const TCHAR* GetViewModeName(EViewModeIndex ViewModeIndex)
 		case VMI_VisualizeBuffer:			return TEXT("VisualizeBuffer");
 		case VMI_VisualizeNanite:			return TEXT("VisualizeNanite");
 		case VMI_VisualizeLumen:			return TEXT("VisualizeLumen");
-		case VMI_VisualizeStrata:			return TEXT("VisualizeStrata");
+		case VMI_VisualizeSubstrate:		return TEXT("VisualizeSubstrate");
 		case VMI_VisualizeGroom:			return TEXT("VisualizeGroom");
 		case VMI_VisualizeVirtualShadowMap:	return TEXT("VisualizeVirtualShadowMap");
 		case VMI_RayTracingDebug:			return TEXT("RayTracingDebug");
diff --git a/Engine/Source/Runtime/Engine/Private/StrataVisualizationData.cpp b/Engine/Source/Runtime/Engine/Private/StrataVisualizationData.cpp
--- a/Engine/Source/Runtime/Engine/Private/StrataVisualizationData.cpp
+++ b/Engine/Source/Runtime/Engine/Private/StrataVisualizationData.cpp
@@ -26,7 +26,7 @@ static FString ConfigureConsoleCommand(FStrataVisualizationData::TModeMap& ModeM
 	}
 
 	FString Out;
-	Out = TEXT("When the viewport view-mode is set to 'Strata Visualization', this command specifies which of the various channels to display. Values entered other than the allowed values shown below will be ignored.");
+	Out = TEXT("When the viewport view-mode is set to 'Substrate Visualization', this command specifies which of the various channels to display. Values entered other than the allowed values shown below will be ignored.");
 	Out += AvailableVisualizationModes;
 
 	IConsoleManager::Get().RegisterConsoleVariable(
@@ -72,7 +72,7 @@ void FStrataVisualizationData::Initialize()
 			AllModeMap,
 			TEXT("MaterialProperties"),
 			LOCTEXT("MaterialProperties", "Material Properties"),
-			LOCTEXT("MaterialPropertiesDesc", "Visualizes Strata material properties under mouse cursor"),
+			LOCTEXT("MaterialPropertiesDesc", "Visualizes Substrate material properties under mouse cursor"),
 			FViewMode::MaterialProperties,
 			true,
 			true,
@@ -82,7 +82,7 @@ void FStrataVisualizationData::Initialize()
 			AllModeMap,
 			TEXT("MaterialCount"),
 			LOCTEXT("MaterialCount", "Material Count"),
-			LOCTEXT("MaterialCountDesc", "Visualizes Strata material count per pixel"),
+			LOCTEXT("MaterialCountDesc", "Visualizes Substrate material count per pixel"),
 			FViewMode::MaterialCount,
 			true,
 			true,
@@ -92,17 +92,17 @@ void FStrataVisualizationData::Initialize()
 			AllModeMap,
 			TEXT("AdvancedMaterialProperties"),
 			LOCTEXT("AdvancedMaterialProperties", "Advanced Material Properties"),
-			LOCTEXT("AdvancedMaterialPropertiesDesc", "Visualizes Strata advanced material properties"),
+			LOCTEXT("AdvancedMaterialPropertiesDesc", "Visualizes Substrate advanced material properties"),
 			FViewMode::AdvancedMaterialProperties,
 			true,
 			Strata::IsAdvancedVisualizationEnabled(),
-			LOCTEXT("IsStrataAdvancedDebugShaderEnabled", "Strata advanced debugging r.Strata.Debug.AdvancedVisualizationShaders is disabled"));
+			LOCTEXT("IsSubstrateAdvancedDebugShaderEnabled", "Substrate advanced debugging r.Substrate.Debug.AdvancedVisualizationShaders is disabled"));
 
 		AddVisualizationMode(
 			AllModeMap,
 			TEXT("MaterialClassification"),
 			LOCTEXT("MaterialClassification", "Material Classification"),
-			LOCTEXT("MaterialClassificationDesc", "Visualizes Strata material classification"),
+			LOCTEXT("MaterialClassificationDesc", "Visualizes Substrate material classification"),
 			FViewMode::MaterialClassification,
 			true,
 			true,
@@ -112,27 +112,27 @@ void FStrataVisualizationData::Initialize()
 			AllModeMap,
 			TEXT("DecalClassification"),
 			LOCTEXT("DecalClassification", "Decal classification"),
-			LOCTEXT("DecalClassificationDesc", "Visualizes Strata decal classification"),
+			LOCTEXT("DecalClassificationDesc", "Visualizes Substrate decal classification"),
 			FViewMode::DecalClassification,
 			true,
 			false, // Disable for now, as it is not important, and is mainly used for debugging
-			LOCTEXT("IsStrataDBufferPassEnabled", "Strata tiled DBuffer pass (r.Strata.DBufferPass and r.Strata.DBufferPass.DedicatedTiles) is disabled"));
+			LOCTEXT("IsSubstrateBufferPassEnabled", "Substrate tiled DBuffer pass (r.Substrate.DBufferPass and r.Substrate.DBufferPass.DedicatedTiles) is disabled"));
 
 		AddVisualizationMode(
 			AllModeMap,
 			TEXT("RoughRefractionClassification"),
 			LOCTEXT("RoughRefractionClassification", "Rough Refraction Classification"),
-			LOCTEXT("RoughRefractionClassificationDesc", "Visualizes Strata rough refraction classification"),
+			LOCTEXT("RoughRefractionClassificationDesc", "Visualizes Substrate rough refraction classification"),
 			FViewMode::RoughRefractionClassification,
 			true,
 			Strata::IsOpaqueRoughRefractionEnabled(),
-			LOCTEXT("IsStrataRoughRefractionEnabled", "Strata rough refraction r.Strata.OpaqueMaterialRoughRefraction is disabled"));
+			LOCTEXT("IsSubstrateRoughRefractionEnabled", "Substrate rough refraction r.Substrate.OpaqueMaterialRoughRefraction is disabled"));
 
 		AddVisualizationMode(
 			AllModeMap,
-			TEXT("StrataInfo"),
-			LOCTEXT("StrataInfo", "Strata Info"),
-			LOCTEXT("StrataInfoDesc", "Visualizes Strata info"),
+			TEXT("SubstrateInfo"),
+			LOCTEXT("SubstrateInfo", "Substrate Info"),
+			LOCTEXT("SubstrateInfoDesc", "Visualizes Substrate info"),
 			FViewMode::StrataInfo,
 			true,
 			true,
diff --git a/Engine/Source/Runtime/Engine/Private/ViewModeNames.cpp b/Engine/Source/Runtime/Engine/Private/ViewModeNames.cpp
--- a/Engine/Source/Runtime/Engine/Private/ViewModeNames.cpp
+++ b/Engine/Source/Runtime/Engine/Private/ViewModeNames.cpp
@@ -281,9 +281,9 @@ TArray<const FSlateBrush*> FillViewModeDisplayIcons()
 		{
 			ViewModeDisplayIcons.Emplace(FAppStyle::Get().GetBrush("EditorViewport.VisualizeLumenMode"));
 		}
-		else if (ViewModeIndex == VMI_VisualizeStrata)
+		else if (ViewModeIndex == VMI_VisualizeSubstrate)
 		{
-			ViewModeDisplayIcons.Emplace(FAppStyle::Get().GetBrush("EditorViewport.VisualizeStrataMode"));
+			ViewModeDisplayIcons.Emplace(FAppStyle::Get().GetBrush("EditorViewport.VisualizeSubstrateMode"));
 		}
 		else if (ViewModeIndex == VMI_VisualizeGroom)
 		{
diff --git a/Engine/Source/Runtime/Engine/Public/MaterialCompiler.h b/Engine/Source/Runtime/Engine/Public/MaterialCompiler.h
--- a/Engine/Source/Runtime/Engine/Public/MaterialCompiler.h
+++ b/Engine/Source/Runtime/Engine/Public/MaterialCompiler.h
@@ -60,7 +60,7 @@ enum class EMaterialCompilerType
 };
 
 
-/** Whether we need some data export from a Strata material from spatially varying properties, e.g. diffuse color for Lighmass to generate lightmaps. */
+/** Whether we need some data export from a Substrate material from spatially varying properties, e.g. diffuse color for Lighmass to generate lightmaps. */
 enum EStrataMaterialExport : uint8
 {
 	SME_None			= 0,
diff --git a/Engine/Source/Runtime/Engine/Public/SceneTypes.h b/Engine/Source/Runtime/Engine/Public/SceneTypes.h
--- a/Engine/Source/Runtime/Engine/Public/SceneTypes.h
+++ b/Engine/Source/Runtime/Engine/Public/SceneTypes.h
@@ -162,8 +162,8 @@ enum EMaterialProperty : int
 	MP_EmissiveColor = 0 UMETA(DisplayName = "Emissive"),
 	MP_Opacity UMETA(DisplayName = "Opacity"),
 	MP_OpacityMask UMETA(DisplayName = "Opacity Mask"),
-	MP_DiffuseColor UMETA(Hidden),			// used in Lightmass, not exposed to user, computed from: BaseColor, Metallic				Also used in Strata which uses Albedo/F0 parameterrization
-	MP_SpecularColor UMETA(Hidden),			// used in Lightmass, not exposed to user, derived from: SpecularColor, Metallic, Specular	Also used in Strata which uses Albedo/F0 parameterrization
+	MP_DiffuseColor UMETA(Hidden),			// used in Lightmass, not exposed to user, computed from: BaseColor, Metallic				Also used in Substrate which uses Albedo/F0 parameterization
+	MP_SpecularColor UMETA(Hidden),			// used in Lightmass, not exposed to user, derived from: SpecularColor, Metallic, Specular	Also used in Substrate which uses Albedo/F0 parameterization
 	MP_BaseColor UMETA(DisplayName = "Diffuse"),
 	MP_Metallic UMETA(DisplayName = "Metallic"),
 	MP_Specular UMETA(DisplayName = "Specular"),
diff --git a/Engine/Source/Runtime/Engine/Public/SceneView.h b/Engine/Source/Runtime/Engine/Public/SceneView.h
--- a/Engine/Source/Runtime/Engine/Public/SceneView.h
+++ b/Engine/Source/Runtime/Engine/Public/SceneView.h
@@ -1185,7 +1185,7 @@ public:
 	/** Current Lumen visualization mode */
 	FName CurrentLumenVisualizationMode;
 
-	/** Current Strata visualization mode */
+	/** Current Substrate visualization mode */
 	FName CurrentStrataVisualizationMode;
 
 	/** Current Groom visualization mode */
diff --git a/Engine/Source/Runtime/Engine/Public/ShowFlags.h b/Engine/Source/Runtime/Engine/Public/ShowFlags.h
--- a/Engine/Source/Runtime/Engine/Public/ShowFlags.h
+++ b/Engine/Source/Runtime/Engine/Public/ShowFlags.h
@@ -376,7 +376,7 @@ private:
 		SetVisualizeBuffer(false);
 		SetVisualizeNanite(false);
 		SetVisualizeLumen(false);
-		SetVisualizeStrata(false);
+		SetVisualizeSubstrate(false);
 		SetVisualizeGroom(false);
 		SetVisualizeVirtualShadowMap(false);
 		SetVectorFields(false);
diff --git a/Engine/Source/Runtime/Engine/Public/ShowFlagsValues.inl b/Engine/Source/Runtime/Engine/Public/ShowFlagsValues.inl
--- a/Engine/Source/Runtime/Engine/Public/ShowFlagsValues.inl
+++ b/Engine/Source/Runtime/Engine/Public/ShowFlagsValues.inl
@@ -66,9 +66,9 @@ SHOWFLAG_FIXED_IN_SHIPPING(0, VisualizeBuffer, SFG_Hidden, NSLOCTEXT("UnrealEd",
 SHOWFLAG_ALWAYS_ACCESSIBLE(VisualizeNanite, SFG_Hidden, NSLOCTEXT("UnrealEd", "VisualizeNaniteSF", "Nanite Visualization"))
 /** Needed for VMI_VisualizeLumen, Whether to enable the Lumen visualization mode. */
 SHOWFLAG_ALWAYS_ACCESSIBLE(VisualizeLumen, SFG_Hidden, NSLOCTEXT("UnrealEd", "VisualizeLumenSF", "Lumen Visualization"))
-/** Needed for VMI_VisualizeStrata, Whether to enable the Strata visualization mode. */
-SHOWFLAG_ALWAYS_ACCESSIBLE(VisualizeStrata, SFG_Hidden, NSLOCTEXT("UnrealEd", "VisualizeStrataSF", "Strata Visualization"))
-/** Needed for VMI_VisualizeGroom, Whether to enable the Strata visualization mode. */
+/** Needed for VMI_VisualizeSubstrate, Whether to enable the Substrate visualization mode. */
+SHOWFLAG_ALWAYS_ACCESSIBLE(VisualizeSubstrate, SFG_Hidden, NSLOCTEXT("UnrealEd", "VisualizeSubstrateSF", "Substrate Visualization"))
+/** Needed for VMI_VisualizeGroom, Whether to enable the Groom visualization mode. */
 SHOWFLAG_ALWAYS_ACCESSIBLE(VisualizeGroom, SFG_Hidden, NSLOCTEXT("UnrealEd", "VisualizeGroomSF", "Groom Visualization"))
 /** Needed for VMI_VisualizeVirtualShadowMap, Whether to enable the virtual shadow map visualization mode. */
 SHOWFLAG_ALWAYS_ACCESSIBLE(VisualizeVirtualShadowMap, SFG_Hidden, NSLOCTEXT("UnrealEd", "VisualizeVirtualShadowMapSF", "Virtual Shadow Map Visualization"))
diff --git a/Engine/Source/Runtime/Engine/Public/StrataVisualizationData.h b/Engine/Source/Runtime/Engine/Public/StrataVisualizationData.h
--- a/Engine/Source/Runtime/Engine/Public/StrataVisualizationData.h
+++ b/Engine/Source/Runtime/Engine/Public/StrataVisualizationData.h
@@ -65,7 +65,7 @@ public:
 	/** Return the console command name for enabling single mode visualization. */
 	static const TCHAR* GetVisualizeConsoleCommandName()
 	{
-		return TEXT("r.Strata.ViewMode");
+		return TEXT("r.Substrate.ViewMode");
 	}
 
 private:
diff --git a/Engine/Source/Runtime/RenderCore/Private/RenderUtils.cpp b/Engine/Source/Runtime/RenderCore/Private/RenderUtils.cpp
--- a/Engine/Source/Runtime/RenderCore/Private/RenderUtils.cpp
+++ b/Engine/Source/Runtime/RenderCore/Private/RenderUtils.cpp
@@ -1516,69 +1516,69 @@ ERayTracingMode GetRayTracingMode()
 ///////////////////////////////////////////////////////////////////////////////////////////////////
 // Strata settings interface
 
-static TAutoConsoleVariable<int32> CVarStrata(
-	TEXT("r.Strata"),
+static TAutoConsoleVariable<int32> CVarSubstrate(
+	TEXT("r.Substrate"),
 	0,
-	TEXT("Enable Strata materials (Beta)."),
+	TEXT("Enable Substrate materials (Beta)."),
 	ECVF_ReadOnly | ECVF_RenderThreadSafe);
 
-static TAutoConsoleVariable<int32> CVarStrataBytePerPixel(
-	TEXT("r.Strata.BytesPerPixel"),
+static TAutoConsoleVariable<int32> CVarSubstrateBytePerPixel(
+	TEXT("r.Substrate.BytesPerPixel"),
 	80,
-	TEXT("Strata allocated byte per pixel to store materials data. Higher value means more complex material can be represented."),
+	TEXT("Substrate allocated byte per pixel to store materials data. Higher value means more complex material can be represented."),
 	ECVF_ReadOnly | ECVF_RenderThreadSafe);
 
-static TAutoConsoleVariable<int32> CVarStrataBackCompatibility(
-	TEXT("r.StrataBackCompatibility"),
+static TAutoConsoleVariable<int32> CVarSubstrateBackCompatibility(
+	TEXT("r.SubstrateBackCompatibility"),
 	0,
-	TEXT("Disables Strata multiple scattering and replaces Chan diffuse by Lambert."),
+	TEXT("Disables Substrate multiple scattering and replaces Chan diffuse by Lambert."),
 	ECVF_ReadOnly | ECVF_RenderThreadSafe);
 
-static TAutoConsoleVariable<int32> CVarStrataOpaqueMaterialRoughRefraction(
-	TEXT("r.Strata.OpaqueMaterialRoughRefraction"),
+static TAutoConsoleVariable<int32> CVarSubstrateOpaqueMaterialRoughRefraction(
+	TEXT("r.Substrate.OpaqueMaterialRoughRefraction"),
 	0,
-	TEXT("Enable Strata opaque material rough refractions effect from top layers over layers below."),
+	TEXT("Enable Substrate opaque material rough refractions effect from top layers over layers below."),
 	ECVF_ReadOnly | ECVF_RenderThreadSafe);
 
-static TAutoConsoleVariable<int32> CVarStrataShadingQuality(
-	TEXT("r.Strata.ShadingQuality"),
+static TAutoConsoleVariable<int32> CVarSubstrateShadingQuality(
+	TEXT("r.Substrate.ShadingQuality"),
 	1,
-	TEXT("Define Strata shading quality (1: accurate lighting, 2: approximate lighting). This variable is read-only."),
+	TEXT("Define Substrate shading quality (1: accurate lighting, 2: approximate lighting). This variable is read-only."),
 	ECVF_ReadOnly | ECVF_RenderThreadSafe);
 
-static TAutoConsoleVariable<int32> CVarStrataDBufferPass(
-	TEXT("r.Strata.DBufferPass"),
+static TAutoConsoleVariable<int32> CVarSubstrateDBufferPass(
+	TEXT("r.Substrate.DBufferPass"),
 	0,
 	TEXT("Apply DBuffer after the base-pass as a separate pass. Read only because when this is changed, it will require the recompilation of all shaders."),
 	ECVF_ReadOnly | ECVF_RenderThreadSafe);
 
-// Transition render settings that will disapear when strata gets enabled
+// Transition render settings that will disappear when Substrate gets enabled
 static TAutoConsoleVariable<int32> CVarMaterialRoughDiffuse(
 	TEXT("r.Material.RoughDiffuse"),
 	0,
 	TEXT("Enable rough diffuse material."),
 	ECVF_ReadOnly | ECVF_RenderThreadSafe);
 
-static TAutoConsoleVariable<int32> CVarStrataRoughDiffuse(
-	TEXT("r.Strata.RoughDiffuse"),
+static TAutoConsoleVariable<int32> CVarSubstrateRoughDiffuse(
+	TEXT("r.Substrate.RoughDiffuse"),
 	1,
-	TEXT("Enable Strata rough diffuse model (works only if r.Material.RoughDiffuse is enabled in the project settings). Togglable at runtime"),
+	TEXT("Enable Substrate rough diffuse model (works only if r.Material.RoughDiffuse is enabled in the project settings). Togglable at runtime"),
 	ECVF_RenderThreadSafe);
 
-static TAutoConsoleVariable<int32> CVarStrataDebugAdvancedVisualizationShaders(
-	TEXT("r.Strata.Debug.AdvancedVisualizationShaders"),
+static TAutoConsoleVariable<int32> CVarSubstrateDebugAdvancedVisualizationShaders(
+	TEXT("r.Substrate.Debug.AdvancedVisualizationShaders"),
 	0,
-	TEXT("Enable advanced strata material debug visualization shaders. Base pass shaders can output such advanced data."),
+	TEXT("Enable advanced Substrate material debug visualization shaders. Base pass shaders can output such advanced data."),
 	ECVF_ReadOnly | ECVF_RenderThreadSafe);
 
-static TAutoConsoleVariable<int32> CVarStrataTileCoord8Bits(
-	TEXT("r.Strata.TileCoord8bits"),
+static TAutoConsoleVariable<int32> CVarSubstrateTileCoord8Bits(
+	TEXT("r.Substrate.TileCoord8bits"),
 	0,
 	TEXT("Format of tile coord. This variable is read-only."),
 	ECVF_ReadOnly | ECVF_RenderThreadSafe);
 
-static TAutoConsoleVariable<int32> CVarStrataAccurateSRGB(
-	TEXT("r.Strata.AccurateSRGB"),
+static TAutoConsoleVariable<int32> CVarSubstrateAccurateSRGB(
+	TEXT("r.Substrate.AccurateSRGB"),
 	0,
 	TEXT("Enable accurate sRGB encoding for pixel data encoding/decoding during rendering passes (base pass, lighting, ...). Otherwise use a simple 'Gamma2' approximation."),
 	ECVF_ReadOnly | ECVF_RenderThreadSafe);
@@ -1587,7 +1587,7 @@ namespace Strata
 {
 	bool IsStrataEnabled()
 	{
-		return CVarStrata.GetValueOnAnyThread() > 0;
+		return CVarSubstrate.GetValueOnAnyThread() > 0;
 	}
 
 	static uint32 InternalGetBytePerPixel(uint32 InBytePerPixel)
@@ -1600,13 +1600,13 @@ namespace Strata
 
 	uint32 GetBytePerPixel()
 	{
-		return InternalGetBytePerPixel(CVarStrataBytePerPixel.GetValueOnAnyThread());
+		return InternalGetBytePerPixel(CVarSubstrateBytePerPixel.GetValueOnAnyThread());
 	}
 
 	uint32 GetBytePerPixel(EShaderPlatform InPlatform)
 	{
 		// Variant for shader compilation per platform
-		static FShaderPlatformCachedIniValue<int32> CVarBudget(TEXT("r.Strata.BytesPerPixel"));
+		static FShaderPlatformCachedIniValue<int32> CVarBudget(TEXT("r.Substrate.BytesPerPixel"));
 		return InternalGetBytePerPixel(CVarBudget.Get(InPlatform));
 	}
 
@@ -1637,34 +1637,33 @@ namespace Strata
 
 	bool IsBackCompatibilityEnabled()
 	{
-		return CVarStrataBackCompatibility.GetValueOnAnyThread() > 0 ? 1 : 0;
+		return CVarSubstrateBackCompatibility.GetValueOnAnyThread() > 0 ? 1 : 0;
 	}
 	
 	bool IsRoughDiffuseEnabled()
 	{
-		return CVarStrataRoughDiffuse.GetValueOnAnyThread() > 0;
+		return CVarSubstrateRoughDiffuse.GetValueOnAnyThread() > 0;
 	}
 
 	bool Is8bitTileCoordEnabled()
 	{
-		return CVarStrataTileCoord8Bits.GetValueOnAnyThread() > 0 ? 1 : 0;
+		return CVarSubstrateTileCoord8Bits.GetValueOnAnyThread() > 0 ? 1 : 0;
 	}
 
 	bool IsAccurateSRGBEnabled()
 	{
-		return CVarStrataAccurateSRGB.GetValueOnAnyThread() > 0 ? 1 : 0;
+		return CVarSubstrateAccurateSRGB.GetValueOnAnyThread() > 0 ? 1 : 0;
 	}
 
 	uint32 GetShadingQuality()
 	{
-		return CVarStrataShadingQuality.GetValueOnAnyThread();
+		return CVarSubstrateShadingQuality.GetValueOnAnyThread();
 	}
 
 	uint32 GetShadingQuality(EShaderPlatform InPlatform)
-
 	{
-		static FShaderPlatformCachedIniValue<int32> CVarStrataShadingQualityPlatform(TEXT("r.Strata.ShadingQuality"));
-		return CVarStrataShadingQualityPlatform.Get(InPlatform);
+		static FShaderPlatformCachedIniValue<int32> CVarSubstrateShadingQualityPlatform(TEXT("r.Substrate.ShadingQuality"));
+		return CVarSubstrateShadingQualityPlatform.Get(InPlatform);
 	}
 	
 	bool IsDBufferPassEnabled(EShaderPlatform InPlatform)
@@ -1672,21 +1671,21 @@ namespace Strata
 		// DBuffer pass is only available if high quality normal is disabled, or if we are on a console.
 		// That is because in high quality normals requires a uint2 UAV which is only supported on some graphic cards on PC 
 		// and this is an unknown when compiling shaders at this stage.
-		// !!! If this is changed, please update all sites reading r.Strata.DBufferPass !!!
-		const uint32 StrataNormalQuality = GetNormalQuality();
-		const bool bDBufferPassSupported = StrataNormalQuality == 0 || (StrataNormalQuality > 0 && IsConsolePlatform(InPlatform));
+		// !!! If this is changed, please update all sites reading r.Substrate.DBufferPass !!!
+		const uint32 NormalQuality = GetNormalQuality();
+		const bool bDBufferPassSupported = NormalQuality == 0 || (NormalQuality > 0 && IsConsolePlatform(InPlatform));
 		
-		static FShaderPlatformCachedIniValue<int32> CVarStrataDBufferPassPlatform(TEXT("r.Strata.DBufferPass"));
-		return IsStrataEnabled() && IsUsingDBuffers(InPlatform) && bDBufferPassSupported && CVarStrataDBufferPassPlatform.Get(InPlatform) > 0;
+		static FShaderPlatformCachedIniValue<int32> CVarSubstrateDBufferPassPlatform(TEXT("r.Substrate.DBufferPass"));
+		return IsStrataEnabled() && IsUsingDBuffers(InPlatform) && bDBufferPassSupported && CVarSubstrateDBufferPassPlatform.Get(InPlatform) > 0;
 	}
 
 	bool IsOpaqueRoughRefractionEnabled()
 	{
-		return IsStrataEnabled() && CVarStrataOpaqueMaterialRoughRefraction.GetValueOnAnyThread() > 0;
+		return IsStrataEnabled() && CVarSubstrateOpaqueMaterialRoughRefraction.GetValueOnAnyThread() > 0;
 	}
 
 	bool IsAdvancedVisualizationEnabled()
 	{
-		return CVarStrataDebugAdvancedVisualizationShaders.GetValueOnAnyThread() > 0;
+		return CVarSubstrateDebugAdvancedVisualizationShaders.GetValueOnAnyThread() > 0;
 	}
 }
\ No newline at end of file
diff --git a/Engine/Source/Runtime/Renderer/Private/ClusteredDeferredShadingPass.cpp b/Engine/Source/Runtime/Renderer/Private/ClusteredDeferredShadingPass.cpp
--- a/Engine/Source/Runtime/Renderer/Private/ClusteredDeferredShadingPass.cpp
+++ b/Engine/Source/Runtime/Renderer/Private/ClusteredDeferredShadingPass.cpp
@@ -189,7 +189,7 @@ static void InternalAddClusteredDeferredShadingPass(
 	
 	const TCHAR* TileTypeName = ToString(TileType);
 	GraphBuilder.AddPass(
-		RDG_EVENT_NAME("Light::ClusteredDeferredShading(%s,Lights:%d,Tile:%s)", bHairStrands ? TEXT("HairStrands") : (bStrata ? TEXT("Strata") : TEXT("GBuffer")), SortedLightsSet.ClusteredSupportedEnd, TileTypeName),
+		RDG_EVENT_NAME("Light::ClusteredDeferredShading(%s,Lights:%d,Tile:%s)", bHairStrands ? TEXT("HairStrands") : (bStrata ? TEXT("Substrate") : TEXT("GBuffer")), SortedLightsSet.ClusteredSupportedEnd, TileTypeName),
 		PassParameters,
 		ERDGPassFlags::Raster,
 		[PassParameters, &View, SceneTextureExtent, bHairStrands, bStrata, TileType, PrimitiveType](FRHICommandListImmediate& InRHICmdList)
diff --git a/Engine/Source/Runtime/Renderer/Private/IlluminanceMeter.cpp b/Engine/Source/Runtime/Renderer/Private/IlluminanceMeter.cpp
--- a/Engine/Source/Runtime/Renderer/Private/IlluminanceMeter.cpp
+++ b/Engine/Source/Runtime/Renderer/Private/IlluminanceMeter.cpp
@@ -213,6 +213,6 @@ void FScene::ProcessAndRenderIlluminanceMeter(FRDGBuilder& GraphBuilder, TArrayV
 		FPrintIlluminanceMeterPS::FPermutationDomain PermutationVector;
 		TShaderMapRef<FPrintIlluminanceMeterPS> PixelShader(MainView.ShaderMap, PermutationVector);
 
-		FPixelShaderUtils::AddFullscreenPass<FPrintIlluminanceMeterPS>(GraphBuilder, MainView.ShaderMap, RDG_EVENT_NAME("Strata::VisualizeMaterial(Draw)"), PixelShader, PassParameters, MainView.ViewRect, PreMultipliedColorTransmittanceBlend);
+		FPixelShaderUtils::AddFullscreenPass<FPrintIlluminanceMeterPS>(GraphBuilder, MainView.ShaderMap, RDG_EVENT_NAME("Substrate::VisualizeMaterial(Draw)"), PixelShader, PassParameters, MainView.ViewRect, PreMultipliedColorTransmittanceBlend);
 	}
 }
diff --git a/Engine/Source/Runtime/Renderer/Private/PostProcess/PostProcessing.cpp b/Engine/Source/Runtime/Renderer/Private/PostProcess/PostProcessing.cpp
--- a/Engine/Source/Runtime/Renderer/Private/PostProcess/PostProcessing.cpp
+++ b/Engine/Source/Runtime/Renderer/Private/PostProcess/PostProcessing.cpp
@@ -282,7 +282,7 @@ void AddPostProcessingPasses(
 		VisualizeStationaryLightOverlap,
 		VisualizeLightCulling,
 		VisualizePostProcessStack,
-		VisualizeStrata,
+		VisualizeSubstrate,
 		VisualizeLightGrid,
 		VisualizeSkyAtmosphere,
 		VisualizeLevelInstance,
@@ -330,7 +330,7 @@ void AddPostProcessingPasses(
 		TEXT("VisualizeStationaryLightOverlap"),
 		TEXT("VisualizeLightCulling"),
 		TEXT("VisualizePostProcessStack"),
-		TEXT("VisualizeStrata"),
+		TEXT("VisualizeSubstrate"),
 		TEXT("VisualizeLightGrid"),
 		TEXT("VisualizeSkyAtmosphere"),
 		TEXT("VisualizeLevelInstance"),
@@ -363,7 +363,7 @@ void AddPostProcessingPasses(
 	PassSequence.SetEnabled(EPass::VisualizePostProcessStack, false);
 #endif
 	PassSequence.SetEnabled(EPass::VisualizeLumenScene, LumenVisualizeMode >= 0 && LumenVisualizeMode != VISUALIZE_MODE_OVERVIEW && bPostProcessingEnabled);
-	PassSequence.SetEnabled(EPass::VisualizeStrata, Strata::ShouldRenderStrataDebugPasses(View));
+	PassSequence.SetEnabled(EPass::VisualizeSubstrate, Strata::ShouldRenderStrataDebugPasses(View));
 	PassSequence.SetEnabled(EPass::VisualizeLightGrid, ShouldVisualizeLightGrid());
 
 #if WITH_EDITOR
@@ -1165,10 +1165,10 @@ void AddPostProcessingPasses(
 	}
 #endif
 
-	if (PassSequence.IsEnabled(EPass::VisualizeStrata))
+	if (PassSequence.IsEnabled(EPass::VisualizeSubstrate))
 	{
 		FScreenPassRenderTarget OverrideOutput;
-		PassSequence.AcceptOverrideIfLastPass(EPass::VisualizeStrata, OverrideOutput);
+		PassSequence.AcceptOverrideIfLastPass(EPass::VisualizeSubstrate, OverrideOutput);
 		SceneColor = Strata::AddStrataDebugPasses(GraphBuilder, View, SceneColor);
 	}
 
diff --git a/Engine/Source/Runtime/Renderer/Private/PostProcess/TemporalSuperResolution.cpp b/Engine/Source/Runtime/Renderer/Private/PostProcess/TemporalSuperResolution.cpp
--- a/Engine/Source/Runtime/Renderer/Private/PostProcess/TemporalSuperResolution.cpp
+++ b/Engine/Source/Runtime/Renderer/Private/PostProcess/TemporalSuperResolution.cpp
@@ -1051,11 +1051,11 @@ ITemporalUpscaler::FOutputs AddTemporalSuperResolutionPasses(
 	// Do not composite translucency if we are visualizing a buffer, unless it is the overview mode.
 	static FName OverviewName = FName(TEXT("Overview"));
 	bHasSeparateTranslucency &= 
-		   (!View.Family->EngineShowFlags.VisualizeBuffer || (View.Family->EngineShowFlags.VisualizeBuffer && View.CurrentBufferVisualizationMode == OverviewName))
-		&& (!View.Family->EngineShowFlags.VisualizeNanite || (View.Family->EngineShowFlags.VisualizeNanite && View.CurrentNaniteVisualizationMode == OverviewName))
-		&& (!View.Family->EngineShowFlags.VisualizeLumen  || (View.Family->EngineShowFlags.VisualizeLumen  && View.CurrentLumenVisualizationMode  == OverviewName))
-		&& (!View.Family->EngineShowFlags.VisualizeStrata || (View.Family->EngineShowFlags.VisualizeStrata && View.CurrentStrataVisualizationMode == OverviewName))
-		&& (!View.Family->EngineShowFlags.VisualizeGroom  || (View.Family->EngineShowFlags.VisualizeGroom  && View.CurrentGroomVisualizationMode  == OverviewName));
+		   (!View.Family->EngineShowFlags.VisualizeBuffer    || (View.Family->EngineShowFlags.VisualizeBuffer    && View.CurrentBufferVisualizationMode == OverviewName))
+		&& (!View.Family->EngineShowFlags.VisualizeNanite    || (View.Family->EngineShowFlags.VisualizeNanite    && View.CurrentNaniteVisualizationMode == OverviewName))
+		&& (!View.Family->EngineShowFlags.VisualizeLumen     || (View.Family->EngineShowFlags.VisualizeLumen     && View.CurrentLumenVisualizationMode  == OverviewName))
+		&& (!View.Family->EngineShowFlags.VisualizeSubstrate || (View.Family->EngineShowFlags.VisualizeSubstrate && View.CurrentStrataVisualizationMode == OverviewName))
+		&& (!View.Family->EngineShowFlags.VisualizeGroom     || (View.Family->EngineShowFlags.VisualizeGroom     && View.CurrentGroomVisualizationMode  == OverviewName));
 #endif
 	if (bHasSeparateTranslucency)
 	{
diff --git a/Engine/Source/Runtime/Renderer/Private/Strata/Strata.cpp b/Engine/Source/Runtime/Renderer/Private/Strata/Strata.cpp
--- a/Engine/Source/Runtime/Renderer/Private/Strata/Strata.cpp
+++ b/Engine/Source/Runtime/Renderer/Private/Strata/Strata.cpp
@@ -22,37 +22,37 @@
 
 // The project setting for Strata
 static TAutoConsoleVariable<int32> CVarUseCmaskClear(
-	TEXT("r.Strata.UseCmaskClear"),
+	TEXT("r.Substrate.UseCmaskClear"),
 	0,
 	TEXT("TEST."),
 	ECVF_RenderThreadSafe);
 
-static TAutoConsoleVariable<float> CVarStrataTileOverflow(
-	TEXT("r.Strata.TileOverflow"),
+static TAutoConsoleVariable<float> CVarSubstrateTileOverflow(
+	TEXT("r.Substrate.TileOverflow"),
 	1.f,
-	TEXT("Scale the number of Strata tile for overflowing tiles containing multi-BSDFs pixels. (0: 0%, 1: 100%. Default 1.0f)."),
+	TEXT("Scale the number of Substrate tile for overflowing tiles containing multi-BSDFs pixels. (0: 0%, 1: 100%. Default 1.0f)."),
 	ECVF_RenderThreadSafe);
 
-static TAutoConsoleVariable<int32> CVarStrataTileOverflowFromMaterial(
-	TEXT("r.Strata.TileOverflowFromMaterial"),
+static TAutoConsoleVariable<int32> CVarSubstrateTileOverflowFromMaterial(
+	TEXT("r.Substrate.TileOverflowFromMaterial"),
 	1,
-	TEXT("When enable, scale the number of Strata tile for overflowing tiles containing multi-BSDFs pixels based on material data. Otherwise use global overflow scale defined by r.Strata.TileOverflow."),
+	TEXT("When enable, scale the number of Substrate tile for overflowing tiles containing multi-BSDFs pixels based on material data. Otherwise use global overflow scale defined by r.Substrate.TileOverflow."),
 	ECVF_RenderThreadSafe);
 
-static TAutoConsoleVariable<int32> CVarStrataDebugPeelLayersAboveDepth(
-	TEXT("r.Strata.Debug.PeelLayersAboveDepth"),
+static TAutoConsoleVariable<int32> CVarSubstrateDebugPeelLayersAboveDepth(
+	TEXT("r.Substrate.Debug.PeelLayersAboveDepth"),
 	0,
-	TEXT("Strata debug control to progressively peel off materials layer by layer."),
+	TEXT("Substrate debug control to progressively peel off materials layer by layer."),
 	ECVF_RenderThreadSafe);
 
-static TAutoConsoleVariable<int32> CVarStrataAsyncClassification(
-	TEXT("r.Strata.AsyncClassification"),
+static TAutoConsoleVariable<int32> CVarSubstrateAsyncClassification(
+	TEXT("r.Substrate.AsyncClassification"),
 	1,
-	TEXT("Run Strata material classification in async (with shadow)."),
+	TEXT("Run Substrate material classification in async (with shadow)."),
 	ECVF_RenderThreadSafe);
 
-static TAutoConsoleVariable<int32> CVarStrataDBufferPassDedicatedTiles(
-	TEXT("r.Strata.DBufferPass.DedicatedTiles"),
+static TAutoConsoleVariable<int32> CVarSubstrateDBufferPassDedicatedTiles(
+	TEXT("r.Substrate.DBufferPass.DedicatedTiles"),
 	0,
 	TEXT("Use dedicated tile for DBuffer application when DBuffer pass is enabled."),
 	ECVF_RenderThreadSafe);
@@ -102,7 +102,7 @@ enum EStrataTileSpace
 
 bool DoesStrataTileOverflowUseMaterialData() 
 {
-	return CVarStrataTileOverflowFromMaterial.GetValueOnRenderThread() > 0;
+	return CVarSubstrateTileOverflowFromMaterial.GetValueOnRenderThread() > 0;
 }
 
 float GetStrataTileOverflowRatio(const FViewInfo& View)
@@ -114,7 +114,7 @@ float GetStrataTileOverflowRatio(const FViewInfo& View)
 	}
 	else
 	{
-		return FMath::Clamp(CVarStrataTileOverflow.GetValueOnRenderThread(), 0.f, 4.0f);
+		return FMath::Clamp(CVarSubstrateTileOverflow.GetValueOnRenderThread(), 0.f, 4.0f);
 	}
 }
 
@@ -155,7 +155,7 @@ bool SupportsCMask(const FStaticShaderPlatform InPlatform)
 
 bool IsClassificationAsync()
 {
-	return CVarStrataAsyncClassification.GetValueOnRenderThread() > 0;
+	return CVarSubstrateAsyncClassification.GetValueOnRenderThread() > 0;
 }
 
 static EPixelFormat GetClassificationTileFormat(const FIntPoint& InResolution)
@@ -185,25 +185,25 @@ static void InitialiseStrataViewData(FRDGBuilder& GraphBuilder, FViewInfo& View,
 
 		const TCHAR* StrataTileListBufferNames[EStrataTileType::ECount] =
 		{
-			TEXT("Strata.StrataTileListBuffer(Simple)"),
-			TEXT("Strata.StrataTileListBuffer(Single)"),
-			TEXT("Strata.StrataTileListBuffer(Complex)"),
-			TEXT("Strata.StrataTileListBuffer(OpaqueRoughRefraction)"),
-			TEXT("Strata.StrataTileListBuffer(OpaqueRoughRefractionSSSWithout)"),
-			TEXT("Strata.StrataTileListBuffer(DecalSimple)"),
-			TEXT("Strata.StrataTileListBuffer(DecalSingle)"),
-			TEXT("Strata.StrataTileListBuffer(DecalComplex)"),
+			TEXT("Substrate.StrataTileListBuffer(Simple)"),
+			TEXT("Substrate.StrataTileListBuffer(Single)"),
+			TEXT("Substrate.StrataTileListBuffer(Complex)"),
+			TEXT("Substrate.StrataTileListBuffer(OpaqueRoughRefraction)"),
+			TEXT("Substrate.StrataTileListBuffer(OpaqueRoughRefractionSSSWithout)"),
+			TEXT("Substrate.StrataTileListBuffer(DecalSimple)"),
+			TEXT("Substrate.StrataTileListBuffer(DecalSingle)"),
+			TEXT("Substrate.StrataTileListBuffer(DecalComplex)"),
 		};
 
 		// Tile classification buffers
 		{
 			// Indirect draw
-			Out.ClassificationTileDrawIndirectBuffer = GraphBuilder.CreateBuffer(FRDGBufferDesc::CreateIndirectDesc<FRHIDrawIndirectParameters>(EStrataTileType::ECount), TEXT("Strata.StrataTileDrawIndirectBuffer"));
+			Out.ClassificationTileDrawIndirectBuffer = GraphBuilder.CreateBuffer(FRDGBufferDesc::CreateIndirectDesc<FRHIDrawIndirectParameters>(EStrataTileType::ECount), TEXT("Substrate.StrataTileDrawIndirectBuffer"));
 			Out.ClassificationTileDrawIndirectBufferUAV = GraphBuilder.CreateUAV(Out.ClassificationTileDrawIndirectBuffer, PF_R32_UINT);
 			AddClearUAVPass(GraphBuilder, Out.ClassificationTileDrawIndirectBufferUAV, 0);
 
 			// Indirect dispatch
-			Out.ClassificationTileDispatchIndirectBuffer = GraphBuilder.CreateBuffer(FRDGBufferDesc::CreateIndirectDesc<FRHIDispatchIndirectParameters>(EStrataTileType::ECount), TEXT("Strata.StrataTileDispatchIndirectBuffer"));
+			Out.ClassificationTileDispatchIndirectBuffer = GraphBuilder.CreateBuffer(FRDGBufferDesc::CreateIndirectDesc<FRHIDispatchIndirectParameters>(EStrataTileType::ECount), TEXT("Substrate.StrataTileDispatchIndirectBuffer"));
 			Out.ClassificationTileDispatchIndirectBufferUAV = GraphBuilder.CreateUAV(Out.ClassificationTileDispatchIndirectBuffer, PF_R32_UINT);
 			AddClearUAVPass(GraphBuilder, Out.ClassificationTileDispatchIndirectBufferUAV, 0);
 
@@ -255,12 +255,12 @@ static void InitialiseStrataViewData(FRDGBuilder& GraphBuilder, FViewInfo& View,
 			Out.OverflowTileCount = GetStrataTextureTileResolution(View, ViewResolution, EStrataTileSpace::StrataTileSpace_Overflow);
 			Out.OverflowTileOffset = GetStrataTextureTileResolution(View, Out.TileOffset * STRATA_TILE_SIZE, EStrataTileSpace::StrataTileSpace_Overflow) + BaseOverflowTileOffset;
 
-			Out.BSDFTileTexture = GraphBuilder.CreateTexture(FRDGTextureDesc::Create2D(BufferSize_Extended, PF_R32_UINT, FClearValueBinding::None, TexCreate_UAV | TexCreate_ShaderResource), TEXT("Strata.BSDFTiles"));
+			Out.BSDFTileTexture = GraphBuilder.CreateTexture(FRDGTextureDesc::Create2D(BufferSize_Extended, PF_R32_UINT, FClearValueBinding::None, TexCreate_UAV | TexCreate_ShaderResource), TEXT("Substrate.BSDFTiles"));
 			AddClearUAVPass(GraphBuilder, GraphBuilder.CreateUAV(Out.BSDFTileTexture), 0u);
 
-			Out.BSDFTilePerThreadDispatchIndirectBuffer = GraphBuilder.CreateBuffer(FRDGBufferDesc::CreateIndirectDesc<FRHIDispatchIndirectParameters>(1), TEXT("Strata.StrataBSDFTilePerThreadDispatchIndirectBuffer"));
-			Out.BSDFTileDispatchIndirectBuffer = GraphBuilder.CreateBuffer(FRDGBufferDesc::CreateIndirectDesc<FRHIDispatchIndirectParameters>(1), TEXT("Strata.StrataBSDFTileDispatchIndirectBuffer"));
-			Out.BSDFTileCountBuffer = GraphBuilder.CreateBuffer(FRDGBufferDesc::CreateBufferDesc(4, 1), TEXT("Strata.BSDFTileCount"));
+			Out.BSDFTilePerThreadDispatchIndirectBuffer = GraphBuilder.CreateBuffer(FRDGBufferDesc::CreateIndirectDesc<FRHIDispatchIndirectParameters>(1), TEXT("Substrate.StrataBSDFTilePerThreadDispatchIndirectBuffer"));
+			Out.BSDFTileDispatchIndirectBuffer = GraphBuilder.CreateBuffer(FRDGBufferDesc::CreateIndirectDesc<FRHIDispatchIndirectParameters>(1), TEXT("Substrate.StrataBSDFTileDispatchIndirectBuffer"));
+			Out.BSDFTileCountBuffer = GraphBuilder.CreateBuffer(FRDGBufferDesc::CreateBufferDesc(4, 1), TEXT("Substrate.BSDFTileCount"));
 		}
 		else
 		{
@@ -297,7 +297,7 @@ static void RecordStrataAnalytics(EShaderPlatform InPlatform)
 		EventAttributes.Add(FAnalyticsEventAttribute(TEXT("Enabled"), 1));
 		EventAttributes.Add(FAnalyticsEventAttribute(TEXT("BytesPerPixel"), GetBytePerPixel(InPlatform)));
 
-		FString OutStr(TEXT("Strata"));
+		FString OutStr(TEXT("Substrate"));
 		FEngineAnalytics::GetProvider().RecordEvent(OutStr, EventAttributes);
 	}
 }
@@ -356,7 +356,7 @@ void InitialiseStrataFrameSceneData(FRDGBuilder& GraphBuilder, FSceneRenderer& S
 
 		// Top layer texture
 		{
-			Out.TopLayerTexture = GraphBuilder.CreateTexture(FRDGTextureDesc::Create2D(SceneTextureExtent, GetTopLayerTextureFormat(), FClearValueBinding::Black, TexCreate_RenderTargetable | TexCreate_ShaderResource | TexCreate_FastVRAM | (bNeedUAV ? TexCreate_UAV : TexCreate_None)), TEXT("Strata.TopLayerTexture"));
+			Out.TopLayerTexture = GraphBuilder.CreateTexture(FRDGTextureDesc::Create2D(SceneTextureExtent, GetTopLayerTextureFormat(), FClearValueBinding::Black, TexCreate_RenderTargetable | TexCreate_ShaderResource | TexCreate_FastVRAM | (bNeedUAV ? TexCreate_UAV : TexCreate_None)), TEXT("Substrate.TopLayerTexture"));
 		}
 
 		// Separated subsurface and rough refraction textures
@@ -365,13 +365,13 @@ void InitialiseStrataFrameSceneData(FRDGBuilder& GraphBuilder, FSceneRenderer& S
 			const FIntPoint OpaqueRoughRefractionSceneExtent		 = bIsStrataOpaqueMaterialRoughRefractionEnabled ? SceneTextureExtent : FIntPoint(4, 4);
 			
 			Out.OpaqueRoughRefractionTexture = GraphBuilder.CreateTexture(
-				FRDGTextureDesc::Create2D(OpaqueRoughRefractionSceneExtent, PF_FloatR11G11B10, FClearValueBinding::Black, TexCreate_ShaderResource | TexCreate_UAV | TexCreate_RenderTargetable), TEXT("Strata.OpaqueRoughRefractionTexture"));
+				FRDGTextureDesc::Create2D(OpaqueRoughRefractionSceneExtent, PF_FloatR11G11B10, FClearValueBinding::Black, TexCreate_ShaderResource | TexCreate_UAV | TexCreate_RenderTargetable), TEXT("Substrate.OpaqueRoughRefractionTexture"));
 			Out.OpaqueRoughRefractionTextureUAV = GraphBuilder.CreateUAV(Out.OpaqueRoughRefractionTexture);
 			
 			Out.SeparatedSubSurfaceSceneColor			= GraphBuilder.CreateTexture(
-				FRDGTextureDesc::Create2D(OpaqueRoughRefractionSceneExtent, PF_FloatR11G11B10, FClearValueBinding::Black, TexCreate_ShaderResource | TexCreate_UAV | TexCreate_RenderTargetable), TEXT("Strata.SeparatedSubSurfaceSceneColor"));
+				FRDGTextureDesc::Create2D(OpaqueRoughRefractionSceneExtent, PF_FloatR11G11B10, FClearValueBinding::Black, TexCreate_ShaderResource | TexCreate_UAV | TexCreate_RenderTargetable), TEXT("Substrate.SeparatedSubSurfaceSceneColor"));
 			Out.SeparatedOpaqueRoughRefractionSceneColor= GraphBuilder.CreateTexture(
-				FRDGTextureDesc::Create2D(OpaqueRoughRefractionSceneExtent, PF_FloatR11G11B10, FClearValueBinding::Black, TexCreate_ShaderResource | TexCreate_UAV | TexCreate_RenderTargetable), TEXT("Strata.SeparatedOpaqueRoughRefractionSceneColor"));
+				FRDGTextureDesc::Create2D(OpaqueRoughRefractionSceneExtent, PF_FloatR11G11B10, FClearValueBinding::Black, TexCreate_ShaderResource | TexCreate_UAV | TexCreate_RenderTargetable), TEXT("Substrate.SeparatedOpaqueRoughRefractionSceneColor"));
 
 			if (bIsStrataOpaqueMaterialRoughRefractionEnabled)
 			{
@@ -385,7 +385,7 @@ void InitialiseStrataFrameSceneData(FRDGBuilder& GraphBuilder, FSceneRenderer& S
 		// BSDF offsets
 		if (bNeedBSDFOffsets)
 		{
-			Out.BSDFOffsetTexture = GraphBuilder.CreateTexture(FRDGTextureDesc::Create2D(SceneTextureExtent, PF_R32_UINT, FClearValueBinding::None, TexCreate_UAV | TexCreate_ShaderResource), TEXT("Strata.BSDFOffsets"));
+			Out.BSDFOffsetTexture = GraphBuilder.CreateTexture(FRDGTextureDesc::Create2D(SceneTextureExtent, PF_R32_UINT, FClearValueBinding::None, TexCreate_UAV | TexCreate_ShaderResource), TEXT("Substrate.BSDFOffsets"));
 			AddClearUAVPass(GraphBuilder, GraphBuilder.CreateUAV(Out.BSDFOffsetTexture), 0u);
 		}
 	}
@@ -403,7 +403,7 @@ void InitialiseStrataFrameSceneData(FRDGBuilder& GraphBuilder, FSceneRenderer& S
 	FRDGTextureDesc MaterialTextureDesc = FRDGTextureDesc::Create2DArray(SceneTextureExtent, PF_R32_UINT, FClearValueBinding::Transparent, TexCreate_TargetArraySlicesIndependently | TexCreate_DisableDCC | TexCreate_NoFastClear | TexCreate_RenderTargetable | TexCreate_ShaderResource | TexCreate_UAV | TexCreate_FastVRAM, SliceCount, 1, 1);
 	MaterialTextureDesc.FastVRAMPercentage = (1.0f / SliceCount) * 0xFF; // Only allocate the first slice into ESRAM
 
-	Out.MaterialTextureArray = GraphBuilder.CreateTexture(MaterialTextureDesc, TEXT("Strata.Material"));
+	Out.MaterialTextureArray = GraphBuilder.CreateTexture(MaterialTextureDesc, TEXT("Substrate.Material"));
 	Out.MaterialTextureArraySRV = GraphBuilder.CreateSRV(FRDGTextureSRVDesc::Create(Out.MaterialTextureArray));
 	Out.MaterialTextureArrayUAV = GraphBuilder.CreateUAV(FRDGTextureUAVDesc(Out.MaterialTextureArray, 0));
 
@@ -414,7 +414,7 @@ void InitialiseStrataFrameSceneData(FRDGBuilder& GraphBuilder, FSceneRenderer& S
 	// Rough diffuse model
 	Out.bRoughDiffuse = IsRoughDiffuseEnabled();
 
-	Out.PeelLayersAboveDepth = FMath::Max(CVarStrataDebugPeelLayersAboveDepth.GetValueOnRenderThread(), 0);
+	Out.PeelLayersAboveDepth = FMath::Max(CVarSubstrateDebugPeelLayersAboveDepth.GetValueOnRenderThread(), 0);
 
 	// STRATA_TODO allocate a slice for StoringDebugStrata only if STRATA_ADVANCED_DEBUG_ENABLED is enabled 
 	Out.SliceStoringDebugStrataTreeData				= SliceCount - SliceCountAdvDebug;										// When we read, there is no slices excluded
@@ -450,10 +450,10 @@ void BindStrataBasePassUniformParameters(FRDGBuilder& GraphBuilder, const FViewI
 	}
 	else
 	{
-		FRDGTextureRef DummyWritableRefracTexture = GraphBuilder.CreateTexture(FRDGTextureDesc::Create2D(FIntPoint(1, 1), PF_R8, FClearValueBinding::None, TexCreate_ShaderResource | TexCreate_UAV), TEXT("Strata.DummyWritableTexture"));
+		FRDGTextureRef DummyWritableRefracTexture = GraphBuilder.CreateTexture(FRDGTextureDesc::Create2D(FIntPoint(1, 1), PF_R8, FClearValueBinding::None, TexCreate_ShaderResource | TexCreate_UAV), TEXT("Substrate.DummyWritableTexture"));
 		FRDGTextureUAVRef DummyWritableRefracTextureUAV = GraphBuilder.CreateUAV(FRDGTextureUAVDesc(DummyWritableRefracTexture));
 
-		FRDGTextureRef DummyWritableTextureArray = GraphBuilder.CreateTexture(FRDGTextureDesc::Create2DArray(FIntPoint(1, 1), PF_R32_UINT, FClearValueBinding::None, TexCreate_ShaderResource | TexCreate_UAV, 1), TEXT("Strata.DummyWritableTexture"));
+		FRDGTextureRef DummyWritableTextureArray = GraphBuilder.CreateTexture(FRDGTextureDesc::Create2DArray(FIntPoint(1, 1), PF_R32_UINT, FClearValueBinding::None, TexCreate_ShaderResource | TexCreate_UAV, 1), TEXT("Substrate.DummyWritableTexture"));
 		FRDGTextureUAVRef DummyWritableTextureArrayUAV = GraphBuilder.CreateUAV(FRDGTextureUAVDesc(DummyWritableTextureArray));
 
 		const FRDGSystemTextures& SystemTextures = FRDGSystemTextures::Get(GraphBuilder);
@@ -1002,7 +1002,7 @@ void AddStrataInternalClassificationTilePass(
 	}
 	
 	GraphBuilder.AddPass(
-		RDG_EVENT_NAME("Strata::%sClassificationPass(%s)", bDebug ? TEXT("Debug") : TEXT("Stencil"), ToString(TileMaterialType)),
+		RDG_EVENT_NAME("Substrate::%sClassificationPass(%s)", bDebug ? TEXT("Debug") : TEXT("Stencil"), ToString(TileMaterialType)),
 		ParametersPS,
 		ERDGPassFlags::Raster,
 		[ParametersPS, VertexShader, PixelShader, ViewRect, DebugOutputResolution, StrataTilePrimitiveType, TileMaterialType, bDebug](FRHICommandList& RHICmdList)
@@ -1156,7 +1156,7 @@ void SetBasePassRenderTargetOutputFormat(const EShaderPlatform Platform, const F
 
 void AddStrataMaterialClassificationPass(FRDGBuilder& GraphBuilder, const FMinimalSceneTextures& SceneTextures, const FDBufferTextures& DBufferTextures, const TArray<FViewInfo>& Views)
 {
-	RDG_EVENT_SCOPE_CONDITIONAL(GraphBuilder, IsStrataEnabled() && Views.Num() > 0, "Strata::MaterialClassification");
+	RDG_EVENT_SCOPE_CONDITIONAL(GraphBuilder, IsStrataEnabled() && Views.Num() > 0, "Substrate::MaterialClassification");
 	if (!IsStrataEnabled())
 	{
 		return;
@@ -1185,11 +1185,11 @@ void AddStrataMaterialClassificationPass(FRDGBuilder& GraphBuilder, const FMinim
 			{
 				// Combine DBuffer RTWriteMasks; will end up in one texture we can load from in the base pass PS and decide whether to do the actual work or not.
 				FRDGTextureRef SourceCMaskTextures[] = { StrataSceneData->TopLayerTexture };
-				FRenderTargetWriteMask::Decode(GraphBuilder, View.ShaderMap, MakeArrayView(SourceCMaskTextures), TopLayerCmaskTexture, GFastVRamConfig.DBufferMask, TEXT("Strata::TopLayerCmask"));
+				FRenderTargetWriteMask::Decode(GraphBuilder, View.ShaderMap, MakeArrayView(SourceCMaskTextures), TopLayerCmaskTexture, GFastVRamConfig.DBufferMask, TEXT("Substrate::TopLayerCmask"));
 			}
 
 			// If Dbuffer pass (i.e. apply DBuffer data after the base-pass) is enabled, run special classification for outputing tile with/without tiles
-			const bool bDBufferTiles = IsDBufferPassEnabled(View.GetShaderPlatform()) && CVarStrataDBufferPassDedicatedTiles.GetValueOnRenderThread() > 0 && DBufferTextures.IsValid() && IsConsolePlatform(View.GetShaderPlatform());
+			const bool bDBufferTiles = IsDBufferPassEnabled(View.GetShaderPlatform()) && CVarSubstrateDBufferPassDedicatedTiles.GetValueOnRenderThread() > 0 && DBufferTextures.IsValid() && IsConsolePlatform(View.GetShaderPlatform());
 
 			FStrataMaterialTileClassificationPassCS::FPermutationDomain PermutationVector;
 			PermutationVector.Set< FStrataMaterialTileClassificationPassCS::FCmask >(bSupportCMask);
@@ -1224,7 +1224,7 @@ void AddStrataMaterialClassificationPass(FRDGBuilder& GraphBuilder, const FMinim
 			const uint32 GroupSize = 8;
 			FComputeShaderUtils::AddPass(
 				GraphBuilder,
-				RDG_EVENT_NAME("Strata::MaterialTileClassification(%s%s)", bWaveOps ? TEXT("Wave") : TEXT("SharedMemory"), bSupportCMask ? TEXT(", CMask") : TEXT("")),
+				RDG_EVENT_NAME("Substrate::MaterialTileClassification(%s%s)", bWaveOps ? TEXT("Wave") : TEXT("SharedMemory"), bSupportCMask ? TEXT(", CMask") : TEXT("")),
 				PassFlags,
 				ComputeShader,
 				PassParameters,
@@ -1240,7 +1240,7 @@ void AddStrataMaterialClassificationPass(FRDGBuilder& GraphBuilder, const FMinim
 
 			FComputeShaderUtils::AddPass(
 				GraphBuilder,
-				RDG_EVENT_NAME("Strata::MaterialTilePrepareArgs"),
+				RDG_EVENT_NAME("Substrate::MaterialTilePrepareArgs"),
 				PassFlags,
 				ComputeShader,
 				PassParameters,
@@ -1276,7 +1276,7 @@ void AddStrataMaterialClassificationPass(FRDGBuilder& GraphBuilder, const FMinim
 
 			FComputeShaderUtils::AddPass(
 				GraphBuilder,
-				RDG_EVENT_NAME("Strata::BSDFTileAndOffsets(%s)", bWaveOps ? TEXT("Wave") : TEXT("SharedMemory")),
+				RDG_EVENT_NAME("Substrate::BSDFTileAndOffsets(%s)", bWaveOps ? TEXT("Wave") : TEXT("SharedMemory")),
 				PassFlags,
 				ComputeShader,
 				PassParameters,
@@ -1299,7 +1299,7 @@ void AddStrataMaterialClassificationPass(FRDGBuilder& GraphBuilder, const FMinim
 
 			FComputeShaderUtils::AddPass(
 				GraphBuilder,
-				RDG_EVENT_NAME("Strata::BSDFTilePrepareArgs"),
+				RDG_EVENT_NAME("Substrate::BSDFTilePrepareArgs"),
 				PassFlags,
 				ComputeShader,
 				PassParameters,
@@ -1311,7 +1311,7 @@ void AddStrataMaterialClassificationPass(FRDGBuilder& GraphBuilder, const FMinim
 
 void AddStrataDBufferPass(FRDGBuilder& GraphBuilder, const FMinimalSceneTextures& SceneTextures, const FDBufferTextures& DBufferTextures, const TArray<FViewInfo>& Views)
 {
-	RDG_EVENT_SCOPE_CONDITIONAL(GraphBuilder, IsStrataEnabled() && Views.Num() > 0, "Strata::DBuffer");
+	RDG_EVENT_SCOPE_CONDITIONAL(GraphBuilder, IsStrataEnabled() && Views.Num() > 0, "Substrate::DBuffer");
 	if (!IsStrataEnabled() || !DBufferTextures.IsValid())
 	{
 		return;
@@ -1377,7 +1377,7 @@ void AddStrataDBufferPass(FRDGBuilder& GraphBuilder, const FMinimalSceneTextures
 			const uint32 GroupSize = 8;
 			FComputeShaderUtils::AddPass(
 				GraphBuilder,
-				RDG_EVENT_NAME("Strata::Dbuffer(%s)", ToString(TileType)),
+				RDG_EVENT_NAME("Substrate::Dbuffer(%s)", ToString(TileType)),
 				ERDGPassFlags::Compute,
 				ComputeShader,
 				PassParameters,
@@ -1385,7 +1385,7 @@ void AddStrataDBufferPass(FRDGBuilder& GraphBuilder, const FMinimalSceneTextures
 				TileTypeDispatchIndirectArgOffset(TileType));
 		};
 
-		const bool bDbufferTiles = CVarStrataDBufferPassDedicatedTiles.GetValueOnRenderThread() > 0;
+		const bool bDbufferTiles = CVarSubstrateDBufferPassDedicatedTiles.GetValueOnRenderThread() > 0;
 		DBufferPass(bDbufferTiles ? EStrataTileType::EDecalComplex : EStrataTileType::EComplex);
 		DBufferPass(bDbufferTiles ? EStrataTileType::EDecalSingle : EStrataTileType::ESingle);
 		DBufferPass(bDbufferTiles ? EStrataTileType::EDecalSimple : EStrataTileType::ESimple);
diff --git a/Engine/Source/Runtime/Renderer/Private/Strata/StrataRoughRefraction.cpp b/Engine/Source/Runtime/Renderer/Private/Strata/StrataRoughRefraction.cpp
--- a/Engine/Source/Runtime/Renderer/Private/Strata/StrataRoughRefraction.cpp
+++ b/Engine/Source/Runtime/Renderer/Private/Strata/StrataRoughRefraction.cpp
@@ -14,8 +14,8 @@
 #include "ShaderCompiler.h"
 
 
-static TAutoConsoleVariable<float> CVarStrataOpaqueMaterialRoughRefractionBlurScale(
-	TEXT("r.Strata.OpaqueMaterialRoughRefraction.BlurScale"),
+static TAutoConsoleVariable<float> CVarSubstrateOpaqueMaterialRoughRefractionBlurScale(
+	TEXT("r.Substrate.OpaqueMaterialRoughRefraction.BlurScale"),
 	1.0f,
 	TEXT("Scale opaque rough refraction strengh for debug purposes."),
 	ECVF_RenderThreadSafe);
@@ -66,7 +66,7 @@ void AddStrataOpaqueRoughRefractionPasses(
 {
 	const uint32 ViewCount = Views.Num();
 	const bool bOpaqueRoughRefractionEnabled = IsOpaqueRoughRefractionEnabled() && ViewCount > 0;
-	RDG_EVENT_SCOPE_CONDITIONAL(GraphBuilder, bOpaqueRoughRefractionEnabled, "Strata::OpaqueRoughRefraction");
+	RDG_EVENT_SCOPE_CONDITIONAL(GraphBuilder, bOpaqueRoughRefractionEnabled, "Substrate::OpaqueRoughRefraction");
 	if (!bOpaqueRoughRefractionEnabled)
 	{
 		return;
@@ -86,7 +86,7 @@ void AddStrataOpaqueRoughRefractionPasses(
 
 		if (TempTexture == nullptr)
 		{
-			TempTexture = GraphBuilder.CreateTexture(FRDGTextureDesc::Create2D(View.GetSceneTexturesConfig().Extent, PF_FloatRGBA, FClearValueBinding::Black, TexCreate_UAV | TexCreate_ShaderResource | TexCreate_RenderTargetable), TEXT("Strata.RoughRefrac.TempTexture"));
+			TempTexture = GraphBuilder.CreateTexture(FRDGTextureDesc::Create2D(View.GetSceneTexturesConfig().Extent, PF_FloatRGBA, FClearValueBinding::Black, TexCreate_UAV | TexCreate_ShaderResource | TexCreate_RenderTargetable), TEXT("Substrate.RoughRefrac.TempTexture"));
 		}
 
 		FRDGTextureRef SeparatedOpaqueRoughRefractionSceneColor = View.StrataViewData.SceneData->SeparatedOpaqueRoughRefractionSceneColor;
@@ -99,7 +99,7 @@ void AddStrataOpaqueRoughRefractionPasses(
 		PassParameters->SeparatedOpaqueRoughRefractionSceneColor = SeparatedOpaqueRoughRefractionSceneColor;
 		PassParameters->RenderTargets[0] = FRenderTargetBinding(TempTexture, LoadAction);
 		PassParameters->BlurDirection = FVector2f(1.0f, 0.0f);
-		PassParameters->BlurScale = FMath::Max(0.0f, CVarStrataOpaqueMaterialRoughRefractionBlurScale.GetValueOnAnyThread());
+		PassParameters->BlurScale = FMath::Max(0.0f, CVarSubstrateOpaqueMaterialRoughRefractionBlurScale.GetValueOnAnyThread());
 
 		FOpaqueRoughRefractionPS::FPermutationDomain PermutationVector;
 		PermutationVector.Set<FOpaqueRoughRefractionPS::FEnableBlur>(true);
@@ -114,7 +114,7 @@ void AddStrataOpaqueRoughRefractionPasses(
 		PassParameters->StrataTile = Strata::SetTileParameters(GraphBuilder, View, StrataTileType, PrimitiveType);
 
 		GraphBuilder.AddPass(
-			RDG_EVENT_NAME("Strata::OpaqueRoughRefraction(Blur,Horizontal)"),
+			RDG_EVENT_NAME("Substrate::OpaqueRoughRefraction(Blur,Horizontal)"),
 			PassParameters,
 			ERDGPassFlags::Raster,
 			[&View, TileVertexShader, PixelShader, PassParameters, StrataTileType, PrimitiveType](FRHICommandList& RHICmdList)
@@ -156,7 +156,7 @@ void AddStrataOpaqueRoughRefractionPasses(
 		PassParameters->SeparatedOpaqueRoughRefractionSceneColor = TempTexture;
 		PassParameters->RenderTargets[0] = FRenderTargetBinding(SceneColorTexture, ERenderTargetLoadAction::ELoad);
 		PassParameters->BlurDirection = FVector2f(0.0f, 1.0f);
-		PassParameters->BlurScale = FMath::Max(0.0f, CVarStrataOpaqueMaterialRoughRefractionBlurScale.GetValueOnAnyThread());
+		PassParameters->BlurScale = FMath::Max(0.0f, CVarSubstrateOpaqueMaterialRoughRefractionBlurScale.GetValueOnAnyThread());
 
 		FOpaqueRoughRefractionPS::FPermutationDomain PermutationVector;
 		PermutationVector.Set<FOpaqueRoughRefractionPS::FEnableBlur>(true);
@@ -171,7 +171,7 @@ void AddStrataOpaqueRoughRefractionPasses(
 		PassParameters->StrataTile = Strata::SetTileParameters(GraphBuilder, View, StrataTileType, PrimitiveType);
 
 		GraphBuilder.AddPass(
-			RDG_EVENT_NAME("Strata::OpaqueRoughRefraction(Blur,Vertical)"),
+			RDG_EVENT_NAME("Substrate::OpaqueRoughRefraction(Blur,Vertical)"),
 			PassParameters,
 			ERDGPassFlags::Raster,
 			[&View, TileVertexShader, PixelShader, PassParameters, StrataTileType, PrimitiveType](FRHICommandList& RHICmdList)
@@ -214,7 +214,7 @@ void AddStrataOpaqueRoughRefractionPasses(
 		PassParameters->SeparatedOpaqueRoughRefractionSceneColor = SeparatedOpaqueRoughRefractionSceneColor;
 		PassParameters->RenderTargets[0] = FRenderTargetBinding(SceneColorTexture, ERenderTargetLoadAction::ELoad);
 		PassParameters->BlurDirection = FVector2f(0.0f, 0.0f);
-		PassParameters->BlurScale = FMath::Max(0.0f, CVarStrataOpaqueMaterialRoughRefractionBlurScale.GetValueOnAnyThread());
+		PassParameters->BlurScale = FMath::Max(0.0f, CVarSubstrateOpaqueMaterialRoughRefractionBlurScale.GetValueOnAnyThread());
 
 		FOpaqueRoughRefractionPS::FPermutationDomain PermutationVector;
 		PermutationVector.Set<FOpaqueRoughRefractionPS::FEnableBlur>(false);
@@ -229,7 +229,7 @@ void AddStrataOpaqueRoughRefractionPasses(
 		PassParameters->StrataTile = Strata::SetTileParameters(GraphBuilder, View, StrataTileType, PrimitiveType);
 
 		GraphBuilder.AddPass(
-			RDG_EVENT_NAME("Strata::OpaqueRoughRefraction(SSS)"),
+			RDG_EVENT_NAME("Substrate::OpaqueRoughRefraction(SSS)"),
 			PassParameters,
 			ERDGPassFlags::Raster,
 			[&View, TileVertexShader, PixelShader, PassParameters, StrataTileType, PrimitiveType](FRHICommandList& RHICmdList)
@@ -265,15 +265,15 @@ void AddStrataOpaqueRoughRefractionPasses(
 
 #if STRATA_ROUGH_REFRACTION_RND
 
-static TAutoConsoleVariable<int32> CVarStrataRoughRefractionShadersShowRoughRefractionRnD(
-	TEXT("r.Strata.ShowRoughRefractionRnD"),
+static TAutoConsoleVariable<int32> CVarSubstrateRoughRefractionShadersShowRoughRefractionRnD(
+	TEXT("r.Substrate.ShowRoughRefractionRnD"),
 	1,
-	TEXT("Enable strata rough refraction shaders."),
+	TEXT("Enable Substrate rough refraction shaders."),
 	ECVF_RenderThreadSafe);
 
 bool ShouldRenderStrataRoughRefractionRnD()
 {
-	return CVarStrataRoughRefractionShadersShowRoughRefractionRnD.GetValueOnAnyThread() > 0;
+	return CVarSubstrateRoughRefractionShadersShowRoughRefractionRnD.GetValueOnAnyThread() > 0;
 }
 
 #else
@@ -423,11 +423,11 @@ void StrataRoughRefractionRnD(FRDGBuilder& GraphBuilder, const FViewInfo& View,
 		// Texture to count
 		const uint32 SampleCountTextureWidth = 64;
 		const FIntPoint SampleCountTextureSize(SampleCountTextureWidth, SampleCountTextureWidth);
-		FRDGTextureRef SampleCountTexture = GraphBuilder.CreateTexture(FRDGTextureDesc::Create2D(SampleCountTextureSize, PF_R32_UINT, FClearValueBinding::Black, TexCreate_UAV | TexCreate_ShaderResource), TEXT("Strata.RoughRefrac.SampleCount"));
+		FRDGTextureRef SampleCountTexture = GraphBuilder.CreateTexture(FRDGTextureDesc::Create2D(SampleCountTextureSize, PF_R32_UINT, FClearValueBinding::Black, TexCreate_UAV | TexCreate_ShaderResource), TEXT("Substrate.RoughRefrac.SampleCount"));
 		FRDGTextureUAVRef SampleCountTextureUAV = GraphBuilder.CreateUAV(FRDGTextureUAVDesc(SampleCountTexture));
 		FRDGTextureSRVRef SampleCountTextureSRV = GraphBuilder.CreateSRV(FRDGTextureSRVDesc::Create(SampleCountTexture));
 
-		FRDGBufferRef LobStatisticsBuffer = GraphBuilder.CreateBuffer(FRDGBufferDesc::CreateStructuredDesc(sizeof(float) * 8, 16), TEXT("Strata.RoughRefrac.LobStat"));
+		FRDGBufferRef LobStatisticsBuffer = GraphBuilder.CreateBuffer(FRDGBufferDesc::CreateStructuredDesc(sizeof(float) * 8, 16), TEXT("Substrate.RoughRefrac.LobStat"));
 		FRDGBufferUAVRef LobStatisticsBufferUAV = GraphBuilder.CreateUAV(LobStatisticsBuffer, PF_R32_UINT);
 		FRDGBufferSRVRef LobStatisticsBufferSRV = GraphBuilder.CreateSRV(LobStatisticsBuffer, PF_R32_UINT);
 
@@ -451,7 +451,7 @@ void StrataRoughRefractionRnD(FRDGBuilder& GraphBuilder, const FViewInfo& View,
 
 			FComputeShaderUtils::AddPass(
 				GraphBuilder,
-				RDG_EVENT_NAME("Strata::EvaluateRoughRefractionLobeCS"),
+				RDG_EVENT_NAME("Substrate::EvaluateRoughRefractionLobeCS"),
 				ComputeShader,
 				PassParameters,
 				FComputeShaderUtils::GetGroupCount(1, FEvaluateRoughRefractionLobeCS::ThreadGroupSize));
@@ -484,7 +484,7 @@ void StrataRoughRefractionRnD(FRDGBuilder& GraphBuilder, const FViewInfo& View,
 
 			FRHIBlendState* PreMultipliedColorTransmittanceBlend = TStaticBlendState<CW_RGB, BO_Add, BF_One, BF_SourceAlpha, BO_Add, BF_Zero, BF_One>::GetRHI();
 
-			FPixelShaderUtils::AddFullscreenPass<FVisualizeRoughRefractionPS>(GraphBuilder, View.ShaderMap, RDG_EVENT_NAME("Strata::VisualizeRoughRefractionPS"),
+			FPixelShaderUtils::AddFullscreenPass<FVisualizeRoughRefractionPS>(GraphBuilder, View.ShaderMap, RDG_EVENT_NAME("Substrate::VisualizeRoughRefractionPS"),
 				PixelShader, PassParameters, View.ViewRect, PreMultipliedColorTransmittanceBlend);
 		}
 
@@ -492,7 +492,7 @@ void StrataRoughRefractionRnD(FRDGBuilder& GraphBuilder, const FViewInfo& View,
 		// Sample variance as a function of roughness
 		{
 			const uint32 RoughRefracDataByteSize = sizeof(float) * 2 * FRoughRefracDataCS::ThreadGroupSize;
-			FRDGBufferRef RoughRefracDataBuffer = GraphBuilder.CreateBuffer(FRDGBufferDesc::CreateStructuredDesc(sizeof(float) * 2, FRoughRefracDataCS::ThreadGroupSize), TEXT("Strata.RoughRefrac.LobStat"));
+			FRDGBufferRef RoughRefracDataBuffer = GraphBuilder.CreateBuffer(FRDGBufferDesc::CreateStructuredDesc(sizeof(float) * 2, FRoughRefracDataCS::ThreadGroupSize), TEXT("Substrate.RoughRefrac.LobStat"));
 			FRDGBufferUAVRef RoughRefracDataBufferUAV = GraphBuilder.CreateUAV(RoughRefracDataBuffer, PF_G32R32F);
 
 			// Generate data in a buffer
@@ -506,7 +506,7 @@ void StrataRoughRefractionRnD(FRDGBuilder& GraphBuilder, const FViewInfo& View,
 
 				FComputeShaderUtils::AddPass(
 					GraphBuilder,
-					RDG_EVENT_NAME("Strata::FRoughRefracDataCS"),
+					RDG_EVENT_NAME("Substrate::FRoughRefracDataCS"),
 					ComputeShader,
 					PassParameters,
 					FComputeShaderUtils::GetGroupCount(FRoughRefracDataCS::ThreadGroupSize, FRoughRefracDataCS::ThreadGroupSize));
@@ -518,7 +518,7 @@ void StrataRoughRefractionRnD(FRDGBuilder& GraphBuilder, const FViewInfo& View,
 				PassParameters->Buffer = RoughRefracDataBuffer;
 
 				GraphBuilder.AddPass(
-					RDG_EVENT_NAME("Strata::FRoughRefracData ReadBack"),
+					RDG_EVENT_NAME("Substrate::FRoughRefracData ReadBack"),
 					PassParameters,
 					ERDGPassFlags::Readback,
 					[RoughRefracDataBuffer, RoughRefracDataByteSize](FRHICommandListImmediate& RHICmdList)
diff --git a/Engine/Source/Runtime/Renderer/Private/Strata/StrataVisualize.cpp b/Engine/Source/Runtime/Renderer/Private/Strata/StrataVisualize.cpp
--- a/Engine/Source/Runtime/Renderer/Private/Strata/StrataVisualize.cpp
+++ b/Engine/Source/Runtime/Renderer/Private/Strata/StrataVisualize.cpp
@@ -229,7 +229,7 @@ static void AddVisualizeMaterialPropertiesPasses(FRDGBuilder& GraphBuilder, cons
 	ShaderPrint::RequestSpaceForCharacters(1024);
 	FRDGBufferUAVRef PrintOffsetBufferUAV = nullptr;
 
-	FRDGBufferRef PrintOffsetBuffer = GraphBuilder.CreateBuffer(FRDGBufferDesc::CreateBufferDesc(4, 2), TEXT("Strata.DebugPrintPositionOffset"));
+	FRDGBufferRef PrintOffsetBuffer = GraphBuilder.CreateBuffer(FRDGBufferDesc::CreateBufferDesc(4, 2), TEXT("Substrate.DebugPrintPositionOffset"));
 	PrintOffsetBufferUAV = GraphBuilder.CreateUAV(PrintOffsetBuffer, PF_R32_UINT);
 	AddClearUAVPass(GraphBuilder, PrintOffsetBufferUAV, 50u);
 	const uint32 MaxBSDFCount = 8;
@@ -245,7 +245,7 @@ static void AddVisualizeMaterialPropertiesPasses(FRDGBuilder& GraphBuilder, cons
 		ShaderPrint::SetParameters(GraphBuilder, View.ShaderPrintData, PassParameters->ShaderPrintParameters);
 
 		TShaderMapRef<FMaterialPrintInfoCS> ComputeShader(View.ShaderMap);
-		FComputeShaderUtils::AddPass(GraphBuilder, RDG_EVENT_NAME("Strata::VisualizeMaterial(Print, BSDF=%d)", BSDFIndex), ComputeShader, PassParameters, FIntVector(1, 1, 1));
+		FComputeShaderUtils::AddPass(GraphBuilder, RDG_EVENT_NAME("Substrate::VisualizeMaterial(Print, BSDF=%d)", BSDFIndex), ComputeShader, PassParameters, FIntVector(1, 1, 1));
 	}
 }
 
@@ -269,7 +269,7 @@ static void AddVisualizeMaterialCountPasses(FRDGBuilder & GraphBuilder, const FV
 	FVisualizeMaterialCountPS::FPermutationDomain PermutationVector;
 	TShaderMapRef<FVisualizeMaterialCountPS> PixelShader(View.ShaderMap, PermutationVector);
 
-	FPixelShaderUtils::AddFullscreenPass<FVisualizeMaterialCountPS>(GraphBuilder, View.ShaderMap, RDG_EVENT_NAME("Strata::VisualizeMaterial(Draw)"), PixelShader, PassParameters, ScreenPassSceneColor.ViewRect, PreMultipliedColorTransmittanceBlend);
+	FPixelShaderUtils::AddFullscreenPass<FVisualizeMaterialCountPS>(GraphBuilder, View.ShaderMap, RDG_EVENT_NAME("Substrate::VisualizeMaterial(Draw)"), PixelShader, PassParameters, ScreenPassSceneColor.ViewRect, PreMultipliedColorTransmittanceBlend);
 }
 
 float GetStrataTileOverflowRatio(const FViewInfo& View);
@@ -302,7 +302,7 @@ static void AddVisualizeSystemInfoPasses(FRDGBuilder& GraphBuilder, const FViewI
 	ShaderPrint::SetParameters(GraphBuilder, View.ShaderPrintData, PassParameters->ShaderPrintParameters);
 
 	TShaderMapRef<FStrataSystemInfoCS> ComputeShader(View.ShaderMap);
-	FComputeShaderUtils::AddPass(GraphBuilder, RDG_EVENT_NAME("Strata::VisualizeSystemInfo"), ComputeShader, PassParameters, FIntVector(1, 1, 1));
+	FComputeShaderUtils::AddPass(GraphBuilder, RDG_EVENT_NAME("Substrate::VisualizeSystemInfo"), ComputeShader, PassParameters, FIntVector(1, 1, 1));
 }
 
 // Draw each material layer independently
@@ -326,7 +326,7 @@ static void AddVisualizeAdvancedMaterialPasses(FRDGBuilder& GraphBuilder, const
 		ShaderPrint::SetParameters(GraphBuilder, View.ShaderPrintData, PassParameters->ShaderPrintParameters);
 
 		TShaderMapRef<FMaterialDebugStrataTreeCS> ComputeShader(View.ShaderMap);
-		FComputeShaderUtils::AddPass(GraphBuilder, RDG_EVENT_NAME("Strata::StrataAdvancedVisualization(Print)"), ComputeShader, PassParameters, FIntVector(1, 1, 1));
+		FComputeShaderUtils::AddPass(GraphBuilder, RDG_EVENT_NAME("Substrate::StrataAdvancedVisualization(Print)"), ComputeShader, PassParameters, FIntVector(1, 1, 1));
 	}
 
 	{
@@ -345,7 +345,7 @@ static void AddVisualizeAdvancedMaterialPasses(FRDGBuilder& GraphBuilder, const
 		FMaterialDebugStrataTreePS::FPermutationDomain PermutationVector;
 		TShaderMapRef<FMaterialDebugStrataTreePS> PixelShader(View.ShaderMap, PermutationVector);
 
-		FPixelShaderUtils::AddFullscreenPass<FMaterialDebugStrataTreePS>(GraphBuilder, View.ShaderMap, RDG_EVENT_NAME("Strata::StrataAdvancedVisualization(Draw)"), PixelShader, PassParameters, ScreenPassSceneColor.ViewRect, PreMultipliedColorTransmittanceBlend);
+		FPixelShaderUtils::AddFullscreenPass<FMaterialDebugStrataTreePS>(GraphBuilder, View.ShaderMap, RDG_EVENT_NAME("Substrate::StrataAdvancedVisualization(Draw)"), PixelShader, PassParameters, ScreenPassSceneColor.ViewRect, PreMultipliedColorTransmittanceBlend);
 	}
 }
 
@@ -358,8 +358,8 @@ static FStrataVisualizationData::FViewMode GetStrataVisualizeMode(const FViewInf
 	if (IsStrataEnabled() && StrataDebugVisualizationCanRunOnPlatform(View.GetShaderPlatform()))
 	{
 		// Variable defined in StrataVisualizationData.h/.cpp
-		static const auto CVarStrataViewMode = IConsoleManager::Get().FindConsoleVariable(FStrataVisualizationData::GetVisualizeConsoleCommandName());
-		const uint32 ViewMode = CVarStrataViewMode && CVarStrataViewMode->AsVariableInt() ? CVarStrataViewMode->AsVariableInt()->GetValueOnRenderThread() : 0;
+		static const auto CVar = IConsoleManager::Get().FindConsoleVariable(FStrataVisualizationData::GetVisualizeConsoleCommandName());
+		const uint32 ViewMode = CVar && CVar->AsVariableInt() ? CVar->AsVariableInt()->GetValueOnRenderThread() : 0;
 		switch (ViewMode)
 		{
 			case 1: return FStrataVisualizationData::FViewMode::MaterialProperties;
@@ -372,7 +372,7 @@ static FStrataVisualizationData::FViewMode GetStrataVisualizeMode(const FViewInf
 		}
 
 		const FStrataVisualizationData& VisualizationData = GetStrataVisualizationData();
-		if (View.Family && View.Family->EngineShowFlags.VisualizeStrata)
+		if (View.Family && View.Family->EngineShowFlags.VisualizeSubstrate)
 		{
 			Out = VisualizationData.GetViewMode(View.CurrentStrataVisualizationMode);
 		}
@@ -392,7 +392,7 @@ FScreenPassTexture AddStrataDebugPasses(FRDGBuilder& GraphBuilder, const FViewIn
 	const FStrataVisualizationData::FViewMode DebugMode = GetStrataVisualizeMode(View);
 	if (DebugMode != FStrataVisualizationData::FViewMode::None)
 	{
-		RDG_EVENT_SCOPE(GraphBuilder, "Strata::VisualizeMaterial");
+		RDG_EVENT_SCOPE(GraphBuilder, "Substrate::VisualizeMaterial");
 
 		const bool bDebugPass = true;
 		if (DebugMode == FStrataVisualizationData::FViewMode::MaterialProperties)
diff --git a/Engine/Source/Runtime/Renderer/Private/VirtualShadowMaps/VirtualShadowMapProjection.cpp b/Engine/Source/Runtime/Renderer/Private/VirtualShadowMaps/VirtualShadowMapProjection.cpp
--- a/Engine/Source/Runtime/Renderer/Private/VirtualShadowMaps/VirtualShadowMapProjection.cpp
+++ b/Engine/Source/Runtime/Renderer/Private/VirtualShadowMaps/VirtualShadowMapProjection.cpp
@@ -171,7 +171,7 @@ const TCHAR* ToString(EVirtualShadowMapProjectionInputType In)
 	switch (In)
 	{
 	case EVirtualShadowMapProjectionInputType::HairStrands: return TEXT("HairStrands");
-	case EVirtualShadowMapProjectionInputType::GBuffer:     return Strata::IsStrataEnabled() ? TEXT("Strata") : TEXT("GBuffer");
+	case EVirtualShadowMapProjectionInputType::GBuffer:     return Strata::IsStrataEnabled() ? TEXT("Substrate") : TEXT("GBuffer");
 	}
 	return TEXT("Invalid");
 }
