Delete some deprecated code in PreReplication

diff --git a/Engine/Source/Runtime/Engine/Classes/Components/ActorComponent.h b/Engine/Source/Runtime/Engine/Classes/Components/ActorComponent.h
--- a/Engine/Source/Runtime/Engine/Classes/Components/ActorComponent.h
+++ b/Engine/Source/Runtime/Engine/Classes/Components/ActorComponent.h
@@ -510,7 +510,7 @@ public:
 	virtual ELifetimeCondition GetReplicationCondition() const { return GetIsReplicated() ? COND_None : COND_Never;  }
 
 	/** Called on the component right before replication occurs */
-	virtual void PreReplication(IRepChangedPropertyTracker & ChangedPropertyTracker);
+	virtual void PreReplication(IRepChangedPropertyTracker & ChangedPropertyTracker) {}
 
 	/** Returns true if this type of component can ever replicate, override to disable the default behavior */
 	virtual bool GetComponentClassCanReplicate() const;
diff --git a/Engine/Source/Runtime/Engine/Classes/Engine/BlueprintGeneratedClass.h b/Engine/Source/Runtime/Engine/Classes/Engine/BlueprintGeneratedClass.h
--- a/Engine/Source/Runtime/Engine/Classes/Engine/BlueprintGeneratedClass.h
+++ b/Engine/Source/Runtime/Engine/Classes/Engine/BlueprintGeneratedClass.h
@@ -895,17 +895,6 @@ public:
 	/** called to gather blueprint replicated properties */
 	virtual void GetLifetimeBlueprintReplicationList(TArray<class FLifetimeProperty>& OutLifetimeProps) const;
 
-	UE_DEPRECATED(5.0, "To be removed in a future release.")
-	/** called prior to replication of an instance of this BP class */
-	virtual void InstancePreReplication(UObject* Obj, class IRepChangedPropertyTracker& ChangedPropertyTracker) const
-	{
-		UBlueprintGeneratedClass* SuperBPClass = Cast<UBlueprintGeneratedClass>(GetSuperStruct());
-		if (SuperBPClass != NULL)
-		{
-			SuperBPClass->InstancePreReplication(Obj, ChangedPropertyTracker);
-		}
-	}
-
 protected:
 	/** Internal helper method used to recursively build the custom property list that's used for post-construct initialization. */
 	bool BuildCustomPropertyListForPostConstruction(FCustomPropertyListNode*& InPropertyList, UStruct* InStruct, const uint8* DataPtr, const uint8* DefaultDataPtr);
diff --git a/Engine/Source/Runtime/Engine/Private/Actor.cpp b/Engine/Source/Runtime/Engine/Private/Actor.cpp
--- a/Engine/Source/Runtime/Engine/Private/Actor.cpp
+++ b/Engine/Source/Runtime/Engine/Private/Actor.cpp
@@ -1505,14 +1505,6 @@ void AActor::PreReplication(IRepChangedPropertyTracker & ChangedPropertyTracker)
 		MARK_PROPERTY_DIRTY_FROM_NAME(AActor, AttachmentReplication, this);
 	}
 #endif
-
-	PRAGMA_DISABLE_DEPRECATION_WARNINGS
-	UBlueprintGeneratedClass* BPClass = Cast<UBlueprintGeneratedClass>(GetClass());
-	if (BPClass != nullptr)
-	{
-		BPClass->InstancePreReplication(this, ChangedPropertyTracker);
-	}
-	PRAGMA_ENABLE_DEPRECATION_WARNINGS
 }
 
 void AActor::CallPreReplication(UNetDriver* NetDriver)
diff --git a/Engine/Source/Runtime/Engine/Private/Components/ActorComponent.cpp b/Engine/Source/Runtime/Engine/Private/Components/ActorComponent.cpp
--- a/Engine/Source/Runtime/Engine/Private/Components/ActorComponent.cpp
+++ b/Engine/Source/Runtime/Engine/Private/Components/ActorComponent.cpp
@@ -2177,17 +2177,6 @@ bool UActorComponent::ReplicateSubobjects(UActorChannel* Channel, FOutBunch* Bun
 	return false;
 }
 
-void UActorComponent::PreReplication(IRepChangedPropertyTracker & ChangedPropertyTracker)
-{
-	PRAGMA_DISABLE_DEPRECATION_WARNINGS
-	UBlueprintGeneratedClass* BPClass = Cast<UBlueprintGeneratedClass>(GetClass());
-	if (BPClass != NULL)
-	{
-		BPClass->InstancePreReplication(this, ChangedPropertyTracker);
-	}
-	PRAGMA_ENABLE_DEPRECATION_WARNINGS
-}
-
 bool UActorComponent::GetComponentClassCanReplicate() const
 {
 	return true;
