FExec: Wrapped remaining Exec overrides with UE_ALLOW_EXEC_COMMANDS so that we can compile the engine without exec command support when desired.

diff --git a/Engine/Plugins/Experimental/WebSocketNetworking/Source/WebSocketNetworking/Private/WebSocketNetDriver.cpp b/Engine/Plugins/Experimental/WebSocketNetworking/Source/WebSocketNetworking/Private/WebSocketNetDriver.cpp
--- a/Engine/Plugins/Experimental/WebSocketNetworking/Source/WebSocketNetworking/Private/WebSocketNetDriver.cpp
+++ b/Engine/Plugins/Experimental/WebSocketNetworking/Source/WebSocketNetworking/Private/WebSocketNetDriver.cpp
@@ -190,6 +190,7 @@ bool UWebSocketNetDriver::HandleSocketsCommand(const TCHAR* Cmd, FOutputDevice&
 	return UNetDriver::Exec(InWorld, TEXT("SOCKETS"), Ar);
 }
 
+#if UE_ALLOW_EXEC_COMMANDS
 bool UWebSocketNetDriver::Exec(UWorld* InWorld, const TCHAR* Cmd, FOutputDevice& Ar)
 {
 	if (FParse::Command(&Cmd, TEXT("SOCKETS")))
@@ -198,6 +199,7 @@ bool UWebSocketNetDriver::Exec(UWorld* InWorld, const TCHAR* Cmd, FOutputDevice&
 	}
 	return UNetDriver::Exec(InWorld, Cmd, Ar);
 }
+#endif // UE_ALLOW_EXEC_COMMANDS
 
 UWebSocketConnection* UWebSocketNetDriver::GetServerConnection()
 {
diff --git a/Engine/Plugins/Experimental/WebSocketNetworking/Source/WebSocketNetworking/Public/WebSocketNetDriver.h b/Engine/Plugins/Experimental/WebSocketNetworking/Source/WebSocketNetworking/Public/WebSocketNetDriver.h
--- a/Engine/Plugins/Experimental/WebSocketNetworking/Source/WebSocketNetworking/Public/WebSocketNetDriver.h
+++ b/Engine/Plugins/Experimental/WebSocketNetworking/Source/WebSocketNetworking/Public/WebSocketNetDriver.h
@@ -38,7 +38,9 @@ class WEBSOCKETNETWORKING_API UWebSocketNetDriver : public UNetDriver
 	//~ End UNetDriver Interface.
 
 	//~ Begin FExec Interface
+#if UE_ALLOW_EXEC_COMMANDS
 	virtual bool Exec(UWorld* InWorld, const TCHAR* Cmd, FOutputDevice& Ar = *GLog) override;
+#endif
 	//~ End FExec Interface
 
 	/**
diff --git a/Engine/Plugins/Online/OnlineFramework/Source/Party/Private/SocialManager.cpp b/Engine/Plugins/Online/OnlineFramework/Source/Party/Private/SocialManager.cpp
--- a/Engine/Plugins/Online/OnlineFramework/Source/Party/Private/SocialManager.cpp
+++ b/Engine/Plugins/Online/OnlineFramework/Source/Party/Private/SocialManager.cpp
@@ -1401,6 +1401,7 @@ void USocialManager::HandleFindSessionForJoinComplete(bool bWasSuccessful, const
 	}
 }
 
+#if UE_ALLOW_EXEC_COMMANDS
 bool USocialManager::Exec(class UWorld* InWorld, const TCHAR* Cmd, FOutputDevice& Out)
 {
 	if (FParse::Command(&Cmd, TEXT("SOCIAL")))
@@ -1423,6 +1424,7 @@ bool USocialManager::Exec(class UWorld* InWorld, const TCHAR* Cmd, FOutputDevice
 	}
 	return false;
 }
+#endif // UE_ALLOW_EXEC_COMMANDS
 
 USocialDebugTools* USocialManager::GetDebugTools() const
 {
diff --git a/Engine/Plugins/Online/OnlineFramework/Source/Party/Public/SocialManager.h b/Engine/Plugins/Online/OnlineFramework/Source/Party/Public/SocialManager.h
--- a/Engine/Plugins/Online/OnlineFramework/Source/Party/Public/SocialManager.h
+++ b/Engine/Plugins/Online/OnlineFramework/Source/Party/Public/SocialManager.h
@@ -40,7 +40,9 @@ class PARTY_API USocialManager : public UObject, public FExec
 
 public:
 	// FExec
+#if UE_ALLOW_EXEC_COMMANDS
 	virtual bool Exec(class UWorld* InWorld, const TCHAR* Cmd, FOutputDevice& Out) override;
+#endif
 
 	static bool IsSocialSubsystemEnabled(ESocialSubsystem SubsystemType);
 	static FName GetSocialOssName(ESocialSubsystem SubsystemType);
diff --git a/Engine/Plugins/Online/OnlineServices/Source/OnlineServicesCommon/Private/Online/OnlineServicesCommon.cpp b/Engine/Plugins/Online/OnlineServices/Source/OnlineServicesCommon/Private/Online/OnlineServicesCommon.cpp
--- a/Engine/Plugins/Online/OnlineServices/Source/OnlineServicesCommon/Private/Online/OnlineServicesCommon.cpp
+++ b/Engine/Plugins/Online/OnlineServices/Source/OnlineServicesCommon/Private/Online/OnlineServicesCommon.cpp
@@ -193,6 +193,7 @@ void FOnlineServicesCommon::RegisterExecHandler(const FString& Name, TUniquePtr<
 	ExecCommands.Emplace(Name, MoveTemp(Handler));
 }
 
+#if UE_ALLOW_EXEC_COMMANDS
 bool FOnlineServicesCommon::Exec(UWorld* World, const TCHAR* Cmd, FOutputDevice& Ar)
 {
 	if (FParse::Command(&Cmd, TEXT("OnlineServices")))
@@ -218,5 +219,6 @@ bool FOnlineServicesCommon::Exec(UWorld* World, const TCHAR* Cmd, FOutputDevice&
 	}
 	return false;
 }
+#endif // UE_ALLOW_EXEC_COMMANDS
 
 /* UE::Online */ }
diff --git a/Engine/Plugins/Online/OnlineServices/Source/OnlineServicesCommon/Public/Online/OnlineServicesCommon.h b/Engine/Plugins/Online/OnlineServices/Source/OnlineServicesCommon/Public/Online/OnlineServicesCommon.h
--- a/Engine/Plugins/Online/OnlineServices/Source/OnlineServicesCommon/Public/Online/OnlineServicesCommon.h
+++ b/Engine/Plugins/Online/OnlineServices/Source/OnlineServicesCommon/Public/Online/OnlineServicesCommon.h
@@ -312,7 +312,9 @@ public:
 	
 	void RegisterExecHandler(const FString& Name, TUniquePtr<IOnlineExecHandler>&& Handler);
 
+#if UE_ALLOW_EXEC_COMMANDS
 	virtual bool Exec(UWorld* World, const TCHAR* Cmd, FOutputDevice& Ar) override;
+#endif
 
 	FOnlineAsyncOpCache OpCache;
 
diff --git a/Engine/Plugins/Online/OnlineSubsystemUtils/Source/OnlineSubsystemUtils/Classes/IpNetDriver.h b/Engine/Plugins/Online/OnlineSubsystemUtils/Source/OnlineSubsystemUtils/Classes/IpNetDriver.h
--- a/Engine/Plugins/Online/OnlineSubsystemUtils/Source/OnlineSubsystemUtils/Classes/IpNetDriver.h
+++ b/Engine/Plugins/Online/OnlineSubsystemUtils/Source/OnlineSubsystemUtils/Classes/IpNetDriver.h
@@ -203,7 +203,9 @@ protected:
 
 public:
 	//~ Begin FExec Interface
+#if UE_ALLOW_EXEC_COMMANDS
 	virtual bool Exec( UWorld* InWorld, const TCHAR* Cmd, FOutputDevice& Ar=*GLog ) override;
+#endif
 	//~ End FExec Interface
 
 
diff --git a/Engine/Plugins/Online/OnlineSubsystemUtils/Source/OnlineSubsystemUtils/Private/IpNetDriver.cpp b/Engine/Plugins/Online/OnlineSubsystemUtils/Source/OnlineSubsystemUtils/Private/IpNetDriver.cpp
--- a/Engine/Plugins/Online/OnlineSubsystemUtils/Source/OnlineSubsystemUtils/Private/IpNetDriver.cpp
+++ b/Engine/Plugins/Online/OnlineSubsystemUtils/Source/OnlineSubsystemUtils/Private/IpNetDriver.cpp
@@ -1710,6 +1710,7 @@ void UIpNetDriver::TestSuddenPortChange(uint32 NumConnections)
 }
 #endif
 
+#if UE_ALLOW_EXEC_COMMANDS
 bool UIpNetDriver::Exec( UWorld* InWorld, const TCHAR* Cmd, FOutputDevice& Ar )
 {
 	if (FParse::Command(&Cmd,TEXT("SOCKETS")))
@@ -1727,6 +1728,7 @@ bool UIpNetDriver::Exec( UWorld* InWorld, const TCHAR* Cmd, FOutputDevice& Ar )
 
 	return UNetDriver::Exec( InWorld, Cmd,Ar);
 }
+#endif // UE_ALLOW_EXEC_COMMANDS
 
 UIpConnection* UIpNetDriver::GetServerConnection() 
 {
diff --git a/Engine/Source/Editor/UnrealEd/Classes/Editor/UnrealEdEngine.h b/Engine/Source/Editor/UnrealEd/Classes/Editor/UnrealEdEngine.h
--- a/Engine/Source/Editor/UnrealEd/Classes/Editor/UnrealEdEngine.h
+++ b/Engine/Source/Editor/UnrealEd/Classes/Editor/UnrealEdEngine.h
@@ -198,7 +198,9 @@ public:
 	//~ End UEditorEngine Interface 
 	
 	//~ Begin FExec Interface
+#if UE_ALLOW_EXEC_COMMANDS
 	virtual bool Exec( UWorld* InWorld, const TCHAR* Cmd, FOutputDevice& Ar=*GLog ) override;
+#endif
 	//~ End FExec Interface
 
 
diff --git a/Engine/Source/Editor/UnrealEd/Private/UnrealEdSrv.cpp b/Engine/Source/Editor/UnrealEd/Private/UnrealEdSrv.cpp
--- a/Engine/Source/Editor/UnrealEd/Private/UnrealEdSrv.cpp
+++ b/Engine/Source/Editor/UnrealEd/Private/UnrealEdSrv.cpp
@@ -629,6 +629,7 @@ bool UUnrealEdEngine::HandleDisasmScriptCommand(const TCHAR* Str, FOutputDevice&
 	return true;
 }
 
+#if UE_ALLOW_EXEC_COMMANDS
 bool UUnrealEdEngine::Exec( UWorld* InWorld, const TCHAR* Stream, FOutputDevice& Ar )
 {
 	const TCHAR* Str = Stream;
@@ -1266,6 +1267,7 @@ bool UUnrealEdEngine::Exec( UWorld* InWorld, const TCHAR* Stream, FOutputDevice&
 	}
 	return false;
 }
+#endif // UE_ALLOW_EXEC_COMMANDS
 
 bool UUnrealEdEngine::AnyWorldsAreDirty( UWorld* InWorld ) const
 {
diff --git a/Engine/Source/Runtime/AudioMixer/Private/AudioMixerDevice.cpp b/Engine/Source/Runtime/AudioMixer/Private/AudioMixerDevice.cpp
--- a/Engine/Source/Runtime/AudioMixer/Private/AudioMixerDevice.cpp
+++ b/Engine/Source/Runtime/AudioMixer/Private/AudioMixerDevice.cpp
@@ -734,6 +734,7 @@ namespace Audio
 		return false;
 	}
 
+#if UE_ALLOW_EXEC_COMMANDS
 	bool FMixerDevice::Exec(UWorld* InWorld, const TCHAR* Cmd, FOutputDevice& Ar)
 	{
 		if (FAudioDevice::Exec(InWorld, Cmd, Ar))
@@ -743,6 +744,7 @@ namespace Audio
 
 		return false;
 	}
+#endif // UE_ALLOW_EXEC_COMMANDS
 
 	void FMixerDevice::CountBytes(FArchive& InArchive)
 	{
diff --git a/Engine/Source/Runtime/AudioMixer/Public/AudioMixerDevice.h b/Engine/Source/Runtime/AudioMixer/Public/AudioMixerDevice.h
--- a/Engine/Source/Runtime/AudioMixer/Public/AudioMixerDevice.h
+++ b/Engine/Source/Runtime/AudioMixer/Public/AudioMixerDevice.h
@@ -115,7 +115,9 @@ namespace Audio
 		virtual class ICompressedAudioInfo* CreateCompressedAudioInfo(const USoundWave* SoundWave) const override;
 		virtual class ICompressedAudioInfo* CreateCompressedAudioInfo(const FSoundWaveProxyPtr& SoundWave) const override;
 		virtual bool ValidateAPICall(const TCHAR* Function, uint32 ErrorCode) override;
+#if UE_ALLOW_EXEC_COMMANDS
 		virtual bool Exec(UWorld* InWorld, const TCHAR* Cmd, FOutputDevice& Ar) override;
+#endif
 		virtual void CountBytes(class FArchive& Ar) override;
 		virtual bool IsExernalBackgroundSoundActive() override;
 		virtual void ResumeContext() override;
diff --git a/Engine/Source/Runtime/Core/Private/HAL/MallocDoubleFreeFinder.cpp b/Engine/Source/Runtime/Core/Private/HAL/MallocDoubleFreeFinder.cpp
--- a/Engine/Source/Runtime/Core/Private/HAL/MallocDoubleFreeFinder.cpp
+++ b/Engine/Source/Runtime/Core/Private/HAL/MallocDoubleFreeFinder.cpp
@@ -134,6 +134,7 @@ void FMallocDoubleFreeFinder::TrackSpecial(void* Ptr)
 	}
 }
 
+#if UE_ALLOW_EXEC_COMMANDS
 bool FMallocDoubleFreeFinder::Exec(UWorld* InWorld, const TCHAR* Cmd, FOutputDevice& Ar)
 {
 	if (FParse::Command(&Cmd, TEXT("DoubleFreeFinderCrash")))
@@ -147,7 +148,7 @@ bool FMallocDoubleFreeFinder::Exec(UWorld* InWorld, const TCHAR* Cmd, FOutputDev
 
 	return UsedMalloc->Exec(InWorld, Cmd, Ar);
 }
-
+#endif // UE_ALLOW_EXEC_COMMANDS
 
 FMalloc* FMallocDoubleFreeFinder::OverrideIfEnabled(FMalloc*InUsedAlloc)
 {
diff --git a/Engine/Source/Runtime/Core/Private/HAL/MallocFrameProfiler.cpp b/Engine/Source/Runtime/Core/Private/HAL/MallocFrameProfiler.cpp
--- a/Engine/Source/Runtime/Core/Private/HAL/MallocFrameProfiler.cpp
+++ b/Engine/Source/Runtime/Core/Private/HAL/MallocFrameProfiler.cpp
@@ -163,6 +163,7 @@ void FMallocFrameProfiler::UpdateStats()
 	CallStackStatsArray.Reset();
 }
 
+#if UE_ALLOW_EXEC_COMMANDS
 bool FMallocFrameProfiler::Exec(UWorld* InWorld, const TCHAR* Cmd, FOutputDevice& Ar)
 {
 	if (FParse::Command(&Cmd, TEXT("MallocFrameProfiler")))
@@ -183,7 +184,7 @@ bool FMallocFrameProfiler::Exec(UWorld* InWorld, const TCHAR* Cmd, FOutputDevice
 
 	return UsedMalloc->Exec(InWorld, Cmd, Ar);
 }
-
+#endif // UE_ALLOW_EXEC_COMMANDS
 
 FMalloc* FMallocFrameProfiler::OverrideIfEnabled(FMalloc*InUsedAlloc)
 {
diff --git a/Engine/Source/Runtime/Core/Private/HAL/UnrealMemory.cpp b/Engine/Source/Runtime/Core/Private/HAL/UnrealMemory.cpp
--- a/Engine/Source/Runtime/Core/Private/HAL/UnrealMemory.cpp
+++ b/Engine/Source/Runtime/Core/Private/HAL/UnrealMemory.cpp
@@ -180,10 +180,12 @@ public:
 		return(UsedMalloc->ValidateHeap());
 	}
 
+#if UE_ALLOW_EXEC_COMMANDS
 	virtual bool Exec(UWorld* InWorld, const TCHAR* Cmd, FOutputDevice& Ar) override
 	{
 		return UsedMalloc->Exec(InWorld, Cmd, Ar);
 	}
+#endif // UE_ALLOW_EXEC_COMMANDS
 
 	/**
 	* If possible determine the size of the memory allocated at the given address
diff --git a/Engine/Source/Runtime/Core/Public/HAL/MallocDoubleFreeFinder.h b/Engine/Source/Runtime/Core/Public/HAL/MallocDoubleFreeFinder.h
--- a/Engine/Source/Runtime/Core/Public/HAL/MallocDoubleFreeFinder.h
+++ b/Engine/Source/Runtime/Core/Public/HAL/MallocDoubleFreeFinder.h
@@ -15,10 +15,12 @@ class CORE_API FMallocDoubleFreeFinder final : public FMallocCallstackHandler
 public:
 	FMallocDoubleFreeFinder(FMalloc* InMalloc);
 
+#if UE_ALLOW_EXEC_COMMANDS
 	/**
 	 * Handles any commands passed in on the command line
 	 */
 	virtual bool Exec(UWorld* InWorld, const TCHAR* Cmd, FOutputDevice& Ar) override;
+#endif
 
 	/**
 	 * If you get an allocation/memory error outside of the allocator you can call this directly
diff --git a/Engine/Source/Runtime/Core/Public/HAL/MallocFrameProfiler.h b/Engine/Source/Runtime/Core/Public/HAL/MallocFrameProfiler.h
--- a/Engine/Source/Runtime/Core/Public/HAL/MallocFrameProfiler.h
+++ b/Engine/Source/Runtime/Core/Public/HAL/MallocFrameProfiler.h
@@ -18,10 +18,12 @@ public:
 
 	virtual void Init() override;
 	
+#if UE_ALLOW_EXEC_COMMANDS
 	/**
 	 * Handles any commands passed in on the command line
 	 */
 	virtual bool Exec(UWorld* InWorld, const TCHAR* Cmd, FOutputDevice& Ar) override;
+#endif
 
 	/** 
 	* Called once per frame, gathers and sets all memory allocator statistics into the corresponding stats. MUST BE THREAD SAFE.
diff --git a/Engine/Source/Runtime/Core/Public/HAL/MallocPoisonProxy.h b/Engine/Source/Runtime/Core/Public/HAL/MallocPoisonProxy.h
--- a/Engine/Source/Runtime/Core/Public/HAL/MallocPoisonProxy.h
+++ b/Engine/Source/Runtime/Core/Public/HAL/MallocPoisonProxy.h
@@ -114,10 +114,12 @@ public:
 		return UsedMalloc->ValidateHeap();
 	}
 
+#if UE_ALLOW_EXEC_COMMANDS
 	virtual bool Exec( UWorld* InWorld, const TCHAR* Cmd, FOutputDevice& Ar ) override
 	{
 		return UsedMalloc->Exec(InWorld, Cmd, Ar);
 	}
+#endif // UE_ALLOW_EXEC_COMMANDS
 
 	virtual bool GetAllocationSize(void *Original, SIZE_T &SizeOut) override
 	{
diff --git a/Engine/Source/Runtime/Core/Public/HAL/MallocStomp.h b/Engine/Source/Runtime/Core/Public/HAL/MallocStomp.h
--- a/Engine/Source/Runtime/Core/Public/HAL/MallocStomp.h
+++ b/Engine/Source/Runtime/Core/Public/HAL/MallocStomp.h
@@ -127,10 +127,12 @@ public:
 		return true;
 	}
 
+#if UE_ALLOW_EXEC_COMMANDS
 	virtual bool Exec( UWorld* InWorld, const TCHAR* Cmd, FOutputDevice& Ar ) override
 	{
 		return false;
 	}
+#endif
 
 	virtual const TCHAR* GetDescriptiveName() override
 	{
diff --git a/Engine/Source/Runtime/Core/Public/HAL/MallocThreadSafeProxy.h b/Engine/Source/Runtime/Core/Public/HAL/MallocThreadSafeProxy.h
--- a/Engine/Source/Runtime/Core/Public/HAL/MallocThreadSafeProxy.h
+++ b/Engine/Source/Runtime/Core/Public/HAL/MallocThreadSafeProxy.h
@@ -87,11 +87,13 @@ public:
 		return( UsedMalloc->ValidateHeap() );
 	}
 
+#if UE_ALLOW_EXEC_COMMANDS
 	virtual bool Exec( UWorld* InWorld, const TCHAR* Cmd, FOutputDevice& Ar ) override
 	{
 		FScopeLock ScopeLock( &SynchronizationObject );
 		return UsedMalloc->Exec( InWorld, Cmd, Ar);
 	}
+#endif // UE_ALLOW_EXEC_COMMANDS
 
 	/**
 	* If possible determine the size of the memory allocated at the given address
diff --git a/Engine/Source/Runtime/Core/Public/HAL/MemoryBase.h b/Engine/Source/Runtime/Core/Public/HAL/MemoryBase.h
--- a/Engine/Source/Runtime/Core/Public/HAL/MemoryBase.h
+++ b/Engine/Source/Runtime/Core/Public/HAL/MemoryBase.h
@@ -162,6 +162,7 @@ public:
 	*/
 	virtual void InitializeStatsMetadata();
 
+#if UE_ALLOW_EXEC_COMMANDS
 	/**
 	 * Handles any commands passed in on the command line
 	 */
@@ -169,6 +170,7 @@ public:
 	{ 
 		return false; 
 	}
+#endif // UE_ALLOW_EXEC_COMMANDS
 
 	/** Called once per frame, gathers and sets all memory allocator statistics into the corresponding stats. MUST BE THREAD SAFE. */
 	virtual void UpdateStats();
diff --git a/Engine/Source/Runtime/Core/Public/Stats/StatsMallocProfilerProxy.h b/Engine/Source/Runtime/Core/Public/Stats/StatsMallocProfilerProxy.h
--- a/Engine/Source/Runtime/Core/Public/Stats/StatsMallocProfilerProxy.h
+++ b/Engine/Source/Runtime/Core/Public/Stats/StatsMallocProfilerProxy.h
@@ -118,10 +118,12 @@ public:
 		return UsedMalloc->ValidateHeap();
 	}
 
+#if UE_ALLOW_EXEC_COMMANDS
 	virtual bool Exec( UWorld* InWorld, const TCHAR* Cmd, FOutputDevice& Ar ) override
 	{
 		return UsedMalloc->Exec( InWorld, Cmd, Ar);
 	}
+#endif
 
 	virtual bool GetAllocationSize(void *Original, SIZE_T &SizeOut) override
 	{
diff --git a/Engine/Source/Runtime/Engine/Classes/Engine/GameEngine.h b/Engine/Source/Runtime/Engine/Classes/Engine/GameEngine.h
--- a/Engine/Source/Runtime/Engine/Classes/Engine/GameEngine.h
+++ b/Engine/Source/Runtime/Engine/Classes/Engine/GameEngine.h
@@ -143,8 +143,9 @@ public:
 public:
 
 	// FExec interface
-
+#if UE_ALLOW_EXEC_COMMANDS
 	virtual bool Exec( UWorld* InWorld, const TCHAR* Cmd, FOutputDevice& Ar=*GLog ) override;
+#endif
 
 public:
 
diff --git a/Engine/Source/Runtime/Engine/Classes/Engine/GameInstance.h b/Engine/Source/Runtime/Engine/Classes/Engine/GameInstance.h
--- a/Engine/Source/Runtime/Engine/Classes/Engine/GameInstance.h
+++ b/Engine/Source/Runtime/Engine/Classes/Engine/GameInstance.h
@@ -210,7 +210,9 @@ public:
 #endif
 
 	//~ Begin FExec Interface
+#if UE_ALLOW_EXEC_COMMANDS
 	virtual bool Exec(UWorld* InWorld, const TCHAR* Cmd, FOutputDevice& Out = *GLog) override;
+#endif
 	//~ End FExec Interface
 
 	//~ Begin UObject Interface
diff --git a/Engine/Source/Runtime/Engine/Classes/Engine/NetConnection.h b/Engine/Source/Runtime/Engine/Classes/Engine/NetConnection.h
--- a/Engine/Source/Runtime/Engine/Classes/Engine/NetConnection.h
+++ b/Engine/Source/Runtime/Engine/Classes/Engine/NetConnection.h
@@ -916,9 +916,9 @@ public:
 
 
 	//~ Begin FExec Interface.
-
+#if UE_ALLOW_EXEC_COMMANDS
 	ENGINE_API virtual bool Exec( UWorld* InWorld, const TCHAR* Cmd, FOutputDevice& Ar=*GLog ) override;
-
+#endif
 	//~ End FExec Interface.
 
 	/** read input */
diff --git a/Engine/Source/Runtime/Engine/Classes/Engine/Player.h b/Engine/Source/Runtime/Engine/Classes/Engine/Player.h
--- a/Engine/Source/Runtime/Engine/Classes/Engine/Player.h
+++ b/Engine/Source/Runtime/Engine/Classes/Engine/Player.h
@@ -37,7 +37,9 @@ class UPlayer : public UObject, public FExec
 
 public:
 	//~ Begin FExec Interface.
+#if UE_ALLOW_EXEC_COMMANDS
 	ENGINE_API virtual bool Exec( UWorld* InWorld, const TCHAR* Cmd,FOutputDevice& Ar) override;
+#endif
 	//~ End FExec Interface.
 
 	/**
diff --git a/Engine/Source/Runtime/Engine/Private/AudioDevice.cpp b/Engine/Source/Runtime/Engine/Private/AudioDevice.cpp
--- a/Engine/Source/Runtime/Engine/Private/AudioDevice.cpp
+++ b/Engine/Source/Runtime/Engine/Private/AudioDevice.cpp
@@ -2262,6 +2262,7 @@ void FAudioDevice::SetMixDebugState(EDebugState InDebugState)
 	DebugState = InDebugState;
 }
 
+#if UE_ALLOW_EXEC_COMMANDS
 bool FAudioDevice::Exec(UWorld* InWorld, const TCHAR* Cmd, FOutputDevice& Ar)
 {
 #if !UE_BUILD_SHIPPING
@@ -2417,6 +2418,7 @@ bool FAudioDevice::Exec(UWorld* InWorld, const TCHAR* Cmd, FOutputDevice& Ar)
 
 	return false;
 }
+#endif // UE_ALLOW_EXEC_COMMANDS
 
 void FAudioDevice::InitSoundClasses()
 {
diff --git a/Engine/Source/Runtime/Engine/Private/GameEngine.cpp b/Engine/Source/Runtime/Engine/Private/GameEngine.cpp
--- a/Engine/Source/Runtime/Engine/Private/GameEngine.cpp
+++ b/Engine/Source/Runtime/Engine/Private/GameEngine.cpp
@@ -1016,6 +1016,7 @@ public:
 		return GEngineIni;
 	}
 
+#if UE_ALLOW_EXEC_COMMANDS
 	virtual bool Exec(UWorld* Inworld, const TCHAR* Cmd, FOutputDevice& Ar) override
 	{
 		if (FParse::Command(&Cmd, TEXT("exitembedded")))
@@ -1092,7 +1093,7 @@ public:
 		}
 		return false;
 	}
-
+#endif // UE_ALLOW_EXEC_COMMANDS
 } GEmbeddedCommunicationExec;
 
 
@@ -1378,6 +1379,7 @@ bool UGameEngine::ShouldDoAsyncEndOfFrameTasks() const
 	Command line executor.
 -----------------------------------------------------------------------------*/
 
+#if UE_ALLOW_EXEC_COMMANDS
 bool UGameEngine::Exec( UWorld* InWorld, const TCHAR* Cmd, FOutputDevice& Ar )
 {
 	if( FParse::Command( &Cmd,TEXT("REATTACHCOMPONENTS")) || FParse::Command( &Cmd,TEXT("REREGISTERCOMPONENTS")))
@@ -1532,6 +1534,7 @@ bool UGameEngine::Exec( UWorld* InWorld, const TCHAR* Cmd, FOutputDevice& Ar )
 		return false;
 	}
 }
+#endif // UE_ALLOW_EXEC_COMMANDS
 
 bool UGameEngine::HandleExitCommand( const TCHAR* Cmd, FOutputDevice& Ar )
 {
diff --git a/Engine/Source/Runtime/Engine/Private/GameInstance.cpp b/Engine/Source/Runtime/Engine/Private/GameInstance.cpp
--- a/Engine/Source/Runtime/Engine/Private/GameInstance.cpp
+++ b/Engine/Source/Runtime/Engine/Private/GameInstance.cpp
@@ -755,6 +755,7 @@ bool UGameInstance::HandleTravelCommand(const TCHAR* Cmd, FOutputDevice& Ar, UWo
 	return Engine->HandleTravelCommand(Cmd, Ar, InWorld);
 }
 
+#if UE_ALLOW_EXEC_COMMANDS
 bool UGameInstance::Exec(UWorld* InWorld, const TCHAR* Cmd, FOutputDevice& Ar)
 {
 	// @todo a bunch of stuff in UEngine probably belongs here as well
@@ -777,6 +778,7 @@ bool UGameInstance::Exec(UWorld* InWorld, const TCHAR* Cmd, FOutputDevice& Ar)
 
 	return false;
 }
+#endif // UE_ALLOW_EXEC_COMMANDS
 
 ULocalPlayer* UGameInstance::CreateInitialPlayer(FString& OutError)
 {
diff --git a/Engine/Source/Runtime/Engine/Private/NetConnection.cpp b/Engine/Source/Runtime/Engine/Private/NetConnection.cpp
--- a/Engine/Source/Runtime/Engine/Private/NetConnection.cpp
+++ b/Engine/Source/Runtime/Engine/Private/NetConnection.cpp
@@ -1400,6 +1400,7 @@ UWorld* UNetConnection::GetWorld() const
 	return World;
 }
 
+#if UE_ALLOW_EXEC_COMMANDS
 bool UNetConnection::Exec( UWorld* InWorld, const TCHAR* Cmd, FOutputDevice& Ar )
 {
 	if ( Super::Exec( InWorld, Cmd,Ar) )
@@ -1412,6 +1413,8 @@ bool UNetConnection::Exec( UWorld* InWorld, const TCHAR* Cmd, FOutputDevice& Ar
 	}
 	return false;
 }
+#endif // UE_ALLOW_EXEC_COMMANDS
+
 void UNetConnection::AssertValid()
 {
 	// Make sure this connection is in a reasonable state.
diff --git a/Engine/Source/Runtime/Engine/Private/Player.cpp b/Engine/Source/Runtime/Engine/Private/Player.cpp
--- a/Engine/Source/Runtime/Engine/Private/Player.cpp
+++ b/Engine/Source/Runtime/Engine/Private/Player.cpp
@@ -91,6 +91,7 @@ APlayerController* UPlayer::GetPlayerController(const UWorld* const InWorld) con
 	return nullptr;
 }
 
+#if UE_ALLOW_EXEC_COMMANDS
 bool UPlayer::Exec( UWorld* InWorld, const TCHAR* Cmd,FOutputDevice& Ar)
 {
 	// Route through Exec_Dev and Exec_Editor first
@@ -154,6 +155,7 @@ bool UPlayer::Exec( UWorld* InWorld, const TCHAR* Cmd,FOutputDevice& Ar)
 	}
 	return false;
 }
+#endif // UE_ALLOW_EXEC_COMMANDS
 
 void UPlayer::SwitchController(class APlayerController* PC)
 {
diff --git a/Engine/Source/Runtime/Engine/Private/SystemSettings.cpp b/Engine/Source/Runtime/Engine/Private/SystemSettings.cpp
--- a/Engine/Source/Runtime/Engine/Private/SystemSettings.cpp
+++ b/Engine/Source/Runtime/Engine/Private/SystemSettings.cpp
@@ -240,11 +240,13 @@ void FSystemSettings::CVarSink()
 	ApplyOverrides();
 }
 
+#if UE_ALLOW_EXEC_COMMANDS
 bool FSystemSettings::Exec( UWorld* InWorld, const TCHAR* Cmd, FOutputDevice& Ar )
 {
 	// no longer needed, we have the "Scalability" console command
 	return false;
 }
+#endif // UE_ALLOW_EXEC_COMMANDS
 
 /*-----------------------------------------------------------------------------
 	Resolution.
diff --git a/Engine/Source/Runtime/Engine/Public/AudioDevice.h b/Engine/Source/Runtime/Engine/Public/AudioDevice.h
--- a/Engine/Source/Runtime/Engine/Public/AudioDevice.h
+++ b/Engine/Source/Runtime/Engine/Public/AudioDevice.h
@@ -441,7 +441,9 @@ class ENGINE_API FAudioDevice : public FExec
 public:
 
 	//Begin FExec Interface
+#if UE_ALLOW_EXEC_COMMANDS
 	virtual bool Exec(UWorld* InWorld, const TCHAR* Cmd, FOutputDevice& Ar = *GLog) override;
+#endif
 	//End FExec Interface
 
 #if !UE_BUILD_SHIPPING
diff --git a/Engine/Source/Runtime/Engine/Public/SystemSettings.h b/Engine/Source/Runtime/Engine/Public/SystemSettings.h
--- a/Engine/Source/Runtime/Engine/Public/SystemSettings.h
+++ b/Engine/Source/Runtime/Engine/Public/SystemSettings.h
@@ -47,7 +47,9 @@ public:
 	void ApplyOverrides();
 	
 	//~ Begin Exec Interface
+#if UE_ALLOW_EXEC_COMMANDS
 	virtual bool Exec( UWorld* InWorld, const TCHAR* Cmd, FOutputDevice& Ar ) override;
+#endif
 	//~ End Exec Interface
 
 	/** Mask where 1 bits mean we want to force the engine show flag to be off */
diff --git a/Engine/Source/Runtime/NetworkReplayStreaming/NetworkReplayStreaming/Private/NetworkReplayStreaming.cpp b/Engine/Source/Runtime/NetworkReplayStreaming/NetworkReplayStreaming/Private/NetworkReplayStreaming.cpp
--- a/Engine/Source/Runtime/NetworkReplayStreaming/NetworkReplayStreaming/Private/NetworkReplayStreaming.cpp
+++ b/Engine/Source/Runtime/NetworkReplayStreaming/NetworkReplayStreaming/Private/NetworkReplayStreaming.cpp
@@ -110,6 +110,7 @@ void FNetworkReplayStreaming::Flush()
 	}
 }
 
+#if UE_ALLOW_EXEC_COMMANDS
 bool FNetworkReplayStreaming::Exec(UWorld* InWorld, const TCHAR* Cmd, FOutputDevice& Ar)
 {
 	// expected usage is "replaystreamer <streamer factory name> <streamer specific args>"
@@ -130,6 +131,7 @@ bool FNetworkReplayStreaming::Exec(UWorld* InWorld, const TCHAR* Cmd, FOutputDev
 
 	return false;
 }
+#endif // UE_ALLOW_EXEC_COMMANDS
 
 FString LexToString(const EReplayStreamerState State)
 {
diff --git a/Engine/Source/Runtime/NetworkReplayStreaming/NetworkReplayStreaming/Public/NetworkReplayStreaming.h b/Engine/Source/Runtime/NetworkReplayStreaming/NetworkReplayStreaming/Public/NetworkReplayStreaming.h
--- a/Engine/Source/Runtime/NetworkReplayStreaming/NetworkReplayStreaming/Public/NetworkReplayStreaming.h
+++ b/Engine/Source/Runtime/NetworkReplayStreaming/NetworkReplayStreaming/Public/NetworkReplayStreaming.h
@@ -723,7 +723,9 @@ public:
 	static NETWORKREPLAYSTREAMING_API FString GetReplayFileExtension();
 
 	// FSelfRegisteringExec interface
+#if UE_ALLOW_EXEC_COMMANDS
 	virtual bool Exec(UWorld* InWorld, const TCHAR* Cmd, FOutputDevice& Ar) override;
+#endif
 
 private:
 	TSet<FName> LoadedFactories;
