Update analytics user id and crash reporting user id when the epic account id for the running process changes.

diff --git a/Engine/Source/Programs/CrashReportClient/Private/CrashReportClientApp.cpp b/Engine/Source/Programs/CrashReportClient/Private/CrashReportClientApp.cpp
--- a/Engine/Source/Programs/CrashReportClient/Private/CrashReportClientApp.cpp
+++ b/Engine/Source/Programs/CrashReportClient/Private/CrashReportClientApp.cpp
@@ -888,6 +888,7 @@ static bool LoadTempCrashContextFromFile(FSharedCrashContext& CrashContext, uint
 	FindAndCopyValue(ContextProperties, TEXT("AppDefaultLocale"), SessionContext.DefaultLocale);
 	FindAndParseValue(ContextProperties, TEXT("IsUERelease"), SessionContext.bIsUERelease);
 	FindAndCopyValue(ContextProperties, TEXT("UserName"), SessionContext.UserName);
+	FindAndCopyValue(ContextProperties, TEXT("EpicAccountId"), SessionContext.EpicAccountId);
 	FindAndCopyValue(ContextProperties, TEXT("BaseDir"), SessionContext.BaseDir);
 	FindAndCopyValue(ContextProperties, TEXT("RootDir"), SessionContext.RootDir);
 	FindAndCopyValue(ContextProperties, TEXT("LoginId"), SessionContext.LoginIdStr);
@@ -1123,7 +1124,7 @@ void RunCrashReportClient(const TCHAR* CommandLine)
 					const bool bReportCrashAnalyticInfo = CrashContext.UserSettings.bSendUsageData;
 					if (bReportCrashAnalyticInfo)
 					{
-						FCrashReportAnalytics::Initialize();
+						FCrashReportAnalytics::Initialize(CrashContext.SessionContext.EpicAccountId);
 					}
 
 					FCrashReportAnalyticsSessionSummary::Get().OnCrashReportCollecting();
@@ -1196,7 +1197,7 @@ void RunCrashReportClient(const TCHAR* CommandLine)
 			FMemory::Memzero(TempCrashContext);
 			if (LoadTempCrashContextFromFile(TempCrashContext, MonitorPid) && TempCrashContext.UserSettings.bSendUsageData)
 			{
-				FCrashReportAnalytics::Initialize();
+				FCrashReportAnalytics::Initialize(TempCrashContext.SessionContext.EpicAccountId);
 				if (FCrashReportAnalytics::IsAvailable())
 				{
 					TOptional<int32> ExitCodeOpt;
diff --git a/Engine/Source/Runtime/Analytics/AnalyticsET/Private/AnalyticsSessionSummaryManager.cpp b/Engine/Source/Runtime/Analytics/AnalyticsET/Private/AnalyticsSessionSummaryManager.cpp
--- a/Engine/Source/Runtime/Analytics/AnalyticsET/Private/AnalyticsSessionSummaryManager.cpp
+++ b/Engine/Source/Runtime/Analytics/AnalyticsET/Private/AnalyticsSessionSummaryManager.cpp
@@ -243,6 +243,9 @@ TSharedPtr<IAnalyticsPropertyStore> FAnalyticsSessionSummaryManager::MakeStore(u
 		PropertyStore.Reset(); // Ensure to return a null pointer if the store failed to create (not enough disk space for example)
 	}
 
+	WeakPropertyStores.Add(PropertyStore);
+	PruneTrackedPropertyStores();
+
 	return StaticCastSharedPtr<IAnalyticsPropertyStore>(PropertyStore);
 }
 
@@ -251,6 +254,22 @@ void FAnalyticsSessionSummaryManager::SetSender(TSharedPtr<IAnalyticsSessionSumm
 	SummarySender = Sender;
 }
 
+void FAnalyticsSessionSummaryManager::SetUserId(const FString& InUserId)
+{
+	UserId = InUserId;
+
+	PruneTrackedPropertyStores();
+
+	for (TWeakPtr<IAnalyticsPropertyStore>& WeakPropertyStore : WeakPropertyStores)
+	{
+		if (TSharedPtr<IAnalyticsPropertyStore> PropertyStore = WeakPropertyStore.Pin())
+		{
+			AnalyticsManagerProperties::InternalSessionUserId.Set(PropertyStore.Get(), UserId);
+			PropertyStore->Flush();
+		}
+	}
+}
+
 void FAnalyticsSessionSummaryManager::Shutdown(bool bDiscard)
 {
 	if (SummarySender)
@@ -312,6 +331,17 @@ bool FAnalyticsSessionSummaryManager::IsPrincipalProcess() const
 	return bIsPrincipal;
 }
 
+void FAnalyticsSessionSummaryManager::PruneTrackedPropertyStores()
+{
+	for (int32 index = WeakPropertyStores.Num() - 1; index >= 0; --index)
+	{
+		if (!WeakPropertyStores[index].IsValid())
+		{
+			WeakPropertyStores.RemoveAtSwap(index);
+		}
+	}
+}
+
 TMap<FString, FAnalyticsSessionSummaryManager::FProcessGroup> FAnalyticsSessionSummaryManager::GetSessionFiles() const
 {
 	TMap<FString, FProcessGroup> ProcessFiles;
diff --git a/Engine/Source/Runtime/Analytics/AnalyticsET/Public/AnalyticsSessionSummaryManager.h b/Engine/Source/Runtime/Analytics/AnalyticsET/Public/AnalyticsSessionSummaryManager.h
--- a/Engine/Source/Runtime/Analytics/AnalyticsET/Public/AnalyticsSessionSummaryManager.h
+++ b/Engine/Source/Runtime/Analytics/AnalyticsET/Public/AnalyticsSessionSummaryManager.h
@@ -105,6 +105,12 @@ public:
 	 */
 	void SetSender(TSharedPtr<IAnalyticsSessionSummarySender> Sender);
 
+	/**
+	 * Sets a the user id used for reporting analytics. Allows for changing the application user after startup. Will set the user id on
+	 * all existing stores created from this session summary manager.
+	 */
+	void SetUserId(const FString& UserId);
+
 	/**
 	 * Shuts down and sends the session summaries that can be sent, if a sender is set. Remember to flush and discard the
 	 * property store returned by MakeStore() before shutting down to prevent delaying the report.
@@ -176,7 +182,11 @@ private:
 	/** Returns whether this manager instance represents the principal process of a group. */
 	bool IsPrincipalProcess() const;
 
+	/** Prune the list of property stores for invalid entries. */
+	void PruneTrackedPropertyStores();
+
 private:
+	TArray<TWeakPtr<IAnalyticsPropertyStore>> WeakPropertyStores;
 	TSharedPtr<IAnalyticsSessionSummarySender> SummarySender;
 	FString ProcessName;
 	FString ProcessGroupId;
diff --git a/Engine/Source/Runtime/Core/Private/GenericPlatform/GenericPlatformCrashContext.cpp b/Engine/Source/Runtime/Core/Private/GenericPlatform/GenericPlatformCrashContext.cpp
--- a/Engine/Source/Runtime/Core/Private/GenericPlatform/GenericPlatformCrashContext.cpp
+++ b/Engine/Source/Runtime/Core/Private/GenericPlatform/GenericPlatformCrashContext.cpp
@@ -1094,6 +1094,11 @@ void FGenericCrashContext::PurgeOldCrashConfig()
 	}
 }
 
+void FGenericCrashContext::SetEpicAccountId(const FString& EpicAccountId)
+{
+	NCached::Set(NCached::Session.EpicAccountId, *EpicAccountId);
+}
+
 void FGenericCrashContext::ResetEngineData()
 {
 	NCached::EngineData.Reset();
diff --git a/Engine/Source/Runtime/Core/Public/GenericPlatform/GenericPlatformCrashContext.h b/Engine/Source/Runtime/Core/Public/GenericPlatform/GenericPlatformCrashContext.h
--- a/Engine/Source/Runtime/Core/Public/GenericPlatform/GenericPlatformCrashContext.h
+++ b/Engine/Source/Runtime/Core/Public/GenericPlatform/GenericPlatformCrashContext.h
@@ -431,6 +431,9 @@ public:
 	/** Helper to clean out old files in the crash report client config folder. */
 	CORE_API static void PurgeOldCrashConfig();
 
+	/** Set or change the epic account id associated with the crash session. Will override the epic account id stored in the registry for reporting when present. */ 
+	static void SetEpicAccountId(const FString& EpicAccountId);
+
 	/** Clears the engine data dictionary */
 	CORE_API static void ResetEngineData();
 
diff --git a/Engine/Source/Runtime/CrashReportCore/Private/CrashReportAnalytics.cpp b/Engine/Source/Runtime/CrashReportCore/Private/CrashReportAnalytics.cpp
--- a/Engine/Source/Runtime/CrashReportCore/Private/CrashReportAnalytics.cpp
+++ b/Engine/Source/Runtime/CrashReportCore/Private/CrashReportAnalytics.cpp
@@ -38,8 +38,8 @@ IAnalyticsProviderET& FCrashReportAnalytics::GetProvider()
 	checkf(bIsInitialized && Analytics.IsValid(), TEXT("FCrashReportAnalytics::GetProvider called outside of Initialize/Shutdown."));
 	return *Analytics.Get();
 }
- 
-void FCrashReportAnalytics::Initialize()
+
+void FCrashReportAnalytics::Initialize(const FString& EpicAccountId)
 {
 	checkf(!bIsInitialized, TEXT("FCrashReportAnalytics::Initialize called more than once."));
 
@@ -68,7 +68,8 @@ void FCrashReportAnalytics::Initialize()
 		Analytics = FAnalyticsET::Get().CreateAnalyticsProvider(Config);
 		if( Analytics.IsValid() )
 		{
-			Analytics->SetUserID(FString::Printf(TEXT("%s|%s|%s"), *FPlatformMisc::GetLoginId(), *FPlatformMisc::GetEpicAccountId(), *FPlatformMisc::GetOperatingSystemId()));
+			const FString AnalyticsEpicAccountId = EpicAccountId.IsEmpty() ? *FPlatformMisc::GetEpicAccountId() : *EpicAccountId;
+			Analytics->SetUserID(FString::Printf(TEXT("%s|%s|%s"), *FPlatformMisc::GetLoginId(), *AnalyticsEpicAccountId, *FPlatformMisc::GetOperatingSystemId()));
 			Analytics->StartSession();
 		}
 	}
diff --git a/Engine/Source/Runtime/CrashReportCore/Public/CrashReportAnalytics.h b/Engine/Source/Runtime/CrashReportCore/Public/CrashReportAnalytics.h
--- a/Engine/Source/Runtime/CrashReportCore/Public/CrashReportAnalytics.h
+++ b/Engine/Source/Runtime/CrashReportCore/Public/CrashReportAnalytics.h
@@ -29,7 +29,7 @@ public:
 		return Analytics.IsValid(); 
 	}
 	/** Called to initialize the singleton. */
-	static void Initialize();
+	static void Initialize(const FString& EpicAccountId = FString());
 	/** Called to shut down the singleton */
 	static void Shutdown();
 
diff --git a/Engine/Source/Runtime/Engine/Private/EngineAnalytics.cpp b/Engine/Source/Runtime/Engine/Private/EngineAnalytics.cpp
--- a/Engine/Source/Runtime/Engine/Private/EngineAnalytics.cpp
+++ b/Engine/Source/Runtime/Engine/Private/EngineAnalytics.cpp
@@ -15,6 +15,7 @@
 #include "GenericPlatform/GenericPlatformCrashContext.h"
 #include "StudioAnalytics.h"
 #include "UObject/Class.h"
+#include "Containers/Set.h"
 
 #if WITH_EDITOR
 #include "AnalyticsSessionSummaryManager.h"
@@ -26,6 +27,7 @@
 
 bool FEngineAnalytics::bIsInitialized;
 TSharedPtr<IAnalyticsProviderET> FEngineAnalytics::Analytics;
+TSet<FString> FEngineAnalytics::SessionEpicAccountIds;
 
 #if WITH_EDITOR
 static TUniquePtr<FAnalyticsSessionSummaryManager> AnalyticsSessionSummaryManager;
@@ -75,6 +77,11 @@ static TSharedPtr<IAnalyticsProviderET> CreateEpicAnalyticsProvider()
 	return FAnalyticsET::Get().CreateAnalyticsProvider(Config);
 }
 
+FString CreateAnalyticsUserId(const FString& EpicAccountId)
+{
+	return FString::Printf(TEXT("%s|%s|%s"), *FPlatformMisc::GetLoginId(), *EpicAccountId, *FPlatformMisc::GetOperatingSystemId());
+}
+
 IAnalyticsProviderET& FEngineAnalytics::GetProvider()
 {
 	checkf(bIsInitialized && IsAvailable(), TEXT("FEngineAnalytics::GetProvider called outside of Initialize/Shutdown."));
@@ -118,77 +125,16 @@ void FEngineAnalytics::Initialize()
 
 		if (Analytics.IsValid())
 		{
+			Analytics->SetUserID(CreateAnalyticsUserId(FPlatformMisc::GetEpicAccountId()));
+			SessionEpicAccountIds.Add(FPlatformMisc::GetEpicAccountId());
+
 			if (UE::Analytics::Private::EngineAnalyticsConfigOverride)
 			{
-				UE::Analytics::Private::EngineAnalyticsConfigOverride->OnProviderCreated(*Analytics);
-			}
-			else
-			{
-				Analytics->SetUserID(FString::Printf(TEXT("%s|%s|%s"), *FPlatformMisc::GetLoginId(), *FPlatformMisc::GetEpicAccountId(), *FPlatformMisc::GetOperatingSystemId()));
+				UE::Analytics::Private::EngineAnalyticsConfigOverride->OnInitialized(*Analytics, UE::Analytics::Private::FOnEpicAccountIdChanged::CreateStatic(&FEngineAnalytics::OnEpicAccountIdChanged));
 			}
 
-			const FString UserID = FPlatformProcess::UserName(false);
-			const UGeneralProjectSettings& ProjectSettings = *GetDefault<UGeneralProjectSettings>();
-
 			TArray<FAnalyticsEventAttribute> StartSessionAttributes;
-			GEngine->CreateStartupAnalyticsAttributes( StartSessionAttributes );
-			// Add project info whether we are in editor or game.
-			FString OSMajor;
-			FString OSMinor;
-			FPlatformMemoryStats Stats = FPlatformMemory::GetStats();
-			StartSessionAttributes.Emplace(TEXT("ProjectName"), ProjectSettings.ProjectName);
-			StartSessionAttributes.Emplace(TEXT("ProjectID"), ProjectSettings.ProjectID);
-			StartSessionAttributes.Emplace(TEXT("ProjectDescription"), ProjectSettings.Description);
-			StartSessionAttributes.Emplace(TEXT("ProjectVersion"), ProjectSettings.ProjectVersion);
-			StartSessionAttributes.Emplace(TEXT("Application.Commandline"), FCommandLine::Get());
-			StartSessionAttributes.Emplace(TEXT("User.ID"), UserID);
-			StartSessionAttributes.Emplace(TEXT("Build.Configuration"), LexToString(FApp::GetBuildConfiguration()));
-			StartSessionAttributes.Emplace(TEXT("Build.IsInternalBuild"), FEngineBuildSettings::IsInternalBuild());
-			StartSessionAttributes.Emplace(TEXT("Build.IsPerforceBuild"), FEngineBuildSettings::IsPerforceBuild());
-			StartSessionAttributes.Emplace(TEXT("Build.IsPromotedBuild"), FApp::GetEngineIsPromotedBuild() == 0 ? false : true);
-			StartSessionAttributes.Emplace(TEXT("Build.BranchName"), FApp::GetBranchName());
-			StartSessionAttributes.Emplace(TEXT("Build.Changelist"), BuildSettings::GetCurrentChangelist());
-			StartSessionAttributes.Emplace(TEXT("Config.IsEditor"), GIsEditor);
-			StartSessionAttributes.Emplace(TEXT("Config.IsUnattended"), FApp::IsUnattended());
-			StartSessionAttributes.Emplace(TEXT("Config.IsBuildMachine"), GIsBuildMachine);
-			StartSessionAttributes.Emplace(TEXT("Config.IsRunningCommandlet"), IsRunningCommandlet());
-			StartSessionAttributes.Emplace(TEXT("Platform.IsRemoteSession"), FPlatformMisc::IsRemoteSession());
-			StartSessionAttributes.Emplace(TEXT("GPUVendorID"), GRHIVendorId);
-			StartSessionAttributes.Emplace(TEXT("GPUDeviceID"), GRHIDeviceId);
-			StartSessionAttributes.Emplace(TEXT("GRHIDeviceRevision"), GRHIDeviceRevision);
-			StartSessionAttributes.Emplace(TEXT("GRHIAdapterInternalDriverVersion"), GRHIAdapterInternalDriverVersion);
-			StartSessionAttributes.Emplace(TEXT("GRHIAdapterUserDriverVersion"), GRHIAdapterUserDriverVersion);
-			StartSessionAttributes.Emplace(TEXT("TotalPhysicalRAM"), static_cast<uint64>(Stats.TotalPhysical));
-			StartSessionAttributes.Emplace(TEXT("CPUPhysicalCores"), FPlatformMisc::NumberOfCores());
-			StartSessionAttributes.Emplace(TEXT("CPULogicalCores"), FPlatformMisc::NumberOfCoresIncludingHyperthreads());
-			StartSessionAttributes.Emplace(TEXT("DesktopGPUAdapter"), FPlatformMisc::GetPrimaryGPUBrand());
-			StartSessionAttributes.Emplace(TEXT("RenderingGPUAdapter"), GRHIAdapterName);
-			StartSessionAttributes.Emplace(TEXT("CPUVendor"), FPlatformMisc::GetCPUVendor());
-			StartSessionAttributes.Emplace(TEXT("CPUBrand"), FPlatformMisc::GetCPUBrand());
-			FPlatformMisc::GetOSVersions(/*out*/ OSMajor, /*out*/ OSMinor);
-			StartSessionAttributes.Emplace(TEXT("OSMajor"), OSMajor);
-			StartSessionAttributes.Emplace(TEXT("OSMinor"), OSMinor);
-			StartSessionAttributes.Emplace(TEXT("OSVersion"), FPlatformMisc::GetOSVersion());
-			StartSessionAttributes.Emplace(TEXT("Is64BitOS"), FPlatformMisc::Is64bitOperatingSystem());
-
-#if WITH_EDITOR
-			StartSessionAttributes.Emplace(TEXT("Horde.TemplateID"), FHorde::GetTemplateId());
-			StartSessionAttributes.Emplace(TEXT("Horde.TemplateName"), FHorde::GetTemplateName());
-			StartSessionAttributes.Emplace(TEXT("Horde.JobURL"), FHorde::GetJobURL());
-			StartSessionAttributes.Emplace(TEXT("Horde.JobID"), FHorde::GetJobId());
-			StartSessionAttributes.Emplace(TEXT("Horde.StepName"), FHorde::GetStepName());
-			StartSessionAttributes.Emplace(TEXT("Horde.StepID"), FHorde::GetStepId());
-			StartSessionAttributes.Emplace(TEXT("Horde.StepURL"), FHorde::GetStepURL());
-			StartSessionAttributes.Emplace(TEXT("Horde.BatchID"), FHorde::GetBatchId());
-#endif
-
-#if PLATFORM_MAC
-#if PLATFORM_MAC_ARM64
-            StartSessionAttributes.Emplace(TEXT("UEBuildArch"), FString(TEXT("AppleSilicon")));
-#else
-            StartSessionAttributes.Emplace(TEXT("UEBuildArch"), FString(TEXT("Intel(Mac)")));
-#endif
-#endif
+			AppendMachineStats(StartSessionAttributes);
 
 			// allow editor events to be correlated to StudioAnalytics events (if there is a studio analytics provider)
 			if (FStudioAnalytics::IsAvailable())
@@ -289,3 +235,107 @@ void FEngineAnalytics::LowDriveSpaceDetected()
 	}
 #endif
 }
+
+void FEngineAnalytics::AppendMachineStats(TArray<FAnalyticsEventAttribute>& EventAttributes)
+{
+	const FString UserID = FPlatformProcess::UserName(false);
+	const UGeneralProjectSettings& ProjectSettings = *GetDefault<UGeneralProjectSettings>();
+	const FPlatformMemoryStats Stats = FPlatformMemory::GetStats();
+
+	FString OSMajor;
+	FString OSMinor;
+	FPlatformMisc::GetOSVersions(/*out*/ OSMajor, /*out*/ OSMinor);
+
+	GEngine->CreateStartupAnalyticsAttributes( EventAttributes );
+
+	// Add project info whether we are in editor or game.
+	EventAttributes.Emplace(TEXT("ProjectName"), ProjectSettings.ProjectName);
+	EventAttributes.Emplace(TEXT("ProjectID"), ProjectSettings.ProjectID);
+	EventAttributes.Emplace(TEXT("ProjectDescription"), ProjectSettings.Description);
+	EventAttributes.Emplace(TEXT("ProjectVersion"), ProjectSettings.ProjectVersion);
+	EventAttributes.Emplace(TEXT("Application.Commandline"), FCommandLine::Get());
+	EventAttributes.Emplace(TEXT("User.ID"), UserID);
+	EventAttributes.Emplace(TEXT("Build.Configuration"), LexToString(FApp::GetBuildConfiguration()));
+	EventAttributes.Emplace(TEXT("Build.IsInternalBuild"), FEngineBuildSettings::IsInternalBuild());
+	EventAttributes.Emplace(TEXT("Build.IsPerforceBuild"), FEngineBuildSettings::IsPerforceBuild());
+	EventAttributes.Emplace(TEXT("Build.IsPromotedBuild"), FApp::GetEngineIsPromotedBuild() == 0 ? false : true);
+	EventAttributes.Emplace(TEXT("Build.BranchName"), FApp::GetBranchName());
+	EventAttributes.Emplace(TEXT("Build.Changelist"), BuildSettings::GetCurrentChangelist());
+	EventAttributes.Emplace(TEXT("Config.IsEditor"), GIsEditor);
+	EventAttributes.Emplace(TEXT("Config.IsUnattended"), FApp::IsUnattended());
+	EventAttributes.Emplace(TEXT("Config.IsBuildMachine"), GIsBuildMachine);
+	EventAttributes.Emplace(TEXT("Config.IsRunningCommandlet"), IsRunningCommandlet());
+	EventAttributes.Emplace(TEXT("Platform.IsRemoteSession"), FPlatformMisc::IsRemoteSession());
+	EventAttributes.Emplace(TEXT("OSMajor"), OSMajor);
+	EventAttributes.Emplace(TEXT("OSMinor"), OSMinor);
+	EventAttributes.Emplace(TEXT("OSVersion"), FPlatformMisc::GetOSVersion());
+	EventAttributes.Emplace(TEXT("Is64BitOS"), FPlatformMisc::Is64bitOperatingSystem());
+	EventAttributes.Emplace(TEXT("GPUVendorID"), GRHIVendorId);
+	EventAttributes.Emplace(TEXT("GPUDeviceID"), GRHIDeviceId);
+	EventAttributes.Emplace(TEXT("GRHIDeviceRevision"), GRHIDeviceRevision);
+	EventAttributes.Emplace(TEXT("GRHIAdapterInternalDriverVersion"), GRHIAdapterInternalDriverVersion);
+	EventAttributes.Emplace(TEXT("GRHIAdapterUserDriverVersion"), GRHIAdapterUserDriverVersion);
+	EventAttributes.Emplace(TEXT("TotalPhysicalRAM"), static_cast<uint64>(Stats.TotalPhysical));
+	EventAttributes.Emplace(TEXT("CPUPhysicalCores"), FPlatformMisc::NumberOfCores());
+	EventAttributes.Emplace(TEXT("CPULogicalCores"), FPlatformMisc::NumberOfCoresIncludingHyperthreads());
+	EventAttributes.Emplace(TEXT("DesktopGPUAdapter"), FPlatformMisc::GetPrimaryGPUBrand());
+	EventAttributes.Emplace(TEXT("RenderingGPUAdapter"), GRHIAdapterName);
+	EventAttributes.Emplace(TEXT("CPUVendor"), FPlatformMisc::GetCPUVendor());
+	EventAttributes.Emplace(TEXT("CPUBrand"), FPlatformMisc::GetCPUBrand());
+#if WITH_EDITOR
+	EventAttributes.Emplace(TEXT("Horde.TemplateID"), FHorde::GetTemplateId());
+	EventAttributes.Emplace(TEXT("Horde.TemplateName"), FHorde::GetTemplateName());
+	EventAttributes.Emplace(TEXT("Horde.JobURL"), FHorde::GetJobURL());
+	EventAttributes.Emplace(TEXT("Horde.JobID"), FHorde::GetJobId());
+	EventAttributes.Emplace(TEXT("Horde.StepName"), FHorde::GetStepName());
+	EventAttributes.Emplace(TEXT("Horde.StepID"), FHorde::GetStepId());
+	EventAttributes.Emplace(TEXT("Horde.StepURL"), FHorde::GetStepURL());
+	EventAttributes.Emplace(TEXT("Horde.BatchID"), FHorde::GetBatchId());
+#endif
+
+#if PLATFORM_MAC
+#if PLATFORM_MAC_ARM64
+	EventAttributes.Emplace(TEXT("UEBuildArch"), FString(TEXT("AppleSilicon")));
+#else
+	EventAttributes.Emplace(TEXT("UEBuildArch"), FString(TEXT("Intel(Mac)")));
+#endif
+#endif
+}
+
+void FEngineAnalytics::OnEpicAccountIdChanged(const FString& EpicAccountId)
+{
+	// For analytics reporting ignore changes to an empty account id when the user logs out.
+	if (!EpicAccountId.IsEmpty())
+	{
+		const FString NewAnalyticsUserId = CreateAnalyticsUserId(EpicAccountId);
+
+		// Update analytics provider user.
+		if (Analytics)
+		{
+			Analytics->SetUserID(NewAnalyticsUserId);
+
+			// Only send SessionMachineStats when the user is changed to a user that has not been seen in the session.
+			if (!SessionEpicAccountIds.Contains(EpicAccountId))
+			{
+				SessionEpicAccountIds.Add(EpicAccountId);
+
+				// When the user id is changed send a SessionMachineStats event.
+				static const FString SZEventName = TEXT("SessionMachineStats");
+				TArray<FAnalyticsEventAttribute> Params;
+				AppendMachineStats(Params);
+				Analytics->RecordEvent(SZEventName, Params);
+			}
+		}
+
+#if WITH_EDITOR
+		// Update the summary manager and all of the data stores.
+		if (AnalyticsSessionSummaryManager)
+		{
+			AnalyticsSessionSummaryManager->SetUserId(NewAnalyticsUserId);
+		}
+#endif
+
+		// Update the crash context so the user id will be sent with runtime events to CRCEditor.
+		FGenericCrashContext::SetEpicAccountId(EpicAccountId);
+	}
+}
diff --git a/Engine/Source/Runtime/Engine/Public/EngineAnalytics.h b/Engine/Source/Runtime/Engine/Public/EngineAnalytics.h
--- a/Engine/Source/Runtime/Engine/Public/EngineAnalytics.h
+++ b/Engine/Source/Runtime/Engine/Public/EngineAnalytics.h
@@ -3,6 +3,8 @@
 #pragma once
 
 #include "CoreMinimal.h"
+#include "Containers/ContainersFwd.h"
+#include "Delegates/Delegate.h"
 #include "IAnalyticsProviderET.h" // NOTE: Consider changing the code to replace IAnalyticProvider.h by IAnalyticProviderET.h
 
 class FEngineSessionManager;
@@ -57,12 +59,18 @@ public:
 	static ENGINE_API void LowDriveSpaceDetected();
 
 private:
+	static void AppendMachineStats(TArray<FAnalyticsEventAttribute>& EventAttributes);
+	static void OnEpicAccountIdChanged(const FString& EpicAccountId);
+
 	static bool bIsInitialized;
 	static ENGINE_API TSharedPtr<IAnalyticsProviderET> Analytics;
+	static ENGINE_API TSet<FString> SessionEpicAccountIds;
 };
 
 namespace UE::Analytics::Private {
 
+DECLARE_DELEGATE_OneParam(FOnEpicAccountIdChanged, const FString&);
+
 /**
   * Interface for allowing changes to engine analytics configuration intended for use by internal tools.
   * 
@@ -73,7 +81,7 @@ class IEngineAnalyticsConfigOverride
 public:
 	virtual ~IEngineAnalyticsConfigOverride() = default;
 	virtual void ApplyConfiguration(FAnalyticsET::Config& Config) = 0;
-	virtual void OnProviderCreated(IAnalyticsProviderET& provider) = 0;
+	virtual void OnInitialized(IAnalyticsProviderET& provider, const FOnEpicAccountIdChanged& OnEpicAccountIdChanged) = 0;
 };
 
 extern ENGINE_API IEngineAnalyticsConfigOverride* EngineAnalyticsConfigOverride;
