Clip asset loading debug name to 1024 characters.

When preloading a very large list of assets, the debug name could grow to a very large string. This clamps the string size to (around) 1024 characters.

diff --git a/Engine/Source/Runtime/Engine/Private/AssetManager.cpp b/Engine/Source/Runtime/Engine/Private/AssetManager.cpp
--- a/Engine/Source/Runtime/Engine/Private/AssetManager.cpp
+++ b/Engine/Source/Runtime/Engine/Private/AssetManager.cpp
@@ -1849,7 +1849,7 @@ TSharedPtr<FStreamableHandle> UAssetManager::PreloadPrimaryAssets(const TArray<F
 
 	for (const FPrimaryAssetId& PrimaryAssetId : AssetsToLoad)
 	{
-		if (GetPrimaryAssetLoadSet(PathsToLoad, PrimaryAssetId, LoadBundles, bLoadRecursive))
+		if (GetPrimaryAssetLoadSet(PathsToLoad, PrimaryAssetId, LoadBundles, bLoadRecursive) && DebugName.Len() < 1024)
 		{
 			if (DebugName.Len() == 0)
 			{
@@ -1863,6 +1863,11 @@ TSharedPtr<FStreamableHandle> UAssetManager::PreloadPrimaryAssets(const TArray<F
 			DebugName << PrimaryAssetId.ToString();
 		}
 	}
+	
+	if (DebugName.Len() >= 1024)
+	{
+		DebugName.ReplaceAt(DebugName.Len() - 3, 3, TEXTVIEW("..."));
+	}
 
 	ReturnHandle = LoadAssetList(PathsToLoad.Array(), MoveTemp(DelegateToCall), Priority, FString(*DebugName));
 
