ContentWorker tests failing in Dev-EngineMerge
#fix remove ensure as it is valid for a cooked animation sequence not to have DataModel

diff --git a/Engine/Source/Runtime/Engine/Private/Animation/AnimSequenceBase.cpp b/Engine/Source/Runtime/Engine/Private/Animation/AnimSequenceBase.cpp
--- a/Engine/Source/Runtime/Engine/Private/Animation/AnimSequenceBase.cpp
+++ b/Engine/Source/Runtime/Engine/Private/Animation/AnimSequenceBase.cpp
@@ -240,7 +240,7 @@ void UAnimSequenceBase::PostLoad()
 	if(USkeleton* MySkeleton = GetSkeleton())
 	{
 #if WITH_EDITOR
-		if (ensure(IsDataModelValid()))
+		if (IsDataModelValid())
 		{
 			const bool bDoNotTransactAction = false;
 			if (GetLinkerCustomVersion(FFortniteMainBranchObjectVersion::GUID) < FFortniteMainBranchObjectVersion::FixUpNoneNameAnimationCurves)
