Set ECVF_RenderThreadSafe flag on r.Material.ExcludeNonPipelinedShaders to fix warning about it being accessed on the rendering thread during cook

diff --git a/Engine/Source/Runtime/Engine/Private/Materials/MaterialShader.cpp b/Engine/Source/Runtime/Engine/Private/Materials/MaterialShader.cpp
--- a/Engine/Source/Runtime/Engine/Private/Materials/MaterialShader.cpp
+++ b/Engine/Source/Runtime/Engine/Private/Materials/MaterialShader.cpp
@@ -32,7 +32,7 @@ static FAutoConsoleVariableRef CVarMaterialExcludeNonPipelinedShaders(
 	TEXT("r.Material.ExcludeNonPipelinedShaders"),
 	GMaterialExcludeNonPipelinedShaders,
 	TEXT("if != 0, standalone shaders that are also part of FShaderPipeline will not be compiled (default)."),
-	ECVF_ReadOnly
+	ECVF_ReadOnly | ECVF_RenderThreadSafe
 );
 
 static TAutoConsoleVariable<FString> CVarMaterialShaderMapDump(
