Implemented first version of mip maps for SparseVolumeTextures. Each mip is currently an individual SVT structure/proxy. Mip generation is currently done on import; this will be moved to cooking in a later CL. Much of the OpenVDB->SVT conversion logic has been moved into SparseVolumeTexture.cpp, decoupling it from the vdb library. OpenVDBUtility.cpp now only implements an interface for querying data from the source file. The SVT preview received a few fixes for issues that became apparent when trying to visualize low res mips (Can't early out if the packed page table result is 0 when no null tile exists; ray stepping needs to be done after volume sampling in order to not miss voxels; sampling needs to be guarded by checking against the virtual volume resolution).

#rnx

diff --git a/Engine/Shaders/Private/SparseVolumeTexture/VisualizeSparseVolumeTexture.usf b/Engine/Shaders/Private/SparseVolumeTexture/VisualizeSparseVolumeTexture.usf
--- a/Engine/Shaders/Private/SparseVolumeTexture/VisualizeSparseVolumeTexture.usf
+++ b/Engine/Shaders/Private/SparseVolumeTexture/VisualizeSparseVolumeTexture.usf
@@ -47,6 +47,8 @@ float3 WorldToLocalNoScale0;
 float3 WorldToLocalNoScale1;
 float3 WorldToLocalNoScale2;
 uint ComponentToVisualize;
+uint MipLevel;
+float Extinction;
 
 
 Texture3D<float4> SparseVolumeTextureA;
@@ -130,10 +132,47 @@ void VisualizeSparseVolumeTexturePS(
 		tMax.y *= tDelta.y;
 		tMax.z *= tDelta.z;
 
+		float StepLength = float(1 << MipLevel);
+
 		LOOP
 		while (P.x >= 0 && P.y >= 0 && P.z >= 0 && P.x <= VolumeResolution.x && P.y <= VolumeResolution.y && P.z <= VolumeResolution.z)
 		{
-			float StepLength = 0.0;
+			if (all(P >= 0.0f) && all(P < SparseVolumeTextureResolution))
+			{
+				const float3 PageCoordF = floor(P * (1.0 / SPARSE_VOLUME_TILE_RES));
+				const uint EncodedPage = SparseVolumeTexturePageTable.Load(int4(int3(PageCoordF), 0));
+				const int3 TileCoord = int3(
+					EncodedPage & 0x7FF,
+					(EncodedPage >> 11) & 0x7FF,
+					(EncodedPage >> 22) & 0x3FF);
+
+				const int3 VoxelCoord = TileCoord * SPARSE_VOLUME_TILE_RES + int3(P - PageCoordF * SPARSE_VOLUME_TILE_RES);
+
+				float Density = 0.0f;
+				if (ComponentToVisualize < 4)
+				{
+					float4 VoxelData = SparseVolumeTextureA.Load(int4(VoxelCoord, 0));
+					Density = VoxelData[ComponentToVisualize];
+				}
+				else
+				{
+					float4 VoxelData = SparseVolumeTextureB.Load(int4(VoxelCoord, 0));
+					Density = VoxelData[min(ComponentToVisualize - 4, 3u)];
+				}
+
+				if (Density > 0.0f)
+				{
+#if 0
+					//OutLuminanceAlpha = float4(Rand3DPCG16(PageCoordF).x / 65535.0f, Rand3DPCG16(VoxelCoord).y / 65535.0f, 0.0, 0.0);
+					//OutLuminanceAlpha = float4(Rand3DPCG16(PageCoordF) / 65535.0f, 0.0);
+					OutLuminanceAlpha = float4(Rand3DPCG16(VoxelCoord) / 65535.0f, 0.0);
+					break;
+#else
+					Transmittance *= exp(-StepLength * Density * Extinction);
+#endif
+				}
+			}
+
 #if 0
 			// Slow reference
 			P += RayDir * 0.005;
@@ -146,15 +185,11 @@ void VisualizeSparseVolumeTexturePS(
 				{
 					P.x += stepSign.x;
 					tMax.x += tDelta.x;
-
-					StepLength += 1;
 				}
 				else
 				{
 					P.z += stepSign.z;
 					tMax.z += tDelta.z;
-
-					StepLength += 1;
 				}
 			}
 			else
@@ -163,55 +198,14 @@ void VisualizeSparseVolumeTexturePS(
 				{
 					P.y += stepSign.y;
 					tMax.y += tDelta.y;
-
-					StepLength += 1;
 				}
 				else
 				{
 					P.z += stepSign.z;
 					tMax.z += tDelta.z;
-
-					StepLength += 1;
 				}
 			}
 #endif
-
-			float3 PageCoordF = floor(P * (1.0 / SPARSE_VOLUME_TILE_RES));
-			const uint EncodedPage = SparseVolumeTexturePageTable.Load(int4(int3(PageCoordF), 0));
-			if (EncodedPage)
-			{
-				const int3 TileCoord = int3(
-					 EncodedPage         & 0x7FF,
-					(EncodedPage >> 11 ) & 0x7FF,
-					(EncodedPage >> 22 ) & 0x3FF);
-
-				int3 VoxelCoord = TileCoord * SPARSE_VOLUME_TILE_RES + int3(P - PageCoordF * SPARSE_VOLUME_TILE_RES);
-
-				float Density = 0.0f;
-				if (ComponentToVisualize < 4)
-				{
-					float4 VoxelData = SparseVolumeTextureA.Load(int4(VoxelCoord, 0));
-					Density = VoxelData[ComponentToVisualize];
-				}
-				else
-				{
-					float4 VoxelData = SparseVolumeTextureB.Load(int4(VoxelCoord, 0));
-					Density = VoxelData[min(ComponentToVisualize - 4, 3u)];
-				}
-
-				if (Density > 0.0f)
-
-				{
-				#if 0
-					//OutLuminanceAlpha = float4(Rand3DPCG16(PageCoordF).x / 65535.0f, Rand3DPCG16(VoxelCoord).y / 65535.0f, 0.0, 0.0);
-					//OutLuminanceAlpha = float4(Rand3DPCG16(PageCoordF) / 65535.0f, 0.0);
-					OutLuminanceAlpha = float4(Rand3DPCG16(VoxelCoord) / 65535.0f, 0.0);
-					break;
-				#else
-					Transmittance *= exp(-StepLength * Density * 0.025);
-				#endif
-				}
-			}
 		}
 	}
 
diff --git a/Engine/Source/Editor/SparseVolumeTexture/Private/SparseVolumeTextureFactory.cpp b/Engine/Source/Editor/SparseVolumeTexture/Private/SparseVolumeTextureFactory.cpp
--- a/Engine/Source/Editor/SparseVolumeTexture/Private/SparseVolumeTextureFactory.cpp
+++ b/Engine/Source/Editor/SparseVolumeTexture/Private/SparseVolumeTextureFactory.cpp
@@ -15,6 +15,7 @@
 #include "Misc/ScopedSlowTask.h"
 #include "Misc/FileHelper.h"
 #include "Async/Async.h"
+#include "Async/ParallelFor.h"
 
 #include "Editor.h"
 
@@ -414,7 +415,6 @@ void USparseVolumeTextureFactory::CleanUp()
 UObject* USparseVolumeTextureFactory::FactoryCreateFile(UClass* InClass, UObject* InParent, FName InName, EObjectFlags Flags, const FString& Filename,
 	const TCHAR* Parms, FFeedbackContext* Warn, bool& bOutOperationCanceled)
 {
-
 #if OPENVDB_AVAILABLE
 
 	GEditor->GetEditorSubsystem<UImportSubsystem>()->BroadcastAssetPreImport(this, InClass, InParent, InName, Parms);
@@ -453,7 +453,7 @@ UObject* USparseVolumeTextureFactory::FactoryCreateFile(UClass* InClass, UObject
 	}
 
 	// Utility function for computing the bounding box encompassing the bounds of all frames in the SVT.
-	auto ExpandVolumeBounds = [](const FOpenVDBImportOptions& ImportOptions, const TArray<FOpenVDBGridInfo>& GridInfoArray, FBox& VolumeBounds)
+	auto ExpandVolumeBounds = [](const FOpenVDBImportOptions& ImportOptions, const TArray<FOpenVDBGridInfo>& GridInfoArray, FIntVector3& VolumeBoundsMin, FIntVector3& VolumeBoundsMax)
 	{
 		for (const FOpenVDBSparseVolumeAttributesDesc& Attributes : ImportOptions.Attributes)
 		{
@@ -462,18 +462,33 @@ UObject* USparseVolumeTextureFactory::FactoryCreateFile(UClass* InClass, UObject
 				if (Mapping.SourceGridIndex != INDEX_NONE)
 				{
 					const FOpenVDBGridInfo& GridInfo = GridInfoArray[Mapping.SourceGridIndex];
-					VolumeBounds.Min.X = FMath::Min(VolumeBounds.Min.X, GridInfo.VolumeActiveAABBMin.X);
-					VolumeBounds.Min.Y = FMath::Min(VolumeBounds.Min.Y, GridInfo.VolumeActiveAABBMin.Y);
-					VolumeBounds.Min.Z = FMath::Min(VolumeBounds.Min.Z, GridInfo.VolumeActiveAABBMin.Z);
+					VolumeBoundsMin.X = FMath::Min(VolumeBoundsMin.X, GridInfo.VolumeActiveAABBMin.X);
+					VolumeBoundsMin.Y = FMath::Min(VolumeBoundsMin.Y, GridInfo.VolumeActiveAABBMin.Y);
+					VolumeBoundsMin.Z = FMath::Min(VolumeBoundsMin.Z, GridInfo.VolumeActiveAABBMin.Z);
 
-					VolumeBounds.Max.X = FMath::Max(VolumeBounds.Max.X, GridInfo.VolumeActiveAABBMax.X);
-					VolumeBounds.Max.Y = FMath::Max(VolumeBounds.Max.Y, GridInfo.VolumeActiveAABBMax.Y);
-					VolumeBounds.Max.Z = FMath::Max(VolumeBounds.Max.Z, GridInfo.VolumeActiveAABBMax.Z);
+					VolumeBoundsMax.X = FMath::Max(VolumeBoundsMax.X, GridInfo.VolumeActiveAABBMax.X);
+					VolumeBoundsMax.Y = FMath::Max(VolumeBoundsMax.Y, GridInfo.VolumeActiveAABBMax.Y);
+					VolumeBoundsMax.Z = FMath::Max(VolumeBoundsMax.Z, GridInfo.VolumeActiveAABBMax.Z);
 				}
 			}
 		}
 	};
 
+	auto ComputeNumMipLevels = [](const FIntVector3& VolumeBoundsMin, const FIntVector3& VolumeBoundsMax)
+	{
+		int32 Levels = 1;
+		FIntVector3 Resolution = VolumeBoundsMax - VolumeBoundsMin;
+		while (Resolution.X > 1 || Resolution.Y > 1 || Resolution.Z > 1)
+		{
+			Resolution /= 2;
+			++Levels;
+		}
+		return Levels;
+	};
+
+	FIntVector3 VolumeBoundsMin = FIntVector3(INT32_MAX, INT32_MAX, INT32_MAX);
+	FIntVector3 VolumeBoundsMax = FIntVector3(INT32_MIN, INT32_MIN, INT32_MIN);
+
 	// Import as either single static SVT or a sequence of frames, making up an animated SVT
 	if (!ImportOptions.bIsSequence)
 	{
@@ -485,24 +500,14 @@ UObject* USparseVolumeTextureFactory::FactoryCreateFile(UClass* InClass, UObject
 		FName NewName(InName.ToString() + TEXT("VDB"));
 		UStaticSparseVolumeTexture* StaticSVTexture = NewObject<UStaticSparseVolumeTexture>(InParent, UStaticSparseVolumeTexture::StaticClass(), NewName, Flags);
 
-		FBox VolumeBounds(FVector(FLT_MAX), FVector(-FLT_MAX));
-		ExpandVolumeBounds(ImportOptions, PreviewData.GridInfo, VolumeBounds);
-		StaticSVTexture->VolumeBounds = VolumeBounds;
+		ExpandVolumeBounds(ImportOptions, PreviewData.GridInfo, VolumeBoundsMin, VolumeBoundsMax);
+		StaticSVTexture->VolumeResolution = VolumeBoundsMax - VolumeBoundsMin;
 		StaticSVTexture->Frames.SetNum(1);
 
-		FSparseVolumeRawSource SparseVolumeRawSource{};
-
-		FOpenVDBToSVTConversionResult SVTResult;
-		SVTResult.Header = &SparseVolumeRawSource.Header;
-		SVTResult.PageTable = &SparseVolumeRawSource.PageTable;
-		SVTResult.PhysicalTileDataA = &SparseVolumeRawSource.PhysicalTileDataA;
-		SVTResult.PhysicalTileDataB = &SparseVolumeRawSource.PhysicalTileDataB;
+		const int32 NumMipLevels = ComputeNumMipLevels(VolumeBoundsMin, VolumeBoundsMax);
 
-		const bool bConversionSuccess = ConvertOpenVDBToSparseVolumeTexture(
-			PreviewData.LoadedFile,
-			ImportOptions,
-			&SVTResult,
-			false, FVector::Zero(), FVector::Zero());
+		FSparseVolumeRawSource SparseVolumeRawSource{};
+		const bool bConversionSuccess = ConvertOpenVDBToSparseVolumeTexture(PreviewData.LoadedFile, ImportOptions, VolumeBoundsMin, SparseVolumeRawSource);
 
 		if (!bConversionSuccess)
 		{
@@ -511,9 +516,17 @@ UObject* USparseVolumeTextureFactory::FactoryCreateFile(UClass* InClass, UObject
 		}
 
 		// Serialize the raw source data into the asset object.
+		// After serializing the data, we generate the next mip level by downsampling the current one (with a call to GenerateMipMap())
+		// and then repeating the loop until we processed all levels.
+		StaticSVTexture->Frames[0].SetNum(NumMipLevels);
+		for (int32 MipLevel = 0; MipLevel < NumMipLevels; ++MipLevel)
 		{
-			UE::Serialization::FEditorBulkDataWriter RawDataArchiveWriter(StaticSVTexture->Frames[0].RawData);
+			UE::Serialization::FEditorBulkDataWriter RawDataArchiveWriter(StaticSVTexture->Frames[0][MipLevel].RawData);
 			SparseVolumeRawSource.Serialize(RawDataArchiveWriter);
+			if ((MipLevel + 1) < NumMipLevels)
+			{
+				SparseVolumeRawSource = SparseVolumeRawSource.GenerateMipMap();
+			}
 		}
 
 		if (ImportTask.ShouldCancel())
@@ -537,21 +550,89 @@ UObject* USparseVolumeTextureFactory::FactoryCreateFile(UClass* InClass, UObject
 		
 		const int32 NumFrames = PreviewData.SequenceFilenames.Num();
 
-		FScopedSlowTask ImportTask(NumFrames, LOCTEXT("ImportingVDBAnim", "Importing OpenVDB animation"));
+		FScopedSlowTask ImportTask(NumFrames + 1, LOCTEXT("ImportingVDBAnim", "Importing OpenVDB animation"));
 		ImportTask.MakeDialog(true);
 
 		// Allocate space for each frame
 		AnimatedSVTexture->Frames.SetNum(NumFrames);
 
-		FEvent* AllTasksFinishedEvent = FPlatformProcess::GetSynchEventFromPool();
 		std::atomic_bool bErrored = false;
 		std::atomic_bool bCanceled = false;
+
+		// Compute volume bounds and check sequence files for compatiblity
+		std::mutex VolumeBoundsMutex;
+		ParallelFor(NumFrames, [&bErrored, &VolumeBoundsMutex, &VolumeBoundsMin, &VolumeBoundsMax, &ExpandVolumeBounds, &PreviewData, &ImportOptions](int32 FrameIdx)
+			{
+				if (bErrored.load())
+				{
+					return;
+				}
+
+				// Load file and get info about each contained grid
+				const FString& FrameFilename = PreviewData.SequenceFilenames[FrameIdx];
+				TArray<uint8> LoadedFrameFile;
+				if (!FFileHelper::LoadFileToArray(LoadedFrameFile, *FrameFilename))
+				{
+					UE_LOG(LogSparseVolumeTextureFactory, Error, TEXT("OpenVDB file could not be loaded: %s"), *FrameFilename);
+					bErrored.store(true);
+					return;
+				}
+
+				TArray<FOpenVDBGridInfo> FrameGridInfo;
+				if (!GetOpenVDBGridInfo(LoadedFrameFile, true, &FrameGridInfo))
+				{
+					UE_LOG(LogSparseVolumeTextureFactory, Error, TEXT("Failed to read OpenVDB file: %s"), *FrameFilename);
+					bErrored.store(true);
+					return;
+				}
+
+				// Sanity check for compatibility
+				for (const FOpenVDBSparseVolumeAttributesDesc& AttributesDesc : ImportOptions.Attributes)
+				{
+					for (const FOpenVDBSparseVolumeComponentMapping& Mapping : AttributesDesc.Mappings)
+					{
+						const uint32 SourceGridIndex = Mapping.SourceGridIndex;
+						if (SourceGridIndex != INDEX_NONE)
+						{
+							if ((int32)SourceGridIndex >= FrameGridInfo.Num())
+							{
+								UE_LOG(LogSparseVolumeTextureFactory, Error, TEXT("OpenVDB file is incompatible with other frames in the sequence: %s"), *FrameFilename);
+								bErrored.store(true);
+								return;
+							}
+							const FOpenVDBGridInfo& OrigSourceGrid = PreviewData.GridInfo[SourceGridIndex];
+							const FOpenVDBGridInfo& FrameSourceGrid = FrameGridInfo[SourceGridIndex];
+							if (OrigSourceGrid.Type != FrameSourceGrid.Type || OrigSourceGrid.Name != FrameSourceGrid.Name)
+							{
+								UE_LOG(LogSparseVolumeTextureFactory, Error, TEXT("OpenVDB file is incompatible with other frames in the sequence: %s"), *FrameFilename);
+								bErrored.store(true);
+								return;
+							}
+						}
+					}
+				}
+
+				// Update sequence volume bounds and increment ProcessedFramesCounter
+				{
+					std::lock_guard<std::mutex> Lock(VolumeBoundsMutex);
+					ExpandVolumeBounds(ImportOptions, FrameGridInfo, VolumeBoundsMin, VolumeBoundsMax);
+				}
+			});
+
+		if (bErrored.load())
+		{
+			return nullptr;
+		}
+
+		ImportTask.EnterProgressFrame(1.0f, LOCTEXT("ConvertingVDBAnim", "Converting OpenVDB animation"));
+
+		const int32 NumMipLevels = ComputeNumMipLevels(VolumeBoundsMin, VolumeBoundsMax);
+
+		FEvent* AllTasksFinishedEvent = FPlatformProcess::GetSynchEventFromPool();
 		std::atomic_int FinishedTasksCounter = 0; // Will be incremented even if frame processing failed
 		std::atomic_int ProcessedFramesCounter = 0;
-		std::mutex Mutex;
-		FBox VolumeBounds(FVector(FLT_MAX), FVector(-FLT_MAX));
 
-		// Load individual frames, check them for compatibility with the first loaded file and append them to the resulting asset
+		// Load individual frames, process/convert them and append them to the resulting asset
 		for (int32 FrameIdx = 0; FrameIdx < NumFrames; ++FrameIdx)
 		{
 			// Increments the atomic counter when going out of scope. Triggers an event once the counter reaches a given value.
@@ -572,8 +653,8 @@ UObject* USparseVolumeTextureFactory::FactoryCreateFile(UClass* InClass, UObject
 			};
 
 			AsyncTask(ENamedThreads::AnyNormalThreadNormalTask, 
-				[FrameIdx, NumFrames, &PreviewData, &ImportOptions, &AnimatedSVTexture, &ExpandVolumeBounds,
-				AllTasksFinishedEvent, &bErrored, &bCanceled, &Mutex, &FinishedTasksCounter, &ProcessedFramesCounter, &VolumeBounds]()
+				[FrameIdx, NumFrames, &PreviewData, &ImportOptions, &AnimatedSVTexture,
+				AllTasksFinishedEvent, &bErrored, &bCanceled, &FinishedTasksCounter, &ProcessedFramesCounter, &VolumeBoundsMin, &NumMipLevels]()
 				{
 					// Ensure the FinishedTasksCounter will be incremented in all cases
 					FScopedIncrementer Incremeter(FinishedTasksCounter, NumFrames, AllTasksFinishedEvent);
@@ -587,7 +668,7 @@ UObject* USparseVolumeTextureFactory::FactoryCreateFile(UClass* InClass, UObject
 
 					UE_LOG(LogSparseVolumeTextureFactory, Display, TEXT("Loading OpenVDB sequence frame #%i %s."), FrameIdx, *FrameFilename);
 
-					// Load file and get info about each contained grid
+					// Load file
 					TArray<uint8> LoadedFrameFile;
 					if (!FFileHelper::LoadFileToArray(LoadedFrameFile, *FrameFilename))
 					{
@@ -596,53 +677,8 @@ UObject* USparseVolumeTextureFactory::FactoryCreateFile(UClass* InClass, UObject
 						return;
 					}
 
-					TArray<FOpenVDBGridInfo> FrameGridInfo;
-					if (!GetOpenVDBGridInfo(LoadedFrameFile, true, &FrameGridInfo))
-					{
-						UE_LOG(LogSparseVolumeTextureFactory, Error, TEXT("Failed to read OpenVDB file: %s"), *FrameFilename);
-						bErrored.store(true);
-						return;
-					}
-
-					// Sanity check for compatibility
-					for (const FOpenVDBSparseVolumeAttributesDesc& AttributesDesc : ImportOptions.Attributes)
-					{
-						for (const FOpenVDBSparseVolumeComponentMapping& Mapping : AttributesDesc.Mappings)
-						{
-							const uint32 SourceGridIndex = Mapping.SourceGridIndex;
-							if (SourceGridIndex != INDEX_NONE)
-							{
-								if ((int32)SourceGridIndex >= FrameGridInfo.Num())
-								{
-									UE_LOG(LogSparseVolumeTextureFactory, Error, TEXT("OpenVDB file is incompatible with other frames in the sequence: %s"), *FrameFilename);
-									bErrored.store(true);
-									return;
-								}
-								const FOpenVDBGridInfo& OrigSourceGrid = PreviewData.GridInfo[SourceGridIndex];
-								const FOpenVDBGridInfo& FrameSourceGrid = FrameGridInfo[SourceGridIndex];
-								if (OrigSourceGrid.Type != FrameSourceGrid.Type || OrigSourceGrid.Name != FrameSourceGrid.Name)
-								{
-									UE_LOG(LogSparseVolumeTextureFactory, Error, TEXT("OpenVDB file is incompatible with other frames in the sequence: %s"), *FrameFilename);
-									bErrored.store(true);
-									return;
-								}
-							}
-						}
-					}
-
 					FSparseVolumeRawSource SparseVolumeRawSource{};
-
-					FOpenVDBToSVTConversionResult SVTResult;
-					SVTResult.Header = &SparseVolumeRawSource.Header;
-					SVTResult.PageTable = &SparseVolumeRawSource.PageTable;
-					SVTResult.PhysicalTileDataA = &SparseVolumeRawSource.PhysicalTileDataA;
-					SVTResult.PhysicalTileDataB = &SparseVolumeRawSource.PhysicalTileDataB;
-
-					const bool bConversionSuccess = ConvertOpenVDBToSparseVolumeTexture(
-						LoadedFrameFile,
-						ImportOptions,
-						&SVTResult,
-						false, FVector::Zero(), FVector::Zero());
+					const bool bConversionSuccess = ConvertOpenVDBToSparseVolumeTexture(LoadedFrameFile, ImportOptions, VolumeBoundsMin, SparseVolumeRawSource);
 
 					if (!bConversionSuccess)
 					{
@@ -652,16 +688,20 @@ UObject* USparseVolumeTextureFactory::FactoryCreateFile(UClass* InClass, UObject
 					}
 
 					// Serialize the raw source data from this frame into the asset object.
+					// After serializing the data, we generate the next mip level by downsampling the current one (with a call to GenerateMipMap())
+					// and then repeating the loop until we processed all levels.
+					AnimatedSVTexture->Frames[FrameIdx].SetNum(NumMipLevels);
+					for (int32 MipLevel = 0; MipLevel < NumMipLevels; ++MipLevel)
 					{
-						UE::Serialization::FEditorBulkDataWriter RawDataArchiveWriter(AnimatedSVTexture->Frames[FrameIdx].RawData);
+						UE::Serialization::FEditorBulkDataWriter RawDataArchiveWriter(AnimatedSVTexture->Frames[FrameIdx][MipLevel].RawData);
 						SparseVolumeRawSource.Serialize(RawDataArchiveWriter);
+						if ((MipLevel + 1) < NumMipLevels)
+						{
+							SparseVolumeRawSource = SparseVolumeRawSource.GenerateMipMap();
+						}
 					}
 
-					// Update sequence volume bounds and increment ProcessedFramesCounter
-					{
-						std::lock_guard<std::mutex> Lock(Mutex);
-						ExpandVolumeBounds(ImportOptions, FrameGridInfo, VolumeBounds);
-					}
+					// Increment ProcessedFramesCounter
 					ProcessedFramesCounter.fetch_add(1);
 				});
 		}
@@ -709,7 +749,7 @@ UObject* USparseVolumeTextureFactory::FactoryCreateFile(UClass* InClass, UObject
 			return nullptr;
 		}
 
-		AnimatedSVTexture->VolumeBounds = VolumeBounds;
+		AnimatedSVTexture->VolumeResolution = VolumeBoundsMax - VolumeBoundsMin;
 
 		ResultAssets.Add(AnimatedSVTexture);
 	}
diff --git a/Engine/Source/Editor/SparseVolumeTexture/Private/SparseVolumeTextureOpenVDBUtility.cpp b/Engine/Source/Editor/SparseVolumeTexture/Private/SparseVolumeTextureOpenVDBUtility.cpp
--- a/Engine/Source/Editor/SparseVolumeTexture/Private/SparseVolumeTextureOpenVDBUtility.cpp
+++ b/Engine/Source/Editor/SparseVolumeTexture/Private/SparseVolumeTextureOpenVDBUtility.cpp
@@ -12,10 +12,6 @@ DEFINE_LOG_CATEGORY_STATIC(LogSparseVolumeTextureOpenVDBUtility, Log, All);
 
 #if OPENVDB_AVAILABLE
 
-// We are using this instead of GMaxVolumeTextureDimensions to be independent of the platform that
-// the asset is imported on. 2048 should be a safe value that should be supported by all our platforms.
-static constexpr int32 SVTMaxVolumeTextureDim = 2048;
-
 namespace
 {
 	// Utility class acting as adapter between TArray<uint8> and std::istream
@@ -43,9 +39,9 @@ static FOpenVDBGridInfo GetOpenVDBGridInfo(openvdb::GridBase::Ptr Grid, uint32 G
 	GridInfo.Index = GridIndex;
 	GridInfo.NumComponents = 0;
 	GridInfo.Type = EOpenVDBGridType::Unknown;
-	GridInfo.VolumeActiveAABBMin = FVector(VolumeActiveAABB.min().x(), VolumeActiveAABB.min().y(), VolumeActiveAABB.min().z());
-	GridInfo.VolumeActiveAABBMax = FVector(VolumeActiveAABB.max().x(), VolumeActiveAABB.max().y(), VolumeActiveAABB.max().z());
-	GridInfo.VolumeActiveDim = FVector(VolumeActiveDim.x(), VolumeActiveDim.y(), VolumeActiveDim.z());
+	GridInfo.VolumeActiveAABBMin = FIntVector3(VolumeActiveAABB.min().x(), VolumeActiveAABB.min().y(), VolumeActiveAABB.min().z());
+	GridInfo.VolumeActiveAABBMax = FIntVector3(VolumeActiveAABB.max().x() + 1, VolumeActiveAABB.max().y() + 1, VolumeActiveAABB.max().z() + 1); // +1 because CoordBBox::Max is inclusive, but we want exclusive
+	GridInfo.VolumeActiveDim = FIntVector3(VolumeActiveDim.x(), VolumeActiveDim.y(), VolumeActiveDim.z());
 	GridInfo.VolumeVoxelSize = FVector(VolumeVoxelSize.x(), VolumeVoxelSize.y(), VolumeVoxelSize.z());
 	GridInfo.bIsInWorldSpace = Grid->isInWorldSpace();
 	GridInfo.bHasUniformVoxels = Grid->hasUniformVoxels();
@@ -199,439 +195,231 @@ static EPixelFormat GetMultiComponentFormat(ESparseVolumeAttributesFormat Format
 	return PF_Unknown;
 }
 
-static uint32 PackPageTableEntry(const FIntVector3& Coord)
-{
-	// A page encodes the physical tile coord as unsigned int of 11 11 10 bits
-	// This means a page coord cannot be larger than 2047 for x and y and 1023 for z
-	// which mean we cannot have more than 2048*2048*1024 = 4 Giga tiles of 16^3 tiles.
-	uint32 Result = (Coord.X & 0x7FF) | ((Coord.Y & 0x7FF) << 11) | ((Coord.Z & 0x3FF) << 22);
-	return Result;
-}
 
-static FIntVector3 UnpackPageTableEntry(uint32 Packed)
-{
-	FIntVector3 Result;
-	Result.X = Packed & 0x7FF;
-	Result.Y = (Packed >> 11) & 0x7FF;
-	Result.Z = (Packed >> 22) & 0x3FF;
-	return Result;
-}
+#if OPENVDB_AVAILABLE
 
-bool ConvertOpenVDBToSparseVolumeTexture(
-	TArray<uint8>& SourceFile,
-	FOpenVDBImportOptions& ImportOptions,
-	FOpenVDBToSVTConversionResult* OutResult,
-	bool bOverrideActiveMinMax,
-	FVector ActiveMin,
-	FVector ActiveMax)
+class FSparseVolumeRawSourceConstructionOpenVDBAdapter : public ISparseVolumeRawSourceConstructionAdapter
 {
-#if OPENVDB_AVAILABLE
+public:
 
-	constexpr uint32 NumAttributesDescs = 2; // AttributesA and AttributesB, representing the two textures with voxel data
-	TStaticArray<FOpenVDBSparseVolumeAttributesDesc, 2>& Attributes = ImportOptions.Attributes;
+	bool Initialize(TArray<uint8>& SourceFile, const FOpenVDBImportOptions& ImportOptions, const FIntVector3& InVolumeBoundsMin)
+	{
+		Attributes = ImportOptions.Attributes;
+		VolumeBoundsMin = InVolumeBoundsMin;
 
-	// Compute some basic info about the number of components and which format to use
-	uint32 NumActualComponents[NumAttributesDescs] = {};
-	EPixelFormat MultiCompFormat[NumAttributesDescs] = {};
-	uint32 FormatSize[NumAttributesDescs] = {};
-	uint32 SingleComponentFormatSize[NumAttributesDescs] = {};
-	bool bNormalizedFormat[NumAttributesDescs] = {};
-	bool bHasValidSourceGrids[NumAttributesDescs] = {};
-	bool bAnySourceGridIndicesValid = false;
+		// Compute some basic info about the number of components and which format to use
+		EPixelFormat MultiCompFormat[NumAttributesDescs] = {};
+		bool bNormalizedFormat[NumAttributesDescs] = {};
+		bool bHasValidSourceGrids[NumAttributesDescs] = {};
+		bool bAnySourceGridIndicesValid = false;
 
-	for (uint32 AttributesIdx = 0; AttributesIdx < NumAttributesDescs; ++AttributesIdx)
-	{
-		uint32 NumRequiredComponents = 0;
-		for (uint32 ComponentIdx = 0; ComponentIdx < 4; ++ComponentIdx)
+		for (int32 AttributesIdx = 0; AttributesIdx < NumAttributesDescs; ++AttributesIdx)
 		{
-			if (Attributes[AttributesIdx].Mappings[ComponentIdx].SourceGridIndex != INDEX_NONE)
+			int32 NumRequiredComponents = 0;
+			for (int32 ComponentIdx = 0; ComponentIdx < 4; ++ComponentIdx)
 			{
-				check(Attributes[AttributesIdx].Mappings[ComponentIdx].SourceComponentIndex != INDEX_NONE);
-				NumRequiredComponents = FMath::Max(ComponentIdx + 1, NumRequiredComponents);
-				bHasValidSourceGrids[AttributesIdx] = true;
-				bAnySourceGridIndicesValid = true;
+				if (Attributes[AttributesIdx].Mappings[ComponentIdx].SourceGridIndex != INDEX_NONE)
+				{
+					check(Attributes[AttributesIdx].Mappings[ComponentIdx].SourceComponentIndex != INDEX_NONE);
+					NumRequiredComponents = FMath::Max(ComponentIdx + 1, NumRequiredComponents);
+					bHasValidSourceGrids[AttributesIdx] = true;
+					bAnySourceGridIndicesValid = true;
+				}
 			}
-		}
-
-		if (bHasValidSourceGrids[AttributesIdx])
-		{
-			NumActualComponents[AttributesIdx] = NumRequiredComponents == 3 ? 4 : NumRequiredComponents; // We don't support formats with only 3 components
-			bNormalizedFormat[AttributesIdx] = Attributes[AttributesIdx].Format == ESparseVolumeAttributesFormat::Unorm8;
-			MultiCompFormat[AttributesIdx] = GetMultiComponentFormat(Attributes[AttributesIdx].Format, NumActualComponents[AttributesIdx]);
 
-			if (MultiCompFormat[AttributesIdx] == PF_Unknown)
+			if (bHasValidSourceGrids[AttributesIdx])
 			{
-				UE_LOG(LogSparseVolumeTextureOpenVDBUtility, Warning, TEXT("SparseVolumeTexture is set to use an unsupported format: %i"), (int32)Attributes[AttributesIdx].Format);
-				return false;
-			}
+				NumComponents[AttributesIdx] = NumRequiredComponents == 3 ? 4 : NumRequiredComponents; // We don't support formats with only 3 components
+				bNormalizedFormat[AttributesIdx] = Attributes[AttributesIdx].Format == ESparseVolumeAttributesFormat::Unorm8;
+				MultiCompFormat[AttributesIdx] = GetMultiComponentFormat(Attributes[AttributesIdx].Format, NumComponents[AttributesIdx]);
 
-			FormatSize[AttributesIdx] = (uint32)GPixelFormats[(SIZE_T)MultiCompFormat[AttributesIdx]].BlockBytes;
-			SingleComponentFormatSize[AttributesIdx] = FormatSize[AttributesIdx] / NumActualComponents[AttributesIdx];
+				if (MultiCompFormat[AttributesIdx] == PF_Unknown)
+				{
+					UE_LOG(LogSparseVolumeTextureOpenVDBUtility, Warning, TEXT("SparseVolumeTexture is set to use an unsupported format: %i"), (int32)Attributes[AttributesIdx].Format);
+					return false;
+				}
+			}
 		}
-	}
-
-	// All source grid indices are INDEX_NONE, so nothing was selected for import
-	if (!bAnySourceGridIndicesValid)
-	{
-		UE_LOG(LogSparseVolumeTextureOpenVDBUtility, Warning, TEXT("SparseVolumeTexture has all components set to <None>, so there is nothing to import."));
-		return false;
-	}
-
-	// Load file
-	FArrayUint8StreamBuf StreamBuf(SourceFile);
-	std::istream IStream(&StreamBuf);
-	openvdb::io::Stream Stream(IStream, false /*delayLoad*/);
 
-	// Check that the source grid indices are valid
-	openvdb::GridPtrVecPtr Grids = Stream.getGrids();
-	const size_t NumSourceGrids = Grids ? Grids->size() : 0;
-	for (const FOpenVDBSparseVolumeAttributesDesc& AttributesDesc : ImportOptions.Attributes)
-	{
-		for (const FOpenVDBSparseVolumeComponentMapping& Mapping : AttributesDesc.Mappings)
+		// All source grid indices are INDEX_NONE, so nothing was selected for import
+		if (!bAnySourceGridIndicesValid)
 		{
-			const int32 SourceGridIndex = Mapping.SourceGridIndex;
-			if (SourceGridIndex != INDEX_NONE && SourceGridIndex >= (int32)NumSourceGrids)
-			{
-				UE_LOG(LogSparseVolumeTextureOpenVDBUtility, Warning, TEXT("SparseVolumeTexture has invalid index into the array of grids in the source file: %i"), SourceGridIndex);
-				return false;
-			}
+			UE_LOG(LogSparseVolumeTextureOpenVDBUtility, Warning, TEXT("SparseVolumeTexture has all components set to <None>, so there is nothing to import."));
+			return false;
 		}
-	}
 
-	FSparseVolumeAssetHeader& Header = *OutResult->Header;
-	Header.AttributesAFormat = MultiCompFormat[0];
-	Header.AttributesBFormat = MultiCompFormat[1];
-	Header.SourceVolumeResolution = FIntVector::ZeroValue;
-	
-	FIntVector SmallestAABBMin = FIntVector(INT32_MAX);
+		// Load file
+		FArrayUint8StreamBuf StreamBuf(SourceFile);
+		std::istream IStream(&StreamBuf);
+		openvdb::io::Stream Stream(IStream, false /*delayLoad*/);
 
-	// Compute per source grid data of up to 4 different grids (one per component)
-	TArray<TSharedPtr<IOpenVDBGridAdapterBase>> UniqueGridAdapters;
-	UniqueGridAdapters.SetNum((int32)NumSourceGrids);
-	TSharedPtr<IOpenVDBGridAdapterBase> GridAdapters[NumAttributesDescs][4]{};
-	float GridBackgroundValues[NumAttributesDescs][4]{};
-	float NormalizeScale[NumAttributesDescs][4]{};
-	float NormalizeBias[NumAttributesDescs][4]{};
-	for (uint32 AttributesIdx = 0; AttributesIdx < NumAttributesDescs; ++AttributesIdx)
-	{
-		for (uint32 CompIdx = 0; CompIdx < 4; ++CompIdx)
+		// Check that the source grid indices are valid
+		openvdb::GridPtrVecPtr Grids = Stream.getGrids();
+		const size_t NumSourceGrids = Grids ? Grids->size() : 0;
+		for (const FOpenVDBSparseVolumeAttributesDesc& AttributesDesc : Attributes)
 		{
-			NormalizeScale[AttributesIdx][CompIdx] = 1.0f;
-			const uint32 SourceGridIndex = Attributes[AttributesIdx].Mappings[CompIdx].SourceGridIndex;
-			const uint32 SourceComponentIndex = Attributes[AttributesIdx].Mappings[CompIdx].SourceComponentIndex;
-			if (SourceGridIndex == INDEX_NONE)
-			{
-				continue;
-			}
-
-			openvdb::GridBase::Ptr GridBase = (*Grids)[SourceGridIndex];
-
-			// Try to reuse adapters. Internally they use caching to accelerate read accesses, 
-			// so using three different adapters to access the three components of a single grid would be wasteful.
-			if (UniqueGridAdapters[SourceGridIndex] == nullptr)
+			for (const FOpenVDBSparseVolumeComponentMapping& Mapping : AttributesDesc.Mappings)
 			{
-				UniqueGridAdapters[SourceGridIndex] = CreateOpenVDBGridAdapter(GridBase);
-				if (!UniqueGridAdapters[SourceGridIndex])
+				const int32 SourceGridIndex = Mapping.SourceGridIndex;
+				if (SourceGridIndex != INDEX_NONE && SourceGridIndex >= (int32)NumSourceGrids)
 				{
+					UE_LOG(LogSparseVolumeTextureOpenVDBUtility, Warning, TEXT("SparseVolumeTexture has invalid index into the array of grids in the source file: %i"), SourceGridIndex);
 					return false;
 				}
 			}
+		}
 
-			GridAdapters[AttributesIdx][CompIdx] = UniqueGridAdapters[SourceGridIndex];
-
-			FOpenVDBGridInfo GridInfo = GetOpenVDBGridInfo(GridBase, 0, false);
-			if (!IsOpenVDBGridValid(GridInfo, TEXT("")))
-			{
-				return false;
-			}
-
-			if (bOverrideActiveMinMax)
-			{
-				GridInfo.VolumeActiveAABBMin = ActiveMin;
-				GridInfo.VolumeActiveAABBMax = ActiveMax;
-				GridInfo.VolumeActiveDim = ActiveMax + 1 - ActiveMin;
-			}
+		AttributesInfo[0].Format = MultiCompFormat[0];
+		AttributesInfo[1].Format = MultiCompFormat[1];
+		AttributesInfo[0].bNormalized = Attributes[0].bRemapInputForUnorm;
+		AttributesInfo[1].bNormalized = Attributes[1].bRemapInputForUnorm;
 
-			Header.SourceVolumeResolution.X = FMath::Max(Header.SourceVolumeResolution.X, GridInfo.VolumeActiveDim.X);
-			Header.SourceVolumeResolution.Y = FMath::Max(Header.SourceVolumeResolution.Y, GridInfo.VolumeActiveDim.Y);
-			Header.SourceVolumeResolution.Z = FMath::Max(Header.SourceVolumeResolution.Z, GridInfo.VolumeActiveDim.Z);
-			SmallestAABBMin.X = FMath::Min(SmallestAABBMin.X, GridInfo.VolumeActiveAABBMin.X);
-			SmallestAABBMin.Y = FMath::Min(SmallestAABBMin.Y, GridInfo.VolumeActiveAABBMin.Y);
-			SmallestAABBMin.Z = FMath::Min(SmallestAABBMin.Z, GridInfo.VolumeActiveAABBMin.Z);
+		FIntVector3 SmallestAABBMin = FIntVector3(INT32_MAX);
+		FIntVector3 LargestAABBMax = FIntVector3(INT32_MIN);
 
-			GridBackgroundValues[AttributesIdx][CompIdx] = GridAdapters[AttributesIdx][CompIdx]->GetBackgroundValue(SourceComponentIndex);
-			if (bNormalizedFormat[AttributesIdx] && Attributes[AttributesIdx].bRemapInputForUnorm)
+		// Compute per source grid data of up to 4 different grids (one per component)
+		UniqueGridAdapters.SetNum((int32)NumSourceGrids);
+		GridToComponentMappings.SetNum((int32)NumSourceGrids);
+		for (int32 AttributesIdx = 0; AttributesIdx < NumAttributesDescs; ++AttributesIdx)
+		{
+			for (int32 CompIdx = 0; CompIdx < 4; ++CompIdx)
 			{
-				float MinVal = 0.0f;
-				float MaxVal = 0.0f;
-				GridAdapters[AttributesIdx][CompIdx]->GetMinMaxValue(SourceComponentIndex, &MinVal, &MaxVal);
-				const float Diff = MaxVal - MinVal;
-				NormalizeScale[AttributesIdx][CompIdx] = MaxVal > SMALL_NUMBER ? (1.0f / Diff) : 1.0f;
-				NormalizeBias[AttributesIdx][CompIdx] = -MinVal * NormalizeScale[AttributesIdx][CompIdx];
-			}
-		}
-	}
-	OutResult->Header->SourceVolumeAABBMin = SmallestAABBMin;
-	
-	FIntVector3 PageTableVolumeResolution = FIntVector3(
-		FMath::DivideAndRoundUp(Header.SourceVolumeResolution.X, SPARSE_VOLUME_TILE_RES),
-		FMath::DivideAndRoundUp(Header.SourceVolumeResolution.Y, SPARSE_VOLUME_TILE_RES),
-		FMath::DivideAndRoundUp(Header.SourceVolumeResolution.Z, SPARSE_VOLUME_TILE_RES));
-	if (PageTableVolumeResolution.X > SVTMaxVolumeTextureDim
-		|| PageTableVolumeResolution.Y > SVTMaxVolumeTextureDim
-		|| PageTableVolumeResolution.Z > SVTMaxVolumeTextureDim)
-	{
-		UE_LOG(LogSparseVolumeTextureOpenVDBUtility, Warning, TEXT("SparseVolumeTexture page table texture dimensions exceed limit (%ix%ix%i): %ix%ix%i"), SVTMaxVolumeTextureDim, SVTMaxVolumeTextureDim, SVTMaxVolumeTextureDim, PageTableVolumeResolution.X, PageTableVolumeResolution.Y, PageTableVolumeResolution.Z);
-		return false;
-	}
-
-	Header.PageTableVolumeResolution = PageTableVolumeResolution;
-	Header.TileDataVolumeResolution = FIntVector::ZeroValue;	// unknown for now
+				AttributesInfo[AttributesIdx].NormalizeScale[CompIdx] = 1.0f;
+				const uint32 SourceGridIndex = Attributes[AttributesIdx].Mappings[CompIdx].SourceGridIndex;
+				const uint32 SourceComponentIndex = Attributes[AttributesIdx].Mappings[CompIdx].SourceComponentIndex;
+				if (SourceGridIndex == INDEX_NONE)
+				{
+					continue;
+				}
 
-	// Tag all pages with valid data
-	TBitArray PagesWithData(false, Header.PageTableVolumeResolution.Z * Header.PageTableVolumeResolution.Y * Header.PageTableVolumeResolution.X);
-	for (uint32 GridIdx = 0; GridIdx < NumSourceGrids; ++GridIdx)
-	{
-		if (!UniqueGridAdapters[GridIdx])
-		{
-			continue;
-		}
+				openvdb::GridBase::Ptr GridBase = (*Grids)[SourceGridIndex];
 
-		UniqueGridAdapters[GridIdx]->IteratePhysical(
-			[&](const FIntVector3& Coord, uint32 NumComponents, float* VoxelValues)
-			{
-				// Check if the source grid component is used at all
-				bool bValueIsUsed = false;
-				for (uint32 AttributesIdx = 0; AttributesIdx < NumAttributesDescs && !bValueIsUsed; ++AttributesIdx)
+				// Try to reuse adapters. Internally they use caching to accelerate read accesses, 
+				// so using three different adapters to access the three components of a single grid would be wasteful.
+				if (UniqueGridAdapters[SourceGridIndex] == nullptr)
 				{
-					for (uint32 DstCompIdx = 0; DstCompIdx < NumActualComponents[AttributesIdx] && !bValueIsUsed; ++DstCompIdx)
+					UniqueGridAdapters[SourceGridIndex] = CreateOpenVDBGridAdapter(GridBase);
+					if (!UniqueGridAdapters[SourceGridIndex])
 					{
-						for (uint32 SrcCompIdx = 0; SrcCompIdx < NumComponents && !bValueIsUsed; ++SrcCompIdx)
-						{
-							const bool bIsBackgroundValue = (VoxelValues[SrcCompIdx] == GridBackgroundValues[AttributesIdx][DstCompIdx]);
-							bValueIsUsed |= !bIsBackgroundValue && (Attributes[AttributesIdx].Mappings[DstCompIdx].SourceGridIndex == GridIdx) && (Attributes[AttributesIdx].Mappings[DstCompIdx].SourceComponentIndex == SrcCompIdx);
-						}
+						return false;
 					}
 				}
-				if (!bValueIsUsed)
+
+				FOpenVDBGridInfo GridInfo = GetOpenVDBGridInfo(GridBase, 0, false);
+				if (!IsOpenVDBGridValid(GridInfo, TEXT("")))
 				{
-					return;
+					return false;
 				}
 
-				const FIntVector3 GridCoord = Coord - SmallestAABBMin;
-				check(GridCoord.X >= 0 && GridCoord.Y >= 0 && GridCoord.Z >= 0);
-				check(GridCoord.X < Header.SourceVolumeResolution.X && GridCoord.Y < Header.SourceVolumeResolution.Y && GridCoord.Z < Header.SourceVolumeResolution.Z);
-				const FIntVector3 PageCoord = GridCoord / SPARSE_VOLUME_TILE_RES;
-				check(PageCoord.X >= 0 && PageCoord.Y >= 0 && PageCoord.Z >= 0);
-				check(PageCoord.X < Header.PageTableVolumeResolution.X && PageCoord.Y < Header.PageTableVolumeResolution.Y && PageCoord.Z < Header.PageTableVolumeResolution.Z);
-
-				const int32 PageIndex = PageCoord.Z * (Header.PageTableVolumeResolution.Y * Header.PageTableVolumeResolution.X) + PageCoord.Y * Header.PageTableVolumeResolution.X + PageCoord.X;
-				PagesWithData[PageIndex] = true;
-			});
-	}
+				SmallestAABBMin.X = FMath::Min(SmallestAABBMin.X, GridInfo.VolumeActiveAABBMin.X);
+				SmallestAABBMin.Y = FMath::Min(SmallestAABBMin.Y, GridInfo.VolumeActiveAABBMin.Y);
+				SmallestAABBMin.Z = FMath::Min(SmallestAABBMin.Z, GridInfo.VolumeActiveAABBMin.Z);
+				LargestAABBMax.X = FMath::Max(LargestAABBMax.X, GridInfo.VolumeActiveAABBMax.X);
+				LargestAABBMax.Y = FMath::Max(LargestAABBMax.Y, GridInfo.VolumeActiveAABBMax.Y);
+				LargestAABBMax.Z = FMath::Max(LargestAABBMax.Z, GridInfo.VolumeActiveAABBMax.Z);
 
-	// Allocate some memory for temp data (worst case)
-	TArray<FIntVector3> LinearAllocatedPages;
-	LinearAllocatedPages.SetNum(Header.PageTableVolumeResolution.X * Header.PageTableVolumeResolution.Y * Header.PageTableVolumeResolution.Z);
-
-	// Go over each potential page from the source data and push allocate it if it has any data.
-	// Otherwise point to the default empty page.
-	bool bAnyEmptyPageExists = false;
-	uint32 NumAllocatedPages = 0;
-	for (int32_t PageZ = 0; PageZ < Header.PageTableVolumeResolution.Z; ++PageZ)
-	{
-		for (int32_t PageY = 0; PageY < Header.PageTableVolumeResolution.Y; ++PageY)
-		{
-			for (int32_t PageX = 0; PageX < Header.PageTableVolumeResolution.X; ++PageX)
-			{
-				const int32 PageIndex = PageZ * (Header.PageTableVolumeResolution.Y * Header.PageTableVolumeResolution.X) + PageY * Header.PageTableVolumeResolution.X + PageX;
-				bool bHasAnyData = PagesWithData[PageIndex];
-				if (bHasAnyData)
+				AttributesInfo[AttributesIdx].FallbackValue[CompIdx] = UniqueGridAdapters[SourceGridIndex]->GetBackgroundValue(SourceComponentIndex);
+				if (bNormalizedFormat[AttributesIdx] && Attributes[AttributesIdx].bRemapInputForUnorm)
 				{
-					LinearAllocatedPages[NumAllocatedPages] = FIntVector3(PageX, PageY, PageZ);
-					NumAllocatedPages++;
+					float MinVal = 0.0f;
+					float MaxVal = 0.0f;
+					UniqueGridAdapters[SourceGridIndex]->GetMinMaxValue(SourceComponentIndex, &MinVal, &MaxVal);
+					const float Diff = MaxVal - MinVal;
+					AttributesInfo[AttributesIdx].NormalizeScale[CompIdx] = MaxVal > SMALL_NUMBER ? (1.0f / Diff) : 1.0f;
+					AttributesInfo[AttributesIdx].NormalizeBias[CompIdx] = -MinVal * AttributesInfo[AttributesIdx].NormalizeScale[CompIdx];
 				}
-				bAnyEmptyPageExists |= !bHasAnyData;
+
+				FSingleGridToComponentMapping Mapping{};
+				Mapping.AttributesIdx = (int32)AttributesIdx;
+				Mapping.ComponentIdx = (int32)CompIdx;
+				Mapping.GridComponentIdx = (int32)SourceComponentIndex;
+				GridToComponentMappings[SourceGridIndex].Add(Mapping);
 			}
 		}
-	}
 
-	// Compute Page and Tile VolumeResolution from allocated pages
-	const uint32 EffectivelyAllocatedPageEntries = NumAllocatedPages + (bAnyEmptyPageExists ? 1 : 0);
-	uint32 TileVolumeResolutionCube = 1;
-	while (TileVolumeResolutionCube * TileVolumeResolutionCube * TileVolumeResolutionCube < EffectivelyAllocatedPageEntries)
-	{
-		TileVolumeResolutionCube++;				// We use a simple loop to compute the minimum resolution of a cube to store all the tile data
+		AABBMin = SmallestAABBMin - VolumeBoundsMin;
+		AABBMax = LargestAABBMax - VolumeBoundsMin;
+		Resolution = LargestAABBMax - SmallestAABBMin;
+
+		return true;
 	}
-	Header.TileDataVolumeResolution = FIntVector3(TileVolumeResolutionCube, TileVolumeResolutionCube, TileVolumeResolutionCube);
-	while (Header.TileDataVolumeResolution.X * Header.TileDataVolumeResolution.Y * (Header.TileDataVolumeResolution.Z - 1) > int32(EffectivelyAllocatedPageEntries))
+
+	void GetAttributesInfo(FAttributesInfo& OutInfoA, FAttributesInfo& OutInfoB) const override
 	{
-		Header.TileDataVolumeResolution.Z--;	// We then trim an edge to get back space.
+		OutInfoA = AttributesInfo[0];
+		OutInfoB = AttributesInfo[1];
 	}
-	const FIntVector3 TileCoordResolution = Header.TileDataVolumeResolution;
-	Header.TileDataVolumeResolution = Header.TileDataVolumeResolution * SPARSE_VOLUME_TILE_RES;
-	if (Header.TileDataVolumeResolution.X > SVTMaxVolumeTextureDim
-		|| Header.TileDataVolumeResolution.Y > SVTMaxVolumeTextureDim
-		|| Header.TileDataVolumeResolution.Z > SVTMaxVolumeTextureDim)
+
+	FIntVector3 GetAABBMin() const override
 	{
-		UE_LOG(LogSparseVolumeTextureOpenVDBUtility, Warning, TEXT("SparseVolumeTexture tile data texture dimensions exceed limit (%ix%ix%i): %ix%ix%i"), SVTMaxVolumeTextureDim, SVTMaxVolumeTextureDim, SVTMaxVolumeTextureDim, Header.TileDataVolumeResolution.X, Header.TileDataVolumeResolution.Y, Header.TileDataVolumeResolution.Z);
-		return false;
+		return AABBMin;
 	}
 
-	// Initialise the SparseVolumeTexture page and tile.
-	OutResult->PageTable->SetNumZeroed(Header.PageTableVolumeResolution.X * Header.PageTableVolumeResolution.Y * Header.PageTableVolumeResolution.Z);
-	OutResult->PhysicalTileDataA->SetNumZeroed(Header.TileDataVolumeResolution.X * Header.TileDataVolumeResolution.Y * Header.TileDataVolumeResolution.Z * FormatSize[0] * (bHasValidSourceGrids[0] ? 1 : 0));
-	OutResult->PhysicalTileDataB->SetNumZeroed(Header.TileDataVolumeResolution.X * Header.TileDataVolumeResolution.Y * Header.TileDataVolumeResolution.Z * FormatSize[1] * (bHasValidSourceGrids[1] ? 1 : 0));
-	uint32* PageTablePtr = OutResult->PageTable->GetData();
-	uint8* PhysicalTileDataPtrs[] = { OutResult->PhysicalTileDataA->GetData(), OutResult->PhysicalTileDataB->GetData() };
-
-	// Generate page table and tile volume data by splatting the data
+	FIntVector3 GetAABBMax() const override
 	{
-		FIntVector3 DstTileCoord = FIntVector3::ZeroValue;
-		auto GoToNextTileCoord = [&]()
-		{
-			DstTileCoord.X++;
-			if (DstTileCoord.X >= TileCoordResolution.X)
-			{
-				DstTileCoord.X = 0;
-				DstTileCoord.Y++;
-			}
-			if (DstTileCoord.Y >= TileCoordResolution.Y)
-			{
-				DstTileCoord.Y = 0;
-				DstTileCoord.Z++;
-			}
-		};
-
-		// Add an empty tile is needed, reserve slot at coord 0
-		if (bAnyEmptyPageExists)
-		{
-			// PageTable is all cleared to zero, simply skip a tile
-			GoToNextTileCoord();
-		}
-
-		for (uint32 i = 0; i < NumAllocatedPages; ++i)
-		{
-			FIntVector3 PageCoordToSplat = LinearAllocatedPages[i];
-			uint32 DestinationTileCoord32bit = PackPageTableEntry(DstTileCoord);
-
-			// Setup the page table entry
-			PageTablePtr
-				[
-					PageCoordToSplat.Z * Header.PageTableVolumeResolution.X * Header.PageTableVolumeResolution.Y +
-					PageCoordToSplat.Y * Header.PageTableVolumeResolution.X +
-				PageCoordToSplat.X
-				] = DestinationTileCoord32bit;
-
-			// Set the next tile to be written to
-			GoToNextTileCoord();
-		}
+		return AABBMax;
 	}
 
-	auto WriteVoxelComponent = [&](const FIntVector3& DstCoord, uint32 AttributesIdx, uint32 CompIdx, float Value)
+	FIntVector3 GetResolution() const override
 	{
-		const size_t DstVoxelIndex = DstCoord.Z * Header.TileDataVolumeResolution.X * Header.TileDataVolumeResolution.Y 
-			+ DstCoord.Y * Header.TileDataVolumeResolution.X 
-			+ DstCoord.X;
-
-		const float VoxelValue = Value;
-		const float VoxelValueNormalized = FMath::Clamp(VoxelValue * NormalizeScale[AttributesIdx][CompIdx] + NormalizeBias[AttributesIdx][CompIdx], 0.0f, 1.0f);
-
-		const size_t DstIdx = DstVoxelIndex * FormatSize[AttributesIdx] + CompIdx * SingleComponentFormatSize[AttributesIdx];
-
-		switch (Attributes[AttributesIdx].Format)
-		{
-		case ESparseVolumeAttributesFormat::Unorm8:
-		{
-			PhysicalTileDataPtrs[AttributesIdx][DstIdx] = uint8(VoxelValueNormalized * 255.0f);
-			break;
-		}
-		case ESparseVolumeAttributesFormat::Float16:
-		{
-			const uint16 VoxelValue16FEncoded = FFloat16(VoxelValue).Encoded;
-			*((uint16*)(&PhysicalTileDataPtrs[AttributesIdx][DstIdx])) = VoxelValue16FEncoded;
-			break;
-		}
-		case ESparseVolumeAttributesFormat::Float32:
-		{
-			*((float*)(&PhysicalTileDataPtrs[AttributesIdx][DstIdx])) = VoxelValue;
-			break;
-		}
-		default:
-			checkNoEntry();
-			break;
-		}
-	};
+		return Resolution;
+	}
 
-	// Fill null tile with fallback/background value
-	if (bAnyEmptyPageExists)
+	void IteratePhysicalSource(TFunctionRef<void(const FIntVector3& Coord, int32 AttributesIdx, int32 ComponentIdx, float VoxelValue)> OnVisit) const override
 	{
-		for (uint32 AttributesIdx = 0; AttributesIdx < NumAttributesDescs; ++AttributesIdx)
+		for (int32 GridIdx = 0; GridIdx < UniqueGridAdapters.Num(); ++GridIdx)
 		{
-			for (uint32 Z = 0; Z < SPARSE_VOLUME_TILE_RES; ++Z)
+			if (!UniqueGridAdapters[GridIdx])
 			{
-				for (uint32 Y = 0; Y < SPARSE_VOLUME_TILE_RES; ++Y)
+				continue;
+			}
+
+			UniqueGridAdapters[GridIdx]->IteratePhysical(
+				[&](const FIntVector3& Coord, uint32 NumVoxelComponents, float* VoxelValues)
 				{
-					for (uint32 X = 0; X < SPARSE_VOLUME_TILE_RES; ++X)
+					for (const FSingleGridToComponentMapping& Mapping : GridToComponentMappings[GridIdx])
 					{
-						for (uint32 CompIdx = 0; CompIdx < NumActualComponents[AttributesIdx]; ++CompIdx)
-						{
-							WriteVoxelComponent(FIntVector3(X, Y, Z), AttributesIdx, CompIdx, GridBackgroundValues[AttributesIdx][CompIdx]);
-						}
+						OnVisit(Coord - VolumeBoundsMin, Mapping.AttributesIdx, Mapping.ComponentIdx, VoxelValues[Mapping.GridComponentIdx]);
 					}
-				}
-			}
+				});
 		}
 	}
 
-	// Write physical tile data
-	for (uint32 GridIdx = 0; GridIdx < NumSourceGrids; ++GridIdx)
-	{
-		if (!UniqueGridAdapters[GridIdx])
-		{
-			continue;
-		}
-
-		UniqueGridAdapters[GridIdx]->IteratePhysical(
-			[&](const FIntVector3& Coord, uint32 NumComponents, float* VoxelValues)
-			{
-				const FIntVector3 GridCoord = Coord - SmallestAABBMin;
-				check(GridCoord.X >= 0 && GridCoord.Y >= 0 && GridCoord.Z >= 0);
-				check(GridCoord.X < Header.SourceVolumeResolution.X&& GridCoord.Y < Header.SourceVolumeResolution.Y&& GridCoord.Z < Header.SourceVolumeResolution.Z);
-				const FIntVector3 PageCoord = GridCoord / SPARSE_VOLUME_TILE_RES;
-				check(PageCoord.X >= 0 && PageCoord.Y >= 0 && PageCoord.Z >= 0);
-				check(PageCoord.X < Header.PageTableVolumeResolution.X&& PageCoord.Y < Header.PageTableVolumeResolution.Y&& PageCoord.Z < Header.PageTableVolumeResolution.Z);
+private:
 
-				const int32 PageIndex = PageCoord.Z * (Header.PageTableVolumeResolution.Y * Header.PageTableVolumeResolution.X) + PageCoord.Y * Header.PageTableVolumeResolution.X + PageCoord.X;
+	struct FSingleGridToComponentMapping
+	{
+		int32 AttributesIdx;
+		int32 ComponentIdx;
+		int32 GridComponentIdx;
+	};
 
-				if (!PagesWithData[PageIndex])
-				{
-					return;
-				}
+	static constexpr int32 NumAttributesDescs = 2;
+	TArray<TSharedPtr<IOpenVDBGridAdapterBase>> UniqueGridAdapters;
+	TArray<TArray<FSingleGridToComponentMapping, TInlineAllocator<4>>> GridToComponentMappings;
+	TStaticArray<FOpenVDBSparseVolumeAttributesDesc, NumAttributesDescs> Attributes;
+	TStaticArray<uint32, NumAttributesDescs> NumComponents;
+	TStaticArray<FAttributesInfo, NumAttributesDescs> AttributesInfo;
+	FIntVector3 AABBMin;
+	FIntVector3 AABBMax;
+	FIntVector3 VolumeBoundsMin;
+	FIntVector3 Resolution;
+};
 
-				const FIntVector3 DstTileCoord = UnpackPageTableEntry(PageTablePtr[PageIndex]);
-				const FIntVector3 TileLocalCoord = GridCoord % SPARSE_VOLUME_TILE_RES;
+#endif // OPENVDB_AVAILABLE
 
-				// Check all output components and splat VoxelValue if they map to this source grid/component
-				for (uint32 SrcCompIdx = 0; SrcCompIdx < NumComponents; ++SrcCompIdx)
-				{
-					for (uint32 AttributesIdx = 0; AttributesIdx < NumAttributesDescs; ++AttributesIdx)
-					{
-						for (uint32 DstCompIdx = 0; DstCompIdx < NumActualComponents[AttributesIdx]; ++DstCompIdx)
-						{
-							if ((Attributes[AttributesIdx].Mappings[DstCompIdx].SourceGridIndex != GridIdx) || (Attributes[AttributesIdx].Mappings[DstCompIdx].SourceComponentIndex != SrcCompIdx))
-							{
-								continue;
-							}
-
-							const FIntVector3 DstCoord = DstTileCoord * SPARSE_VOLUME_TILE_RES + TileLocalCoord;
-							WriteVoxelComponent(DstCoord, AttributesIdx, DstCompIdx, VoxelValues[SrcCompIdx]);
-						}
-					}
-				}
-			});
+bool ConvertOpenVDBToSparseVolumeTexture(TArray<uint8>& SourceFile, const FOpenVDBImportOptions& ImportOptions, const FIntVector3& VolumeBoundsMin, FSparseVolumeRawSource& OutResult)
+{
+#if OPENVDB_AVAILABLE
+	FSparseVolumeRawSourceConstructionOpenVDBAdapter Adapter;
+	if (!Adapter.Initialize(SourceFile, ImportOptions, VolumeBoundsMin))
+	{
+		return false;
+	}
+	if (!OutResult.Construct(Adapter))
+	{
+		return false;
 	}
-
 	return true;
 #else
 	return false;
diff --git a/Engine/Source/Editor/SparseVolumeTexture/Private/SparseVolumeTextureOpenVDBUtility.h b/Engine/Source/Editor/SparseVolumeTexture/Private/SparseVolumeTextureOpenVDBUtility.h
--- a/Engine/Source/Editor/SparseVolumeTexture/Private/SparseVolumeTextureOpenVDBUtility.h
+++ b/Engine/Source/Editor/SparseVolumeTexture/Private/SparseVolumeTextureOpenVDBUtility.h
@@ -22,9 +22,9 @@ enum class EOpenVDBGridType : uint8
 struct FOpenVDBGridInfo
 {
 	FMatrix44f Transform;
-	FVector VolumeActiveAABBMin;
-	FVector VolumeActiveAABBMax;
-	FVector VolumeActiveDim;
+	FIntVector3 VolumeActiveAABBMin;
+	FIntVector3 VolumeActiveAABBMax;
+	FIntVector3 VolumeActiveDim;
 	FVector VolumeVoxelSize;
 	FString Name;
 	FString DisplayString; // Contains Index (into source file grids), Type and Name
@@ -47,13 +47,7 @@ bool IsOpenVDBGridValid(const FOpenVDBGridInfo& GridInfo, const FString& Filenam
 
 bool GetOpenVDBGridInfo(TArray<uint8>& SourceFile, bool bCreateStrings, TArray<FOpenVDBGridInfo>* OutGridInfo);
 
-bool ConvertOpenVDBToSparseVolumeTexture(
-	TArray<uint8>& SourceFile,
-	struct FOpenVDBImportOptions& ImportOptions,
-	struct FOpenVDBToSVTConversionResult* OutResult,
-	bool bOverrideActiveMinMax,
-	FVector ActiveMin,
-	FVector ActiveMax);
+bool ConvertOpenVDBToSparseVolumeTexture(TArray<uint8>& SourceFile, const struct FOpenVDBImportOptions& ImportOptions, const FIntVector3& VolumeBoundsMin, struct FSparseVolumeRawSource& OutResult);
 
 const TCHAR* OpenVDBGridTypeToString(EOpenVDBGridType Type);
 
diff --git a/Engine/Source/Runtime/Engine/Classes/SparseVolumeTexture/SparseVolumeTexture.h b/Engine/Source/Runtime/Engine/Classes/SparseVolumeTexture/SparseVolumeTexture.h
--- a/Engine/Source/Runtime/Engine/Classes/SparseVolumeTexture/SparseVolumeTexture.h
+++ b/Engine/Source/Runtime/Engine/Classes/SparseVolumeTexture/SparseVolumeTexture.h
@@ -18,32 +18,49 @@ namespace UE { namespace DerivedData { class FRequestOwner; } }
 
 #define SPARSE_VOLUME_TILE_RES	16
 
+uint32 SparseVolumeTexturePackPageTableEntry(const FIntVector3& Coord);
+FIntVector3 SparseVolumeTextureUnpackPageTableEntry(uint32 Packed);
+
 struct ENGINE_API FSparseVolumeAssetHeader
 {
-	FIntVector3 PageTableVolumeResolution;
-	FIntVector3 TileDataVolumeResolution;
-	FIntVector3 SourceVolumeResolution;
-	FIntVector3 SourceVolumeAABBMin;
-	EPixelFormat AttributesAFormat;
-	EPixelFormat AttributesBFormat;
+	FIntVector3 VirtualVolumeResolution = FIntVector3(0, 0, 0);
+	FIntVector3 VirtualVolumeAABBMin = FIntVector3(INT32_MAX, INT32_MAX, INT32_MAX);
+	FIntVector3 VirtualVolumeAABBMax = FIntVector3(INT32_MIN, INT32_MIN, INT32_MIN);
+	FIntVector3 PageTableVolumeResolution = FIntVector3(0, 0, 0);
+	FIntVector3 PageTableVolumeAABBMin = FIntVector3(INT32_MAX, INT32_MAX, INT32_MAX);
+	FIntVector3 PageTableVolumeAABBMax = FIntVector3(INT32_MIN, INT32_MIN, INT32_MIN);
+	FIntVector3 TileDataVolumeResolution = FIntVector3(0, 0, 0);
+	TStaticArray<EPixelFormat, 2> AttributesFormats = TStaticArray<EPixelFormat, 2>(InPlace, PF_Unknown);
+	int32 MipLevel = 0;
+	bool bHasNullTile = false;
 
 	// The current data format version for the header.
 	static const uint32 kVersion = 0;
 
 	// This version can be used to convert existing header to new version later.
-	uint32 Version;
+	uint32 Version = kVersion;
 
 	void Serialize(FArchive& Ar);
+};
 
-	FSparseVolumeAssetHeader()
-		: PageTableVolumeResolution(FIntVector3(0, 0, 0))
-		, TileDataVolumeResolution(FIntVector3(0, 0, 0))
-		, SourceVolumeResolution(FIntVector3(0, 0, 0))
-		, AttributesAFormat(PF_Unknown)
-		, AttributesBFormat(PF_Unknown)
-		, Version(kVersion)
+class ISparseVolumeRawSourceConstructionAdapter
+{
+public:
+	struct FAttributesInfo
 	{
-	}
+		EPixelFormat Format;
+		FVector4f FallbackValue;
+		FVector4f NormalizeScale;
+		FVector4f NormalizeBias;
+		bool bNormalized;
+	};
+
+	virtual void GetAttributesInfo(FAttributesInfo& OutInfoA, FAttributesInfo& OutInfoB) const = 0;
+	virtual FIntVector3 GetAABBMin() const = 0;
+	virtual FIntVector3 GetAABBMax() const = 0;
+	virtual FIntVector3 GetResolution() const = 0;
+	virtual void IteratePhysicalSource(TFunctionRef<void(const FIntVector3& Coord, int32 AttributesIdx, int32 ComponentIdx, float VoxelValue)> OnVisit) const = 0;
+	virtual ~ISparseVolumeRawSourceConstructionAdapter() = default;
 };
 
 // The structure represent the source asset in high quality. It is used to cook the runtime data
@@ -62,6 +79,16 @@ struct ENGINE_API FSparseVolumeRawSource
 
 	void Serialize(FArchive& Ar);
 
+	bool Construct(const ISparseVolumeRawSourceConstructionAdapter& Adapter);
+	uint32 ReadPageTablePacked(const FIntVector3& PageTableCoord) const;
+	FIntVector3 ReadPageTable(const FIntVector3& PageTableCoord) const;
+	FVector4f ReadTileDataVoxel(const FIntVector3& TileDataCoord, int32 AttributesIdx) const;
+	FVector4f Sample(const FIntVector3& VolumeCoord, int32 AttributesIdx) const;
+	void Sample(const FIntVector3& VolumeCoord, FVector4f& OutAttributesA, FVector4f& OutAttributesB) const;
+	void WriteTileDataVoxel(const FIntVector3& TileDataCoord, int32 AttributesIdx, const FVector4f& Value, int32 DstComponent = -1);
+	void FillNullTile(const FVector4f& FallbackValueA, const FVector4f& FallbackValueB);
+	FSparseVolumeRawSource GenerateMipMap() const;
+
 	FSparseVolumeRawSource()
 		: Version(kVersion)
 	{
@@ -163,6 +190,8 @@ struct ENGINE_API FSparseVolumeTextureFrame
 	void Serialize(FArchive& Ar, UStreamableSparseVolumeTexture* Owner, int32 FrameIndex);
 };
 
+using FSparseVolumeTextureFrameMips = TArray<FSparseVolumeTextureFrame, TInlineAllocator<1u>>;
+
 
 enum ESparseVolumeTextureShaderUniform
 {
@@ -185,17 +214,18 @@ public:
 	virtual ~USparseVolumeTexture() = default;
 
 	UFUNCTION(BlueprintCallable, Category = "Texture")
-	virtual int32 GetSizeX() const { return FMath::CeilToInt32(GetVolumeBounds().GetSize().X); }
+	virtual int32 GetSizeX() const { return GetVolumeResolution().X; }
 
 	UFUNCTION(BlueprintCallable, Category = "Texture")
-	virtual int32 GetSizeY() const { return FMath::CeilToInt32(GetVolumeBounds().GetSize().Y); }
+	virtual int32 GetSizeY() const { return GetVolumeResolution().Y; }
 
 	UFUNCTION(BlueprintCallable, Category = "Texture")
-	virtual int32 GetSizeZ() const { return FMath::CeilToInt32(GetVolumeBounds().GetSize().Z); }
+	virtual int32 GetSizeZ() const { return GetVolumeResolution().Z; }
 
 	UFUNCTION(BlueprintCallable, Category = "Texture")
-	virtual int32 GetFrameCount() const { return 0; }
-	virtual FBox GetVolumeBounds() const { return FBox(); }
+	virtual int32 GetNumFrames() const { return 0; }
+	virtual int32 GetNumMipLevels() const { return 0; }
+	virtual FIntVector GetVolumeResolution() const { return FIntVector(); }
 	virtual const FSparseVolumeTextureSceneProxy* GetSparseVolumeTextureSceneProxy() const { return nullptr; }
 
 	/** Getter for the shader uniform parameters with index as ESparseVolumeTextureShaderUniform. */
@@ -221,8 +251,8 @@ class UStreamableSparseVolumeTexture : public USparseVolumeTexture
 
 public:
 
-	UPROPERTY(VisibleAnywhere, Category = "Rendering", meta = (DisplayName = "Volume Bounds"))
-	FBox VolumeBounds;
+	UPROPERTY(VisibleAnywhere, Category = "Rendering")
+	FIntVector VolumeResolution;
 
 	UStreamableSparseVolumeTexture();
 	virtual ~UStreamableSparseVolumeTexture() = default;
@@ -239,16 +269,17 @@ public:
 	//~ End UObject Interface.
 
 	//~ Begin USparseVolumeTexture Interface.
-	int32 GetFrameCount() const override { return Frames.Num(); }
-	FBox GetVolumeBounds() const override { return VolumeBounds; };
-	const FSparseVolumeTextureSceneProxy* GetSparseVolumeTextureSceneProxy() const override { return GetStreamedFrameProxyOrFallback(0); };
+	int32 GetNumFrames() const override { return Frames.Num(); }
+	int32 GetNumMipLevels() const override { return Frames.IsEmpty() ? 0 : Frames[0].Num(); }
+	FIntVector GetVolumeResolution() const override { return VolumeResolution; };
+	const FSparseVolumeTextureSceneProxy* GetSparseVolumeTextureSceneProxy() const override { return GetStreamedFrameProxyOrFallback(0 /*FrameIndex*/, 0 /*MipLevel*/); };
 	//~ End USparseVolumeTexture Interface.
 
-	const FSparseVolumeTextureSceneProxy* GetStreamedFrameProxyOrFallback(int32 FrameIndex) const;
-	TArrayView<const FSparseVolumeTextureFrame> GetFrames() const;
+	const FSparseVolumeTextureSceneProxy* GetStreamedFrameProxyOrFallback(int32 FrameIndex, int32 MipLevel) const;
+	TArrayView<const FSparseVolumeTextureFrameMips> GetFrames() const;
 
 protected:
-	TArray<FSparseVolumeTextureFrame, TInlineAllocator<1u>> Frames;
+	TArray<FSparseVolumeTextureFrameMips> Frames;
 
 	void GenerateOrLoadDDCRuntimeDataAndCreateSceneProxy();
 	void GenerateOrLoadDDCRuntimeDataForFrame(FSparseVolumeTextureFrame& Frame, UE::DerivedData::FRequestOwner& DDCRequestOwner);
@@ -265,7 +296,7 @@ public:
 	virtual ~UStaticSparseVolumeTexture() = default;
 
 	//~ Begin USparseVolumeTexture Interface.
-	int32 GetFrameCount() const override { return 1; }
+	int32 GetNumFrames() const override { return 1; }
 	//~ End USparseVolumeTexture Interface.
 
 private:
@@ -291,8 +322,8 @@ public:
 	//~ End USparseVolumeTexture Interface.
 
 	// Used for debugging a specific frame of an animated sequence.
-	const FSparseVolumeTextureSceneProxy* GetSparseVolumeTextureFrameSceneProxy(int32 FrameIndex) const;
-	const FSparseVolumeAssetHeader* GetSparseVolumeTextureFrameHeader(int32 FrameIndex) const;
+	const FSparseVolumeTextureSceneProxy* GetSparseVolumeTextureFrameSceneProxy(int32 FrameIndex, int32 MipLevel) const;
+	const FSparseVolumeAssetHeader* GetSparseVolumeTextureFrameHeader(int32 FrameIndex, int32 MipLevel) const;
 
 private:
 
@@ -301,6 +332,7 @@ private:
 #endif
 	
 	int32 PreviewFrameIndex;
+	int32 PreviewMipLevel;
 };
 
 // USparseVolumeTextureFrame inherits from USparseVolumeTexture to be viewed using any given frame of a UAnimatedSparseVolumeTexture (or UStaticSparseVolumeTexture)
@@ -313,18 +345,19 @@ public:
 	USparseVolumeTextureFrame();
 	virtual ~USparseVolumeTextureFrame() = default;
 
-	static USparseVolumeTextureFrame* CreateFrame(USparseVolumeTexture* Texture, int32 FrameIndex);
+	static USparseVolumeTextureFrame* CreateFrame(USparseVolumeTexture* Texture, int32 FrameIndex, int32 MipLevel);
 
-	void Initialize(const FSparseVolumeTextureSceneProxy* InSceneProxy, const FBox& InVolumeBounds);
+	void Initialize(const FSparseVolumeTextureSceneProxy* InSceneProxy, const FIntVector& InVolumeResolution);
 
 	//~ Begin USparseVolumeTexture Interface.
-	int32 GetFrameCount() const override { return 1; }
-	FBox GetVolumeBounds() const override { return VolumeBounds; };
+	int32 GetNumFrames() const override { return 1; }
+	int32 GetNumMipLevels() const override { return 1; }
+	FIntVector GetVolumeResolution() const override { return VolumeResolution; };
 	const FSparseVolumeTextureSceneProxy* GetSparseVolumeTextureSceneProxy() const override { return SceneProxy; };
 	//~ End USparseVolumeTexture Interface.
 
 private:
-	FBox VolumeBounds;
+	FIntVector3 VolumeResolution;
 	const FSparseVolumeTextureSceneProxy* SceneProxy;
 };
 
@@ -335,9 +368,12 @@ class ENGINE_API UAnimatedSparseVolumeTextureController : public UObject
 
 public:
 
-	UPROPERTY(BlueprintReadWrite, Category = "Animation", meta = (DisplayName = "Frame Rate"))
+	UPROPERTY(BlueprintReadWrite, Category = "Animation")
 	float FrameRate = 24.0f;
 
+	UPROPERTY(BlueprintReadWrite, Category = "Rendering")
+	int32 MipLevel = 0;
+
 	UAnimatedSparseVolumeTextureController();
 	virtual ~UAnimatedSparseVolumeTextureController() = default;
 
diff --git a/Engine/Source/Runtime/Engine/Private/SparseVolumeTexture/SparseVolumeTexture.cpp b/Engine/Source/Runtime/Engine/Private/SparseVolumeTexture/SparseVolumeTexture.cpp
--- a/Engine/Source/Runtime/Engine/Private/SparseVolumeTexture/SparseVolumeTexture.cpp
+++ b/Engine/Source/Runtime/Engine/Private/SparseVolumeTexture/SparseVolumeTexture.cpp
@@ -30,6 +30,10 @@
 
 DEFINE_LOG_CATEGORY_STATIC(LogSparseVolumeTexture, Log, All);
 
+// We are using this instead of GMaxVolumeTextureDimensions to be independent of the platform that
+// the asset is imported on. 2048 should be a safe value that should be supported by all our platforms.
+static constexpr int32 SVTMaxVolumeTextureDim = 2048;
+
 namespace UE
 {
 namespace SVT
@@ -47,10 +51,69 @@ namespace Private
 			Target = static_cast<T>(Buffer);
 		}
 	}
+
+	static float F16ToF32(uint16 Packed)
+	{
+		FFloat16 F16{};
+		F16.Encoded = Packed;
+		return F16.GetFloat();
+	};
+
+	static FIntVector3 AdvanceTileCoord(const FIntVector3& TileCoord, const FIntVector3& TileCoordResolution)
+	{
+		FIntVector3 Result = TileCoord;
+		Result.X++;
+		if (Result.X >= TileCoordResolution.X)
+		{
+			Result.X = 0;
+			Result.Y++;
+		}
+		if (Result.Y >= TileCoordResolution.Y)
+		{
+			Result.Y = 0;
+			Result.Z++;
+		}
+		return Result;
+	};
+
+	static FIntVector3 ComputeTileDataVolumeResolution(int32 NumAllocatedPages, bool bAnyEmptyPageExists)
+	{
+		const int32 EffectivelyAllocatedPageEntries = NumAllocatedPages + (bAnyEmptyPageExists ? 1 : 0);
+		int32 TileVolumeResolutionCube = 1;
+		while (TileVolumeResolutionCube * TileVolumeResolutionCube * TileVolumeResolutionCube < EffectivelyAllocatedPageEntries)
+		{
+			TileVolumeResolutionCube++;				// We use a simple loop to compute the minimum resolution of a cube to store all the tile data
+		}
+		FIntVector3 TileDataVolumeResolution = FIntVector3(TileVolumeResolutionCube, TileVolumeResolutionCube, TileVolumeResolutionCube);
+		while (TileDataVolumeResolution.X * TileDataVolumeResolution.Y * (TileDataVolumeResolution.Z - 1) > EffectivelyAllocatedPageEntries)
+		{
+			TileDataVolumeResolution.Z--;	// We then trim an edge to get back space.
+		}
+
+		return TileDataVolumeResolution * SPARSE_VOLUME_TILE_RES;
+	}
 } // Private
 } // SVT
 } // UE
 
+uint32 SparseVolumeTexturePackPageTableEntry(const FIntVector3& Coord)
+{
+	// A page encodes the physical tile coord as unsigned int of 11 11 10 bits
+	// This means a page coord cannot be larger than 2047 for x and y and 1023 for z
+	// which mean we cannot have more than 2048*2048*1024 = 4 Giga tiles of 16^3 tiles.
+	uint32 Result = (Coord.X & 0x7FF) | ((Coord.Y & 0x7FF) << 11) | ((Coord.Z & 0x3FF) << 22);
+	return Result;
+}
+
+FIntVector3 SparseVolumeTextureUnpackPageTableEntry(uint32 Packed)
+{
+	FIntVector3 Result;
+	Result.X = Packed & 0x7FF;
+	Result.Y = (Packed >> 11) & 0x7FF;
+	Result.Z = (Packed >> 22) & 0x3FF;
+	return Result;
+}
+
 ////////////////////////////////////////////////////////////////////////////////////////////////
 
 void FSparseVolumeRawSource::Serialize(FArchive& Ar)
@@ -72,6 +135,471 @@ void FSparseVolumeRawSource::Serialize(FArchive& Ar)
 	}
 }
 
+bool FSparseVolumeRawSource::Construct(const ISparseVolumeRawSourceConstructionAdapter& Adapter)
+{
+	ISparseVolumeRawSourceConstructionAdapter::FAttributesInfo AttributesInfo[2];
+	Adapter.GetAttributesInfo(AttributesInfo[0], AttributesInfo[1]);
+
+	Header.VirtualVolumeResolution = Adapter.GetResolution();
+	Header.VirtualVolumeAABBMin = Adapter.GetAABBMin();
+	Header.VirtualVolumeAABBMax = Adapter.GetAABBMax();
+	Header.AttributesFormats[0] = AttributesInfo[0].Format;
+	Header.AttributesFormats[1] = AttributesInfo[1].Format;
+
+	Header.PageTableVolumeAABBMin = Header.VirtualVolumeAABBMin / SPARSE_VOLUME_TILE_RES;
+	Header.PageTableVolumeAABBMax = (Header.VirtualVolumeAABBMax + FIntVector3(SPARSE_VOLUME_TILE_RES - 1)) / SPARSE_VOLUME_TILE_RES;
+	Header.bHasNullTile = false;
+
+	const FIntVector3 PageTableVolumeResolution = Header.PageTableVolumeAABBMax - Header.PageTableVolumeAABBMin;
+
+	if (PageTableVolumeResolution.X > SVTMaxVolumeTextureDim
+		|| PageTableVolumeResolution.Y > SVTMaxVolumeTextureDim
+		|| PageTableVolumeResolution.Z > SVTMaxVolumeTextureDim)
+	{
+		UE_LOG(LogSparseVolumeTexture, Warning, TEXT("SparseVolumeTexture page table texture dimensions exceed limit (%ix%ix%i): %ix%ix%i"), 
+			SVTMaxVolumeTextureDim, SVTMaxVolumeTextureDim, SVTMaxVolumeTextureDim, 
+			PageTableVolumeResolution.X, PageTableVolumeResolution.Y, PageTableVolumeResolution.Z);
+		return false;
+	}
+
+	Header.PageTableVolumeResolution = PageTableVolumeResolution;
+	Header.TileDataVolumeResolution = FIntVector::ZeroValue;	// unknown for now
+
+	// Tag all pages with valid data
+	TBitArray PagesWithData(false, Header.PageTableVolumeResolution.Z * Header.PageTableVolumeResolution.Y * Header.PageTableVolumeResolution.X);
+	Adapter.IteratePhysicalSource([&](const FIntVector3& Coord, int32 AttributesIdx, int32 ComponentIdx, float VoxelValue)
+		{
+			const bool bIsFallbackValue = (VoxelValue == AttributesInfo[AttributesIdx].FallbackValue[ComponentIdx]);
+			if (!bIsFallbackValue)
+			{
+				const FIntVector3 GridCoord = Coord;
+				check(GridCoord.X >= 0 && GridCoord.Y >= 0 && GridCoord.Z >= 0);
+				check(GridCoord.X < Header.VirtualVolumeAABBMax.X && GridCoord.Y < Header.VirtualVolumeAABBMax.Y && GridCoord.Z < Header.VirtualVolumeAABBMax.Z);
+				const FIntVector3 PageCoord = (GridCoord / SPARSE_VOLUME_TILE_RES) - Header.PageTableVolumeAABBMin;
+				check(PageCoord.X >= 0 && PageCoord.Y >= 0 && PageCoord.Z >= 0);
+				check(PageCoord.X < Header.PageTableVolumeResolution.X && PageCoord.Y < Header.PageTableVolumeResolution.Y && PageCoord.Z < Header.PageTableVolumeResolution.Z);
+
+				const int32 PageIndex = PageCoord.Z * (Header.PageTableVolumeResolution.Y * Header.PageTableVolumeResolution.X) + PageCoord.Y * Header.PageTableVolumeResolution.X + PageCoord.X;
+				PagesWithData[PageIndex] = true;
+			}
+		});
+
+	// Allocate some memory for temp data (worst case)
+	TArray<FIntVector3> LinearAllocatedPages;
+	LinearAllocatedPages.SetNum(Header.PageTableVolumeResolution.X * Header.PageTableVolumeResolution.Y * Header.PageTableVolumeResolution.Z);
+
+	// Go over each potential page from the source data and push allocate it if it has any data.
+	// Otherwise point to the default empty page.
+	int32 NumAllocatedPages = 0;
+	for (int32_t PageZ = 0; PageZ < Header.PageTableVolumeResolution.Z; ++PageZ)
+	{
+		for (int32_t PageY = 0; PageY < Header.PageTableVolumeResolution.Y; ++PageY)
+		{
+			for (int32_t PageX = 0; PageX < Header.PageTableVolumeResolution.X; ++PageX)
+			{
+				const int32 PageIndex = PageZ * (Header.PageTableVolumeResolution.Y * Header.PageTableVolumeResolution.X) + PageY * Header.PageTableVolumeResolution.X + PageX;
+				const bool bHasAnyData = PagesWithData[PageIndex];
+				if (bHasAnyData)
+				{
+					LinearAllocatedPages[NumAllocatedPages] = FIntVector3(PageX, PageY, PageZ);
+					NumAllocatedPages++;
+				}
+				Header.bHasNullTile |= !bHasAnyData;
+			}
+		}
+	}
+
+	// Compute Page and Tile VolumeResolution from allocated pages
+	Header.TileDataVolumeResolution = UE::SVT::Private::ComputeTileDataVolumeResolution(NumAllocatedPages, Header.bHasNullTile);
+	const FIntVector3 TileCoordResolution = Header.TileDataVolumeResolution / SPARSE_VOLUME_TILE_RES;
+
+	if (Header.TileDataVolumeResolution.X > SVTMaxVolumeTextureDim
+		|| Header.TileDataVolumeResolution.Y > SVTMaxVolumeTextureDim
+		|| Header.TileDataVolumeResolution.Z > SVTMaxVolumeTextureDim)
+	{
+		UE_LOG(LogSparseVolumeTexture, Warning, TEXT("SparseVolumeTexture tile data texture dimensions exceed limit (%ix%ix%i): %ix%ix%i"), SVTMaxVolumeTextureDim, SVTMaxVolumeTextureDim, SVTMaxVolumeTextureDim, Header.TileDataVolumeResolution.X, Header.TileDataVolumeResolution.Y, Header.TileDataVolumeResolution.Z);
+		return false;
+	}
+
+	// Initialize the SparseVolumeTexture page and tile.
+	const int32 FormatSize[] = { GPixelFormats[(SIZE_T)Header.AttributesFormats[0]].BlockBytes, GPixelFormats[(SIZE_T)Header.AttributesFormats[1]].BlockBytes };
+	PageTable.SetNumZeroed(Header.PageTableVolumeResolution.X * Header.PageTableVolumeResolution.Y * Header.PageTableVolumeResolution.Z);
+	PhysicalTileDataA.SetNumZeroed(Header.TileDataVolumeResolution.X * Header.TileDataVolumeResolution.Y * Header.TileDataVolumeResolution.Z * FormatSize[0]);
+	PhysicalTileDataB.SetNumZeroed(Header.TileDataVolumeResolution.X * Header.TileDataVolumeResolution.Y * Header.TileDataVolumeResolution.Z * FormatSize[1]);
+
+	// Generate page table and tile volume data by splatting the data
+	{
+		FIntVector3 DstTileCoord = FIntVector3::ZeroValue;
+
+		// Add an empty tile is needed, reserve slot at coord 0
+		if (Header.bHasNullTile)
+		{
+			// Fill null tile with fallback/background value
+			FVector4f FallbackValues[] = { AttributesInfo[0].FallbackValue, AttributesInfo[1].FallbackValue };
+			for (int32 AttributesIdx = 0; AttributesIdx < 2; ++AttributesIdx)
+			{
+				if (AttributesInfo[AttributesIdx].bNormalized)
+				{
+					FallbackValues[AttributesIdx] = FallbackValues[AttributesIdx] * AttributesInfo[AttributesIdx].NormalizeScale + AttributesInfo[AttributesIdx].NormalizeBias;
+				}
+			}
+			FillNullTile(FallbackValues[0], FallbackValues[1]);
+
+			// PageTable is all cleared to zero, simply skip a tile
+			DstTileCoord = UE::SVT::Private::AdvanceTileCoord(DstTileCoord, TileCoordResolution);
+		}
+
+		for (int32 i = 0; i < NumAllocatedPages; ++i)
+		{
+			const FIntVector3 PageCoordToSplat = LinearAllocatedPages[i];
+			const uint32 DestinationTileCoord32bit = SparseVolumeTexturePackPageTableEntry(DstTileCoord);
+
+			// Setup the page table entry
+			PageTable
+				[
+					PageCoordToSplat.Z * Header.PageTableVolumeResolution.X * Header.PageTableVolumeResolution.Y +
+					PageCoordToSplat.Y * Header.PageTableVolumeResolution.X +
+					PageCoordToSplat.X
+				] = DestinationTileCoord32bit;
+
+			// Set the next tile to be written to
+			DstTileCoord = UE::SVT::Private::AdvanceTileCoord(DstTileCoord, TileCoordResolution);
+		}
+	}
+
+	// Write physical tile data
+	Adapter.IteratePhysicalSource([&](const FIntVector3& Coord, int32 AttributesIdx, int32 ComponentIdx, float VoxelValue)
+		{
+			const FIntVector3 GridCoord = Coord;
+			check(GridCoord.X >= 0 && GridCoord.Y >= 0 && GridCoord.Z >= 0);
+			check(GridCoord.X < Header.VirtualVolumeAABBMax.X && GridCoord.Y < Header.VirtualVolumeAABBMax.Y && GridCoord.Z < Header.VirtualVolumeAABBMax.Z);
+			const FIntVector3 PageCoord = (GridCoord / SPARSE_VOLUME_TILE_RES) - Header.PageTableVolumeAABBMin;
+			check(PageCoord.X >= 0 && PageCoord.Y >= 0 && PageCoord.Z >= 0);
+			check(PageCoord.X < Header.PageTableVolumeResolution.X && PageCoord.Y < Header.PageTableVolumeResolution.Y && PageCoord.Z < Header.PageTableVolumeResolution.Z);
+
+			const int32 PageIndex = PageCoord.Z * (Header.PageTableVolumeResolution.Y * Header.PageTableVolumeResolution.X) + PageCoord.Y * Header.PageTableVolumeResolution.X + PageCoord.X;
+
+			if (!PagesWithData[PageIndex])
+			{
+				return;
+			}
+
+			float WriteValue = VoxelValue;
+			if (AttributesInfo[AttributesIdx].bNormalized)
+			{
+				WriteValue = WriteValue * AttributesInfo[AttributesIdx].NormalizeScale[ComponentIdx] + AttributesInfo[AttributesIdx].NormalizeBias[ComponentIdx];
+			}
+
+			const FIntVector3 WriteTileCoord = SparseVolumeTextureUnpackPageTableEntry(PageTable[PageIndex]);
+			const FIntVector3 TileLocalCoord = GridCoord % SPARSE_VOLUME_TILE_RES;
+			const FIntVector3 WriteCoord = WriteTileCoord * SPARSE_VOLUME_TILE_RES + TileLocalCoord;
+			WriteTileDataVoxel(WriteCoord, AttributesIdx, FVector4f(WriteValue, WriteValue, WriteValue, WriteValue), ComponentIdx);
+		});
+
+	return true;
+}
+
+uint32 FSparseVolumeRawSource::ReadPageTablePacked(const FIntVector3& PageTableCoord) const
+{
+	if (PageTableCoord.X < 0 || PageTableCoord.Y < 0 || PageTableCoord.Z < 0
+		|| PageTableCoord.X >= Header.PageTableVolumeResolution.X 
+		|| PageTableCoord.Y >= Header.PageTableVolumeResolution.Y 
+		|| PageTableCoord.Z >= Header.PageTableVolumeResolution.Z)
+	{
+		return 0;
+	}
+
+	const int32 PageIndex = PageTableCoord.Z * (Header.PageTableVolumeResolution.Y * Header.PageTableVolumeResolution.X) + PageTableCoord.Y * Header.PageTableVolumeResolution.X + PageTableCoord.X;
+	if (PageTable.IsValidIndex(PageIndex))
+	{
+		return PageTable[PageIndex];
+	}
+	return 0;
+}
+
+FIntVector3 FSparseVolumeRawSource::ReadPageTable(const FIntVector3& PageTableCoord) const
+{
+	const uint32 Packed = ReadPageTablePacked(PageTableCoord);
+	const FIntVector3 Unpacked = SparseVolumeTextureUnpackPageTableEntry(Packed);
+	return Unpacked;
+}
+
+FVector4f FSparseVolumeRawSource::ReadTileDataVoxel(const FIntVector3& TileDataCoord, int32 AttributesIdx) const
+{
+	using namespace UE::SVT::Private;
+	check(AttributesIdx >= 0 && AttributesIdx <= 1);
+	if ((AttributesIdx == 0 && PhysicalTileDataA.IsEmpty()) || (AttributesIdx == 1 && PhysicalTileDataB.IsEmpty()))
+	{
+		return FVector4f();
+	}
+
+	if (TileDataCoord.X < 0 || TileDataCoord.Y < 0 || TileDataCoord.Z < 0
+		|| TileDataCoord.X >= Header.TileDataVolumeResolution.X 
+		|| TileDataCoord.Y >= Header.TileDataVolumeResolution.Y 
+		|| TileDataCoord.Z >= Header.TileDataVolumeResolution.Z)
+	{
+		return FVector4f();
+	}
+
+	const int32 VoxelIndex = TileDataCoord.Z * (Header.TileDataVolumeResolution.Y * Header.TileDataVolumeResolution.X) + TileDataCoord.Y * Header.TileDataVolumeResolution.X + TileDataCoord.X;
+	const uint8* TileData = AttributesIdx == 0 ? PhysicalTileDataA.GetData() : PhysicalTileDataB.GetData();
+	const EPixelFormat Format = Header.AttributesFormats[AttributesIdx];
+
+	switch (Format)
+	{
+	case PF_R8:
+		return FVector4f(TileData[VoxelIndex] / 255.0f, 0.0f, 0.0f, 0.0f);
+	case PF_R8G8:
+		return FVector4f(TileData[VoxelIndex * 2 + 0] / 255.0f, TileData[VoxelIndex * 2 + 1] / 255.0f, 0.0f, 0.0f);
+	case PF_R8G8B8A8:
+		return FVector4f(TileData[VoxelIndex * 4 + 0] / 255.0f, TileData[VoxelIndex * 4 + 1] / 255.0f, TileData[VoxelIndex * 4 + 2] / 255.0f, TileData[VoxelIndex * 4 + 3] / 255.0f);
+	case PF_R16F:
+		return FVector4f(F16ToF32(((const uint16*)TileData)[VoxelIndex]), 0.0f, 0.0f, 0.0f);
+	case PF_G16R16F:
+		return FVector4f(F16ToF32(((const uint16*)TileData)[VoxelIndex * 2 + 0]), F16ToF32(((const uint16*)TileData)[VoxelIndex * 2 + 1]), 0.0f, 0.0f);
+	case PF_FloatRGBA:
+		return FVector4f(F16ToF32(((const uint16*)TileData)[VoxelIndex * 4 + 0]), F16ToF32(((const uint16*)TileData)[VoxelIndex * 4 + 1]), F16ToF32(((const uint16*)TileData)[VoxelIndex * 4 + 2]), F16ToF32(((const uint16*)TileData)[VoxelIndex * 4 + 3]));
+	case PF_R32_FLOAT:
+		return FVector4f(((const float*)TileData)[VoxelIndex], 0.0f, 0.0f, 0.0f);
+	case PF_G32R32F:
+		return FVector4f(((const float*)TileData)[VoxelIndex * 2 + 0], ((const float*)TileData)[VoxelIndex * 2 + 1], 0.0f, 0.0f);
+	case PF_A32B32G32R32F:
+		return FVector4f(((const float*)TileData)[VoxelIndex * 4 + 0], ((const float*)TileData)[VoxelIndex * 4 + 1], ((const float*)TileData)[VoxelIndex * 4 + 2], ((const float*)TileData)[VoxelIndex * 4 + 3]);
+	default:
+		checkNoEntry();
+		return FVector4f();
+	}
+}
+
+FVector4f FSparseVolumeRawSource::Sample(const FIntVector3& VolumeCoord, int32 AttributesIdx) const
+{
+	const FIntVector3 PageTableCoord = (VolumeCoord / SPARSE_VOLUME_TILE_RES) - Header.PageTableVolumeAABBMin;
+	const FIntVector3 TileCoord = ReadPageTable(PageTableCoord);
+	const FIntVector3 VoxelCoord = (TileCoord * SPARSE_VOLUME_TILE_RES) + (VolumeCoord % SPARSE_VOLUME_TILE_RES);
+	const FVector4f Sample = ReadTileDataVoxel(VoxelCoord, AttributesIdx);
+	return Sample;
+}
+
+void FSparseVolumeRawSource::Sample(const FIntVector3& VolumeCoord, FVector4f& OutAttributesA, FVector4f& OutAttributesB) const
+{
+	const FIntVector3 PageTableCoord = (VolumeCoord / SPARSE_VOLUME_TILE_RES) - Header.PageTableVolumeAABBMin;
+	const FIntVector3 TileCoord = ReadPageTable(PageTableCoord);
+	const FIntVector3 VoxelCoord = (TileCoord * SPARSE_VOLUME_TILE_RES) + (VolumeCoord % SPARSE_VOLUME_TILE_RES);
+	OutAttributesA = ReadTileDataVoxel(VoxelCoord, 0);
+	OutAttributesB = ReadTileDataVoxel(VoxelCoord, 1);
+}
+
+void FSparseVolumeRawSource::WriteTileDataVoxel(const FIntVector3& TileDataCoord, int32 AttributesIdx, const FVector4f& Value, int32 DstComponent)
+{
+	check(AttributesIdx >= 0 && AttributesIdx <= 1);
+	if ((AttributesIdx == 0 && PhysicalTileDataA.IsEmpty()) || (AttributesIdx == 1 && PhysicalTileDataB.IsEmpty()))
+	{
+		return;
+	}
+
+	const int32 VoxelIndex = TileDataCoord.Z * (Header.TileDataVolumeResolution.Y * Header.TileDataVolumeResolution.X) + TileDataCoord.Y * Header.TileDataVolumeResolution.X + TileDataCoord.X;
+	uint8* TileData = AttributesIdx == 0 ? PhysicalTileDataA.GetData() : PhysicalTileDataB.GetData();
+	const EPixelFormat Format = Header.AttributesFormats[AttributesIdx];
+
+	switch (Format)
+	{
+	case PF_R8:
+		if (DstComponent == -1 || DstComponent == 0) TileData[VoxelIndex] = uint8(FMath::Clamp(Value.X, 0.0f, 1.0f) * 255.0f);
+		break;
+	case PF_R8G8:
+		if (DstComponent == -1 || DstComponent == 0) TileData[VoxelIndex * 2 + 0] = uint8(FMath::Clamp(Value.X, 0.0f, 1.0f) * 255.0f);
+		if (DstComponent == -1 || DstComponent == 1) TileData[VoxelIndex * 2 + 1] = uint8(FMath::Clamp(Value.Y, 0.0f, 1.0f) * 255.0f);
+		break;
+	case PF_R8G8B8A8:
+		if (DstComponent == -1 || DstComponent == 0) TileData[VoxelIndex * 4 + 0] = uint8(FMath::Clamp(Value.X, 0.0f, 1.0f) * 255.0f);
+		if (DstComponent == -1 || DstComponent == 1) TileData[VoxelIndex * 4 + 1] = uint8(FMath::Clamp(Value.Y, 0.0f, 1.0f) * 255.0f);
+		if (DstComponent == -1 || DstComponent == 2) TileData[VoxelIndex * 4 + 2] = uint8(FMath::Clamp(Value.Z, 0.0f, 1.0f) * 255.0f);
+		if (DstComponent == -1 || DstComponent == 3) TileData[VoxelIndex * 4 + 3] = uint8(FMath::Clamp(Value.W, 0.0f, 1.0f) * 255.0f);
+		break;
+	case PF_R16F:
+		if (DstComponent == -1 || DstComponent == 0) ((uint16*)TileData)[VoxelIndex] = FFloat16(Value.X).Encoded;
+		break;
+	case PF_G16R16F:
+		if (DstComponent == -1 || DstComponent == 0) ((uint16*)TileData)[VoxelIndex * 2 + 0] = FFloat16(Value.X).Encoded;
+		if (DstComponent == -1 || DstComponent == 1) ((uint16*)TileData)[VoxelIndex * 2 + 1] = FFloat16(Value.Y).Encoded;
+		break;
+	case PF_FloatRGBA:
+		if (DstComponent == -1 || DstComponent == 0) ((uint16*)TileData)[VoxelIndex * 4 + 0] = FFloat16(Value.X).Encoded;
+		if (DstComponent == -1 || DstComponent == 1) ((uint16*)TileData)[VoxelIndex * 4 + 1] = FFloat16(Value.Y).Encoded;
+		if (DstComponent == -1 || DstComponent == 2) ((uint16*)TileData)[VoxelIndex * 4 + 2] = FFloat16(Value.Z).Encoded;
+		if (DstComponent == -1 || DstComponent == 3) ((uint16*)TileData)[VoxelIndex * 4 + 3] = FFloat16(Value.W).Encoded;
+		break;
+	case PF_R32_FLOAT:
+		if (DstComponent == -1 || DstComponent == 0) ((float*)TileData)[VoxelIndex] = Value.X;
+		break;
+	case PF_G32R32F:
+		if (DstComponent == -1 || DstComponent == 0) ((float*)TileData)[VoxelIndex * 2 + 0] = Value.X;
+		if (DstComponent == -1 || DstComponent == 1) ((float*)TileData)[VoxelIndex * 2 + 1] = Value.Y;
+		break;
+	case PF_A32B32G32R32F:
+		if (DstComponent == -1 || DstComponent == 0) ((float*)TileData)[VoxelIndex * 4 + 0] = Value.X;
+		if (DstComponent == -1 || DstComponent == 1) ((float*)TileData)[VoxelIndex * 4 + 1] = Value.Y;
+		if (DstComponent == -1 || DstComponent == 2) ((float*)TileData)[VoxelIndex * 4 + 2] = Value.Z;
+		if (DstComponent == -1 || DstComponent == 3) ((float*)TileData)[VoxelIndex * 4 + 3] = Value.W;
+		break;
+	default:
+		checkNoEntry();
+	}
+}
+
+void FSparseVolumeRawSource::FillNullTile(const FVector4f& FallbackValueA, const FVector4f& FallbackValueB)
+{
+	for (int32 AttributesIdx = 0; AttributesIdx < 2; ++AttributesIdx)
+	{
+		if ((AttributesIdx == 0 && PhysicalTileDataA.IsEmpty()) || (AttributesIdx == 1 && PhysicalTileDataB.IsEmpty()))
+		{
+			continue;
+		}
+
+		const FVector4f& FallbackValue = AttributesIdx == 0 ? FallbackValueA : FallbackValueB;
+
+		for (int32 Z = 0; Z < SPARSE_VOLUME_TILE_RES; ++Z)
+		{
+			for (int32 Y = 0; Y < SPARSE_VOLUME_TILE_RES; ++Y)
+			{
+				for (int32 X = 0; X < SPARSE_VOLUME_TILE_RES; ++X)
+				{
+					const FIntVector3 WriteCoord = FIntVector3(X, Y, Z);
+					WriteTileDataVoxel(WriteCoord, AttributesIdx, FallbackValue);
+				}
+			}
+		}
+	}
+}
+
+FSparseVolumeRawSource FSparseVolumeRawSource::GenerateMipMap() const
+{
+	FSparseVolumeRawSource Result{};
+	Result.Header = Header; // Overwrite values later
+	Result.Header.VirtualVolumeAABBMin = Header.VirtualVolumeAABBMin / 2;
+	Result.Header.VirtualVolumeAABBMax = FIntVector3(FMath::Max(1, Header.VirtualVolumeAABBMax.X / 2), FMath::Max(1, Header.VirtualVolumeAABBMax.Y / 2), FMath::Max(1, Header.VirtualVolumeAABBMax.Z / 2));
+	Result.Header.VirtualVolumeResolution = Result.Header.VirtualVolumeAABBMax - Result.Header.VirtualVolumeAABBMin;
+	Result.Header.PageTableVolumeAABBMin = Result.Header.VirtualVolumeAABBMin / SPARSE_VOLUME_TILE_RES;
+	Result.Header.PageTableVolumeAABBMax = (Result.Header.VirtualVolumeAABBMax + FIntVector3(SPARSE_VOLUME_TILE_RES - 1)) / SPARSE_VOLUME_TILE_RES;
+	Result.Header.PageTableVolumeResolution = Result.Header.PageTableVolumeAABBMax - Result.Header.PageTableVolumeAABBMin;
+	Result.Header.MipLevel = Header.MipLevel + 1;
+	Result.Header.bHasNullTile = false;
+
+	// Allocate some memory for temp data (worst case)
+	TArray<FIntVector3> LinearAllocatedPages;
+	LinearAllocatedPages.SetNum(Result.Header.PageTableVolumeResolution.X * Result.Header.PageTableVolumeResolution.Y * Result.Header.PageTableVolumeResolution.Z);
+
+	// Go over each potential page from the source data and push allocate it if it has any data.
+	// Otherwise point to the default empty page.
+	int32 NumAllocatedPages = 0;
+	for (int32_t PageZ = 0; PageZ < Result.Header.PageTableVolumeResolution.Z; ++PageZ)
+	{
+		for (int32_t PageY = 0; PageY < Result.Header.PageTableVolumeResolution.Y; ++PageY)
+		{
+			for (int32_t PageX = 0; PageX < Result.Header.PageTableVolumeResolution.X; ++PageX)
+			{
+				const FIntVector3 PageCoord = FIntVector3(PageX, PageY, PageZ);
+				bool bHasAnyData = false;
+				for (int32_t OffsetIdx = 0; OffsetIdx < 8; ++OffsetIdx)
+				{
+					const FIntVector3 Offset = FIntVector3(OffsetIdx, OffsetIdx >> 1, OffsetIdx >> 2) & 1;
+					const FIntVector3 ParentPageTableCoord = (Result.Header.PageTableVolumeAABBMin + PageCoord) * 2 + Offset - Header.PageTableVolumeAABBMin;
+					const uint32 PageSample = ReadPageTablePacked(ParentPageTableCoord);
+					
+					if (PageSample != 0 || !Header.bHasNullTile)
+					{
+						bHasAnyData = true;
+					}
+				}
+				if (bHasAnyData)
+				{
+					LinearAllocatedPages[NumAllocatedPages] = PageCoord;
+					NumAllocatedPages++;
+				}
+				Result.Header.bHasNullTile |= !bHasAnyData;
+			}
+		}
+	}
+
+	// Compute Page and Tile VolumeResolution from allocated pages
+	Result.Header.TileDataVolumeResolution = UE::SVT::Private::ComputeTileDataVolumeResolution(NumAllocatedPages, Result.Header.bHasNullTile);
+	const FIntVector3 TileCoordResolution = Result.Header.TileDataVolumeResolution / SPARSE_VOLUME_TILE_RES;
+
+	// Compute byte sizes of formats as NumTotalBytes/NumVoxels. Note that we do this on the data from this object, not the new downsampled one.
+	const int32 FormatSizeA = PhysicalTileDataA.Num() / (Header.TileDataVolumeResolution.X * Header.TileDataVolumeResolution.Y * Header.TileDataVolumeResolution.Z);
+	const int32 FormatSizeB = PhysicalTileDataB.Num() / (Header.TileDataVolumeResolution.X * Header.TileDataVolumeResolution.Y * Header.TileDataVolumeResolution.Z);
+
+	// Initialize the SparseVolumeTexture page and tile.
+	Result.PageTable.SetNumZeroed(Result.Header.PageTableVolumeResolution.X * Result.Header.PageTableVolumeResolution.Y * Result.Header.PageTableVolumeResolution.Z);
+	Result.PhysicalTileDataA.SetNumZeroed(Result.Header.TileDataVolumeResolution.X * Result.Header.TileDataVolumeResolution.Y * Result.Header.TileDataVolumeResolution.Z * FormatSizeA);
+	Result.PhysicalTileDataB.SetNumZeroed(Result.Header.TileDataVolumeResolution.X * Result.Header.TileDataVolumeResolution.Y * Result.Header.TileDataVolumeResolution.Z * FormatSizeB);
+
+	// Generate page table and tile volume data by splatting the data
+	{
+		FIntVector3 DstTileCoord = FIntVector3::ZeroValue;
+
+		// Add an empty tile is needed, reserve slot at coord 0
+		if (Result.Header.bHasNullTile)
+		{
+			Result.FillNullTile(ReadTileDataVoxel(FIntVector3(), 0), ReadTileDataVoxel(FIntVector3(), 1));
+			
+			// PageTable is all cleared to zero, simply skip a tile
+			DstTileCoord = UE::SVT::Private::AdvanceTileCoord(DstTileCoord, TileCoordResolution);
+		}
+
+		for (int32 i = 0; i < NumAllocatedPages; ++i)
+		{
+			const FIntVector3 PageCoordToSplat = LinearAllocatedPages[i];
+			const uint32 DestinationTileCoord32bit = SparseVolumeTexturePackPageTableEntry(DstTileCoord);
+
+			// Setup the page table entry
+			Result.PageTable
+				[
+					PageCoordToSplat.Z * Result.Header.PageTableVolumeResolution.X * Result.Header.PageTableVolumeResolution.Y +
+					PageCoordToSplat.Y * Result.Header.PageTableVolumeResolution.X +
+					PageCoordToSplat.X
+				] = DestinationTileCoord32bit;
+
+			// Write tile data
+			const FIntVector3 ParentVolumeCoordBase = (Result.Header.PageTableVolumeAABBMin + PageCoordToSplat) * SPARSE_VOLUME_TILE_RES * 2;
+			for (int32 AttributesIdx = 0; AttributesIdx < 2; ++AttributesIdx)
+			{
+				if ((AttributesIdx == 0 && PhysicalTileDataA.IsEmpty()) || (AttributesIdx == 1 && PhysicalTileDataB.IsEmpty()))
+				{
+					continue;
+				}
+
+				for (int32 Z = 0; Z < SPARSE_VOLUME_TILE_RES; ++Z)
+				{
+					for (int32 Y = 0; Y < SPARSE_VOLUME_TILE_RES; ++Y)
+					{
+						for (int32 X = 0; X < SPARSE_VOLUME_TILE_RES; ++X)
+						{
+							FVector4f DownsampledValue = FVector4f(0.0f, 0.0f, 0.0f, 0.0f);
+							for (int32_t OffsetIdx = 0; OffsetIdx < 8; ++OffsetIdx)
+							{
+								const FIntVector3 Offset = FIntVector3(OffsetIdx, OffsetIdx >> 1, OffsetIdx >> 2) & 1;
+								const FIntVector3 SourceCoord = ParentVolumeCoordBase + FIntVector3(X, Y, Z) * 2 + Offset;
+								const FVector4f SampleValue = Sample(SourceCoord, AttributesIdx);
+								DownsampledValue += SampleValue;
+							}
+							DownsampledValue /= 8.0f;
+
+							const FIntVector3 WriteCoord = DstTileCoord * SPARSE_VOLUME_TILE_RES + FIntVector3(X, Y, Z);
+							Result.WriteTileDataVoxel(WriteCoord, AttributesIdx, DownsampledValue);
+						}
+					}
+				}
+			}
+
+			// Set the next tile to be written to
+			DstTileCoord = UE::SVT::Private::AdvanceTileCoord(DstTileCoord, TileCoordResolution);
+		}
+	}
+
+	return Result;
+}
+
 ////////////////////////////////////////////////////////////////////////////////////////////////
 
 void FSparseVolumeAssetHeader::Serialize(FArchive& Ar)
@@ -80,12 +608,17 @@ void FSparseVolumeAssetHeader::Serialize(FArchive& Ar)
 
 	if (Version == 0)
 	{
+		Ar << VirtualVolumeResolution;
+		Ar << VirtualVolumeAABBMin;
+		Ar << VirtualVolumeAABBMax;
 		Ar << PageTableVolumeResolution;
+		Ar << PageTableVolumeAABBMin;
+		Ar << PageTableVolumeAABBMax;
 		Ar << TileDataVolumeResolution;
-		Ar << SourceVolumeResolution;
-		Ar << SourceVolumeAABBMin;
-		UE::SVT::Private::SerializeEnumAs<uint8>(Ar, AttributesAFormat);
-		UE::SVT::Private::SerializeEnumAs<uint8>(Ar, AttributesBFormat);
+		UE::SVT::Private::SerializeEnumAs<uint8>(Ar, AttributesFormats[0]);
+		UE::SVT::Private::SerializeEnumAs<uint8>(Ar, AttributesFormats[1]);
+		Ar << MipLevel;
+		Ar << bHasNullTile;
 	}
 	else
 	{
@@ -166,8 +699,8 @@ void FSparseVolumeTextureSceneProxy::InitRHI()
 {
 	// Page table
 	{
-		EPixelFormat PageEntryFormat = PF_R32_UINT;
-		FIntVector3 PageTableVolumeResolution = SparseVolumeTextureRuntime.Header.PageTableVolumeResolution;
+		const EPixelFormat PageEntryFormat = PF_R32_UINT;
+		const FIntVector3 PageTableVolumeResolution = SparseVolumeTextureRuntime.Header.PageTableVolumeResolution;
 		const FRHITextureCreateDesc Desc =
 			FRHITextureCreateDesc::Create3D(TEXT("SparseVolumeTexture.PageTable.RHITexture"),
 				PageTableVolumeResolution.X, PageTableVolumeResolution.Y, PageTableVolumeResolution.Z, PageEntryFormat)
@@ -182,9 +715,9 @@ void FSparseVolumeTextureSceneProxy::InitRHI()
 
 	// Tile data
 	{
-		FIntVector3 TileDataVolumeResolution = SparseVolumeTextureRuntime.Header.TileDataVolumeResolution;
-		EPixelFormat VoxelFormatA = SparseVolumeTextureRuntime.Header.AttributesAFormat;
-		EPixelFormat VoxelFormatB = SparseVolumeTextureRuntime.Header.AttributesBFormat;
+		const FIntVector3 TileDataVolumeResolution = SparseVolumeTextureRuntime.Header.TileDataVolumeResolution;
+		const EPixelFormat VoxelFormatA = SparseVolumeTextureRuntime.Header.AttributesFormats[0];
+		const EPixelFormat VoxelFormatB = SparseVolumeTextureRuntime.Header.AttributesFormats[1];
 		const FUpdateTextureRegion3D UpdateRegion(0, 0, 0, 0, 0, 0, TileDataVolumeResolution.X, TileDataVolumeResolution.Y, TileDataVolumeResolution.Z);
 
 		// A
@@ -398,19 +931,15 @@ void USparseVolumeTexture::GetFrameUVScaleBias(FVector* OutScale, FVector* OutBi
 	if (Proxy)
 	{
 		const FSparseVolumeAssetHeader& Header = Proxy->GetHeader();
-		const FBox VolumeBounds = GetVolumeBounds();
-		check(VolumeBounds.GetVolume() > 0.0);
-		const FBox FrameBounds = FBox(FVector(Header.SourceVolumeAABBMin), FVector(Header.SourceVolumeAABBMin + Header.SourceVolumeResolution)); // AABB of current frame
-		const FBox FrameBoundsPadded = FBox(FVector(Header.SourceVolumeAABBMin), FVector(Header.SourceVolumeAABBMin + Header.PageTableVolumeResolution * SPARSE_VOLUME_TILE_RES)); // padded to multiple of page size
-		const FVector FrameExtent = FrameBounds.GetExtent();
-
-		// 3D UV coordinates are specified in [0, 1]. Before addressing the page tables which might have padding,
-		// since source volume resolution might not align to tile resolution, we have to rescale the uv so that [0,1] maps to the source texture boundaries.
-		const FVector UVToPhysical = FrameExtent / FrameBoundsPadded.GetExtent();
-
-		*OutScale = VolumeBounds.GetExtent() / FrameExtent;
-		*OutScale *= UVToPhysical;
-		*OutBias = -((FrameBounds.Min - VolumeBounds.Min) / (VolumeBounds.Max - VolumeBounds.Min) * *OutScale);
+		const int32 MipFactor = 1 << Header.MipLevel;
+		const FVector GlobalVolumeRes = FVector(GetVolumeResolution()) / MipFactor;
+		check(GlobalVolumeRes.X > 0.0 && GlobalVolumeRes.Y > 0.0 && GlobalVolumeRes.Z > 0.0);
+		const FVector FrameBoundsPaddedMin = FVector(Header.PageTableVolumeAABBMin * SPARSE_VOLUME_TILE_RES); // padded to multiple of page size
+		const FVector FrameBoundsPaddedMax = FVector(Header.PageTableVolumeAABBMax * SPARSE_VOLUME_TILE_RES);
+		const FVector FramePaddedSize = FrameBoundsPaddedMax - FrameBoundsPaddedMin;
+
+		*OutScale = GlobalVolumeRes / FramePaddedSize; // scale from SVT UV space to frame (padded) local UV space
+		*OutBias = -(FrameBoundsPaddedMin / GlobalVolumeRes * *OutScale);
 	}
 }
 
@@ -457,17 +986,20 @@ void UStreamableSparseVolumeTexture::FinishDestroy()
 void UStreamableSparseVolumeTexture::BeginDestroy()
 {
 	Super::BeginDestroy();
-	for (FSparseVolumeTextureFrame& Frame : Frames)
+	for (FSparseVolumeTextureFrameMips& FrameMips : Frames)
 	{
-		if (Frame.SparseVolumeTextureSceneProxy)
+		for (FSparseVolumeTextureFrame& Frame : FrameMips)
 		{
-			ENQUEUE_RENDER_COMMAND(UStreamableSparseVolumeTexture_DeleteSVTProxy)(
-				[Proxy = Frame.SparseVolumeTextureSceneProxy](FRHICommandListImmediate& RHICmdList)
-				{
-					Proxy->ReleaseResource();
-					delete Proxy;
-				});
-			Frame.SparseVolumeTextureSceneProxy = nullptr;
+			if (Frame.SparseVolumeTextureSceneProxy)
+			{
+				ENQUEUE_RENDER_COMMAND(UStreamableSparseVolumeTexture_DeleteSVTProxy)(
+					[Proxy = Frame.SparseVolumeTextureSceneProxy](FRHICommandListImmediate& RHICmdList)
+					{
+						Proxy->ReleaseResource();
+						delete Proxy;
+					});
+				Frame.SparseVolumeTextureSceneProxy = nullptr;
+			}
 		}
 	}
 }
@@ -478,6 +1010,8 @@ void UStreamableSparseVolumeTexture::Serialize(FArchive& Ar)
 
 	int32 NumFrames = Frames.Num();
 	Ar << NumFrames;
+	int32 NumMipLevels = Frames.IsEmpty() ? 0 : Frames[0].Num();
+	Ar << NumMipLevels;
 
 	if (Ar.IsLoading())
 	{
@@ -486,7 +1020,11 @@ void UStreamableSparseVolumeTexture::Serialize(FArchive& Ar)
 	}
 	for (int32 FrameIndex = 0; FrameIndex < NumFrames; ++FrameIndex)
 	{
-		Frames[FrameIndex].Serialize(Ar, this, FrameIndex);
+		Frames[FrameIndex].SetNum(NumMipLevels);
+		for (int32 MipLevel = 0; MipLevel < NumMipLevels; ++MipLevel)
+		{
+			Frames[FrameIndex][MipLevel].Serialize(Ar, this, FrameIndex);
+		}
 	}
 }
 
@@ -505,11 +1043,15 @@ void UStreamableSparseVolumeTexture::GetResourceSizeEx(FResourceSizeEx& Cumulati
 	SIZE_T SizeCPU = sizeof(Frames);
 	SIZE_T SizeGPU = 0;
 	SizeCPU += Frames.GetAllocatedSize();
-	for (const FSparseVolumeTextureFrame& Frame : Frames)
+	for (const FSparseVolumeTextureFrameMips& FrameMips : Frames)
 	{
-		if (Frame.SparseVolumeTextureSceneProxy)
+		SizeCPU += FrameMips.GetAllocatedSize();
+		for (const FSparseVolumeTextureFrame& Frame : FrameMips)
 		{
-			Frame.SparseVolumeTextureSceneProxy->GetMemorySize(&SizeCPU, &SizeGPU);
+			if (Frame.SparseVolumeTextureSceneProxy)
+			{
+				Frame.SparseVolumeTextureSceneProxy->GetMemorySize(&SizeCPU, &SizeGPU);
+			}
 		}
 	}
 	ISparseVolumeTextureStreamingManager& StreamingManager = IStreamingManager::Get().GetSparseVolumeTextureStreamingManager();
@@ -518,10 +1060,10 @@ void UStreamableSparseVolumeTexture::GetResourceSizeEx(FResourceSizeEx& Cumulati
 	CumulativeResourceSize.AddDedicatedVideoMemoryBytes(SizeGPU);
 }
 
-const FSparseVolumeTextureSceneProxy* UStreamableSparseVolumeTexture::GetStreamedFrameProxyOrFallback(int32 FrameIndex) const
+const FSparseVolumeTextureSceneProxy* UStreamableSparseVolumeTexture::GetStreamedFrameProxyOrFallback(int32 FrameIndex, int32 MipLevel) const
 {
 	ISparseVolumeTextureStreamingManager& StreamingManager = IStreamingManager::Get().GetSparseVolumeTextureStreamingManager();
-	const FSparseVolumeTextureSceneProxy* Proxy = StreamingManager.GetSparseVolumeTextureSceneProxy(this, FrameIndex, true);
+	const FSparseVolumeTextureSceneProxy* Proxy = StreamingManager.GetSparseVolumeTextureSceneProxy(this, FrameIndex, MipLevel, true);
 
 	int32 FallbackFrameIndex = FrameIndex;
 	while (!Proxy)
@@ -532,13 +1074,13 @@ const FSparseVolumeTextureSceneProxy* UStreamableSparseVolumeTexture::GetStreame
 			UE_LOG(LogSparseVolumeTexture, Warning, TEXT("Failed to get ANY streamed SparseVolumeTexture frame  SVT: %s, FrameIndex: %i"), *GetName(), FrameIndex);
 			return nullptr;
 		}
-		Proxy = StreamingManager.GetSparseVolumeTextureSceneProxy(this, FallbackFrameIndex, false);
+		Proxy = StreamingManager.GetSparseVolumeTextureSceneProxy(this, FallbackFrameIndex, MipLevel, false);
 	}
 
 	return Proxy;
 }
 
-TArrayView<const FSparseVolumeTextureFrame> UStreamableSparseVolumeTexture::GetFrames() const
+TArrayView<const FSparseVolumeTextureFrameMips> UStreamableSparseVolumeTexture::GetFrames() const
 {
 	return Frames;
 }
@@ -549,29 +1091,35 @@ void UStreamableSparseVolumeTexture::GenerateOrLoadDDCRuntimeDataAndCreateSceneP
 	UE::DerivedData::FRequestOwner DDCRequestOwner(UE::DerivedData::EPriority::Normal);
 	{
 		UE::DerivedData::FRequestBarrier DDCRequestBarrier(DDCRequestOwner);
-		for (FSparseVolumeTextureFrame& Frame : Frames)
+		for (FSparseVolumeTextureFrameMips& FrameMips : Frames)
 		{
-			// Release any previously allocated render thread proxy
-			if (Frame.SparseVolumeTextureSceneProxy)
+			for (FSparseVolumeTextureFrame& Frame : FrameMips)
 			{
-				BeginReleaseResource(Frame.SparseVolumeTextureSceneProxy);
-			}
-			else
-			{
-				Frame.SparseVolumeTextureSceneProxy = new FSparseVolumeTextureSceneProxy();
-			}
+				// Release any previously allocated render thread proxy
+				if (Frame.SparseVolumeTextureSceneProxy)
+				{
+					BeginReleaseResource(Frame.SparseVolumeTextureSceneProxy);
+				}
+				else
+				{
+					Frame.SparseVolumeTextureSceneProxy = new FSparseVolumeTextureSceneProxy();
+				}
 
-			GenerateOrLoadDDCRuntimeDataForFrame(Frame, DDCRequestOwner);
+				GenerateOrLoadDDCRuntimeDataForFrame(Frame, DDCRequestOwner);
+			}
 		}
 	}
 
 	// Wait for all DDC requests to complete before creating the proxies
 	DDCRequestOwner.Wait();
 
-	for (FSparseVolumeTextureFrame& Frame : Frames)
+	for (FSparseVolumeTextureFrameMips& FrameMips : Frames)
 	{
-		// Runtime data is now valid, initialize the render thread proxy
-		BeginInitResource(Frame.SparseVolumeTextureSceneProxy);
+		for (FSparseVolumeTextureFrame& Frame : FrameMips)
+		{
+			// Runtime data is now valid, initialize the render thread proxy
+			BeginInitResource(Frame.SparseVolumeTextureSceneProxy);
+		}
 	}
 
 	IStreamingManager::Get().GetSparseVolumeTextureStreamingManager().AddSparseVolumeTexture(this);
@@ -664,21 +1212,24 @@ const FSparseVolumeTextureSceneProxy* UAnimatedSparseVolumeTexture::GetSparseVol
 	// When an AnimatedSparseVolumeTexture is used as SparseVolumeTexture, it can only be previewed using a single preview frame.
 	check(!Frames.IsEmpty());
 	const int32 FrameIndex = PreviewFrameIndex % Frames.Num();
-	return GetStreamedFrameProxyOrFallback(FrameIndex);
+	const int32 MipLevel = FMath::Clamp(PreviewMipLevel, 0, GetNumMipLevels() - 1);
+	return GetStreamedFrameProxyOrFallback(FrameIndex, MipLevel);
 }
 
-const FSparseVolumeTextureSceneProxy* UAnimatedSparseVolumeTexture::GetSparseVolumeTextureFrameSceneProxy(int32 FrameIndex) const
+const FSparseVolumeTextureSceneProxy* UAnimatedSparseVolumeTexture::GetSparseVolumeTextureFrameSceneProxy(int32 FrameIndex, int32 MipLevel) const
 {
 	check(!Frames.IsEmpty());
 	FrameIndex = FrameIndex % Frames.Num();
-	return GetStreamedFrameProxyOrFallback(FrameIndex);
+	MipLevel = FMath::Clamp(MipLevel, 0, GetNumMipLevels() - 1);
+	return GetStreamedFrameProxyOrFallback(FrameIndex, MipLevel);
 }
 
-const FSparseVolumeAssetHeader* UAnimatedSparseVolumeTexture::GetSparseVolumeTextureFrameHeader(int32 FrameIndex) const
+const FSparseVolumeAssetHeader* UAnimatedSparseVolumeTexture::GetSparseVolumeTextureFrameHeader(int32 FrameIndex, int32 MipLevel) const
 {
 	check(!Frames.IsEmpty());
 	FrameIndex = FrameIndex % Frames.Num();
-	const FSparseVolumeTextureSceneProxy* Proxy = GetStreamedFrameProxyOrFallback(FrameIndex);
+	MipLevel = FMath::Clamp(MipLevel, 0, GetNumMipLevels() - 1);
+	const FSparseVolumeTextureSceneProxy* Proxy = GetStreamedFrameProxyOrFallback(FrameIndex, MipLevel);
 	return Proxy ? &Proxy->GetHeader() : nullptr;
 }
 
@@ -689,7 +1240,7 @@ USparseVolumeTextureFrame::USparseVolumeTextureFrame(const FObjectInitializer& O
 {
 }
 
-USparseVolumeTextureFrame* USparseVolumeTextureFrame::CreateFrame(USparseVolumeTexture* Texture, int32 FrameIndex)
+USparseVolumeTextureFrame* USparseVolumeTextureFrame::CreateFrame(USparseVolumeTexture* Texture, int32 FrameIndex, int32 MipLevel)
 {
 	if (!Texture)
 	{
@@ -701,7 +1252,7 @@ USparseVolumeTextureFrame* USparseVolumeTextureFrame::CreateFrame(USparseVolumeT
 	{
 		UStreamableSparseVolumeTexture* StreamableSVT = CastChecked<UStreamableSparseVolumeTexture>(Texture);
 		check(StreamableSVT);
-		Proxy = StreamableSVT->GetStreamedFrameProxyOrFallback(FrameIndex);
+		Proxy = StreamableSVT->GetStreamedFrameProxyOrFallback(FrameIndex, MipLevel);
 	}
 	else
 	{
@@ -711,17 +1262,17 @@ USparseVolumeTextureFrame* USparseVolumeTextureFrame::CreateFrame(USparseVolumeT
 	if (Proxy)
 	{
 		USparseVolumeTextureFrame* Frame = NewObject<USparseVolumeTextureFrame>();
-		Frame->Initialize(Proxy, Texture->GetVolumeBounds());
+		Frame->Initialize(Proxy, Texture->GetVolumeResolution());
 		return Frame;
 	}
 	
 	return nullptr;
 }
 
-void USparseVolumeTextureFrame::Initialize(const FSparseVolumeTextureSceneProxy* InSceneProxy, const FBox& InVolumeBounds)
+void USparseVolumeTextureFrame::Initialize(const FSparseVolumeTextureSceneProxy* InSceneProxy, const FIntVector& InVolumeResolution)
 {
 	SceneProxy = InSceneProxy;
-	VolumeBounds = InVolumeBounds;
+	VolumeResolution = InVolumeResolution;
 }
 
 ////////////////////////////////////////////////////////////////////////////////////////////////
@@ -792,7 +1343,7 @@ void UAnimatedSparseVolumeTextureController::SetFractionalFrameIndex(float Frame
 		return;
 	}
 
-	const int32 FrameCount = SparseVolumeTexture->GetFrameCount();
+	const int32 FrameCount = SparseVolumeTexture->GetNumFrames();
 	Frame = FMath::Fmod(Frame, (float)FrameCount);
 	Time = Frame / (FrameRate + UE_SMALL_NUMBER);
 }
@@ -814,7 +1365,7 @@ float UAnimatedSparseVolumeTextureController::GetFractionalFrameIndex()
 		return 0.0f;
 	}
 
-	const int32 FrameCount = SparseVolumeTexture->GetFrameCount();
+	const int32 FrameCount = SparseVolumeTexture->GetNumFrames();
 	const float FrameIndexF = FMath::Fmod(Time * FrameRate, (float)FrameCount);
 	return FrameIndexF;
 }
@@ -831,7 +1382,7 @@ USparseVolumeTextureFrame* UAnimatedSparseVolumeTextureController::GetCurrentFra
 	const int32 FrameIndex = (int32)FrameIndexF;
 
 	// Create and initialize a USparseVolumeTextureFrame which holds the frame to sample and can be bound to shaders
-	USparseVolumeTextureFrame* Frame = USparseVolumeTextureFrame::CreateFrame(SparseVolumeTexture, FrameIndex);
+	USparseVolumeTextureFrame* Frame = USparseVolumeTextureFrame::CreateFrame(SparseVolumeTexture, FrameIndex, MipLevel);
 
 	return Frame;
 }
@@ -849,8 +1400,8 @@ void UAnimatedSparseVolumeTextureController::GetLerpFrames(USparseVolumeTextureF
 	LerpAlpha = FMath::Frac(FrameIndexF);
 
 	// Create and initialize a USparseVolumeTextureFrame which holds the frame to sample and can be bound to shaders
-	Frame0 = USparseVolumeTextureFrame::CreateFrame(SparseVolumeTexture, FrameIndex);
-	Frame1 = USparseVolumeTextureFrame::CreateFrame(SparseVolumeTexture, FrameIndex + 1);
+	Frame0 = USparseVolumeTextureFrame::CreateFrame(SparseVolumeTexture, FrameIndex, MipLevel);
+	Frame1 = USparseVolumeTextureFrame::CreateFrame(SparseVolumeTexture, FrameIndex + 1, MipLevel);
 }
 
 float UAnimatedSparseVolumeTextureController::GetDuration()
@@ -860,7 +1411,7 @@ float UAnimatedSparseVolumeTextureController::GetDuration()
 		return 0.0f;
 	}
 
-	const int32 FrameCount = SparseVolumeTexture->GetFrameCount();
+	const int32 FrameCount = SparseVolumeTexture->GetNumFrames();
 	const float AnimationDuration = FrameCount / (FrameRate + UE_SMALL_NUMBER);
 	return AnimationDuration;
 }
diff --git a/Engine/Source/Runtime/Engine/Private/SparseVolumeTexture/SparseVolumeTextureStreaming.cpp b/Engine/Source/Runtime/Engine/Private/SparseVolumeTexture/SparseVolumeTextureStreaming.cpp
--- a/Engine/Source/Runtime/Engine/Private/SparseVolumeTexture/SparseVolumeTextureStreaming.cpp
+++ b/Engine/Source/Runtime/Engine/Private/SparseVolumeTexture/SparseVolumeTextureStreaming.cpp
@@ -17,9 +17,29 @@ static FAutoConsoleVariableRef CVarSVTStreamingPrefetchCount(
 	TEXT("Number of frames to prefetch when streaming animated SparseVolumeTexture frames."),
 	ECVF_Scalability);
 
+static int32 SVTFrameAndLevelToChunkIndex(int32 FrameIndex, int32 MipLevel, int32 NumFrames)
+{
+	return MipLevel * NumFrames + FrameIndex;
+}
+
+static int32 SVTChunkIndexToFrame(int32 ChunkIndex, int32 NumFrames)
+{
+	return ChunkIndex % NumFrames;
+}
+
+static int32 SVTChunkIndexToMipLevel(int32 ChunkIndex, int32 NumFrames)
+{
+	return ChunkIndex / NumFrames;
+}
+
 ////////////////////////////////////////////////////////////////////////////////////////////////
 
-void FLoadedSparseVolumeTextureFrame::CleanUpIORequest()
+FLoadedSparseVolumeTextureChunk::~FLoadedSparseVolumeTextureChunk()
+{
+	checkf(Proxy.load() == nullptr, TEXT("Render proxy ptr not null (%p), ChunkIndex: %i"), Proxy.load(), ChunkIndex);
+}
+
+void FLoadedSparseVolumeTextureChunk::CleanUpIORequest()
 {
 	if (IORequest)
 	{
@@ -33,7 +53,7 @@ void FLoadedSparseVolumeTextureFrame::CleanUpIORequest()
 
 FStreamingSparseVolumeTextureData::FStreamingSparseVolumeTextureData()
 {
-	ResetRequestedFrames();
+	ResetRequestedChunks();
 }
 
 FStreamingSparseVolumeTextureData::~FStreamingSparseVolumeTextureData()
@@ -54,12 +74,13 @@ bool FStreamingSparseVolumeTextureData::Initialize(UStreamableSparseVolumeTextur
 	StreamingManager = InStreamingManager;
 
 	// Always get the first frame of data so we can play immediately
-	check(LoadedFrames.IsEmpty());
-	check(LoadedFrameIndices.IsEmpty());
+	check(LoadedChunks.IsEmpty());
+	check(LoadedChunkIndices.IsEmpty());
 	check(!InSparseVolumeTexture->GetFrames().IsEmpty()); // Must hold at least the first frame
+	check(!InSparseVolumeTexture->GetFrames()[0].IsEmpty()); // Must hold at least the first frame
 
-	AddNewLoadedFrame(0, InSparseVolumeTexture->GetFrames()[0].SparseVolumeTextureSceneProxy);
-	LoadedFrameIndices.Add(0);
+	AddNewLoadedChunk(0, InSparseVolumeTexture->GetFrames()[0][0].SparseVolumeTextureSceneProxy);
+	LoadedChunkIndices.Add(0);
 
 	return true;
 }
@@ -77,9 +98,9 @@ void FStreamingSparseVolumeTextureData::FreeResources()
 		check(Pass < 2); // we should be done after two passes. Pass 0 will start anything we need and pass 1 will complete those requests
 	}
 
-	for (FLoadedSparseVolumeTextureFrame& LoadedFrame : LoadedFrames)
+	for (FLoadedSparseVolumeTextureChunk& LoadedChunk : LoadedChunks)
 	{
-		FreeLoadedFrame(LoadedFrame);
+		FreeLoadedChunk(LoadedChunk);
 	}
 }
 
@@ -90,24 +111,24 @@ bool FStreamingSparseVolumeTextureData::UpdateStreamingStatus()
 		return false;
 	}
 
-	// Handle failed frames first
-	for (int32 LoadedFrameIdx = 0; LoadedFrameIdx < LoadedFrames.Num(); ++LoadedFrameIdx)
+	// Handle failed chunks first
+	for (int32 LoadedChunkIdx = 0; LoadedChunkIdx < LoadedChunks.Num(); ++LoadedChunkIdx)
 	{
-		FLoadedSparseVolumeTextureFrame& LoadedFrame = LoadedFrames[LoadedFrameIdx];
-		if (LoadFailedFrameIndices.Contains(LoadedFrame.FrameIndex))
+		FLoadedSparseVolumeTextureChunk& LoadedChunk = LoadedChunks[LoadedChunkIdx];
+		if (LoadFailedChunkIndices.Contains(LoadedChunk.ChunkIndex))
 		{
 			// Mark as not loaded
-			LoadedFrameIndices.Remove(LoadedFrame.FrameIndex);
+			LoadedChunkIndices.Remove(LoadedChunk.ChunkIndex);
 
-			// Remove this frame
-			FreeLoadedFrame(LoadedFrame);
+			// Remove this chunk
+			FreeLoadedChunk(LoadedChunk);
 
-			FScopeLock LoadedFramesLock(&LoadedFramesCriticalSection);
-			LoadedFrames.RemoveAtSwap(LoadedFrameIdx, 1, false);
+			FScopeLock LoadedChunksLock(&LoadedChunksCriticalSection);
+			LoadedChunks.RemoveAtSwap(LoadedChunkIdx, 1, false);
 		}
 	}
 
-	LoadFailedFrameIndices.Reset();
+	LoadFailedChunkIndices.Reset();
 
 	bool bHasPendingRequestInFlight = false;
 	TArray<int32> IndicesToLoad;
@@ -115,25 +136,25 @@ bool FStreamingSparseVolumeTextureData::UpdateStreamingStatus()
 
 	if (HasPendingRequests(IndicesToLoad, IndicesToFree))
 	{
-		for (FLoadedSparseVolumeTextureFrame& LoadedFrame : LoadedFrames)
+		for (FLoadedSparseVolumeTextureChunk& LoadedChunk : LoadedChunks)
 		{
-			if (LoadedFrame.IORequest)
+			if (LoadedChunk.IORequest)
 			{
-				const bool bRequestFinished = LoadedFrame.IORequest->PollCompletion();
+				const bool bRequestFinished = LoadedChunk.IORequest->PollCompletion();
 				bHasPendingRequestInFlight |= !bRequestFinished;
 				if (bRequestFinished)
 				{
-					LoadedFrame.CleanUpIORequest();
+					LoadedChunk.CleanUpIORequest();
 				}
 			}
 		}
 
-		LoadedFrameIndices = RequestedFrameIndices;
+		LoadedChunkIndices = RequestedChunkIndices;
 
 		BeginPendingRequests(IndicesToLoad, IndicesToFree);
 	}
 
-	ResetRequestedFrames();
+	ResetRequestedChunks();
 
 	// Print out the currently resident frames to the screen
 #if 0
@@ -144,21 +165,22 @@ bool FStreamingSparseVolumeTextureData::UpdateStreamingStatus()
 #else
 		FString Message = TEXT("Streaming SVT Frames in Memory: ");
 #endif
-		TArray<int32> FramesInMemoryIndices;
-		for (FLoadedSparseVolumeTextureFrame& LoadedFrame : LoadedFrames)
+		TArray<int32> InMemoryChunkIndices;
+		for (FLoadedSparseVolumeTextureChunk& LoadedChunk : LoadedChunks)
 		{
-			if (LoadedFrame.bOwnsProxy)
+			if (LoadedChunk.bOwnsProxy)
 			{
-				FramesInMemoryIndices.Add(LoadedFrame.FrameIndex);
+				InMemoryChunkIndices.Add(LoadedChunk.ChunkIndex);
 			}
 		}
 		
-		if (!FramesInMemoryIndices.IsEmpty())
+		if (!InMemoryChunkIndices.IsEmpty())
 		{
-			FramesInMemoryIndices.Sort();
-			for (int32 FrameIndex : FramesInMemoryIndices)
+			const int32 NumFrames = SparseVolumeTexture->GetNumFrames();
+			InMemoryChunkIndices.Sort();
+			for (int32 ChunkIndex : InMemoryChunkIndices)
 			{
-				Message += FString::Format(TEXT("{0}, "), { FrameIndex });
+				Message += FString::Format(TEXT("F:{0} M:{1}, "), { SVTChunkIndexToFrame(ChunkIndex, NumFrames), SVTChunkIndexToMipLevel(ChunkIndex, NumFrames) });
 			}
 			GEngine->AddOnScreenDebugMessage(-1, 0.1f, FColor::Yellow, Message);
 		}
@@ -174,18 +196,18 @@ bool FStreamingSparseVolumeTextureData::HasPendingRequests(TArray<int32>& Indice
 	IndicesToFree.Reset();
 
 	// Find indices that aren't loaded
-	for (int32 RequestedIndex : RequestedFrameIndices)
+	for (int32 RequestedIndex : RequestedChunkIndices)
 	{
-		if (!LoadedFrameIndices.Contains(RequestedIndex))
+		if (!LoadedChunkIndices.Contains(RequestedIndex))
 		{
 			IndicesToLoad.AddUnique(RequestedIndex);
 		}
 	}
 
 	// Find indices that aren't needed anymore
-	for (int32 LoadedIndex : LoadedFrameIndices)
+	for (int32 LoadedIndex : LoadedChunkIndices)
 	{
-		if (!RequestedFrameIndices.Contains(LoadedIndex))
+		if (!RequestedChunkIndices.Contains(LoadedIndex))
 		{
 			IndicesToFree.AddUnique(LoadedIndex);
 		}
@@ -196,19 +218,19 @@ bool FStreamingSparseVolumeTextureData::HasPendingRequests(TArray<int32>& Indice
 
 void FStreamingSparseVolumeTextureData::BeginPendingRequests(const TArray<int32>& IndicesToLoad, const TArray<int32>& IndicesToFree)
 {
-	// Mark frames for removal in case they can be reused
+	// Mark chunks for removal in case they can be reused
 	{
 		for (int32 IndexToFree : IndicesToFree)
 		{
-			for (int32 LoadedFrameIdx = 0; LoadedFrameIdx < LoadedFrames.Num(); ++LoadedFrameIdx)
+			for (int32 LoadedChunkIdx = 0; LoadedChunkIdx < LoadedChunks.Num(); ++LoadedChunkIdx)
 			{
 				check(IndexToFree != 0);
-				if (LoadedFrames[LoadedFrameIdx].FrameIndex == IndexToFree)
+				if (LoadedChunks[LoadedChunkIdx].ChunkIndex == IndexToFree)
 				{
-					FreeLoadedFrame(LoadedFrames[LoadedFrameIdx]);
+					FreeLoadedChunk(LoadedChunks[LoadedChunkIdx]);
 
-					FScopeLock LoadedFramesLock(&LoadedFramesCriticalSection);
-					LoadedFrames.RemoveAtSwap(LoadedFrameIdx, 1, false);
+					FScopeLock LoadedChunksLock(&LoadedChunksCriticalSection);
+					LoadedChunks.RemoveAtSwap(LoadedChunkIdx, 1, false);
 					break;
 				}
 			}
@@ -216,32 +238,36 @@ void FStreamingSparseVolumeTextureData::BeginPendingRequests(const TArray<int32>
 	}
 
 	// Set off all IO Requests
-
+	const int32 NumFrames = SparseVolumeTexture->GetNumFrames();
+	const int32 NumMipLevels = SparseVolumeTexture->GetNumMipLevels();
+	check(NumFrames > 0 && NumMipLevels > 0);
 	const EAsyncIOPriorityAndFlags AsyncIOPriority = AIOP_CriticalPath; //Set to Crit temporarily as emergency speculative fix for streaming issue
-	TArrayView<const FSparseVolumeTextureFrame> SVTFrames = SparseVolumeTexture->GetFrames();
+	TArrayView<const FSparseVolumeTextureFrameMips> SVTFrames = SparseVolumeTexture->GetFrames();
 	for (int32 IndexToLoad : IndicesToLoad)
 	{
-		const FSparseVolumeTextureFrame& Frame = SVTFrames[IndexToLoad];
+		const int32 FrameToLoad = SVTChunkIndexToFrame(IndexToLoad, NumFrames) % NumFrames;
+		const int32 MipLevelToLoad = FMath::Clamp(SVTChunkIndexToMipLevel(IndexToLoad, NumFrames), 0, NumMipLevels - 1);
+		const FSparseVolumeTextureFrame& Frame = SVTFrames[FrameToLoad][MipLevelToLoad];
 		FSparseVolumeTextureSceneProxy* ExistingProxy = Frame.SparseVolumeTextureSceneProxy;
-		FLoadedSparseVolumeTextureFrame& FrameStorage = AddNewLoadedFrame(IndexToLoad, ExistingProxy);
+		FLoadedSparseVolumeTextureChunk& ChunkStorage = AddNewLoadedChunk(IndexToLoad, ExistingProxy);
 
 		if (!ExistingProxy)
 		{
-			UE_CLOG(FrameStorage.Proxy.load() != nullptr, LogSparseVolumeTextureStreaming, Fatal, TEXT("Existing render proxy for streaming SparseVolumeTexture frame."));
-			UE_CLOG(FrameStorage.IORequest, LogSparseVolumeTextureStreaming, Fatal, TEXT("Streaming SparseVolumeTexture frame already has IORequest."));
+			UE_CLOG(ChunkStorage.Proxy.load() != nullptr, LogSparseVolumeTextureStreaming, Fatal, TEXT("Existing render proxy for streaming SparseVolumeTexture frame."));
+			UE_CLOG(ChunkStorage.IORequest, LogSparseVolumeTextureStreaming, Fatal, TEXT("Streaming SparseVolumeTexture frame already has IORequest."));
 
-			int64 FrameSize = Frame.RuntimeStreamedInData.GetBulkDataSize();
-			FrameStorage.RequestStart = FPlatformTime::Seconds();
-			UE_LOG(LogSparseVolumeTextureStreaming, Warning, TEXT("SparseVolumeTexture streaming request started %s Frame:%i At:%.3f\n"), *SparseVolumeTexture->GetName(), IndexToLoad, FrameStorage.RequestStart);
-			FBulkDataIORequestCallBack AsyncFileCallBack = [this, IndexToLoad, FrameSize](bool bWasCancelled, IBulkDataIORequest* Req)
+			const int64 ChunkSize = Frame.RuntimeStreamedInData.GetBulkDataSize();
+			ChunkStorage.RequestStart = FPlatformTime::Seconds();
+			UE_LOG(LogSparseVolumeTextureStreaming, Warning, TEXT("SparseVolumeTexture streaming request started %s Frame:%i At:%.3f\n"), *SparseVolumeTexture->GetName(), IndexToLoad, ChunkStorage.RequestStart);
+			FBulkDataIORequestCallBack AsyncFileCallBack = [this, IndexToLoad, ChunkSize](bool bWasCancelled, IBulkDataIORequest* Req)
 			{
-				StreamingManager->OnAsyncFileCallback(this, IndexToLoad, FrameSize, Req, bWasCancelled);
+				StreamingManager->OnAsyncFileCallback(this, IndexToLoad, ChunkSize, Req, bWasCancelled);
 			};
 
 			UE_LOG(LogSparseVolumeTextureStreaming, Warning, TEXT("Loading streaming SparseVolumeTexture %s Frame:%i Offset:%i Size:%i File:%s\n"),
 				*SparseVolumeTexture->GetName(), IndexToLoad, Frame.RuntimeStreamedInData.GetBulkDataOffsetInFile(), Frame.RuntimeStreamedInData.GetBulkDataSize(), *Frame.RuntimeStreamedInData.GetDebugName());
-			FrameStorage.IORequest = Frame.RuntimeStreamedInData.CreateStreamingRequest(AsyncIOPriority, &AsyncFileCallBack, nullptr);
-			if (!FrameStorage.IORequest)
+			ChunkStorage.IORequest = Frame.RuntimeStreamedInData.CreateStreamingRequest(AsyncIOPriority, &AsyncFileCallBack, nullptr);
+			if (!ChunkStorage.IORequest)
 			{
 				UE_LOG(LogSparseVolumeTextureStreaming, Error, TEXT("SparseVolumeTexture streaming read request failed."));
 			}
@@ -254,26 +280,26 @@ bool FStreamingSparseVolumeTextureData::BlockTillAllRequestsFinished(float TimeL
 	QUICK_SCOPE_CYCLE_COUNTER(FStreamingSparseVolumeTextureData_BlockTillAllRequestsFinished);
 	if (TimeLimit == 0.0f)
 	{
-		for (FLoadedSparseVolumeTextureFrame& LoadedFrame : LoadedFrames)
+		for (FLoadedSparseVolumeTextureChunk& LoadedChunk : LoadedChunks)
 		{
-			LoadedFrame.CleanUpIORequest();
+			LoadedChunk.CleanUpIORequest();
 		}
 	}
 	else
 	{
 		double EndTime = FPlatformTime::Seconds() + TimeLimit;
-		for (FLoadedSparseVolumeTextureFrame& LoadedFrame : LoadedFrames)
+		for (FLoadedSparseVolumeTextureChunk& LoadedChunk : LoadedChunks)
 		{
-			if (LoadedFrame.IORequest)
+			if (LoadedChunk.IORequest)
 			{
 				float ThisTimeLimit = EndTime - FPlatformTime::Seconds();
 				if (ThisTimeLimit < 0.001f || // one ms is the granularity of the platform event system
-					!LoadedFrame.IORequest->WaitCompletion(ThisTimeLimit))
+					!LoadedChunk.IORequest->WaitCompletion(ThisTimeLimit))
 				{
 					return false;
 				}
 
-				LoadedFrame.CleanUpIORequest();
+				LoadedChunk.CleanUpIORequest();
 			}
 		}
 	}
@@ -282,41 +308,41 @@ bool FStreamingSparseVolumeTextureData::BlockTillAllRequestsFinished(float TimeL
 
 void FStreamingSparseVolumeTextureData::GetMemorySize(SIZE_T* SizeCPU, SIZE_T* SizeGPU) const
 {
-	FScopeLock LoadedFramesLock(&LoadedFramesCriticalSection);
-	for (const FLoadedSparseVolumeTextureFrame& LoadedFrame : LoadedFrames)
+	FScopeLock LoadedChunksLock(&LoadedChunksCriticalSection);
+	for (const FLoadedSparseVolumeTextureChunk& LoadedChunk : LoadedChunks)
 	{
-		const FSparseVolumeTextureSceneProxy* Proxy = LoadedFrame.Proxy.load();
-		if (LoadedFrame.bOwnsProxy && Proxy)
+		const FSparseVolumeTextureSceneProxy* Proxy = LoadedChunk.Proxy.load();
+		if (LoadedChunk.bOwnsProxy && Proxy)
 		{
 			Proxy->GetMemorySize(SizeCPU, SizeGPU);
 		}
 	}
 }
 
-FLoadedSparseVolumeTextureFrame& FStreamingSparseVolumeTextureData::AddNewLoadedFrame(int32 FrameIndex, FSparseVolumeTextureSceneProxy* ExistingProxy)
+FLoadedSparseVolumeTextureChunk& FStreamingSparseVolumeTextureData::AddNewLoadedChunk(int32 ChunkIndex, FSparseVolumeTextureSceneProxy* ExistingProxy)
 {
-	FScopeLock LoadedFramesLock(&LoadedFramesCriticalSection);
-	const int32 NewIndex = LoadedFrames.AddDefaulted();
+	FScopeLock LoadedChunksLock(&LoadedChunksCriticalSection);
+	const int32 NewIndex = LoadedChunks.AddDefaulted();
 
-	LoadedFrames[NewIndex].Proxy = ExistingProxy;
-	LoadedFrames[NewIndex].FrameIndex = FrameIndex;
-	LoadedFrames[NewIndex].bOwnsProxy = false;
-	return LoadedFrames[NewIndex];
+	LoadedChunks[NewIndex].Proxy = ExistingProxy;
+	LoadedChunks[NewIndex].ChunkIndex = ChunkIndex;
+	LoadedChunks[NewIndex].bOwnsProxy = false;
+	return LoadedChunks[NewIndex];
 }
 
-void FStreamingSparseVolumeTextureData::FreeLoadedFrame(FLoadedSparseVolumeTextureFrame& LoadedFrame)
+void FStreamingSparseVolumeTextureData::FreeLoadedChunk(FLoadedSparseVolumeTextureChunk& LoadedChunk)
 {
-	if (LoadedFrame.IORequest)
+	if (LoadedChunk.IORequest)
 	{
-		LoadedFrame.IORequest->Cancel();
-		LoadedFrame.IORequest->WaitCompletion();
-		delete LoadedFrame.IORequest;
-		LoadedFrame.IORequest = nullptr;
+		LoadedChunk.IORequest->Cancel();
+		LoadedChunk.IORequest->WaitCompletion();
+		delete LoadedChunk.IORequest;
+		LoadedChunk.IORequest = nullptr;
 	}
 
-	if (LoadedFrame.bOwnsProxy)
+	if (LoadedChunk.bOwnsProxy)
 	{
-		FSparseVolumeTextureSceneProxy* Proxy = LoadedFrame.Proxy.load();
+		FSparseVolumeTextureSceneProxy* Proxy = LoadedChunk.Proxy.load();
 		ENQUEUE_RENDER_COMMAND(FStreamingSparseVolumeTextureData_DeleteSVTProxy)(
 			[Proxy](FRHICommandListImmediate& RHICmdList)
 			{
@@ -325,15 +351,15 @@ void FStreamingSparseVolumeTextureData::FreeLoadedFrame(FLoadedSparseVolumeTextu
 			});
 	}
 
-	LoadedFrame.Proxy = nullptr;
-	LoadedFrame.bOwnsProxy = false;
-	LoadedFrame.FrameIndex = INDEX_NONE;
+	LoadedChunk.Proxy = nullptr;
+	LoadedChunk.bOwnsProxy = false;
+	LoadedChunk.ChunkIndex = INDEX_NONE;
 }
 
-void FStreamingSparseVolumeTextureData::ResetRequestedFrames()
+void FStreamingSparseVolumeTextureData::ResetRequestedChunks()
 {
-	RequestedFrameIndices.Reset();
-	RequestedFrameIndices.Add(0); //Always want first frame
+	RequestedChunkIndices.Reset();
+	RequestedChunkIndices.Add(0); //Always want first frame
 }
 
 ////////////////////////////////////////////////////////////////////////////////////////////////
@@ -431,29 +457,30 @@ void FSparseVolumeTextureStreamingManager::GetMemorySizeForSparseVolumeTexture(c
 	}
 }
 
-const FSparseVolumeTextureSceneProxy* FSparseVolumeTextureStreamingManager::GetSparseVolumeTextureSceneProxy(const UStreamableSparseVolumeTexture* SparseVolumeTexture, int32 FrameIndex, bool bTrackAsRequested)
+const FSparseVolumeTextureSceneProxy* FSparseVolumeTextureStreamingManager::GetSparseVolumeTextureSceneProxy(const UStreamableSparseVolumeTexture* SparseVolumeTexture, int32 FrameIndex, int32 MipLevel, bool bTrackAsRequested)
 {
 	FScopeLock Lock(&CriticalSection);
 
 	FStreamingSparseVolumeTextureData* StreamingData = StreamingSparseVolumeTextures.FindRef(SparseVolumeTexture);
 	if (StreamingData)
 	{
+		const int32 NumFrames = SparseVolumeTexture->GetNumFrames();
+		const int32 ChunkIndex = SVTFrameAndLevelToChunkIndex(FrameIndex, MipLevel, NumFrames);
 		if (bTrackAsRequested)
 		{
-			StreamingData->RequestedFrameIndices.AddUnique(FrameIndex);
+			StreamingData->RequestedChunkIndices.AddUnique(ChunkIndex);
 
 			// Prefetch next frames
-			const int32 NumFrames = SparseVolumeTexture->GetFrames().Num();
 			const int32 NumPrefetchFrames = FMath::Clamp(GSVTNumPrefetchFrames, 0, NumFrames - 1);
 			for (int32 i = 0; i < NumPrefetchFrames; ++i)
 			{
-				StreamingData->RequestedFrameIndices.AddUnique((FrameIndex + 1 + i) % NumFrames);
+				StreamingData->RequestedChunkIndices.AddUnique(SVTFrameAndLevelToChunkIndex((FrameIndex + 1 + i) % NumFrames, MipLevel, NumFrames));
 			}
 		}
 
-		if (StreamingData->LoadedFrameIndices.Contains(FrameIndex))
+		if (StreamingData->LoadedChunkIndices.Contains(ChunkIndex))
 		{
-			if (const FLoadedSparseVolumeTextureFrame* Frame = Algo::FindBy(StreamingData->LoadedFrames, FrameIndex, &FLoadedSparseVolumeTextureFrame::FrameIndex))
+			if (const FLoadedSparseVolumeTextureChunk* Frame = Algo::FindBy(StreamingData->LoadedChunks, ChunkIndex, &FLoadedSparseVolumeTextureChunk::ChunkIndex))
 			{
 				if (Frame->Proxy.load() == nullptr)
 				{
@@ -464,7 +491,7 @@ const FSparseVolumeTextureSceneProxy* FSparseVolumeTextureStreamingManager::GetS
 			}
 			else
 			{
-				UE_LOG(LogSparseVolumeTextureStreaming, Warning, TEXT("Unable to find requested frame: %i, SVT: %s - Is in LoadedFrameIndices however"), FrameIndex, *SparseVolumeTexture->GetFullName());
+				UE_LOG(LogSparseVolumeTextureStreaming, Warning, TEXT("Unable to find requested frame: %i, SVT: %s - Is in LoadedChunkIndices however"), FrameIndex, *SparseVolumeTexture->GetFullName());
 			}
 		}
 		else
@@ -480,23 +507,23 @@ const FSparseVolumeTextureSceneProxy* FSparseVolumeTextureStreamingManager::GetS
 	return nullptr;
 }
 
-void FSparseVolumeTextureStreamingManager::OnAsyncFileCallback(FStreamingSparseVolumeTextureData* StreamingSVTData, int32 FrameIndex, int64 ReadSize, IBulkDataIORequest* ReadRequest, bool bWasCancelled)
+void FSparseVolumeTextureStreamingManager::OnAsyncFileCallback(FStreamingSparseVolumeTextureData* StreamingSVTData, int32 ChunkIndex, int64 ReadSize, IBulkDataIORequest* ReadRequest, bool bWasCancelled)
 {
 	// Check to see if we successfully managed to load anything
 	uint8* Mem = ReadRequest->GetReadResults();
 
-	FScopeLock Lock(&StreamingSVTData->LoadedFramesCriticalSection);
+	FScopeLock Lock(&StreamingSVTData->LoadedChunksCriticalSection);
 
-	const int32 LoadedFrameIdx = StreamingSVTData->LoadedFrames.IndexOfByPredicate([FrameIndex](const FLoadedSparseVolumeTextureFrame& Frame) { return Frame.FrameIndex == FrameIndex; });
-	check(LoadedFrameIdx != INDEX_NONE);
-	FLoadedSparseVolumeTextureFrame& FrameStorage = StreamingSVTData->LoadedFrames[LoadedFrameIdx];
+	const int32 LoadedChunkIdx = StreamingSVTData->LoadedChunks.IndexOfByPredicate([ChunkIndex](const FLoadedSparseVolumeTextureChunk& Frame) { return Frame.ChunkIndex == ChunkIndex; });
+	check(LoadedChunkIdx != INDEX_NONE);
+	FLoadedSparseVolumeTextureChunk& ChunkStorage = StreamingSVTData->LoadedChunks[LoadedChunkIdx];
 
 	const double CurrentTime = FPlatformTime::Seconds();
-	const double RequestDuration = CurrentTime - FrameStorage.RequestStart;
+	const double RequestDuration = CurrentTime - ChunkStorage.RequestStart;
 
 	if (Mem)
 	{
-		checkf(FrameStorage.Proxy.load() == nullptr, TEXT("Frame storage already has data. (0x%p) FrameIndex:%i LoadedFrameIdx:%i"), FrameStorage.Proxy.load(), FrameIndex, LoadedFrameIdx);
+		checkf(ChunkStorage.Proxy.load() == nullptr, TEXT("Chunk storage already has data. (0x%p) ChunkIndex:%i LoadedChunkIdx:%i"), ChunkStorage.Proxy.load(), ChunkIndex, LoadedChunkIdx);
 
 		TArrayView<const uint8> MemView(Mem, ReadSize);
 		FMemoryReaderView Reader(MemView);
@@ -505,9 +532,9 @@ void FSparseVolumeTextureStreamingManager::OnAsyncFileCallback(FStreamingSparseV
 		NewProxy->GetRuntimeData().Serialize(Reader);
 		BeginInitResource(NewProxy);
 
-		FrameStorage.Proxy = NewProxy;
-		FrameStorage.bOwnsProxy = true;
-		FrameStorage.RequestStart = -2.0; //Signify we have finished loading
+		ChunkStorage.Proxy = NewProxy;
+		ChunkStorage.bOwnsProxy = true;
+		ChunkStorage.RequestStart = -2.0; //Signify we have finished loading
 
 		UE_LOG(LogSparseVolumeTextureStreaming, Log, TEXT("Request Finished %.2f\n SparseVolumeTexture Frame Streamed %.4f\n"), CurrentTime, RequestDuration);
 
@@ -516,8 +543,8 @@ void FSparseVolumeTextureStreamingManager::OnAsyncFileCallback(FStreamingSparseV
 	else
 	{
 		const TCHAR* WasCancelledText = bWasCancelled ? TEXT("Yes") : TEXT("No");
-		UE_LOG(LogSparseVolumeTextureStreaming, Warning, TEXT("Streaming SparseVolumeTexture failed to load frame: %i Load Duration:%.3f, SVT:%s WasCancelled: %s\n"), FrameIndex, RequestDuration, *StreamingSVTData->SparseVolumeTexture->GetName(), WasCancelledText);
+		UE_LOG(LogSparseVolumeTextureStreaming, Warning, TEXT("Streaming SparseVolumeTexture failed to load chunk: %i Load Duration:%.3f, SVT:%s WasCancelled: %s\n"), ChunkIndex, RequestDuration, *StreamingSVTData->SparseVolumeTexture->GetName(), WasCancelledText);
 		
-		StreamingSVTData->LoadFailedFrameIndices.AddUnique(FrameIndex);
+		StreamingSVTData->LoadFailedChunkIndices.AddUnique(ChunkIndex);
 	}
 }
diff --git a/Engine/Source/Runtime/Engine/Public/ContentStreaming.h b/Engine/Source/Runtime/Engine/Public/ContentStreaming.h
--- a/Engine/Source/Runtime/Engine/Public/ContentStreaming.h
+++ b/Engine/Source/Runtime/Engine/Public/ContentStreaming.h
@@ -648,9 +648,10 @@ struct ISparseVolumeTextureStreamingManager : public IStreamingManager
 	 *
 	 * @param SparseVolumeTexture	SparseVolumeTexture we want a frame from
 	 * @param FrameIndex			Index of the frame we want
+	 * @param MipLevel				MipLevel of the frame we want
 	 * @return Either the desired frame proxy or NULL if it's not loaded
 	 */
-	virtual const FSparseVolumeTextureSceneProxy* GetSparseVolumeTextureSceneProxy(const UStreamableSparseVolumeTexture* SparseVolumeTexture, int32 FrameIndex, bool bTrackAsRequested) = 0;
+	virtual const FSparseVolumeTextureSceneProxy* GetSparseVolumeTextureSceneProxy(const UStreamableSparseVolumeTexture* SparseVolumeTexture, int32 FrameIndex, int32 MipLevel, bool bTrackAsRequested) = 0;
 };
 
 
diff --git a/Engine/Source/Runtime/Engine/Public/SparseVolumeTexture/SparseVolumeTextureStreaming.h b/Engine/Source/Runtime/Engine/Public/SparseVolumeTexture/SparseVolumeTextureStreaming.h
--- a/Engine/Source/Runtime/Engine/Public/SparseVolumeTexture/SparseVolumeTextureStreaming.h
+++ b/Engine/Source/Runtime/Engine/Public/SparseVolumeTexture/SparseVolumeTextureStreaming.h
@@ -10,37 +10,26 @@
 class UStreamableSparseVolumeTexture;
 class FSparseVolumeTextureSceneProxy;
 
-struct FLoadedSparseVolumeTextureFrame
+struct FLoadedSparseVolumeTextureChunk
 {
-	std::atomic<FSparseVolumeTextureSceneProxy*> Proxy;
-	class IBulkDataIORequest* IORequest;
-	double RequestStart;
-	int32 FrameIndex;
-	bool bOwnsProxy;
-
-	explicit FLoadedSparseVolumeTextureFrame()
-		:Proxy(),
-		IORequest(),
-		RequestStart(-1.0),
-		FrameIndex(INDEX_NONE)
-	{
-	}
-	~FLoadedSparseVolumeTextureFrame()
-	{
-		checkf(Proxy.load() == nullptr, TEXT("Render proxy ptr not null (%p), FrameIndex: %u"), Proxy.load(), FrameIndex);
-	}
+	std::atomic<FSparseVolumeTextureSceneProxy*> Proxy = nullptr;
+	class IBulkDataIORequest* IORequest = nullptr;
+	double RequestStart = -1.0;
+	int32 ChunkIndex = INDEX_NONE;
+	bool bOwnsProxy = false;
 
+	~FLoadedSparseVolumeTextureChunk();
 	void CleanUpIORequest();
 };
 
 struct FStreamingSparseVolumeTextureData
 {
 	UStreamableSparseVolumeTexture* SparseVolumeTexture;
-	TArray<FLoadedSparseVolumeTextureFrame> LoadedFrames;
-	TArray<int32> LoadedFrameIndices;
-	TArray<int32> RequestedFrameIndices;
-	TArray<int32> LoadFailedFrameIndices;
-	mutable FCriticalSection LoadedFramesCriticalSection;
+	TArray<FLoadedSparseVolumeTextureChunk> LoadedChunks;
+	TArray<int32> LoadedChunkIndices;
+	TArray<int32> RequestedChunkIndices;
+	TArray<int32> LoadFailedChunkIndices;
+	mutable FCriticalSection LoadedChunksCriticalSection;
 
 	explicit FStreamingSparseVolumeTextureData();
 	~FStreamingSparseVolumeTextureData();
@@ -60,9 +49,9 @@ struct FStreamingSparseVolumeTextureData
 private:
 	class FSparseVolumeTextureStreamingManager* StreamingManager;
 
-	FLoadedSparseVolumeTextureFrame& AddNewLoadedFrame(int32 FrameIndex, FSparseVolumeTextureSceneProxy* ExistingProxy);
-	void FreeLoadedFrame(FLoadedSparseVolumeTextureFrame& LoadedFrame);
-	void ResetRequestedFrames();
+	FLoadedSparseVolumeTextureChunk& AddNewLoadedChunk(int32 ChunkIndex, FSparseVolumeTextureSceneProxy* ExistingProxy);
+	void FreeLoadedChunk(FLoadedSparseVolumeTextureChunk& LoadedChunk);
+	void ResetRequestedChunks();
 };
 
 class FSparseVolumeTextureStreamingManager : public ISparseVolumeTextureStreamingManager
@@ -86,10 +75,10 @@ public:
 	virtual void AddSparseVolumeTexture(UStreamableSparseVolumeTexture* SparseVolumeTexture) override;
 	virtual bool RemoveSparseVolumeTexture(UStreamableSparseVolumeTexture* SparseVolumeTexture) override;
 	virtual void GetMemorySizeForSparseVolumeTexture(const UStreamableSparseVolumeTexture* SparseVolumeTexture, SIZE_T* SizeCPU, SIZE_T* SizeGPU) const override;
-	virtual const FSparseVolumeTextureSceneProxy* GetSparseVolumeTextureSceneProxy(const UStreamableSparseVolumeTexture* SparseVolumeTexture, int32 FrameIndex, bool bTrackAsRequested) override;
+	virtual const FSparseVolumeTextureSceneProxy* GetSparseVolumeTextureSceneProxy(const UStreamableSparseVolumeTexture* SparseVolumeTexture, int32 FrameIndex, int32 MipLevel, bool bTrackAsRequested) override;
 	//~ End ISparseVolumeTextureStreamingManager Interface
 
-	void OnAsyncFileCallback(FStreamingSparseVolumeTextureData* StreamingSVTData, int32 FrameIndex, int64 ReadSize, IBulkDataIORequest* ReadRequest, bool bWasCancelled);
+	void OnAsyncFileCallback(FStreamingSparseVolumeTextureData* StreamingSVTData, int32 ChunkIndex, int64 ReadSize, IBulkDataIORequest* ReadRequest, bool bWasCancelled);
 
 private:
 	TMap<UStreamableSparseVolumeTexture*, FStreamingSparseVolumeTextureData*> StreamingSparseVolumeTextures;
diff --git a/Engine/Source/Runtime/Renderer/Private/SparseVolumeTexture/SparseVolumeTextureViewerComponent.cpp b/Engine/Source/Runtime/Renderer/Private/SparseVolumeTexture/SparseVolumeTextureViewerComponent.cpp
--- a/Engine/Source/Runtime/Renderer/Private/SparseVolumeTexture/SparseVolumeTextureViewerComponent.cpp
+++ b/Engine/Source/Runtime/Renderer/Private/SparseVolumeTexture/SparseVolumeTextureViewerComponent.cpp
@@ -55,7 +55,7 @@ void USparseVolumeTextureViewerComponent::PostEditChangeProperty(FPropertyChange
 
 	if (!bAnimate && SparseVolumeTexturePreview)
 	{
-		FrameIndex = int32(AnimationFrame * float(SparseVolumeTexturePreview->GetFrameCount()));
+		FrameIndex = int32(AnimationFrame * float(SparseVolumeTexturePreview->GetNumFrames()));
 	}
 	MarkRenderStateDirty();
 
@@ -86,10 +86,9 @@ FBoxSphereBounds USparseVolumeTextureViewerComponent::CalcBounds(const FTransfor
 		// Then the other dimensions are scaled relatively.
 		// All this account for the size of volume with page table (they add padding).
 		// TODO can we recover world size from OpenVDB meta data in meter?
-		const FBox VolumeBounds = SparseVolumeTexturePreview->GetVolumeBounds();
-		const FVector Extent = VolumeBounds.GetExtent();
-		const double MaxDim = FMath::Max(FMath::Max(Extent.X, Extent.Y), Extent.Z);
-		NormalizedBound.BoxExtent = Extent / MaxDim * SVTViewerDefaultVolumeExtent;
+		const FVector VolumeResolution = FVector(SparseVolumeTexturePreview->GetVolumeResolution());
+		const double MaxDim = FMath::Max(FMath::Max(VolumeResolution.X, VolumeResolution.Y), VolumeResolution.Z);
+		NormalizedBound.BoxExtent = VolumeResolution / MaxDim * SVTViewerDefaultVolumeExtent;
 	}
 	else
 	{
@@ -162,9 +161,9 @@ void USparseVolumeTextureViewerComponent::SendRenderTransformCommand()
 		
 		if (SparseVolumeTextureFrame)
 		{
-			const FVector VolumeBoundsExtent = SparseVolumeTextureFrame->GetVolumeBounds().GetExtent();
-			const double MaxBoundsDim = FMath::Max(FMath::Max(VolumeBoundsExtent.X, VolumeBoundsExtent.Y), VolumeBoundsExtent.Z);
-			VolumeExtent = VolumeBoundsExtent / MaxBoundsDim * SVTViewerDefaultVolumeExtent;
+			const FVector VolumeResolution = FVector(SparseVolumeTextureFrame->GetVolumeResolution());
+			const double MaxBoundsDim = FMath::Max(FMath::Max(VolumeResolution.X, VolumeResolution.Y), VolumeResolution.Z);
+			VolumeExtent = VolumeResolution / MaxBoundsDim * SVTViewerDefaultVolumeExtent;
 
 			SparseVolumeTextureFrame->GetFrameUVScaleBias(&UVScale, &UVBias);
 		}
@@ -188,11 +187,12 @@ void USparseVolumeTextureViewerComponent::SendRenderTransformCommand()
 
 		FSparseVolumeTextureViewerSceneProxy* SVTViewerSceneProxy = SparseVolumeTextureViewerSceneProxy;
 		ENQUEUE_RENDER_COMMAND(FUpdateSparseVolumeTextureViewerProxyTransformCommand)(
-			[SVTViewerSceneProxy, ToLocalMat, ToLocalMatNoScale, CompToVisualize](FRHICommandList& RHICmdList)
+			[SVTViewerSceneProxy, ToLocalMat, ToLocalMatNoScale, CompToVisualize, Ext = Extinction](FRHICommandList& RHICmdList)
 			{
 				SVTViewerSceneProxy->WorldToLocal = ToLocalMat;
 				SVTViewerSceneProxy->WorldToLocalNoScale = ToLocalMatNoScale;
 				SVTViewerSceneProxy->ComponentToVisualize = CompToVisualize;
+				SVTViewerSceneProxy->Extinction = Ext;
 			});
 	}
 }
@@ -203,17 +203,17 @@ void USparseVolumeTextureViewerComponent::TickComponent(float DeltaTime, enum EL
 	{
 		if (bAnimate)
 		{
-			const int32 NumFrames = SparseVolumeTexturePreview->GetFrameCount();
+			const int32 NumFrames = SparseVolumeTexturePreview->GetNumFrames();
 			const float AnimationDuration = NumFrames / (FrameRate + UE_SMALL_NUMBER);
 			AnimationTime = FMath::Fmod(AnimationTime + DeltaTime, (AnimationDuration + UE_SMALL_NUMBER));
 			FrameIndex = int32(AnimationTime * FrameRate);
 		}
 		else
 		{
-			FrameIndex = int32(AnimationFrame * float(SparseVolumeTexturePreview->GetFrameCount()));
+			FrameIndex = int32(AnimationFrame * float(SparseVolumeTexturePreview->GetNumFrames()));
 		}
 
-		SparseVolumeTextureFrame = USparseVolumeTextureFrame::CreateFrame(SparseVolumeTexturePreview, FrameIndex);
+		SparseVolumeTextureFrame = USparseVolumeTextureFrame::CreateFrame(SparseVolumeTexturePreview, FrameIndex, MipLevel);
 	}
 	else
 	{
diff --git a/Engine/Source/Runtime/Renderer/Private/SparseVolumeTexture/SparseVolumeTextureViewerComponent.h b/Engine/Source/Runtime/Renderer/Private/SparseVolumeTexture/SparseVolumeTextureViewerComponent.h
--- a/Engine/Source/Runtime/Renderer/Private/SparseVolumeTexture/SparseVolumeTextureViewerComponent.h
+++ b/Engine/Source/Runtime/Renderer/Private/SparseVolumeTexture/SparseVolumeTextureViewerComponent.h
@@ -44,6 +44,12 @@ class USparseVolumeTextureViewerComponent : public UPrimitiveComponent
 	UPROPERTY(EditAnywhere, BlueprintReadWrite, Category = "Asset Preview", meta = (UIMin = 0, UIMax = 7, ClampMin = 0, ClampMax = 7))
 	int32 ComponentToVisualize;
 
+	UPROPERTY(EditAnywhere, BlueprintReadWrite, Category = "Asset Preview", meta = (UIMin = 0, UIMax = 11, ClampMin = 0, ClampMax = 11))
+	int32 MipLevel;
+
+	UPROPERTY(EditAnywhere, BlueprintReadWrite, Category = "Asset Preview", meta = (UIMin = 0.0, UIMax = 1.0, ClampMin = 0))
+	float Extinction = 0.025f;
+
 	TObjectPtr<class USparseVolumeTextureFrame> SparseVolumeTextureFrame;
 
 	//~ Begin UPrimitiveComponent Interface.
diff --git a/Engine/Source/Runtime/Renderer/Private/SparseVolumeTexture/SparseVolumeTextureViewerRendering.cpp b/Engine/Source/Runtime/Renderer/Private/SparseVolumeTexture/SparseVolumeTextureViewerRendering.cpp
--- a/Engine/Source/Runtime/Renderer/Private/SparseVolumeTexture/SparseVolumeTextureViewerRendering.cpp
+++ b/Engine/Source/Runtime/Renderer/Private/SparseVolumeTexture/SparseVolumeTextureViewerRendering.cpp
@@ -65,6 +65,8 @@ class FVisualizeSparseVolumeTexturePS : public FGlobalShader
 		SHADER_PARAMETER(FVector3f, SparseVolumeTextureResolution)
 		SHADER_PARAMETER(FVector3f, SparseVolumeTexturePageTableResolution)
 		SHADER_PARAMETER(uint32, ComponentToVisualize)
+		SHADER_PARAMETER(uint32, MipLevel)
+		SHADER_PARAMETER(float, Extinction)
 		SHADER_PARAMETER_TEXTURE(Texture3D, SparseVolumeTextureA)
 		SHADER_PARAMETER_TEXTURE(Texture3D, SparseVolumeTextureB)
 		SHADER_PARAMETER_TEXTURE(Texture3D<uint>, SparseVolumeTexturePageTable)
@@ -133,6 +135,8 @@ void AddSparseVolumeTextureViewerRenderPass(FRDGBuilder& GraphBuilder, FSceneRen
 			PsPassParameters->WorldToLocalNoScale1 = FVector3f(WorldToLocalNoScale.M[0][1], WorldToLocalNoScale.M[1][1], WorldToLocalNoScale.M[2][1]);
 			PsPassParameters->WorldToLocalNoScale2 = FVector3f(WorldToLocalNoScale.M[0][2], WorldToLocalNoScale.M[1][2], WorldToLocalNoScale.M[2][2]);
 			PsPassParameters->ComponentToVisualize = SVTProxy->ComponentToVisualize;
+			PsPassParameters->MipLevel = 0;
+			PsPassParameters->Extinction = SVTProxy->Extinction;
 			PsPassParameters->SparseVolumeTextureResolution = FVector3f::OneVector;
 			PsPassParameters->SparseVolumeTexturePageTableResolution = FVector3f::OneVector;
 			PsPassParameters->SparseVolumeTextureA = GBlackVolumeTexture->TextureRHI;
@@ -144,8 +148,9 @@ void AddSparseVolumeTextureViewerRenderPass(FRDGBuilder& GraphBuilder, FSceneRen
 			{
 				const FSparseVolumeAssetHeader& Header = SVTProxy->SparseVolumeTextureSceneProxy->GetHeader();
 
-				PsPassParameters->SparseVolumeTextureResolution = FVector3f(Header.SourceVolumeResolution);
+				PsPassParameters->SparseVolumeTextureResolution = FVector3f(Header.VirtualVolumeResolution);
 				PsPassParameters->SparseVolumeTexturePageTableResolution = FVector3f(Header.PageTableVolumeResolution);
+				PsPassParameters->MipLevel = (uint32)Header.MipLevel;
 
 				FRHITexture* TextureA = SVTProxy->SparseVolumeTextureSceneProxy->GetPhysicalTileDataATextureRHI();
 				if (TextureA)
diff --git a/Engine/Source/Runtime/Renderer/Private/SparseVolumeTexture/SparseVolumeTextureViewerSceneProxy.h b/Engine/Source/Runtime/Renderer/Private/SparseVolumeTexture/SparseVolumeTextureViewerSceneProxy.h
--- a/Engine/Source/Runtime/Renderer/Private/SparseVolumeTexture/SparseVolumeTextureViewerSceneProxy.h
+++ b/Engine/Source/Runtime/Renderer/Private/SparseVolumeTexture/SparseVolumeTextureViewerSceneProxy.h
@@ -21,6 +21,7 @@ public:
 	FMatrix44f WorldToLocal;
 	FMatrix44f WorldToLocalNoScale;
 	uint32 ComponentToVisualize;
+	float Extinction;
 
 protected:
 
