WorldPartitionLevelStreamingDynamic class is now unavailable as a blueprint type.

diff --git a/Engine/Source/Runtime/Engine/Public/WorldPartition/WorldPartitionLevelStreamingDynamic.h b/Engine/Source/Runtime/Engine/Public/WorldPartition/WorldPartitionLevelStreamingDynamic.h
--- a/Engine/Source/Runtime/Engine/Public/WorldPartition/WorldPartitionLevelStreamingDynamic.h
+++ b/Engine/Source/Runtime/Engine/Public/WorldPartition/WorldPartitionLevelStreamingDynamic.h
@@ -20,7 +20,7 @@
 #endif
 #include "WorldPartitionLevelStreamingDynamic.generated.h"
 
-UCLASS()
+UCLASS(NotPlaceable, NotBlueprintable, HideDropdown)
 class ENGINE_API UWorldPartitionLevelStreamingDynamic : public ULevelStreamingDynamic
 {
 	GENERATED_UCLASS_BODY()
