The command has been implemented for a test which covers a system updated via FCoreDelegates::OnBeginFrame. Since the system resolves requests every frame (and that's what the test needs to detect), such a bespoke command seems more accurate than using FEngineWaitLatentCommand.

diff --git a/Engine/Source/Runtime/Engine/Private/Tests/AutomationCommon.cpp b/Engine/Source/Runtime/Engine/Private/Tests/AutomationCommon.cpp
--- a/Engine/Source/Runtime/Engine/Private/Tests/AutomationCommon.cpp
+++ b/Engine/Source/Runtime/Engine/Private/Tests/AutomationCommon.cpp
@@ -675,7 +675,15 @@ bool FWaitForInteractiveFrameRate::Update()
 	return false;
 }
 
+bool FWaitForNextEngineFrameCommand::Update()
+{
+	if (LastFrame == 0)
+	{
+		LastFrame = GFrameCounter;
+	}
 
+	return LastFrame != GFrameCounter;
+}
 
 ///////////////////////////////////////////////////////////////////////
 // Common Latent commands which are used across test type. I.e. Engine, Network, etc...
diff --git a/Engine/Source/Runtime/Engine/Public/Tests/AutomationCommon.h b/Engine/Source/Runtime/Engine/Public/Tests/AutomationCommon.h
--- a/Engine/Source/Runtime/Engine/Public/Tests/AutomationCommon.h
+++ b/Engine/Source/Runtime/Engine/Public/Tests/AutomationCommon.h
@@ -214,6 +214,21 @@ private:
 	const int kSampleCount = (int32)kTickRate * 5;
 };
 
+/**
+ * Latent command to wait for one engine frame
+ */
+class ENGINE_API FWaitForNextEngineFrameCommand : public IAutomationLatentCommand
+{
+public:
+
+	bool Update() override;
+
+private:
+
+	// Frame to be passed to consider the waiting over
+	uint64 LastFrame = 0;
+};
+
 /**
 * Request an Image Comparison and queue the result to the test report
 * @param InImageName	Name used to identify the comparison
