Fix groom visualization typos.

diff --git a/Engine/Source/Runtime/Engine/Private/GroomVisualizationData.cpp b/Engine/Source/Runtime/Engine/Private/GroomVisualizationData.cpp
--- a/Engine/Source/Runtime/Engine/Private/GroomVisualizationData.cpp
+++ b/Engine/Source/Runtime/Engine/Private/GroomVisualizationData.cpp
@@ -199,7 +199,7 @@ void FGroomVisualizationData::Initialize()
 		AddVisualizationMode(ModeMap, true, EGroomViewMode::MaterialTangent,			LOCTEXT("MaterialTangent", "Tangent"),						LOCTEXT("MaterialTangentDesc", "Hair material tangent"));
 		AddVisualizationMode(ModeMap, true, EGroomViewMode::Tile,						LOCTEXT("Tile", "Tile"),									LOCTEXT("TileDesc", "Hair tile cotegorization"));
 		AddVisualizationMode(ModeMap, true, EGroomViewMode::SimHairStrands,				LOCTEXT("SimHairStrands", "Guides"),						LOCTEXT("SimHairStrandsDesc", "Simulation strands"));
-		AddVisualizationMode(ModeMap, true, EGroomViewMode::RenderHairStrands,			LOCTEXT("RenderHairStrands", "Stands Guides Influences"),	LOCTEXT("RenderHairStrandsDesc", "Rendering strands influences"));
+		AddVisualizationMode(ModeMap, true, EGroomViewMode::RenderHairStrands,			LOCTEXT("RenderHairStrands", "Strands Guides Influences"),	LOCTEXT("RenderHairStrandsDesc", "Rendering strands influences"));
 		AddVisualizationMode(ModeMap, true, EGroomViewMode::CardGuides, 				LOCTEXT("CardGuides", "Cards Guides"),						LOCTEXT("CardGuidesDesc", "Cards Guides"));
 		AddVisualizationMode(ModeMap, true, EGroomViewMode::RootUV,						LOCTEXT("RootUV", "Root UV"),								LOCTEXT("RootUVDesc", "Roots UV"));
 		AddVisualizationMode(ModeMap, true, EGroomViewMode::RootUDIM,					LOCTEXT("RootUDIM", "Root UDIM"),							LOCTEXT("RootUDIMDesc", "Roots UV UDIM texture index"));
@@ -211,9 +211,9 @@ void FGroomVisualizationData::Initialize()
 		AddVisualizationMode(ModeMap, true, EGroomViewMode::ControlPoints,				LOCTEXT("ControlPoints", "Control Points"),					LOCTEXT("ControlPointsDesc", "Hair control points"));
 		AddVisualizationMode(ModeMap, true, EGroomViewMode::Color,						LOCTEXT("Color", "Per-CV Color"),							LOCTEXT("ColorDesc", "Control point color (optional)"));
 		AddVisualizationMode(ModeMap, true, EGroomViewMode::Roughness,					LOCTEXT("Roughness", "Per-CV Roughness"),					LOCTEXT("RoughnessDesc", "Control point roughness  (optional)"));
-		AddVisualizationMode(ModeMap, true, EGroomViewMode::Cluster,					LOCTEXT("Cluster", "Clusters"),								LOCTEXT("ClusterDesc", "Hair visility clusters"));
-		AddVisualizationMode(ModeMap, true, EGroomViewMode::ClusterAABB,				LOCTEXT("ClusterAABB", "Clusters Bounds"),					LOCTEXT("ClusterAABBDesc", "Hair visility clusters AABBs"));
-		AddVisualizationMode(ModeMap, true, EGroomViewMode::Group,						LOCTEXT("Group", "Groups"),									LOCTEXT("GroupDesc", "Hair hair groups"));
+		AddVisualizationMode(ModeMap, true, EGroomViewMode::Cluster,					LOCTEXT("Cluster", "Clusters"),								LOCTEXT("ClusterDesc", "Hair clusters"));
+		AddVisualizationMode(ModeMap, true, EGroomViewMode::ClusterAABB,				LOCTEXT("ClusterAABB", "Clusters Bounds"),					LOCTEXT("ClusterAABBDesc", "Hair clusters AABBs"));
+		AddVisualizationMode(ModeMap, true, EGroomViewMode::Group,						LOCTEXT("Group", "Groups"),									LOCTEXT("GroupDesc", "Hair groups"));
 		AddVisualizationMode(ModeMap, true, EGroomViewMode::LODColoration,				LOCTEXT("LODColoration", "LOD Color"),						LOCTEXT("LODColorationDesc", "Hair LOD coloring"));
 		AddVisualizationMode(ModeMap, true, EGroomViewMode::AO,							LOCTEXT("AO", "Per-CV AO"),									LOCTEXT("AODesc", "Control point AO  (optional)"));
 		AddVisualizationMode(ModeMap, true, EGroomViewMode::ClumpID,					LOCTEXT("ClumpID", "Clump ID"),								LOCTEXT("ClumpIDDesc", "Curve clumpID  (optional)"));
