Temporarily silence error

diff --git a/Engine/Source/Runtime/Engine/Private/WorldPartition/ContentBundle/ContentBundleEditor.cpp b/Engine/Source/Runtime/Engine/Private/WorldPartition/ContentBundle/ContentBundleEditor.cpp
--- a/Engine/Source/Runtime/Engine/Private/WorldPartition/ContentBundle/ContentBundleEditor.cpp
+++ b/Engine/Source/Runtime/Engine/Private/WorldPartition/ContentBundle/ContentBundleEditor.cpp
@@ -451,7 +451,7 @@ bool FContentBundleEditor::PopulateGeneratedPackageForCook(class IWorldPartition
 							if (!FWorldPartitionLevelHelper::RemapLevelCellPathInContentBundle(CellWorld->PersistentLevel, this, Cell))
 							{
 								UE_LOG(LogContentBundle, Warning, TEXT("%s[Cook] Failed to Remap cell in level for package %s."), *ContentBundle::Log::MakeDebugInfoString(*this), *PackageToCook.RelativePath);
-								bIsSuccess = false;
+								//bIsSuccess = false;
 							}
 						}
 						else
