Allowing some replay recording options to be set via the demo url
Deprecate some internal functions

diff --git a/Engine/Source/Runtime/Engine/Classes/Engine/DemoNetDriver.h b/Engine/Source/Runtime/Engine/Classes/Engine/DemoNetDriver.h
--- a/Engine/Source/Runtime/Engine/Classes/Engine/DemoNetDriver.h
+++ b/Engine/Source/Runtime/Engine/Classes/Engine/DemoNetDriver.h
@@ -334,8 +334,12 @@ public:
 	virtual void ProcessRemoteFunction(class AActor* Actor, class UFunction* Function, void* Parameters, struct FOutParmRec* OutParms, struct FFrame* Stack, class UObject* SubObject = nullptr) override;
 	virtual bool IsAvailable() const override { return true; }
 	void SkipTime(const float InTimeToSkip);
+	
+	UE_DEPRECATED(5.3, "Internal call will be made private in the future.")
 	void SkipTimeInternal(const float SecondsToSkip, const bool InFastForward, const bool InIsForCheckpoint);
+	UE_DEPRECATED(5.3, "Internal call will be made private in the future.")
 	bool InitConnectInternal(FString& Error);
+
 	virtual bool ShouldClientDestroyTearOffActors() const override;
 	virtual bool ShouldSkipRepNotifies() const override;
 	virtual bool ShouldQueueBunchesForActorGUID(FNetworkGUID InGUID) const override;
@@ -768,4 +772,5 @@ private:
 	bool bSkipStartupActorRollback = false;
 
 	friend class UDemoNetConnection;
+	friend class FQueuedReplayTask;
 };
\ No newline at end of file
diff --git a/Engine/Source/Runtime/Engine/Private/DemoNetDriver.cpp b/Engine/Source/Runtime/Engine/Private/DemoNetDriver.cpp
--- a/Engine/Source/Runtime/Engine/Private/DemoNetDriver.cpp
+++ b/Engine/Source/Runtime/Engine/Private/DemoNetDriver.cpp
@@ -588,7 +588,9 @@ public:
 
 		Driver->GetReplayStreamer()->SetHighPriorityTimeRange(Driver->GetDemoCurrentTimeInMS(), TimeInMSToCheck);
 
+		PRAGMA_DISABLE_DEPRECATION_WARNINGS
 		Driver->SkipTimeInternal(SecondsToSkip, true, false);
+		PRAGMA_ENABLE_DEPRECATION_WARNINGS
 	}
 
 	virtual bool Tick() override
@@ -874,6 +876,21 @@ bool UDemoNetDriver::InitBase(bool bInitAsClient, FNetworkNotify* InNotify, cons
 
 		ResetDemoState();
 
+		if (URL.HasOption(TEXT("MaxDesiredReplayRecordTimeMS")))
+		{
+			MaxDesiredRecordTimeMS = FCString::Atof(URL.GetOption(TEXT("MaxDesiredReplayRecordTimeMS="), nullptr));
+		}
+
+		if (URL.HasOption(TEXT("CheckpointSaveMaxMSPerFrame")))
+		{
+			CheckpointSaveMaxMSPerFrame = FCString::Atof(URL.GetOption(TEXT("CheckpointSaveMaxMSPerFrame="), nullptr));
+		}
+
+		if (URL.HasOption(TEXT("ActorPrioritizationEnabled")))
+		{
+			bPrioritizeActors = FCString::Atoi(URL.GetOption(TEXT("ActorPrioritizationEnabled="), nullptr)) != 0;
+		}
+
 		ReplayStreamer = ReplayHelper.Init(URL);
 
 		ReplayHelper.SetAnalyticsProvider(AnalyticsProvider);
@@ -3203,10 +3220,12 @@ void UDemoNetDriver::ReplayStreamingReady(const FStartStreamingResult& Result)
 		
 		const double StartTime = FPlatformTime::Seconds();
 
+		PRAGMA_DISABLE_DEPRECATION_WARNINGS
 		if (!InitConnectInternal(Error))
 		{
 			return;
 		}
+		PRAGMA_ENABLE_DEPRECATION_WARNINGS
 
 		// InitConnectInternal calls ResetDemoState which will reset this, so restore the value
 		bWasStartStreamingSuccessful = Result.WasSuccessful();
@@ -4282,7 +4301,9 @@ bool UDemoNetDriver::LoadCheckpoint(const FGotoResult& GotoResult)
 
 		if (GotoResult.ExtraTimeMS != -1)
 		{
+			PRAGMA_DISABLE_DEPRECATION_WARNINGS
 			SkipTimeInternal((float)GotoResult.ExtraTimeMS / 1000.0f, true, true);
+			PRAGMA_ENABLE_DEPRECATION_WARNINGS
 		}
 		else
 		{
@@ -4490,8 +4511,10 @@ bool UDemoNetDriver::LoadCheckpoint(const FGotoResult& GotoResult)
 
 	if (GotoResult.ExtraTimeMS != -1)
 	{
+		PRAGMA_DISABLE_DEPRECATION_WARNINGS
 		// If we need to skip more time for fine scrubbing, set that up now
 		SkipTimeInternal((float)GotoResult.ExtraTimeMS / 1000.0f, true, true);
+		PRAGMA_ENABLE_DEPRECATION_WARNINGS
 	}
 	else
 	{
diff --git a/Engine/Source/Runtime/Engine/Private/ReplayHelper.cpp b/Engine/Source/Runtime/Engine/Private/ReplayHelper.cpp
--- a/Engine/Source/Runtime/Engine/Private/ReplayHelper.cpp
+++ b/Engine/Source/Runtime/Engine/Private/ReplayHelper.cpp
@@ -87,6 +87,11 @@ TSharedPtr<INetworkReplayStreamer> FReplayHelper::Init(const FURL& URL)
 	FWorldDelegates::LevelRemovedFromWorld.AddRaw(this, &FReplayHelper::OnLevelRemovedFromWorld);
 	FWorldDelegates::LevelAddedToWorld.AddRaw(this, &FReplayHelper::OnLevelAddedToWorld);
 
+	if (DemoURL.HasOption(TEXT("CheckpointSaveMaxMSPerFrame")))
+	{
+		CheckpointSaveMaxMSPerFrame = FCString::Atof(DemoURL.GetOption(TEXT("CheckpointSaveMaxMSPerFrame="), nullptr));
+	}
+
 	return ReplayStreamer;
 }
 
