Fix ensure in LWI caused when calling FindActorForHandle when the actor already exists.

diff --git a/Engine/Source/Runtime/Engine/Private/LightWeightInstanceManager.cpp b/Engine/Source/Runtime/Engine/Private/LightWeightInstanceManager.cpp
--- a/Engine/Source/Runtime/Engine/Private/LightWeightInstanceManager.cpp
+++ b/Engine/Source/Runtime/Engine/Private/LightWeightInstanceManager.cpp
@@ -221,7 +221,10 @@ bool ALightWeightInstanceManager::IsIndexValid(int32 Index) const
 
 bool ALightWeightInstanceManager::FindActorForHandle(const FActorInstanceHandle& Handle) const
 {
-	ensure(!Handle.Actor.IsValid());
+	if (Handle.Actor.IsValid())
+	{
+		return true;
+	}
 
 	AActor* const* FoundActor = Actors.Find(Handle.GetInstanceIndex());
 	Handle.Actor = FoundActor ? *FoundActor : nullptr;
