Make UAnimInstance::IsSlotActive() threadsafe.

diff --git a/Engine/Source/Runtime/Engine/Classes/Animation/AnimInstance.h b/Engine/Source/Runtime/Engine/Classes/Animation/AnimInstance.h
--- a/Engine/Source/Runtime/Engine/Classes/Animation/AnimInstance.h
+++ b/Engine/Source/Runtime/Engine/Classes/Animation/AnimInstance.h
@@ -548,7 +548,7 @@ public:
 	bool IsPlayingSlotAnimation(const UAnimSequenceBase* Asset, FName SlotNodeName, UAnimMontage*& OutMontage) const;
 
 	/** Return true if this instance has an active montage in the given slot. A UAnimMontage that is playing in the slot and blending out is not determined to be "active". */
-	UFUNCTION(BlueprintPure, Category="Animation|Montage")
+	UFUNCTION(BlueprintPure, Category="Animation|Montage", meta=(BlueprintThreadSafe))
 	bool IsSlotActive(FName SlotNodeName) const;
 	
 	/*********************************************************************************************
diff --git a/Engine/Source/Runtime/Engine/Private/Animation/AnimInstance.cpp b/Engine/Source/Runtime/Engine/Private/Animation/AnimInstance.cpp
--- a/Engine/Source/Runtime/Engine/Private/Animation/AnimInstance.cpp
+++ b/Engine/Source/Runtime/Engine/Private/Animation/AnimInstance.cpp
@@ -2109,18 +2109,20 @@ bool UAnimInstance::IsPlayingSlotAnimation(const UAnimSequenceBase* Asset, FName
 
 bool UAnimInstance::IsSlotActive(FName SlotNodeName) const
 {
-	if (SlotNodeName != NAME_None)
-	{
-		for (const FAnimMontageInstance* Instance : MontageInstances)
-		{
-			if (Instance && Instance->Montage && Instance->Montage->IsValidSlot(SlotNodeName) && Instance->IsActive())
-			{
-				return true;
-			}
-		}
-	}
-	
-	return false;
+   if (SlotNodeName != NAME_None)
+   {
+      const FAnimInstanceProxy& Proxy = GetProxyOnAnyThread<FAnimInstanceProxy>();
+      
+      for (const FMontageEvaluationState & EvaluationState : Proxy.GetMontageEvaluationData())
+      {
+         if (EvaluationState.Montage.IsValid() && EvaluationState.Montage->IsValidSlot(SlotNodeName) && EvaluationState.bIsActive)
+         {
+            return true;
+         }
+      }
+   }
+   
+   return false;
 }
 
 float UAnimInstance::Montage_PlayInternal(UAnimMontage* MontageToPlay, const FMontageBlendSettings& BlendInSettings, float InPlayRate /*= 1.f*/, EMontagePlayReturnType ReturnValueType /*= EMontagePlayReturnType::MontageLength*/, float InTimeToStartMontageAt /*= 0.f*/, bool bStopAllMontages /*= true*/)
