Fix cell remapping not having the full package path
Re-enable associated warning & error

diff --git a/Engine/Source/Runtime/Engine/Private/WorldPartition/ContentBundle/ContentBundleEditor.cpp b/Engine/Source/Runtime/Engine/Private/WorldPartition/ContentBundle/ContentBundleEditor.cpp
--- a/Engine/Source/Runtime/Engine/Private/WorldPartition/ContentBundle/ContentBundleEditor.cpp
+++ b/Engine/Source/Runtime/Engine/Private/WorldPartition/ContentBundle/ContentBundleEditor.cpp
@@ -450,8 +450,8 @@ bool FContentBundleEditor::PopulateGeneratedPackageForCook(class IWorldPartition
 							// We already know that the cell will reside in the ExternalStreamingObject package at runtime.
 							if (!FWorldPartitionLevelHelper::RemapLevelCellPathInContentBundle(CellWorld->PersistentLevel, this, Cell))
 							{
-								//UE_LOG(LogContentBundle, Warning, TEXT("%s[Cook] Failed to Remap cell in level for package %s."), *ContentBundle::Log::MakeDebugInfoString(*this), *PackageToCook.RelativePath);
-								//bIsSuccess = false;
+								UE_LOG(LogContentBundle, Warning, TEXT("%s[Cook] Failed to Remap cell in level for package %s."), *ContentBundle::Log::MakeDebugInfoString(*this), *PackageToCook.RelativePath);
+								bIsSuccess = false;
 							}
 						}
 						else
diff --git a/Engine/Source/Runtime/Engine/Private/WorldPartition/WorldPartitionLevelHelper.cpp b/Engine/Source/Runtime/Engine/Private/WorldPartition/WorldPartitionLevelHelper.cpp
--- a/Engine/Source/Runtime/Engine/Private/WorldPartition/WorldPartitionLevelHelper.cpp
+++ b/Engine/Source/Runtime/Engine/Private/WorldPartition/WorldPartitionLevelHelper.cpp
@@ -244,7 +244,7 @@ bool FWorldPartitionLevelHelper::RemapActorPath(const FActorContainerID& InConta
 
 bool FWorldPartitionLevelHelper::RemapLevelCellPathInContentBundle(ULevel* Level, const class FContentBundleEditor* ContentBundleEditor, const UWorldPartitionRuntimeCell* Cell)
 {
-	FString CellPath = ContentBundleEditor->GetExternalStreamingObjectPackageName();
+	FString CellPath = ContentBundleEditor->GetExternalStreamingObjectPackagePath();
 	CellPath += TEXT(".");
 	CellPath += ContentBundleEditor->GetExternalStreamingObjectName();
 	CellPath += TEXT(".");
