Fix deadlock after deleting some imported assets and we import again, there was a mixup with re-import in this case that did not release the skeletalmesh properties.

#rnx

diff --git a/Engine/Source/Runtime/Engine/Classes/Engine/SkeletalMesh.h b/Engine/Source/Runtime/Engine/Classes/Engine/SkeletalMesh.h
--- a/Engine/Source/Runtime/Engine/Classes/Engine/SkeletalMesh.h
+++ b/Engine/Source/Runtime/Engine/Classes/Engine/SkeletalMesh.h
@@ -2129,10 +2129,6 @@ public:
 	/** Returns false if there is currently an async task running. USkinnedAsset interface. */
 	virtual bool IsAsyncTaskComplete() const override;
 
-	/** Try to cancel any pending async tasks.
-	 *  Returns true if there is no more async tasks pending, false otherwise.
-	 */
-	bool TryCancelAsyncTasks();
 #endif // WITH_EDITOR
 
 	/**
diff --git a/Engine/Source/Runtime/Engine/Classes/Engine/SkinnedAsset.h b/Engine/Source/Runtime/Engine/Classes/Engine/SkinnedAsset.h
--- a/Engine/Source/Runtime/Engine/Classes/Engine/SkinnedAsset.h
+++ b/Engine/Source/Runtime/Engine/Classes/Engine/SkinnedAsset.h
@@ -303,6 +303,11 @@ protected:
 	/** Complete the async task process - Can't be done in parallel. */
 	virtual void FinishAsyncTaskInternal(FSkinnedAsyncTaskContext& Context) {}
 
+	/** Try to cancel any pending async tasks.
+	 *  Returns true if there is no more async tasks pending, false otherwise.
+	 */
+	virtual bool TryCancelAsyncTasks();
+
 	/** Holds the pointer to an async task if one exists. */
 	TUniquePtr<FSkinnedAssetAsyncBuildTask> AsyncTask;
 #endif // WITH_EDITOR
diff --git a/Engine/Source/Runtime/Engine/Classes/Engine/SkinnedAssetAsyncCompileUtils.h b/Engine/Source/Runtime/Engine/Classes/Engine/SkinnedAssetAsyncCompileUtils.h
--- a/Engine/Source/Runtime/Engine/Classes/Engine/SkinnedAssetAsyncCompileUtils.h
+++ b/Engine/Source/Runtime/Engine/Classes/Engine/SkinnedAssetAsyncCompileUtils.h
@@ -41,13 +41,11 @@ public:
 class FSkinnedAsyncTaskContext : public FSkinnedAssetCompilationContext
 {
 public:
-	FSkinnedAsyncTaskContext(bool bInResetAsyncFlagOnFinish, TFunctionRef<void()> InAsyncTaskFunction)
-		: bResetAsyncFlagOnFinish(bInResetAsyncFlagOnFinish)
-		, AsyncTaskFunction(InAsyncTaskFunction)
+	FSkinnedAsyncTaskContext(TFunctionRef<void()> InAsyncTaskFunction)
+		: AsyncTaskFunction(InAsyncTaskFunction)
 	{
 	}
 
-	bool bResetAsyncFlagOnFinish = true;
 	TFunctionRef<void()> AsyncTaskFunction;
 };
 
diff --git a/Engine/Source/Runtime/Engine/Private/SkeletalMesh.cpp b/Engine/Source/Runtime/Engine/Private/SkeletalMesh.cpp
--- a/Engine/Source/Runtime/Engine/Private/SkeletalMesh.cpp
+++ b/Engine/Source/Runtime/Engine/Private/SkeletalMesh.cpp
@@ -1936,8 +1936,6 @@ void USkeletalMesh::BeginBuildInternal(FSkinnedAssetBuildContext& Context)
 {
 	TRACE_CPUPROFILER_EVENT_SCOPE(USkeletalMesh::BeginBuildInternal);
 
-	SetInternalFlags(EInternalObjectFlags::Async);
-	
 	// Unregister all instances of this component
 	Context.RecreateRenderStateContext = MakeUnique<FSkinnedMeshComponentRecreateRenderStateContext>(this, false);
 
@@ -1988,8 +1986,6 @@ void USkeletalMesh::FinishBuildInternal(FSkinnedAssetBuildContext& Context)
 {
 	TRACE_CPUPROFILER_EVENT_SCOPE(USkeletalMesh::FinishBuildInternal);
 
-	ClearInternalFlags(EInternalObjectFlags::Async);
-
 	ReleaseAsyncProperty();
 
 	ApplyFinishBuildInternalData(&Context);
@@ -2019,7 +2015,7 @@ void USkeletalMesh::LockPropertiesUntil(FEvent* Event)
 	};
 
 	//Use the async task compile to lock the properties
-	FSkinnedAsyncTaskContext Context(!HasAnyInternalFlags(EInternalObjectFlags::Async), AsyncTaskFunction);
+	FSkinnedAsyncTaskContext Context(AsyncTaskFunction);
 	BeginAsyncTaskInternal(Context);
 	PrepareForAsyncCompilation();
 	FQueuedThreadPool* SkeletalMeshThreadPool = FSkinnedAssetCompilingManager::Get().GetThreadPool();
@@ -2034,10 +2030,6 @@ void USkeletalMesh::BeginAsyncTaskInternal(FSkinnedAsyncTaskContext& Context)
 {
 	check(IsInGameThread());
 
-	if (Context.bResetAsyncFlagOnFinish)
-	{
-		SetInternalFlags(EInternalObjectFlags::Async);
-	}
 	AcquireAsyncProperty();
 	//Allow thumbnail data so content browser get refresh properly
 	ReleaseAsyncProperty((uint64)ESkeletalMeshAsyncProperties::ThumbnailInfo);
@@ -2051,10 +2043,6 @@ void USkeletalMesh::ExecuteAsyncTaskInternal(FSkinnedAsyncTaskContext& Context)
 void USkeletalMesh::FinishAsyncTaskInternal(FSkinnedAsyncTaskContext& Context)
 {
 	check(IsInGameThread());
-	if (Context.bResetAsyncFlagOnFinish)
-	{
-		ClearInternalFlags(EInternalObjectFlags::Async);
-	}
 	ReleaseAsyncProperty();
 }
 
@@ -2690,19 +2678,6 @@ bool USkeletalMesh::IsAsyncTaskComplete() const
 	return AsyncTask == nullptr || AsyncTask->IsWorkDone();
 }
 
-bool USkeletalMesh::TryCancelAsyncTasks()
-{
-	if (AsyncTask)
-	{
-		if (AsyncTask->IsDone() || AsyncTask->Cancel())
-		{
-			AsyncTask.Reset();
-		}
-	}
-
-	return AsyncTask == nullptr;
-}
-
 void USkeletalMesh::PostLoadEnsureImportDataExist()
 {
 	//If we have a LODModel with no import data and the LOD model have at least one section using more bone then any platform max GPU bone count. We will recreate the import data to allow the asset to be build and chunk properly.
@@ -2986,8 +2961,6 @@ void USkeletalMesh::BeginPostLoadInternal(FSkinnedAssetPostLoadContext& Context)
 #if WITH_EDITOR
 	TRACE_CPUPROFILER_EVENT_SCOPE(USkeletalMesh::BeginPostLoadInternal);
 
-	SetInternalFlags(EInternalObjectFlags::Async);
-
 	// Lock all properties that should not be modified/accessed during async post-load
 	AcquireAsyncProperty();
 
@@ -3216,8 +3189,6 @@ void USkeletalMesh::FinishPostLoadInternal(FSkinnedAssetPostLoadContext& Context
 	
 #if WITH_EDITOR
 
-	ClearInternalFlags(EInternalObjectFlags::Async);
-
 	// This scope allows us to use any locked properties without causing stalls
 	FSkinnedAssetAsyncBuildScope AsyncBuildScope(this);
 
diff --git a/Engine/Source/Runtime/Engine/Private/SkinnedAsset.cpp b/Engine/Source/Runtime/Engine/Private/SkinnedAsset.cpp
--- a/Engine/Source/Runtime/Engine/Private/SkinnedAsset.cpp
+++ b/Engine/Source/Runtime/Engine/Private/SkinnedAsset.cpp
@@ -392,6 +392,21 @@ TArray<FSkeletalMaterial>& USkinnedAsset::GetSkeletalMaterialDummyArray()
 	return Dummy;
 }
 
+#if WITH_EDITOR
+bool USkinnedAsset::TryCancelAsyncTasks()
+{
+	if (AsyncTask)
+	{
+		if (AsyncTask->IsDone() || AsyncTask->Cancel())
+		{
+			AsyncTask.Reset();
+		}
+	}
+
+	return AsyncTask == nullptr;
+}
+#endif //WITH_EDITOR
+
 FString USkinnedAsset::GetLODPathName(const USkinnedAsset* Mesh, int32 LODIndex)
 {
 #if RHI_ENABLE_RESOURCE_INFO
diff --git a/Engine/Source/Runtime/Engine/Private/SkinnedAssetCompiler.cpp b/Engine/Source/Runtime/Engine/Private/SkinnedAssetCompiler.cpp
--- a/Engine/Source/Runtime/Engine/Private/SkinnedAssetCompiler.cpp
+++ b/Engine/Source/Runtime/Engine/Private/SkinnedAssetCompiler.cpp
@@ -53,6 +53,7 @@ FSkinnedAssetCompilingManager::FSkinnedAssetCompilingManager()
 	: Notification(GetAssetNameFormat())
 {
 	SkinnedAssetCompilingManagerImpl::EnsureInitializedCVars();
+	PostReachabilityAnalysisHandle = FCoreUObjectDelegates::PostReachabilityAnalysis.AddRaw(this, &FSkinnedAssetCompilingManager::OnPostReachabilityAnalysis);
 }
 
 FName FSkinnedAssetCompilingManager::GetAssetTypeName() const
@@ -146,6 +147,8 @@ void FSkinnedAssetCompilingManager::Shutdown()
 			FinishCompilation(PendingSkinnedAssets);
 		}
 	}
+
+	FCoreUObjectDelegates::PostReachabilityAnalysis.Remove(PostReachabilityAnalysisHandle);
 }
 
 bool FSkinnedAssetCompilingManager::IsAsyncCompilationEnabled() const
@@ -466,6 +469,37 @@ void FSkinnedAssetCompilingManager::ProcessAsyncTasks(bool bLimitExecutionTime)
 	UpdateCompilationNotification();
 }
 
+void FSkinnedAssetCompilingManager::OnPostReachabilityAnalysis()
+{
+	if (GetNumRemainingJobs())
+	{
+		TRACE_CPUPROFILER_EVENT_SCOPE(FSkinnedAssetCompilingManager::CancelUnreachableMeshes);
+
+		TArray<USkinnedAsset*> PendingSkinnedMeshes;
+		PendingSkinnedMeshes.Reserve(GetNumRemainingJobs());
+
+		for (auto Iterator = RegisteredSkinnedAsset.CreateIterator(); Iterator; ++Iterator)
+		{
+			USkinnedAsset* SkinnedMesh = Iterator->GetEvenIfUnreachable();
+			if (SkinnedMesh && SkinnedMesh->IsUnreachable())
+			{
+				UE_LOG(LogSkeletalMesh, Verbose, TEXT("Cancelling skinned mesh %s async compilation because it's being garbage collected"), *SkinnedMesh->GetName());
+
+				if (SkinnedMesh->TryCancelAsyncTasks())
+				{
+					Iterator.RemoveCurrent();
+				}
+				else
+				{
+					PendingSkinnedMeshes.Add(SkinnedMesh);
+				}
+			}
+		}
+
+		FinishCompilation(PendingSkinnedMeshes);
+	}
+}
+
 #endif // #if WITH_EDITOR
 
 #undef LOCTEXT_NAMESPACE
diff --git a/Engine/Source/Runtime/Engine/Public/SkinnedAssetCompiler.h b/Engine/Source/Runtime/Engine/Public/SkinnedAssetCompiler.h
--- a/Engine/Source/Runtime/Engine/Public/SkinnedAssetCompiler.h
+++ b/Engine/Source/Runtime/Engine/Public/SkinnedAssetCompiler.h
@@ -96,6 +96,9 @@ private:
 
 	void PostCompilation(USkinnedAsset* SkinnedAsset);
 	void PostCompilation(TArrayView<USkinnedAsset* const> InSkinnedAssets);
+
+	void OnPostReachabilityAnalysis();
+	FDelegateHandle PostReachabilityAnalysisHandle;
 };
 
 #endif // #if WITH_EDITOR
\ No newline at end of file
diff --git a/Engine/Source/Runtime/Interchange/Engine/Private/Tasks/InterchangeTaskCompletion.cpp b/Engine/Source/Runtime/Interchange/Engine/Private/Tasks/InterchangeTaskCompletion.cpp
--- a/Engine/Source/Runtime/Interchange/Engine/Private/Tasks/InterchangeTaskCompletion.cpp
+++ b/Engine/Source/Runtime/Interchange/Engine/Private/Tasks/InterchangeTaskCompletion.cpp
@@ -61,12 +61,14 @@ void UE::Interchange::FTaskPreCompletion::DoTask(ENamedThreads::Type CurrentThre
 						ObjectInfo.Factory->Cancel();
 					}
 				}
-				break;
+				//Skip if cancel
+				continue;
 			}
 
 			const int32 SourceIndex = ObjectInfosPerSourceIndexPair.Key;
 			const bool bCallPostImportGameThreadCallback = ensure(AsyncHelper->SourceDatas.IsValidIndex(SourceIndex));
 
+			//First iteration to call SetupObject_GameThread and pipeline ExecutePostFactoryPipeline
 			for (const FImportAsyncHelper::FImportedObjectInfo& ObjectInfo : ObjectInfosPerSourceIndexPair.Value)
 			{
 				UObject* ImportedObject = ObjectInfo.ImportedObject;
@@ -130,11 +132,29 @@ void UE::Interchange::FTaskPreCompletion::DoTask(ENamedThreads::Type CurrentThre
 						, ImportedObject
 						, ObjectInfo.bIsReimport);
 				}
+			}
 
 #if WITH_EDITOR
+			//Second iteration to call PostEditChange
+			for (const FImportAsyncHelper::FImportedObjectInfo& ObjectInfo : ObjectInfosPerSourceIndexPair.Value)
+			{
+				UObject* ImportedObject = ObjectInfo.ImportedObject;
+				if (ImportedObject == nullptr || !IsValid(ImportedObject))
+				{
+					continue;
+				}
 				ImportedObject->PostEditChange();
-#endif
+			}
+#endif //WITH_EDITOR
 
+			//Third iteration to call FinalizeObject_GameThread
+			for (const FImportAsyncHelper::FImportedObjectInfo& ObjectInfo : ObjectInfosPerSourceIndexPair.Value)
+			{
+				UObject* ImportedObject = ObjectInfo.ImportedObject;
+				if (ImportedObject == nullptr || !IsValid(ImportedObject))
+				{
+					continue;
+				}
 				//Post import broadcast
 				if (bIsAsset)
 				{
