Make sure SessionMachineStats is fired for the initial user.

diff --git a/Engine/Source/Runtime/Engine/Private/EngineAnalytics.cpp b/Engine/Source/Runtime/Engine/Private/EngineAnalytics.cpp
--- a/Engine/Source/Runtime/Engine/Private/EngineAnalytics.cpp
+++ b/Engine/Source/Runtime/Engine/Private/EngineAnalytics.cpp
@@ -126,7 +126,6 @@ void FEngineAnalytics::Initialize()
 		if (Analytics.IsValid())
 		{
 			Analytics->SetUserID(CreateAnalyticsUserId(FPlatformMisc::GetEpicAccountId()));
-			SessionEpicAccountIds.Add(FPlatformMisc::GetEpicAccountId());
 
 			if (UE::Analytics::Private::EngineAnalyticsConfigOverride)
 			{
@@ -143,7 +142,7 @@ void FEngineAnalytics::Initialize()
 			}
 
 			Analytics->StartSession(MoveTemp(StartSessionAttributes));
-
+			SendMachineInfoForAccount(FPlatformMisc::GetEpicAccountId());
 			bIsInitialized = true;
 		}
 
@@ -302,6 +301,31 @@ void FEngineAnalytics::AppendMachineStats(TArray<FAnalyticsEventAttribute>& Even
 #endif
 }
 
+void FEngineAnalytics::SendMachineInfoForAccount(const FString& EpicAccountId)
+{
+	if (!SessionEpicAccountIds.Contains(EpicAccountId))
+	{
+		TArray<FAnalyticsEventAttribute> EventAttributes;
+		AppendMachineStats(EventAttributes);
+		SendMachineInfoForAccount(EpicAccountId, EventAttributes);
+	}
+}
+
+void FEngineAnalytics::SendMachineInfoForAccount(const FString& EpicAccountId, const TArray<FAnalyticsEventAttribute>& EventAttributes)
+{
+	// Only send SessionMachineStats when the user is changed to a user that has not been seen in the session.
+	if (!SessionEpicAccountIds.Contains(EpicAccountId))
+	{
+		SessionEpicAccountIds.Add(EpicAccountId);
+
+		// When the user id is changed send a SessionMachineStats event.
+		static const FString SZEventName = TEXT("SessionMachineStats");
+		TArray<FAnalyticsEventAttribute> Params;
+		AppendMachineStats(Params);
+		Analytics->RecordEvent(SZEventName, Params);
+	}
+}
+
 void FEngineAnalytics::OnEpicAccountIdChanged(const FString& EpicAccountId)
 {
 	// For analytics reporting ignore changes to an empty account id when the user logs out.
@@ -313,18 +337,7 @@ void FEngineAnalytics::OnEpicAccountIdChanged(const FString& EpicAccountId)
 		if (Analytics)
 		{
 			Analytics->SetUserID(NewAnalyticsUserId);
-
-			// Only send SessionMachineStats when the user is changed to a user that has not been seen in the session.
-			if (!SessionEpicAccountIds.Contains(EpicAccountId))
-			{
-				SessionEpicAccountIds.Add(EpicAccountId);
-
-				// When the user id is changed send a SessionMachineStats event.
-				static const FString SZEventName = TEXT("SessionMachineStats");
-				TArray<FAnalyticsEventAttribute> Params;
-				AppendMachineStats(Params);
-				Analytics->RecordEvent(SZEventName, Params);
-			}
+			SendMachineInfoForAccount(EpicAccountId);
 		}
 
 #if WITH_EDITOR
diff --git a/Engine/Source/Runtime/Engine/Public/EngineAnalytics.h b/Engine/Source/Runtime/Engine/Public/EngineAnalytics.h
--- a/Engine/Source/Runtime/Engine/Public/EngineAnalytics.h
+++ b/Engine/Source/Runtime/Engine/Public/EngineAnalytics.h
@@ -60,6 +60,8 @@ public:
 
 private:
 	static void AppendMachineStats(TArray<FAnalyticsEventAttribute>& EventAttributes);
+	static void SendMachineInfoForAccount(const FString& EpicAccountId);
+	static void SendMachineInfoForAccount(const FString& EpicAccountId, const TArray<FAnalyticsEventAttribute>& EventAttributes);
 	static void OnEpicAccountIdChanged(const FString& EpicAccountId);
 
 	static bool bIsInitialized;
