Add an << operator overload to FHardwareDeviceIdentifier so that you can serialize it out to an archive if you want

#rnx
#pf 63f69339ac06ce789f5f90da
#rnx

diff --git a/Engine/Source/Runtime/Engine/Classes/GameFramework/InputSettings.h b/Engine/Source/Runtime/Engine/Classes/GameFramework/InputSettings.h
--- a/Engine/Source/Runtime/Engine/Classes/GameFramework/InputSettings.h
+++ b/Engine/Source/Runtime/Engine/Classes/GameFramework/InputSettings.h
@@ -340,6 +340,7 @@ struct ENGINE_API FHardwareDeviceIdentifier
 	}
 
 	ENGINE_API friend uint32 GetTypeHash(const FHardwareDeviceIdentifier& InDevice);
+	ENGINE_API friend FArchive& operator<<(FArchive& Ar, FHardwareDeviceIdentifier& InDevice);
 };
 
 /** Per-Platform input options */
diff --git a/Engine/Source/Runtime/Engine/Private/UserInterface/InputSettings.cpp b/Engine/Source/Runtime/Engine/Private/UserInterface/InputSettings.cpp
--- a/Engine/Source/Runtime/Engine/Private/UserInterface/InputSettings.cpp
+++ b/Engine/Source/Runtime/Engine/Private/UserInterface/InputSettings.cpp
@@ -529,6 +529,14 @@ uint32 GetTypeHash(const FHardwareDeviceIdentifier& InDevice)
 	return HashCombine(GetTypeHash(InDevice.InputClassName), GetTypeHash(InDevice.HardwareDeviceIdentifier));
 }
 
+FArchive& operator<<(FArchive& Ar, FHardwareDeviceIdentifier& InDevice)
+{
+	Ar << InDevice.InputClassName;
+	Ar << InDevice.HardwareDeviceIdentifier;
+	
+	return Ar;
+}
+
 bool FHardwareDeviceIdentifier::IsValid() const
 {
 	return InputClassName.IsValid() && HardwareDeviceIdentifier.IsValid();
