World Partition - HLOD: Register cells from content bundles that are injected after the world partition initialization

diff --git a/Engine/Source/Runtime/Engine/Private/WorldPartition/ContentBundle/ContentBundle.cpp b/Engine/Source/Runtime/Engine/Private/WorldPartition/ContentBundle/ContentBundle.cpp
--- a/Engine/Source/Runtime/Engine/Private/WorldPartition/ContentBundle/ContentBundle.cpp
+++ b/Engine/Source/Runtime/Engine/Private/WorldPartition/ContentBundle/ContentBundle.cpp
@@ -61,7 +61,7 @@ void FContentBundle::DoInjectContent()
 {
 	if (ExternalStreamingObject != nullptr)
 	{
-		if (GetInjectedWorld()->GetWorldPartition()->RuntimeHash->InjectExternalStreamingObject(ExternalStreamingObject))
+		if (GetInjectedWorld()->GetWorldPartition()->InjectExternalStreamingObject(ExternalStreamingObject))
 		{
 			UE_LOG(LogContentBundle, Log, TEXT("%s Streaming Object Injected."), *ContentBundle::Log::MakeDebugInfoString(*this));
 			SetStatus(EContentBundleStatus::ContentInjected);
@@ -85,7 +85,7 @@ void FContentBundle::DoRemoveContent()
 {
 	if (ExternalStreamingObject != nullptr)
 	{
-		if (GetInjectedWorld()->GetWorldPartition()->RuntimeHash->RemoveExternalStreamingObject(ExternalStreamingObject))
+		if (GetInjectedWorld()->GetWorldPartition()->RemoveExternalStreamingObject(ExternalStreamingObject))
 		{
 			ContentBundlesEpoch++;
 		}
diff --git a/Engine/Source/Runtime/Engine/Private/WorldPartition/HLOD/HLODSubsystem.cpp b/Engine/Source/Runtime/Engine/Private/WorldPartition/HLOD/HLODSubsystem.cpp
--- a/Engine/Source/Runtime/Engine/Private/WorldPartition/HLOD/HLODSubsystem.cpp
+++ b/Engine/Source/Runtime/Engine/Private/WorldPartition/HLOD/HLODSubsystem.cpp
@@ -262,6 +262,8 @@ void UHLODSubsystem::OnWorldPartitionInitialized(UWorldPartition* InWorldPartiti
 		{
 			InWorldPartition->RuntimeHash->ForEachStreamingCells([&WorldPartitionHLODRuntimeData](const UWorldPartitionRuntimeCell* Cell)
 			{
+				UE_LOG(LogHLODSubsystem, Verbose, TEXT("Registering cell %s - %s"), *Cell->GetGuid().ToString(), *Cell->GetDebugName());
+
 				WorldPartitionHLODRuntimeData.CellsData.Emplace(Cell->GetGuid());
 				return true;
 			});
@@ -274,10 +276,50 @@ void UHLODSubsystem::OnWorldPartitionUninitialized(UWorldPartition* InWorldParti
 	if (InWorldPartition && InWorldPartition->IsStreamingEnabled())
 	{
 		check(WorldPartitionsHLODRuntimeData.Contains(InWorldPartition));
+
+#if !NO_LOGGING
+		FWorldPartitionHLODRuntimeData& WorldPartitionHLODRuntimeData = WorldPartitionsHLODRuntimeData.FindChecked(InWorldPartition);
+		InWorldPartition->RuntimeHash->ForEachStreamingCells([&WorldPartitionHLODRuntimeData](const UWorldPartitionRuntimeCell* Cell)
+		{
+			UE_LOG(LogHLODSubsystem, Verbose, TEXT("Unregistering cell %s - %s"), *Cell->GetGuid().ToString(), *Cell->GetDebugName());
+			return true;
+		});
+#endif
+
 		WorldPartitionsHLODRuntimeData.Remove(InWorldPartition);
 	}
 }
 
+void UHLODSubsystem::OnExternalStreamingObjectInjected(URuntimeHashExternalStreamingObjectBase* ExternalStreamingObject)
+{
+	UWorldPartition* OwnerPartition = ExternalStreamingObject->GetOuterWorld()->GetWorldPartition();
+	FWorldPartitionHLODRuntimeData* WorldPartitionHLODRuntimeData = WorldPartitionsHLODRuntimeData.Find(OwnerPartition);
+	if (WorldPartitionHLODRuntimeData)
+	{
+		ExternalStreamingObject->ForEachStreamingCells([WorldPartitionHLODRuntimeData](const UWorldPartitionRuntimeCell& Cell)
+		{
+			UE_LOG(LogHLODSubsystem, Verbose, TEXT("Registering external cell %s - %s"), *Cell.GetGuid().ToString(), *Cell.GetDebugName());
+			WorldPartitionHLODRuntimeData->CellsData.Emplace(Cell.GetGuid());
+			return true;
+		});
+	}
+}
+
+void UHLODSubsystem::OnExternalStreamingObjectRemoved(URuntimeHashExternalStreamingObjectBase* ExternalStreamingObject)
+{
+	UWorldPartition* OwnerPartition = ExternalStreamingObject->GetOuterWorld()->GetWorldPartition();
+	FWorldPartitionHLODRuntimeData* WorldPartitionHLODRuntimeData = WorldPartitionsHLODRuntimeData.Find(OwnerPartition);
+	if (WorldPartitionHLODRuntimeData)
+	{
+		ExternalStreamingObject->ForEachStreamingCells([WorldPartitionHLODRuntimeData](const UWorldPartitionRuntimeCell& Cell)
+		{
+			UE_LOG(LogHLODSubsystem, Verbose, TEXT("Unregistering external cell %s - %s"), *Cell.GetGuid().ToString(), *Cell.GetDebugName());
+			WorldPartitionHLODRuntimeData->CellsData.Remove(Cell.GetGuid());
+			return true;
+		});
+	}
+}
+
 const UHLODSubsystem::FCellData* UHLODSubsystem::GetCellData(const UWorldPartitionRuntimeCell* InCell) const
 {
 	return const_cast<UHLODSubsystem*>(this)->GetCellData(InCell);
@@ -330,15 +372,14 @@ void UHLODSubsystem::RegisterHLODActor(AWorldPartitionHLOD* InWorldPartitionHLOD
 	
 	if (FCellData* CellData = GetCellData(InWorldPartitionHLOD))
 	{
-#if WITH_EDITOR
-		UE_LOG(LogHLODSubsystem, Verbose, TEXT("Registering HLOD %s (%s) for cell %s"), *InWorldPartitionHLOD->GetActorLabel(), *InWorldPartitionHLOD->GetActorGuid().ToString(), *InWorldPartitionHLOD->GetSourceCellGuid().ToString());
-#endif
+		UE_LOG(LogHLODSubsystem, Verbose, TEXT("Registering HLOD %s for cell %s"), *InWorldPartitionHLOD->GetActorNameOrLabel(), *InWorldPartitionHLOD->GetSourceCellGuid().ToString());
+
 		CellData->LoadedHLODs.Add(InWorldPartitionHLOD);
 		InWorldPartitionHLOD->SetVisibility(UHLODSubsystem::WorldPartitionHLODEnabled && !CellData->bIsCellVisible);
 	}
 	else
 	{
-		UE_LOG(LogHLODSubsystem, Verbose, TEXT("Found HLOD referencing nonexistent cell '%s'"), *InWorldPartitionHLOD->GetSourceCellGuid().ToString());
+		UE_LOG(LogHLODSubsystem, Verbose, TEXT("Found HLOD %s referencing nonexistent cell '%s'"), *InWorldPartitionHLOD->GetActorNameOrLabel(), *InWorldPartitionHLOD->GetSourceCellGuid().ToString());
 		InWorldPartitionHLOD->SetVisibility(false);
 	}
 
@@ -352,9 +393,7 @@ void UHLODSubsystem::UnregisterHLODActor(AWorldPartitionHLOD* InWorldPartitionHL
 
 	if (FCellData* CellData = GetCellData(InWorldPartitionHLOD))
 	{
-#if WITH_EDITOR
-		UE_LOG(LogHLODSubsystem, Verbose, TEXT("Unregistering HLOD %s (%s) for cell %s"), *InWorldPartitionHLOD->GetActorLabel(), *InWorldPartitionHLOD->GetActorGuid().ToString(), *InWorldPartitionHLOD->GetSourceCellGuid().ToString());
-#endif
+		UE_LOG(LogHLODSubsystem, Verbose, TEXT("Unregistering HLOD %s for cell %s"), *InWorldPartitionHLOD->GetActorNameOrLabel(), *InWorldPartitionHLOD->GetSourceCellGuid().ToString());
 
 		int32 NumRemoved = CellData->LoadedHLODs.Remove(InWorldPartitionHLOD);
 		check(NumRemoved == 1);
@@ -369,16 +408,15 @@ void UHLODSubsystem::OnCellShown(const UWorldPartitionRuntimeCell* InCell)
 	{
 		CellData->bIsCellVisible = true;
 
-#if WITH_EDITOR
-		UE_LOG(LogHLODSubsystem, Verbose, TEXT("Cell shown - %s - hiding %d HLOD actors"), *InCell->GetName(), CellData->LoadedHLODs.Num());
-#endif
-
-		for (AWorldPartitionHLOD* HLODActor : CellData->LoadedHLODs)
+		if (!CellData->LoadedHLODs.IsEmpty())
 		{
-#if WITH_EDITOR
-			UE_LOG(LogHLODSubsystem, Verbose, TEXT("\t\t%s - %s"), *HLODActor->GetActorLabel(), *HLODActor->GetActorGuid().ToString());
-#endif
-			HLODActor->SetVisibility(false);
+			UE_LOG(LogHLODSubsystem, Verbose, TEXT("Cell shown - %s - hiding %d HLOD actors"), *InCell->GetGuid().ToString(), CellData->LoadedHLODs.Num());
+
+			for (AWorldPartitionHLOD* HLODActor : CellData->LoadedHLODs)
+			{
+				UE_LOG(LogHLODSubsystem, Verbose, TEXT("\t\t* %s"), *HLODActor->GetActorNameOrLabel());
+				HLODActor->SetVisibility(false);
+			}
 		}
 	}
 }
@@ -389,17 +427,16 @@ void UHLODSubsystem::OnCellHidden(const UWorldPartitionRuntimeCell* InCell)
 	{
 		CellData->bIsCellVisible = false;
 
-#if WITH_EDITOR
-		UE_LOG(LogHLODSubsystem, Verbose, TEXT("Cell hidden - %s - showing %d HLOD actors"), *InCell->GetName(), CellData->LoadedHLODs.Num());
-#endif
-
-		for (AWorldPartitionHLOD* HLODActor : CellData->LoadedHLODs)
+		if (!CellData->LoadedHLODs.IsEmpty())
 		{
-#if WITH_EDITOR
-			UE_LOG(LogHLODSubsystem, Verbose, TEXT("\t\t%s - %s"), *HLODActor->GetActorLabel(), *HLODActor->GetActorGuid().ToString());
-#endif
-			HLODActor->SetVisibility(UHLODSubsystem::WorldPartitionHLODEnabled);
-			HLODActorsToWarmup.Remove(HLODActor);
+			UE_LOG(LogHLODSubsystem, Verbose, TEXT("Cell hidden - %s - showing %d HLOD actors"), *InCell->GetGuid().ToString(), CellData->LoadedHLODs.Num());
+
+			for (AWorldPartitionHLOD* HLODActor : CellData->LoadedHLODs)
+			{
+				UE_LOG(LogHLODSubsystem, Verbose, TEXT("\t\t* %s"), *HLODActor->GetActorNameOrLabel());
+				HLODActor->SetVisibility(UHLODSubsystem::WorldPartitionHLODEnabled);
+				HLODActorsToWarmup.Remove(HLODActor);
+			}
 		}
 	}
 }
diff --git a/Engine/Source/Runtime/Engine/Private/WorldPartition/WorldPartition.cpp b/Engine/Source/Runtime/Engine/Private/WorldPartition/WorldPartition.cpp
--- a/Engine/Source/Runtime/Engine/Private/WorldPartition/WorldPartition.cpp
+++ b/Engine/Source/Runtime/Engine/Private/WorldPartition/WorldPartition.cpp
@@ -1431,6 +1431,28 @@ void UWorldPartition::UpdateStreamingState()
 	}
 }
 
+bool UWorldPartition::InjectExternalStreamingObject(URuntimeHashExternalStreamingObjectBase* ExternalStreamingObject)
+{
+	bool bInjected = RuntimeHash->InjectExternalStreamingObject(ExternalStreamingObject);
+	if (bInjected)
+	{
+		GetWorld()->GetSubsystem<UHLODSubsystem>()->OnExternalStreamingObjectInjected(ExternalStreamingObject);
+	}
+
+	return bInjected;
+}
+
+bool UWorldPartition::RemoveExternalStreamingObject(URuntimeHashExternalStreamingObjectBase* ExternalStreamingObject)
+{
+	bool bRemoved = RuntimeHash->RemoveExternalStreamingObject(ExternalStreamingObject);
+	if (bRemoved)
+	{
+		GetWorld()->GetSubsystem<UHLODSubsystem>()->OnExternalStreamingObjectRemoved(ExternalStreamingObject);
+	}
+
+	return bRemoved;
+}
+
 bool UWorldPartition::GetIntersectingCells(const TArray<FWorldPartitionStreamingQuerySource>& InSources, TArray<const IWorldPartitionCell*>& OutCells) const
 {
 	if (StreamingPolicy)
diff --git a/Engine/Source/Runtime/Engine/Public/WorldPartition/HLOD/HLODSubsystem.h b/Engine/Source/Runtime/Engine/Public/WorldPartition/HLOD/HLODSubsystem.h
--- a/Engine/Source/Runtime/Engine/Public/WorldPartition/HLOD/HLODSubsystem.h
+++ b/Engine/Source/Runtime/Engine/Public/WorldPartition/HLOD/HLODSubsystem.h
@@ -13,6 +13,7 @@ class FSceneViewFamily;
 class FHLODResourcesResidencySceneViewExtension;
 class UWorldPartition;
 class UWorldPartitionRuntimeCell;
+class URuntimeHashExternalStreamingObjectBase;
 class UWorld;
 
 DECLARE_MULTICAST_DELEGATE_OneParam(FWorldPartitionHLODActorRegisteredEvent, AWorldPartitionHLOD* /* InHLODActor */);
@@ -57,6 +58,9 @@ public:
 	FWorldPartitionHLODActorRegisteredEvent& OnHLODActorRegisteredEvent() { return HLODActorRegisteredEvent; }
 	FWorldPartitionHLODActorUnregisteredEvent& OnHLODActorUnregisteredEvent() { return HLODActorUnregisteredEvent; }
 
+	void OnExternalStreamingObjectInjected(URuntimeHashExternalStreamingObjectBase* ExternalStreamingObject);
+	void OnExternalStreamingObjectRemoved(URuntimeHashExternalStreamingObjectBase* ExternalStreamingObject);
+
 	void SetHLODAlwaysLoadedCullDistance(int32 InCullDistance);
 
 	void OnCVarsChanged();
diff --git a/Engine/Source/Runtime/Engine/Public/WorldPartition/WorldPartition.h b/Engine/Source/Runtime/Engine/Public/WorldPartition/WorldPartition.h
--- a/Engine/Source/Runtime/Engine/Public/WorldPartition/WorldPartition.h
+++ b/Engine/Source/Runtime/Engine/Public/WorldPartition/WorldPartition.h
@@ -30,6 +30,7 @@ class UActorDescContainer;
 class UWorldPartitionEditorHash;
 class UWorldPartitionRuntimeCell;
 class UWorldPartitionRuntimeHash;
+class URuntimeHashExternalStreamingObjectBase;
 class UWorldPartitionStreamingPolicy;
 class IWorldPartitionCell;
 class UDataLayerManager;
@@ -271,6 +272,9 @@ public:
 	bool IsStreamingCompleted(EWorldPartitionRuntimeCellState QueryState, const TArray<FWorldPartitionStreamingQuerySource>& QuerySources, bool bExactState) const;
 	bool GetIntersectingCells(const TArray<FWorldPartitionStreamingQuerySource>& InSources, TArray<const IWorldPartitionCell*>& OutCells) const;
 
+	bool InjectExternalStreamingObject(URuntimeHashExternalStreamingObjectBase* ExternalStreamingObject);
+	bool RemoveExternalStreamingObject(URuntimeHashExternalStreamingObjectBase* ExternalStreamingObject);
+
 	const TArray<FWorldPartitionStreamingSource>& GetStreamingSources() const;
 
 	// Debugging
