MergeActor - Mesh Merging: "Merge Materials" option should be editable if LOD Selection Type is set to "Specific LOD"

diff --git a/Engine/Source/Runtime/Engine/Classes/Engine/MeshMerging.h b/Engine/Source/Runtime/Engine/Classes/Engine/MeshMerging.h
--- a/Engine/Source/Runtime/Engine/Classes/Engine/MeshMerging.h
+++ b/Engine/Source/Runtime/Engine/Classes/Engine/MeshMerging.h
@@ -527,7 +527,7 @@ struct FMeshMergingSettings
 	uint8 bMergeMeshSockets : 1;
 
 	/** Whether to merge source materials into one flat material, ONLY available when LOD Selection Type is set to LowestDetailLOD */
-	UPROPERTY(EditAnywhere, BlueprintReadWrite, Category = MaterialSettings, meta=(EditCondition="LODSelectionType == EMeshLODSelectionType::LowestDetailLOD"))
+	UPROPERTY(EditAnywhere, BlueprintReadWrite, Category = MaterialSettings, meta=(EditCondition="LODSelectionType == EMeshLODSelectionType::LowestDetailLOD || LODSelectionType == EMeshLODSelectionType::SpecificLOD"))
 	uint8 bMergeMaterials:1;
 
 	/** Whether or not vertex data such as vertex colours should be baked into the resulting mesh */
