Only draw Skinned Level Set lattice in Physics Asset Editor when physics are drawing in wireframe (or wireframe + shaded) mode.

diff --git a/Engine/Source/Runtime/Engine/Private/PhysicsEngine/PhysDrawing.cpp b/Engine/Source/Runtime/Engine/Private/PhysicsEngine/PhysDrawing.cpp
--- a/Engine/Source/Runtime/Engine/Private/PhysicsEngine/PhysDrawing.cpp
+++ b/Engine/Source/Runtime/Engine/Private/PhysicsEngine/PhysDrawing.cpp
@@ -901,9 +901,21 @@ void FKSkinnedLevelSetElem::DrawElemWire(class FPrimitiveDrawInterface* PDI, con
 {
 	if (WeightedLevelSet.IsValid())
 	{
-		DrawSkinnedLevelSetLattice(PDI, ElemTM.ToMatrixWithScale(), Color, WeightedLevelSet.Get());
+		constexpr float HSVHueShift = 240.f;
+		constexpr float HSVSaturationMult = 0.5f;
+		constexpr float HSVValueMult = 0.8f;
+		constexpr float HSVAlphaMult = 0.5f;
+		FLinearColor ShiftedColorHSV = Color.ReinterpretAsLinear().LinearRGBToHSV();
+		ShiftedColorHSV.R += HSVHueShift;
+		if (ShiftedColorHSV.R >= 360.f)
+		{
+			ShiftedColorHSV.R -= 360.f;
+		}
+		ShiftedColorHSV.G *= HSVSaturationMult;
+		ShiftedColorHSV.B *= HSVValueMult;
+		ShiftedColorHSV.A *= HSVAlphaMult;
 
-		// TODO grid cells
+		DrawSkinnedLevelSetLattice(PDI, ElemTM.ToMatrixWithScale(), ShiftedColorHSV.HSVToLinearRGB().ToFColor(true), WeightedLevelSet.Get());
 	}
 }
 
@@ -911,9 +923,6 @@ void FKSkinnedLevelSetElem::DrawElemSolid(class FPrimitiveDrawInterface* PDI, co
 {
 	if (WeightedLevelSet.IsValid())
 	{
-		const FColor LatticeColor = FColor::Cyan; // TODO figure out what color to actually draw lattice
-		DrawSkinnedLevelSetLattice(PDI, ElemTM.ToMatrixWithScale(), LatticeColor, WeightedLevelSet.Get());
-
 		TArray<FVector3f> Vertices;
 		TArray<FIntVector> Tris;
 		const Chaos::FLevelSet* const LevelSet = WeightedLevelSet->GetEmbeddedObject();
@@ -937,9 +946,6 @@ void FKSkinnedLevelSetElem::GetElemSolid(const FTransform& ElemTM, const FVector
 {
 	if (WeightedLevelSet.IsValid())
 	{
-		const FColor LatticeColor = FColor::Cyan; // TODO figure out what color to actually draw lattice
-		DrawSkinnedLevelSetLattice(Collector.GetPDI(ViewIndex), ElemTM.ToMatrixWithScale(), LatticeColor, WeightedLevelSet.Get());
-
 		TArray<FVector3f> Vertices;
 		TArray<FIntVector> Tris;
 		const Chaos::FLevelSet* const LevelSet = WeightedLevelSet->GetEmbeddedObject();
