Unloading WP Actors no longer clears the selection (only removes the unloaded actors from it)

Handle Deselect of Actors/Components/Instances in ULevel::RemoveLoadedActors:
- Remove SelectNone call from IWorldPartitionActorLoaderInterface::ILoaderAdapter::PostLoadedStateChanged
- Remove Selection update from FActorBrowsingMode::UnpinItems

diff --git a/Engine/Source/Editor/SceneOutliner/Private/ActorBrowsingMode.cpp b/Engine/Source/Editor/SceneOutliner/Private/ActorBrowsingMode.cpp
--- a/Engine/Source/Editor/SceneOutliner/Private/ActorBrowsingMode.cpp
+++ b/Engine/Source/Editor/SceneOutliner/Private/ActorBrowsingMode.cpp
@@ -2260,23 +2260,7 @@ void FActorBrowsingMode::UnpinItems(const TArray<FSceneOutlinerTreeItemPtr>& InI
 
 	if (ActorsToUnpin.Num())
 	{
-		GEditor->GetSelectedActors()->BeginBatchSelectOperation();
-
-		AActor* LastPinnedActor = nullptr;
-		for (const FGuid& ActorGuid : ActorsToUnpin)
-		{
-			if (FWorldPartitionHandle ActorHandle(WorldPartition, ActorGuid); ActorHandle.IsValid())
-			{
-				if (AActor* PinnedActor = ActorHandle->GetActor())
-				{
-					GEditor->SelectActor(PinnedActor, /*bInSelected=*/false, /*bNotify=*/false);
-				}
-			}
-		}
-
 		WorldPartition->UnpinActors(ActorsToUnpin);
-
-		GEditor->GetSelectedActors()->EndBatchSelectOperation(/*bNotify=*/true);
 	}
 }
 
diff --git a/Engine/Source/Runtime/Engine/Private/Level.cpp b/Engine/Source/Runtime/Engine/Private/Level.cpp
--- a/Engine/Source/Runtime/Engine/Private/Level.cpp
+++ b/Engine/Source/Runtime/Engine/Private/Level.cpp
@@ -59,6 +59,9 @@ Level.cpp: Level-related functions
 #include "WorldPartition/WorldPartitionHelpers.h"
 #include "HAL/PlatformFileManager.h"
 #include "Misc/PathViews.h"
+#include "Selection.h"
+#include "Elements/Framework/EngineElementsLibrary.h"
+#include "Elements/Framework/TypedElementList.h"
 #endif
 #include "WorldPartition/WorldPartition.h"
 #include "WorldPartition/WorldPartitionLog.h"
@@ -935,29 +938,55 @@ void ULevel::RemoveLoadedActors(const TArray<AActor*>& ActorList, const FTransfo
 
 	FScopedSlowTask SlowTask(ActorsQueue.Num(), LOCTEXT("UnregisteringActors", "Unregistering actors..."));
 	SlowTask.MakeDialogDelayed(1.0f);
+			
+	bool bSelectionChanged = false;
 
-	for (AActor* Actor : ActorsQueue)
 	{
-		// UnregisterAllComponents can destroy child actors
-		if (!IsValid(Actor))
-		{
-			continue;
-		}
+		UTypedElementSelectionSet* ElementSelectionSet = (GEditor && GEditor->GetSelectedActors()) ? GEditor->GetSelectedActors()->GetElementSelectionSet() : nullptr;
+		TUniquePtr<FTypedElementList::FLegacySyncScopedBatch> LegacySyncBatch = ElementSelectionSet ? MakeUnique<FTypedElementList::FLegacySyncScopedBatch>(*ElementSelectionSet->GetElementList()) : nullptr;
 
-		Actor->UnregisterAllComponents();
-		Actor->RegisterAllActorTickFunctions(false, true);
+		const FTypedElementSelectionOptions SelectionOptions = FTypedElementSelectionOptions()
+			.SetAllowHidden(true)
+			.SetAllowGroups(false)
+			.SetWarnIfLocked(false)
+			.SetChildElementInclusionMethod(ETypedElementChildInclusionMethod::Recursive);
 
-		if (TransformToRemove)
+		for (AActor* Actor : ActorsQueue)
 		{
-			FLevelUtils::FApplyLevelTransformParams TransformParams(this, TransformToRemove->Inverse());
-			TransformParams.Actor = Actor;
-			TransformParams.bDoPostEditMove = true;
-			FLevelUtils::ApplyLevelTransform(TransformParams);
-		}
+			// UnregisterAllComponents can destroy child actors
+			if (!IsValid(Actor))
+			{
+				continue;
+			}
+
+			if (ElementSelectionSet)
+			{
+				if (FTypedElementHandle ActorHandle = UEngineElementsLibrary::AcquireEditorActorElementHandle(Actor, /*bAllowCreate*/false))
+				{
+					bSelectionChanged |= ElementSelectionSet->DeselectElement(ActorHandle, SelectionOptions);
+				}
+			}
 
-		OnLoadedActorRemovedFromLevelEvent.Broadcast(*Actor);
+			Actor->UnregisterAllComponents();
+			Actor->RegisterAllActorTickFunctions(false, true);
 
-		SlowTask.EnterProgressFrame(1);
+			if (TransformToRemove)
+			{
+				FLevelUtils::FApplyLevelTransformParams TransformParams(this, TransformToRemove->Inverse());
+				TransformParams.Actor = Actor;
+				TransformParams.bDoPostEditMove = true;
+				FLevelUtils::ApplyLevelTransform(TransformParams);
+			}
+
+			OnLoadedActorRemovedFromLevelEvent.Broadcast(*Actor);
+
+			SlowTask.EnterProgressFrame(1);
+		}
+	}
+
+	if (bSelectionChanged)
+	{
+		GEditor->NoteSelectionChange();
 	}
 }
 #endif
diff --git a/Engine/Source/Runtime/Engine/Private/WorldPartition/WorldPartitionActorLoaderInterface.cpp b/Engine/Source/Runtime/Engine/Private/WorldPartition/WorldPartitionActorLoaderInterface.cpp
--- a/Engine/Source/Runtime/Engine/Private/WorldPartition/WorldPartitionActorLoaderInterface.cpp
+++ b/Engine/Source/Runtime/Engine/Private/WorldPartition/WorldPartitionActorLoaderInterface.cpp
@@ -255,13 +255,7 @@ void IWorldPartitionActorLoaderInterface::ILoaderAdapter::PostLoadedStateChanged
 
 	if (!IsRunningCommandlet())
 	{
-		if (NumUnloads)
-		{
-			GEditor->SelectNone(true, true);
-		}
-
 		GEngine->BroadcastLevelActorListChanged();
-		GEditor->NoteSelectionChange();
 
 		if (NumUnloads)
 		{
