Cleanup CanSkipUpdate and deleted some deprecated functions

diff --git a/Engine/Source/Runtime/Engine/Private/DataReplication.cpp b/Engine/Source/Runtime/Engine/Private/DataReplication.cpp
--- a/Engine/Source/Runtime/Engine/Private/DataReplication.cpp
+++ b/Engine/Source/Runtime/Engine/Private/DataReplication.cpp
@@ -339,15 +339,6 @@ FObjectReplicator::~FObjectReplicator()
 	CleanUp();
 }
 
-bool FObjectReplicator::SendCustomDeltaProperty(UObject* InObject, FProperty* Property, uint32 ArrayIndex, FNetBitWriter& OutBunch, TSharedPtr<INetDeltaBaseState>& NewFullState, TSharedPtr<INetDeltaBaseState>& OldState)
-{
-	check(RepLayout);
-
-	const uint16 CustomDeltaIndex = RepLayout->GetCustomDeltaIndexFromPropertyRepIndex(Property->RepIndex + ArrayIndex);
-
-	return SendCustomDeltaProperty(InObject, CustomDeltaIndex, OutBunch, NewFullState, OldState);
-}
-
 bool FObjectReplicator::SendCustomDeltaProperty(UObject* InObject, uint16 CustomDeltaIndex, FNetBitWriter& OutBunch, TSharedPtr<INetDeltaBaseState>& NewFullState, TSharedPtr<INetDeltaBaseState>& OldState)
 {
 	check(!NewFullState.IsValid()); // NewState is passed in as nullptr and instantiated within this function if necessary
@@ -1738,24 +1729,22 @@ void FObjectReplicator::ReplicateCustomDeltaProperties( FNetBitWriter & Bunch, F
 bool FObjectReplicator::CanSkipUpdate(FReplicationFlags RepFlags)
 {
 #if WITH_PUSH_MODEL
-	const FSendingRepState& SendingRepState = *RepState->GetSendingRepState();
-	const FRepChangelistState& RepChangelistState = *ChangelistMgr->GetRepChangelistState();
+	if (!bCanUseNonDirtyOptimization)
+	{
+		// This class must always replicate
+		return false;
+	}
 
-	if (bCanUseNonDirtyOptimization &&
-		(RemoteFunctions == nullptr || RemoteFunctions->GetNumBits() == 0) &&
-		(RepLayout->IsEmpty() ||
-			(!RepFlags.bNetInitial &&
-			SendingRepState.RepFlags.Value == RepFlags.Value &&
-			SendingRepState.NumNaks == 0 &&
-			SendingRepState.LastCompareIndex > 1 &&
-			!(SendingRepState.bOpenAckedCalled && SendingRepState.PreOpenAckHistory.Num() > 0) &&
-			SendingRepState.LastChangelistIndex == RepChangelistState.HistoryEnd &&
-			Connection->ResendAllDataState == EResendAllDataState::None &&
-			!OwningChannel->bForceCompareProperties &&
-			(!GbPushModelSkipUndirtiedFastArrays || ((RepChangelistState.CustomDeltaChangeIndex == SendingRepState.CustomDeltaChangeIndex) && !SendingRepState.HasAnyPendingRetirements())) &&
-			!RepChangelistState.HasAnyDirtyProperties())
-		))
+	const bool bHasRPCQueued = RemoteFunctions && RemoteFunctions->GetNumBits() >= 0;
+	if (bHasRPCQueued)
 	{
+		return false;
+	}
+
+	const bool bHasNoRepLayout = RepLayout->IsEmpty();
+	if (bHasNoRepLayout)
+	{
+		// No properties to replicate and no RPCs queued, let's skip!
 #if CSV_PROFILER
 		++GNumSkippedObjectEmptyUpdates;
 #endif
@@ -1763,9 +1752,58 @@ bool FObjectReplicator::CanSkipUpdate(FReplicationFlags RepFlags)
 		bLastUpdateEmpty = true;
 		return true;
 	}
+
+	const bool bIsNetInitial = RepFlags.bNetInitial;
+	if (bIsNetInitial)
+	{
+		return false;
+	}
+
+	const FSendingRepState& SendingRepState = *RepState->GetSendingRepState();
+	const FRepChangelistState& RepChangelistState = *ChangelistMgr->GetRepChangelistState();
+
+	bool bCanSkip = true;
+
+	// Have the RepFlags changed ?
+	bCanSkip = bCanSkip && SendingRepState.RepFlags.Value == RepFlags.Value; 
+
+	// Any Naks to handle ?
+	bCanSkip = bCanSkip && SendingRepState.NumNaks == 0;
+
+	// Have we compared the properties twice ?
+	bCanSkip = bCanSkip && SendingRepState.LastCompareIndex > 1;
+
+	// Do we have open Ack's ?
+	bCanSkip = bCanSkip && !(SendingRepState.bOpenAckedCalled && SendingRepState.PreOpenAckHistory.Num() > 0);
+
+	// Any changelists to send ?
+	bCanSkip = bCanSkip && SendingRepState.LastChangelistIndex == RepChangelistState.HistoryEnd;
+
+	// Are we resending data for replay ?
+	bCanSkip = bCanSkip && Connection->ResendAllDataState == EResendAllDataState::None;
+
+	// Are we forcing a compare property ?
+	bCanSkip = bCanSkip && !OwningChannel->bForceCompareProperties;
+
+	// Are any CustomDelta properties dirty ?
+	bCanSkip = bCanSkip && (RepChangelistState.CustomDeltaChangeIndex == SendingRepState.CustomDeltaChangeIndex && !SendingRepState.HasAnyPendingRetirements());
+
+	// Are any replicated properties dirty ?
+	bCanSkip = bCanSkip && RepChangelistState.HasAnyDirtyProperties() == false;
+
+	if (bCanSkip)
+	{
+#if CSV_PROFILER
+		++GNumSkippedObjectEmptyUpdates;
 #endif
 
+		bLastUpdateEmpty = true;
+	}
+
+	return bCanSkip;
+#else // WITH_PUSH_MODEL
 	return false;
+#endif
 }
 
 bool FObjectReplicator::ReplicateProperties(FOutBunch& Bunch, FReplicationFlags RepFlags, FNetBitWriter& Writer)
@@ -1792,17 +1830,16 @@ bool FObjectReplicator::ReplicateProperties_r( FOutBunch & Bunch, FReplicationFl
 		return false;
 	}
 
+#if !(UE_BUILD_SHIPPING || UE_BUILD_TEST)
 	// some games ship checks() in Shipping so we cannot rely on DO_CHECK here, and these checks are in an extremely hot path
-	if (!UE_BUILD_SHIPPING && !UE_BUILD_TEST)
-	{
-		check(OwningChannel);
-		check(RepLayout.IsValid());
-		check(RepState.IsValid());
-		check(RepState->GetSendingRepState());
-		check(ChangelistMgr.IsValid());
-		check(ChangelistMgr->GetRepChangelistState() != nullptr);
-		check((ChangelistMgr->GetRepChangelistState()->StaticBuffer.Num() == 0) == RepLayout->IsEmpty());
-	}
+	check(OwningChannel);
+	check(RepLayout.IsValid());
+	check(RepState.IsValid());
+	check(RepState->GetSendingRepState());
+	check(ChangelistMgr.IsValid());
+	check(ChangelistMgr->GetRepChangelistState() != nullptr);
+	check((ChangelistMgr->GetRepChangelistState()->StaticBuffer.Num() == 0) == RepLayout->IsEmpty());
+#endif
 
 	UNetConnection* OwningChannelConnection = OwningChannel->Connection;
 
diff --git a/Engine/Source/Runtime/Engine/Public/Net/DataReplication.h b/Engine/Source/Runtime/Engine/Public/Net/DataReplication.h
--- a/Engine/Source/Runtime/Engine/Public/Net/DataReplication.h
+++ b/Engine/Source/Runtime/Engine/Public/Net/DataReplication.h
@@ -134,16 +134,6 @@ public:
 
 	//~ Both of these should be private, IMO, but we'll leave them public for now for back compat
 	//~ in case anyone was using SerializeCustomDeltaProperty already.
-
-	UE_DEPRECATED(5.0, "No longer used.")
-	bool SendCustomDeltaProperty(
-		UObject* InObject,
-		FProperty* Property,
-		uint32 ArrayIndex,
-		FNetBitWriter& OutBunch,
-		TSharedPtr<INetDeltaBaseState>& NewFullState,
-		TSharedPtr<INetDeltaBaseState> & OldState);
-
 	bool SendCustomDeltaProperty(
 		UObject* InObject,
 		uint16 CustomDeltaProperty,
@@ -171,10 +161,6 @@ public:
 
 	void PostSendBunch(FPacketIdRange& PacketRange, uint8 bReliable);
 
-	/** Updates the custom delta state for a replay delta checkpoint */
-	UE_DEPRECATED(4.27, "No longer used")
-	void UpdateCheckpoint() {}
-	
 	bool ReceivedBunch(
 		FNetBitReader& Bunch,
 		const FReplicationFlags& RepFlags,
