Fix for crash in MaterialInstance.cpp caused by dereferencing a null UTexture* when the material has SparseVolumeTexture parameters.

#rnx

diff --git a/Engine/Source/Runtime/Engine/Private/Materials/MaterialInstance.cpp b/Engine/Source/Runtime/Engine/Private/Materials/MaterialInstance.cpp
--- a/Engine/Source/Runtime/Engine/Private/Materials/MaterialInstance.cpp
+++ b/Engine/Source/Runtime/Engine/Private/Materials/MaterialInstance.cpp
@@ -1265,20 +1265,18 @@ void UMaterialInstance::ValidateTextureOverrides(ERHIFeatureLevel::Type InFeatur
 	for (uint32 TypeIndex = 0; TypeIndex < NumMaterialTextureParameterTypes; ++TypeIndex)
 	{
 		const EMaterialTextureParameterType ParameterType = (EMaterialTextureParameterType)TypeIndex;
+		
+		// SVT currently do not derive from UTexture and checking for GetMaterialType() validity is not necessary here because SVT are always MCT_SparseVolumeTexture.
+		if (ParameterType == EMaterialTextureParameterType::SparseVolume)
+		{
+			continue;
+		}
+
 		for (const FMaterialTextureParameterInfo& TextureInfo : CurrentResource->GetUniformTextureExpressions(ParameterType))
 		{
 			UTexture* Texture = nullptr;
-			USparseVolumeTexture* SVTexture = nullptr;
-			if (ParameterType != EMaterialTextureParameterType::SparseVolume)
-			{
-				TextureInfo.GetGameThreadTextureValue(this, *CurrentResource, Texture);
-			}
-			else
-			{
-				TextureInfo.GetGameThreadTextureValue(this, *CurrentResource, SVTexture);
-			}
-			
-			if (Texture || SVTexture)
+			TextureInfo.GetGameThreadTextureValue(this, *CurrentResource, Texture);
+			if (Texture)
 			{
 				const EMaterialValueType TextureType = Texture->GetMaterialType();
 				switch (ParameterType)
@@ -1327,12 +1325,6 @@ void UMaterialInstance::ValidateTextureOverrides(ERHIFeatureLevel::Type InFeatur
 						UE_LOG(LogMaterial, Error, TEXT("MaterialInstance \"%s\" parameter '%s' assigned texture \"%s\" requires virtual texture"), *MaterialName, *TextureInfo.GetParameterName().ToString(), *Texture->GetName());
 					}
 					break;
-				case EMaterialTextureParameterType::SparseVolume:
-					if (!(TextureType & MCT_SparseVolumeTexture))
-					{
-						UE_LOG(LogMaterial, Error, TEXT("MaterialInstance \"%s\" parameter '%s' assigned texture \"%s\" has invalid type, required sparse volume texture"), *MaterialName, *TextureInfo.GetParameterName().ToString(), *SVTexture->GetName());
-					}
-					break;
 				default:
 					checkNoEntry();
 					break;
