- LWC fixes for Cascade and Beams

diff --git a/Engine/Source/Runtime/Engine/Private/Particles/ParticleBeamModules.cpp b/Engine/Source/Runtime/Engine/Private/Particles/ParticleBeamModules.cpp
--- a/Engine/Source/Runtime/Engine/Private/Particles/ParticleBeamModules.cpp
+++ b/Engine/Source/Runtime/Engine/Private/Particles/ParticleBeamModules.cpp
@@ -127,8 +127,8 @@ void UParticleModuleTypeDataBeam2::Spawn(FParticleEmitterInstance* Owner, int32
 	// If there is no Source module, use the emitter position as the source point
 	if (BeamInst->BeamModule_Source == NULL)
 	{
-		BeamData->SourcePoint	= (FVector3f)Component->GetComponentLocation();
-		BeamData->SourceTangent = FVector3f(Component->GetComponentTransform().GetScaledAxis( EAxis::X ));
+		BeamData->SourcePoint	= Component->GetComponentLocation();
+		BeamData->SourceTangent = Component->GetComponentTransform().GetScaledAxis( EAxis::X );
 		BeamData->SourceStrength = 1.0f;
 	}
 
@@ -136,9 +136,9 @@ void UParticleModuleTypeDataBeam2::Spawn(FParticleEmitterInstance* Owner, int32
 	if ((BeamInst->BeamModule_Target == NULL) && (BeamInst->BeamMethod == PEB2M_Distance))
 	{
 		// Set the particle target based on the distance
-		float	TotalDistance	= Distance.GetValue(Particle.RelativeTime, Component);
+		double	TotalDistance	= Distance.GetValue(Particle.RelativeTime, Component);
 		// Always use the X-axis of the component as the direction
-		FVector3f	Direction		= (FVector3f)Component->GetComponentTransform().GetScaledAxis( EAxis::X );
+		FVector	Direction		= Component->GetComponentTransform().GetScaledAxis( EAxis::X );
 		Direction.Normalize();
 		// Calculate the final target point
 		BeamData->TargetPoint	= BeamData->SourcePoint + Direction * TotalDistance;
@@ -246,16 +246,16 @@ void UParticleModuleTypeDataBeam2::Update(FParticleEmitterInstance* Owner, int32
 		// If there is no Source module, use the emitter location
 		if (BeamInst->BeamModule_Source == NULL)
 		{
-			BeamData->SourcePoint	= (FVector3f)Component->GetComponentLocation();
-			BeamData->SourceTangent = FVector3f(Component->GetComponentTransform().GetScaledAxis( EAxis::X ));
+			BeamData->SourcePoint	= Component->GetComponentLocation();
+			BeamData->SourceTangent = Component->GetComponentTransform().GetScaledAxis( EAxis::X );
 		}
 
 		// If the method is set for distance, or there is no target, determine the target point
 		if ((BeamInst->BeamModule_Target == NULL) && (BeamInst->BeamMethod == PEB2M_Distance))
 		{
 			// Set the particle target based on the distance
-			float	TotalDistance	= Distance.GetValue(Particle.RelativeTime, Component);
-			FVector3f	Direction		= (FVector3f)Component->GetComponentTransform().GetScaledAxis( EAxis::X );
+			double	TotalDistance	= Distance.GetValue(Particle.RelativeTime, Component);
+			FVector	Direction		= Component->GetComponentTransform().GetScaledAxis( EAxis::X );
 			Direction.Normalize();
 			BeamData->TargetPoint	= BeamData->SourcePoint + Direction * TotalDistance;
 			BeamData->TargetTangent = -Direction;
@@ -294,10 +294,10 @@ void UParticleModuleTypeDataBeam2::Update(FParticleEmitterInstance* Owner, int32
 		if ((Speed != 0.0f) && (!BEAM2_TYPEDATA_LOCKED(BeamData->Lock_Max_NumNoisePoints)))
 		{
 			// If the beam isn't locked, then move it towards the target...
-			if (Particle.Location != (FVector)BeamData->TargetPoint)
+			if (Particle.Location != BeamData->TargetPoint)
 			{
 				// Determine the direction of travel
-				FVector Direction	= (FVector)BeamData->TargetPoint - Particle.Location;
+				FVector Direction	= BeamData->TargetPoint - Particle.Location;
 				Direction.Normalize();
 				// Setup the offset and the current distance travelled
 				const FVector BeamOffset		= Direction * Speed * DeltaTime;
@@ -307,7 +307,7 @@ void UParticleModuleTypeDataBeam2::Update(FParticleEmitterInstance* Owner, int32
 					(FMath::Abs(Sum.Z - BeamData->TargetPoint.Z) < LockRadius))
 				{
 					// We are within the lock radius, so lock the beam
-					Particle.Location	= (FVector)BeamData->TargetPoint;
+					Particle.Location	= BeamData->TargetPoint;
 					BEAM2_TYPEDATA_SETLOCKED(BeamData->Lock_Max_NumNoisePoints, 1);
 				}
 				else
@@ -320,13 +320,13 @@ void UParticleModuleTypeDataBeam2::Update(FParticleEmitterInstance* Owner, int32
 		else
 		{
 			// Pop right to the target and set the beam as locked
-			Particle.Location = (FVector)BeamData->TargetPoint;
+			Particle.Location = BeamData->TargetPoint;
 			BEAM2_TYPEDATA_SETLOCKED(BeamData->Lock_Max_NumNoisePoints, 1);
 		}
 
 		// Determine the step size, count, and travelled ratio
-		BeamData->Direction		= FVector3f(BeamData->TargetPoint - BeamData->SourcePoint);
-		float	FullMagnitude	= FMath::Max(BeamData->Direction.Size(), 0.001f);
+		BeamData->Direction		= BeamData->TargetPoint - BeamData->SourcePoint;
+		double	FullMagnitude	= FMath::Max(BeamData->Direction.Size(), 0.001);
 		BeamData->Direction.Normalize();
 
 		int32 InterpSteps = 0;
@@ -337,34 +337,34 @@ void UParticleModuleTypeDataBeam2::Update(FParticleEmitterInstance* Owner, int32
 			if (BEAM2_TYPEDATA_LOCKED(BeamData->Lock_Max_NumNoisePoints))
 			{
 				// If the beam is locked to the target, the steps are the interpolation count
-				BeamData->StepSize		= FullMagnitude / InterpolationCount;
+				BeamData->StepSize		= FullMagnitude / double(InterpolationCount);
 				BeamData->Steps			= InterpolationCount;
 				BeamData->TravelRatio	= 0.0f;
 			}
 			else
 			{
 				// Determine the number of steps we have traveled
-				FVector	TrueDistance	= Particle.Location - (FVector)BeamData->SourcePoint;
-				float	TrueMagnitude	= TrueDistance.Size();
+				FVector	TrueDistance	= Particle.Location - BeamData->SourcePoint;
+				double	TrueMagnitude	= TrueDistance.Size();
 				if (TrueMagnitude > FullMagnitude)
 				{
 					// Lock to the target if we are over-shooting and determine the steps and step size
-					Particle.Location	= (FVector)BeamData->TargetPoint;
-					TrueDistance		= Particle.Location - (FVector)BeamData->SourcePoint;
+					Particle.Location	= BeamData->TargetPoint;
+					TrueDistance		= Particle.Location - BeamData->SourcePoint;
 					TrueMagnitude		= TrueDistance.Size();
 					BEAM2_TYPEDATA_SETLOCKED(BeamData->Lock_Max_NumNoisePoints, true);
-					BeamData->StepSize		= FullMagnitude / InterpolationCount;
+					BeamData->StepSize		= FullMagnitude / double(InterpolationCount);
 					BeamData->Steps			= InterpolationCount;
 					BeamData->TravelRatio	= 0.0f;
 				}
 				else
 				{
 					// Determine the steps and step size
-					BeamData->StepSize		= FullMagnitude / InterpolationCount;
+					BeamData->StepSize		= FullMagnitude / double(InterpolationCount);
 					BeamData->TravelRatio	= TrueMagnitude / FullMagnitude;
 					BeamData->Steps			= FMath::FloorToInt(BeamData->TravelRatio * InterpolationCount);
 					// Readjust the travel ratio
-					BeamData->TravelRatio	= (TrueMagnitude - (BeamData->StepSize * BeamData->Steps)) / BeamData->StepSize;
+					BeamData->TravelRatio	= (TrueMagnitude - (BeamData->StepSize * double(BeamData->Steps))) / BeamData->StepSize;
 				}
 			}
 			InterpSteps = BeamData->Steps;
@@ -386,7 +386,7 @@ void UParticleModuleTypeDataBeam2::Update(FParticleEmitterInstance* Owner, int32
 					// Determine the number of points to drop.
 					int32 Count = FMath::TruncToInt(FullMagnitude / BeamNoise->FrequencyDistance);
 					Count = FMath::Min<int32>(Count, Freq);
-					BeamData->StepSize		= FullMagnitude / (Count + 1);
+					BeamData->StepSize		= FullMagnitude / double(Count + 1);
 					BeamData->Steps			= Count;
 					BeamData->TravelRatio	= 0.0f;
 					if (NoiseDistanceScale != NULL)
@@ -398,7 +398,7 @@ void UParticleModuleTypeDataBeam2::Update(FParticleEmitterInstance* Owner, int32
 				else
 				{
 					// If locked, just use the noise frequency to determine steps
-					BeamData->StepSize		= FullMagnitude / (Freq + 1);
+					BeamData->StepSize		= FullMagnitude / double(Freq + 1);
 					BeamData->Steps			= Freq;
 					BeamData->TravelRatio	= 0.0f;
 					if (NoiseDistanceScale != NULL)
@@ -410,14 +410,14 @@ void UParticleModuleTypeDataBeam2::Update(FParticleEmitterInstance* Owner, int32
 			else
 			{
 				// Determine that actual distance traveled, and its magnitude
-				FVector	TrueDistance	= Particle.Location - (FVector)BeamData->SourcePoint;
-				float	TrueMagnitude	= TrueDistance.Size();
+				FVector	TrueDistance	= Particle.Location - BeamData->SourcePoint;
+				double	TrueMagnitude	= TrueDistance.Size();
 
 				if (BeamNoise->FrequencyDistance > 0.0f)
 				{
 					int32 Count = FMath::TruncToInt(FullMagnitude / BeamNoise->FrequencyDistance);
 					Count = FMath::Min<int32>(Count, Freq);
-					BeamData->StepSize		= FullMagnitude / (Count + 1);
+					BeamData->StepSize		= FullMagnitude / double(Count + 1);
 					// Determine the partial trail amount and the steps taken
 					BeamData->TravelRatio	= TrueMagnitude / FullMagnitude;
 					BeamData->Steps			= FMath::FloorToInt(BeamData->TravelRatio * (Count + 1));
@@ -430,12 +430,12 @@ void UParticleModuleTypeDataBeam2::Update(FParticleEmitterInstance* Owner, int32
 					if (BeamData->Steps == Count)
 					{
 						BeamData->TravelRatio	= 
-							(TrueMagnitude - (BeamData->StepSize * BeamData->Steps)) / 
-							(FullMagnitude - (BeamData->StepSize * BeamData->Steps));
+							(TrueMagnitude - (BeamData->StepSize * double(BeamData->Steps))) /
+							(FullMagnitude - (BeamData->StepSize * double(BeamData->Steps)));
 					}
 					else
 					{
-						BeamData->TravelRatio	= (TrueMagnitude - (BeamData->StepSize * BeamData->Steps)) / BeamData->StepSize;
+						BeamData->TravelRatio	= (TrueMagnitude - (BeamData->StepSize * double(BeamData->Steps))) / BeamData->StepSize;
 					}
 
 					if (NoiseDistanceScale != NULL)
@@ -448,7 +448,7 @@ void UParticleModuleTypeDataBeam2::Update(FParticleEmitterInstance* Owner, int32
 				{
 					// If we are not doing noisy interpolation
 					// Determine the step size for the full beam
-					BeamData->StepSize		= FullMagnitude / (Freq + 1);
+					BeamData->StepSize		= FullMagnitude / double(Freq + 1);
 					// Determine the partial trail amount and the steps taken
 					BeamData->TravelRatio	= TrueMagnitude / FullMagnitude;
 					BeamData->Steps			= FMath::FloorToInt(BeamData->TravelRatio * (Freq + 1));
@@ -461,12 +461,12 @@ void UParticleModuleTypeDataBeam2::Update(FParticleEmitterInstance* Owner, int32
 					if (BeamData->Steps == Freq)
 					{
 						BeamData->TravelRatio	= 
-							(TrueMagnitude - (BeamData->StepSize * BeamData->Steps)) / 
-							(FullMagnitude - (BeamData->StepSize * BeamData->Steps));
+							(TrueMagnitude - (BeamData->StepSize * double(BeamData->Steps))) /
+							(FullMagnitude - (BeamData->StepSize * double(BeamData->Steps)));
 					}
 					else
 					{
-						BeamData->TravelRatio	= (TrueMagnitude - (BeamData->StepSize * BeamData->Steps)) / BeamData->StepSize;
+						BeamData->TravelRatio	= (TrueMagnitude - (BeamData->StepSize * double(BeamData->Steps))) / BeamData->StepSize;
 					}
 				}
 
@@ -489,10 +489,10 @@ void UParticleModuleTypeDataBeam2::Update(FParticleEmitterInstance* Owner, int32
 			FVector	TargetPosition;
 			FVector	TargetTangent;
 
-			float	InvTess	= 1.0f / InterpolationPoints;
+			double	InvTess	= 1.0 / double(InterpolationPoints);
 
-			SourcePosition	 = (FVector)BeamData->SourcePoint;
-			SourceTangent	 = FVector(BeamData->SourceTangent);
+			SourcePosition	 = BeamData->SourcePoint;
+			SourceTangent	 = BeamData->SourceTangent;
 #if defined(_BEAM2_TYPEDATA_NORMAL_TANGENTS_)
 			SourceTangent.Normalize();
 #endif	//#if defined(_BEAM2_TYPEDATA_NORMAL_TANGENTS_)
@@ -502,8 +502,8 @@ void UParticleModuleTypeDataBeam2::Update(FParticleEmitterInstance* Owner, int32
 			}
 			SourceTangent	*= BeamData->SourceStrength;
 
-			TargetPosition	 = (FVector)BeamData->TargetPoint;
-			TargetTangent	 = FVector(BeamData->TargetTangent);
+			TargetPosition	 = BeamData->TargetPoint;
+			TargetTangent	 = BeamData->TargetTangent;
 #if defined(_BEAM2_TYPEDATA_NORMAL_TANGENTS_)
 			TargetTangent.Normalize();
 #endif	//#if defined(_BEAM2_TYPEDATA_NORMAL_TANGENTS_)
@@ -521,7 +521,7 @@ void UParticleModuleTypeDataBeam2::Update(FParticleEmitterInstance* Owner, int32
 				InterpolatedPoints[ii] = FMath::CubicInterp(
 					SourcePosition, SourceTangent,
 					TargetPosition, TargetTangent,
-					InvTess * (ii + 1));
+					InvTess * double(ii + 1));
 				LastPosition		= InterpolatedPoints[ii];
 			}
 
@@ -537,7 +537,7 @@ void UParticleModuleTypeDataBeam2::Update(FParticleEmitterInstance* Owner, int32
 				InterpolatedPoints[ii] = FMath::CubicInterp(
 					SourcePosition, SourceTangent,
 					TargetPosition, TargetTangent,
-					InvTess * (ii + 1));
+					InvTess * double(ii + 1));
 			}
 
 			if (bLowFreqNoise == true)
@@ -1097,7 +1097,7 @@ void UParticleModuleBeamModifier::Spawn(FParticleEmitterInstance* Owner, int32 O
 			ModifierPayload->bModifyPosition = PositionOptions.bModify;
 			if (PositionOptions.bModify == true)
 			{
-				ModifierPayload->Position = (FVector3f)Position.GetValue(Owner->EmitterTime, Owner->Component);
+				ModifierPayload->Position = Position.GetValue(Owner->EmitterTime, Owner->Component);
 				ModifierPayload->bScalePosition = PositionOptions.bScale;
 			}
 
@@ -1105,7 +1105,7 @@ void UParticleModuleBeamModifier::Spawn(FParticleEmitterInstance* Owner, int32 O
 			ModifierPayload->bModifyTangent = TangentOptions.bModify;
 			if (TangentOptions.bModify == true)
 			{
-				ModifierPayload->Tangent = (FVector3f)Tangent.GetValue(Owner->EmitterTime, Owner->Component);
+				ModifierPayload->Tangent = Tangent.GetValue(Owner->EmitterTime, Owner->Component);
 				ModifierPayload->bScaleTangent = TangentOptions.bScale;
 			}
 
@@ -1147,7 +1147,7 @@ void UParticleModuleBeamModifier::Update(FParticleEmitterInstance* Owner, int32
 			ModifierPayload->bModifyPosition = PositionOptions.bModify;
 			if ((PositionOptions.bModify == true) && (PositionOptions.bLock == false))
 			{
-				ModifierPayload->Position = (FVector3f)Position.GetValue(Owner->EmitterTime, Owner->Component);
+				ModifierPayload->Position = Position.GetValue(Owner->EmitterTime, Owner->Component);
 				ModifierPayload->bScalePosition = PositionOptions.bScale;
 			}
 
@@ -1155,7 +1155,7 @@ void UParticleModuleBeamModifier::Update(FParticleEmitterInstance* Owner, int32
 			ModifierPayload->bModifyTangent = TangentOptions.bModify;
 			if ((TangentOptions.bModify == true) && (TangentOptions.bLock == false))
 			{
-				ModifierPayload->Tangent = (FVector3f)Tangent.GetValue(Owner->EmitterTime, Owner->Component);
+				ModifierPayload->Tangent = Tangent.GetValue(Owner->EmitterTime, Owner->Component);
 				ModifierPayload->bScaleTangent = TangentOptions.bScale;
 			}
 
@@ -1706,7 +1706,7 @@ void UParticleModuleBeamSource::Spawn(FParticleEmitterInstance* Owner, int32 Off
 	ResolveSourceData(BeamInst, BeamData, (uint8*)ParticleBase, Offset, BeamInst->ActiveParticles, true, SourceModifier);
 
 	// Set the location and clear the initial data flags
-	Particle.Location					= (FVector)BeamData->SourcePoint - BeamInst->PositionOffsetThisTick;
+	Particle.Location					= BeamData->SourcePoint - BeamInst->PositionOffsetThisTick;
 	BeamData->Lock_Max_NumNoisePoints	= 0;
 	BeamData->StepSize					= 0.0f;
 	BeamData->Steps						= 0;
@@ -1874,18 +1874,18 @@ bool UParticleModuleBeamSource::ResolveSourceData(FParticleBeam2EmitterInstance*
 				if (ParticleIndex >= BeamInst->UserSetSourceArray.Num())
 				{
 					//@todo. How to handle this situation???
-					BeamData->SourcePoint	= (FVector3f)BeamInst->UserSetSourceArray[0];
+					BeamData->SourcePoint	= BeamInst->UserSetSourceArray[0];
 				}
 				else
 				{
-					BeamData->SourcePoint	= (FVector3f)BeamInst->UserSetSourceArray[ParticleIndex];
+					BeamData->SourcePoint	= BeamInst->UserSetSourceArray[ParticleIndex];
 				}
 				bSetSource	= true;
 			}
 			break;
 		case PEB2STM_Emitter:
 			// The position of the owner component is the source
-			BeamData->SourcePoint	= (FVector3f)BeamInst->Component->GetComponentLocation();
+			BeamData->SourcePoint	= BeamInst->Component->GetComponentLocation();
 			bSetSource				= true;
 			break;
 		case PEB2STM_Particle:
@@ -1985,7 +1985,7 @@ bool UParticleModuleBeamSource::ResolveSourceData(FParticleBeam2EmitterInstance*
 						}
 
 						//@todo. fill this in correctly...
-						BeamData->SourcePoint = (FVector3f)CalcSourcePosition;
+						BeamData->SourcePoint = CalcSourcePosition;
 						bSetSource = true;
 					}
 				}
@@ -1998,7 +1998,7 @@ bool UParticleModuleBeamSource::ResolveSourceData(FParticleBeam2EmitterInstance*
 				// Use the actor position as the source
 				if (BeamInst->SourceActor)
 				{
-					BeamData->SourcePoint	= (FVector3f)BeamInst->SourceActor->ActorToWorld().GetLocation();
+					BeamData->SourcePoint	= BeamInst->SourceActor->ActorToWorld().GetLocation();
 					bSetSource = true;
 				}
 			}
@@ -2012,14 +2012,14 @@ bool UParticleModuleBeamSource::ResolveSourceData(FParticleBeam2EmitterInstance*
 			if (bSourceAbsolute)
 			{
 				// Use the value as a world space position
-				BeamData->SourcePoint	= (FVector3f)Source.GetValue(BeamInst->EmitterTime, BeamInst->Component);
+				BeamData->SourcePoint	= Source.GetValue(BeamInst->EmitterTime, BeamInst->Component);
 				// Take into account current world origin offset
-				BeamData->SourcePoint  -= FVector3f(BeamInst->Component->GetWorld()->OriginLocation);
+				BeamData->SourcePoint  -= (FVector)BeamInst->Component->GetWorld()->OriginLocation;
 			}
 			else
 			{
 				// Use the value as a local space position.
-				BeamData->SourcePoint	= (FVector3f)BeamInst->Component->GetComponentTransform().TransformPosition(
+				BeamData->SourcePoint	= BeamInst->Component->GetComponentTransform().TransformPosition(
 					Source.GetValue(BeamInst->EmitterTime, BeamInst->Component));
 			}
 		}
@@ -2033,7 +2033,7 @@ bool UParticleModuleBeamSource::ResolveSourceData(FParticleBeam2EmitterInstance*
 		{
 		case PEB2STTM_Direct:
 			// Use the emitter direction as the tangent
-			BeamData->SourceTangent	= (FVector3f)BeamInst->Component->GetComponentTransform().GetScaledAxis( EAxis::X );
+			BeamData->SourceTangent	= BeamInst->Component->GetComponentTransform().GetScaledAxis( EAxis::X );
 			bSetSourceTangent		= true;
 			break;
 		case PEB2STTM_UserSet:
@@ -2043,23 +2043,23 @@ bool UParticleModuleBeamSource::ResolveSourceData(FParticleBeam2EmitterInstance*
 				if (ParticleIndex >= BeamInst->UserSetSourceTangentArray.Num())
 				{
 					//@todo. How to handle this situation???
-					BeamData->SourceTangent	= (FVector3f)BeamInst->UserSetSourceTangentArray[0];
+					BeamData->SourceTangent	= BeamInst->UserSetSourceTangentArray[0];
 				}
 				else
 				{
-					BeamData->SourceTangent	= (FVector3f)BeamInst->UserSetSourceTangentArray[ParticleIndex];
+					BeamData->SourceTangent	= BeamInst->UserSetSourceTangentArray[ParticleIndex];
 				}
 				bSetSourceTangent	= true;
 			}
 			break;
 		case PEB2STTM_Distribution:
 			// Use the tangent contained in the distribution
-			BeamData->SourceTangent	= (FVector3f)SourceTangent.GetValue(Particle.RelativeTime, BeamInst->Component);
+			BeamData->SourceTangent	= SourceTangent.GetValue(Particle.RelativeTime, BeamInst->Component);
 			bSetSourceTangent		= true;
 			break;
 		case PEB2STTM_Emitter:
 			// Use the emitter direction as the tangent
-			BeamData->SourceTangent	= (FVector3f)BeamInst->Component->GetComponentTransform().GetScaledAxis( EAxis::X );
+			BeamData->SourceTangent	= BeamInst->Component->GetComponentTransform().GetScaledAxis( EAxis::X );
 			bSetSourceTangent		= true;
 			break;
 		}
@@ -2068,11 +2068,11 @@ bool UParticleModuleBeamSource::ResolveSourceData(FParticleBeam2EmitterInstance*
 		{
 			// By default, use the distribution. This will allow artists an easier setup phase...
 //			BeamData->SourceTangent	= BeamInst->Component->LocalToWorld.GetScaledAxis( EAxis::X );
-			BeamData->SourceTangent	= (FVector3f)SourceTangent.GetValue(Particle.RelativeTime, BeamInst->Component);
+			BeamData->SourceTangent	= SourceTangent.GetValue(Particle.RelativeTime, BeamInst->Component);
 			if (bSourceAbsolute == false)
 			{
 				// If not tagged as absolute, transform it to world space
-				BeamData->SourceTangent	= (FVector3f)BeamInst->Component->GetComponentTransform().TransformVector((FVector)BeamData->SourceTangent);
+				BeamData->SourceTangent	= BeamInst->Component->GetComponentTransform().TransformVector(BeamData->SourceTangent);
 			}
 		}
 	}
@@ -2329,7 +2329,7 @@ bool UParticleModuleBeamTarget::ResolveTargetData(FParticleBeam2EmitterInstance*
 			}
 			FVector	Direction		= BeamInst->Component->GetComponentTransform().GetScaledAxis( EAxis::X );
 			Direction.Normalize();
-			BeamData->TargetPoint	= BeamData->SourcePoint + FVector3f(Direction * Distance);
+			BeamData->TargetPoint	= BeamData->SourcePoint + Direction * Distance;
 			bSetTarget				= true;
 		}
 
@@ -2343,11 +2343,11 @@ bool UParticleModuleBeamTarget::ResolveTargetData(FParticleBeam2EmitterInstance*
 					if (ParticleIndex >= BeamInst->UserSetTargetArray.Num())
 					{
 						//@todo. How to handle this situation???
-						BeamData->TargetPoint	= (FVector3f)BeamInst->UserSetTargetArray[0];
+						BeamData->TargetPoint	= BeamInst->UserSetTargetArray[0];
 					}
 					else
 					{
-						BeamData->TargetPoint	= (FVector3f)BeamInst->UserSetTargetArray[ParticleIndex];
+						BeamData->TargetPoint	= BeamInst->UserSetTargetArray[ParticleIndex];
 					}
 					bSetTarget	= true;
 				}
@@ -2450,7 +2450,7 @@ bool UParticleModuleBeamTarget::ResolveTargetData(FParticleBeam2EmitterInstance*
 							}
 
 							//@todo. fill this in correctly...
-							BeamData->TargetPoint = (FVector3f)LocalTargetPosition; //BeamInst->TargetEmitter->Component->GetComponentLocation();
+							BeamData->TargetPoint = LocalTargetPosition; //BeamInst->TargetEmitter->Component->GetComponentLocation();
 							bSetTarget = true;
 						}
 					}
@@ -2462,7 +2462,7 @@ bool UParticleModuleBeamTarget::ResolveTargetData(FParticleBeam2EmitterInstance*
 					BeamInst->ResolveTarget();
 					if (BeamInst->TargetActor)
 					{
-						BeamData->TargetPoint	= (FVector3f)BeamInst->TargetActor->ActorToWorld().GetLocation();
+						BeamData->TargetPoint	= BeamInst->TargetActor->ActorToWorld().GetLocation();
 						bSetTarget = true;
 					}
 				}
@@ -2474,13 +2474,13 @@ bool UParticleModuleBeamTarget::ResolveTargetData(FParticleBeam2EmitterInstance*
 		{
 			if (bTargetAbsolute)
 			{
-				BeamData->TargetPoint	= (FVector3f)Target.GetValue(BeamInst->EmitterTime, BeamInst->Component);
+				BeamData->TargetPoint	= Target.GetValue(BeamInst->EmitterTime, BeamInst->Component);
 				// Take into account current world origin offset
-				BeamData->TargetPoint  -= FVector3f(BeamInst->Component->GetWorld()->OriginLocation);
+				BeamData->TargetPoint  -= (FVector)BeamInst->Component->GetWorld()->OriginLocation;
 			}
 			else
 			{
-				BeamData->TargetPoint	= (FVector3f)BeamInst->Component->GetComponentTransform().TransformPosition(
+				BeamData->TargetPoint	= BeamInst->Component->GetComponentTransform().TransformPosition(
 					Target.GetValue(BeamInst->EmitterTime, BeamInst->Component));
 			}
 		}
@@ -2493,7 +2493,7 @@ bool UParticleModuleBeamTarget::ResolveTargetData(FParticleBeam2EmitterInstance*
 		switch (TargetTangentMethod)
 		{
 		case PEB2STTM_Direct:
-			BeamData->TargetTangent	= (FVector3f)BeamInst->Component->GetComponentTransform().GetScaledAxis( EAxis::X );
+			BeamData->TargetTangent	= BeamInst->Component->GetComponentTransform().GetScaledAxis( EAxis::X );
 			bSetTargetTangent		= true;
 			break;
 		case PEB2STTM_UserSet:
@@ -2502,21 +2502,21 @@ bool UParticleModuleBeamTarget::ResolveTargetData(FParticleBeam2EmitterInstance*
 				if (ParticleIndex >= BeamInst->UserSetTargetTangentArray.Num())
 				{
 					//@todo. How to handle this situation???
-					BeamData->TargetTangent	= (FVector3f)BeamInst->UserSetTargetTangentArray[0];
+					BeamData->TargetTangent	= BeamInst->UserSetTargetTangentArray[0];
 				}
 				else
 				{
-					BeamData->TargetTangent	= (FVector3f)BeamInst->UserSetTargetTangentArray[ParticleIndex];
+					BeamData->TargetTangent	= BeamInst->UserSetTargetTangentArray[ParticleIndex];
 				}
 				bSetTargetTangent	= true;
 			}
 			break;
 		case PEB2STTM_Distribution:
-			BeamData->TargetTangent	= (FVector3f)TargetTangent.GetValue(Particle.RelativeTime, BeamInst->Component);
+			BeamData->TargetTangent	= TargetTangent.GetValue(Particle.RelativeTime, BeamInst->Component);
 			bSetTargetTangent		= true;
 			break;
 		case PEB2STTM_Emitter:
-			BeamData->TargetTangent	= (FVector3f)BeamInst->Component->GetComponentTransform().GetScaledAxis( EAxis::X );
+			BeamData->TargetTangent	= BeamInst->Component->GetComponentTransform().GetScaledAxis( EAxis::X );
 			bSetTargetTangent		= true;
 			break;
 		}
@@ -2524,10 +2524,10 @@ bool UParticleModuleBeamTarget::ResolveTargetData(FParticleBeam2EmitterInstance*
 		if (bSetTargetTangent == false)
 		{
 //			BeamData->TargetTangent	= BeamInst->Component->LocalToWorld.GetScaledAxis( EAxis::X );
-			BeamData->TargetTangent	= (FVector3f)TargetTangent.GetValue(Particle.RelativeTime, BeamInst->Component);
+			BeamData->TargetTangent	= TargetTangent.GetValue(Particle.RelativeTime, BeamInst->Component);
 			if (bTargetAbsolute == false)
 			{
-				BeamData->TargetTangent	= (FVector3f)BeamInst->Component->GetComponentTransform().TransformVector((FVector)BeamData->TargetTangent);
+				BeamData->TargetTangent	= BeamInst->Component->GetComponentTransform().TransformVector(BeamData->TargetTangent);
 			}
 		}
 	}
diff --git a/Engine/Source/Runtime/Engine/Public/ParticleHelper.h b/Engine/Source/Runtime/Engine/Public/ParticleHelper.h
--- a/Engine/Source/Runtime/Engine/Public/ParticleHelper.h
+++ b/Engine/Source/Runtime/Engine/Public/ParticleHelper.h
@@ -630,36 +630,36 @@ struct FLightParticlePayload
 struct FBeam2TypeDataPayload
 {
 	/** The source of this beam											*/
-	FVector3f		SourcePoint;
+	FVector		SourcePoint;
 	/** The source tangent of this beam									*/
-	FVector3f		SourceTangent;
+	FVector		SourceTangent;
 	/** The stength of the source tangent of this beam					*/
 	float		SourceStrength;
 
 	/** The target of this beam											*/
-	FVector3f		TargetPoint;
+	FVector		TargetPoint;
 	/** The target tangent of this beam									*/
-	FVector3f		TargetTangent;
+	FVector		TargetTangent;
 	/** The stength of the Target tangent of this beam					*/
 	float		TargetStrength;
 
 	/** Target lock, extreme max, Number of noise points				*/
-	int32			Lock_Max_NumNoisePoints;
+	int32		Lock_Max_NumNoisePoints;
 
 	/** Number of segments to render (steps)							*/
-	int32			InterpolationSteps;
+	int32		InterpolationSteps;
 
 	/** Direction to step in											*/
-	FVector3f		Direction;
+	FVector		Direction;
 	/** StepSize (for each segment to be rendered)						*/
-	float		StepSize;
+	double		StepSize;
 	/** Number of segments to render (steps)							*/
-	int32			Steps;
+	int32		Steps;
 	/** The 'extra' amount to travel (partial segment)					*/
 	float		TravelRatio;
 
 	/** The number of triangles to render for this beam					*/
-	int32			TriangleCount;
+	int32		TriangleCount;
 
 	/**
 	 *	Type and indexing flags
@@ -699,12 +699,12 @@ struct FBeamParticleModifierPayloadData
 	uint32	bScaleTangent:1;
 	uint32	bModifyStrength:1;
 	uint32	bScaleStrength:1;
-	FVector3f		Position;
-	FVector3f		Tangent;
-	float		Strength;
+	FVector	Position;
+	FVector	Tangent;
+	float	Strength;
 
 	// Helper functions
-	FORCEINLINE void UpdatePosition(FVector3f& Value)
+	FORCEINLINE void UpdatePosition(FVector& Value)
 	{
 		if (bModifyPosition == true)
 		{
@@ -719,16 +719,16 @@ struct FBeamParticleModifierPayloadData
 		}
 	}
 
-	FORCEINLINE void UpdateTangent(FVector3f& Value, bool bAbsolute)
+	FORCEINLINE void UpdateTangent(FVector& Value, bool bAbsolute)
 	{
 		if (bModifyTangent == true)
 		{
-			FVector3f ModTangent;
+			FVector ModTangent;
 			if (bAbsolute == false)
 			{
 				// Transform the modified tangent so it is relative to the real tangent
-				const FQuat RotQuat = FQuat::FindBetweenNormals(FVector(1.0f, 0.0f, 0.0f), FVector(Value));
-				ModTangent = FVector3f(RotQuat.RotateVector(FVector(Tangent)));
+				const FQuat RotQuat = FQuat::FindBetweenNormals(FVector(1.0f, 0.0f, 0.0f), Value);
+				ModTangent = RotQuat.RotateVector(Tangent);
 			}
 			else
 			{
