Fixed missing changes on DataLayerInstanceWithAsset following 24341161 where outer is now accessed through UDataLayerInstance::GetOuterWorldDataLayers.

diff --git a/Engine/Source/Runtime/Engine/Private/WorldPartition/DataLayer/DataLayerInstanceWithAsset.cpp b/Engine/Source/Runtime/Engine/Private/WorldPartition/DataLayer/DataLayerInstanceWithAsset.cpp
--- a/Engine/Source/Runtime/Engine/Private/WorldPartition/DataLayer/DataLayerInstanceWithAsset.cpp
+++ b/Engine/Source/Runtime/Engine/Private/WorldPartition/DataLayer/DataLayerInstanceWithAsset.cpp
@@ -24,7 +24,7 @@ FName UDataLayerInstanceWithAsset::MakeName(const UDataLayerAsset* DeprecatedDat
 
 void UDataLayerInstanceWithAsset::OnCreated(const UDataLayerAsset* Asset)
 {
-	check(!GetOuterAWorldDataLayers()->HasDeprecatedDataLayers() || IsRunningCommandlet());
+	check(!GetOuterWorldDataLayers()->HasDeprecatedDataLayers() || IsRunningCommandlet());
 
 	Modify(/*bAlwaysMarkDirty*/false);
 
@@ -40,7 +40,7 @@ bool UDataLayerInstanceWithAsset::IsReadOnly() const
 	{
 		return true;
 	}
-	return GetOuterAWorldDataLayers()->IsSubWorldDataLayers();
+	return GetOuterWorldDataLayers()->IsSubWorldDataLayers();
 }
 
 bool UDataLayerInstanceWithAsset::IsLocked() const
@@ -120,7 +120,7 @@ void UDataLayerInstanceWithAsset::PostEditChangeProperty(FPropertyChangedEvent&
 
 	if (PropertyName == GET_MEMBER_NAME_CHECKED(UDataLayerInstanceWithAsset, DataLayerAsset))
 	{
-		GetOuterAWorldDataLayers()->ResolveActorDescContainers();
+		GetOuterWorldDataLayers()->ResolveActorDescContainers();
 	}
 }
 
@@ -135,7 +135,7 @@ void UDataLayerInstanceWithAsset::PostEditUndo()
 	Super::PostEditUndo();
 	if (CachedDataLayerAsset != DataLayerAsset)
 	{
-		GetOuterAWorldDataLayers()->ResolveActorDescContainers();
+		GetOuterWorldDataLayers()->ResolveActorDescContainers();
 	}
 	CachedDataLayerAsset = nullptr;
 }
diff --git a/Engine/Source/Runtime/Engine/Public/WorldPartition/DataLayer/DataLayerInstanceWithAsset.h b/Engine/Source/Runtime/Engine/Public/WorldPartition/DataLayer/DataLayerInstanceWithAsset.h
--- a/Engine/Source/Runtime/Engine/Public/WorldPartition/DataLayer/DataLayerInstanceWithAsset.h
+++ b/Engine/Source/Runtime/Engine/Public/WorldPartition/DataLayer/DataLayerInstanceWithAsset.h
@@ -6,7 +6,7 @@
 
 #include "DataLayerInstanceWithAsset.generated.h"
 
-UCLASS(Config = Engine, PerObjectConfig, Within = WorldDataLayers, AutoCollapseCategories = ("Data Layer|Advanced"), AutoExpandCategories = ("Data Layer|Editor", "Data Layer|Advanced|Runtime"))
+UCLASS(Config = Engine, PerObjectConfig, AutoCollapseCategories = ("Data Layer|Advanced"), AutoExpandCategories = ("Data Layer|Editor", "Data Layer|Advanced|Runtime"))
 class ENGINE_API UDataLayerInstanceWithAsset : public UDataLayerInstance
 {
 	GENERATED_UCLASS_BODY()
