Return a nullptr DataLayerInstance when passing a nullptr DataLayerAsset

diff --git a/Engine/Source/Runtime/Engine/Private/WorldPartition/DataLayer/WorldDataLayers.cpp b/Engine/Source/Runtime/Engine/Private/WorldPartition/DataLayer/WorldDataLayers.cpp
--- a/Engine/Source/Runtime/Engine/Private/WorldPartition/DataLayer/WorldDataLayers.cpp
+++ b/Engine/Source/Runtime/Engine/Private/WorldPartition/DataLayer/WorldDataLayers.cpp
@@ -690,6 +690,11 @@ const UDataLayerInstance* AWorldDataLayers::GetDataLayerInstanceFromAssetName(co
 
 const UDataLayerInstance* AWorldDataLayers::GetDataLayerInstance(const UDataLayerAsset* InDataLayerAsset) const
 {
+	if (InDataLayerAsset == nullptr)
+	{
+		return nullptr;
+	}
+
 #if WITH_EDITOR	
 	for (const UDataLayerInstance* DataLayerInstance : DataLayerInstances)
 	{
@@ -703,12 +708,9 @@ const UDataLayerInstance* AWorldDataLayers::GetDataLayerInstance(const UDataLaye
 		}
 	}
 #else
-	if (InDataLayerAsset != nullptr)
+	if (const UDataLayerInstance* const* FoundDataLayerInstance = AssetNameToInstance.Find(InDataLayerAsset->GetFullName()))
 	{
-		if (const UDataLayerInstance* const* FoundDataLayerInstance = AssetNameToInstance.Find(InDataLayerAsset->GetFullName()))
-		{
-			return *FoundDataLayerInstance;
-		}
+		return *FoundDataLayerInstance;
 	}
 #endif
 
