Remove ensure around exceeding the maximum number of bits in FCharacterNetworkSerializationPackedBits::NetSerialize
#rnx

diff --git a/Engine/Source/Runtime/Engine/Private/Components/CharacterMovementComponent.cpp b/Engine/Source/Runtime/Engine/Private/Components/CharacterMovementComponent.cpp
--- a/Engine/Source/Runtime/Engine/Private/Components/CharacterMovementComponent.cpp
+++ b/Engine/Source/Runtime/Engine/Private/Components/CharacterMovementComponent.cpp
@@ -9092,10 +9092,10 @@ void UCharacterMovementComponent::ServerMovePacked_ServerReceive(const FCharacte
 	}
 
 	const int32 NumBits = PackedBits.DataBits.Num();
-	if (!ensureMsgf(NumBits <= CharacterMovementCVars::NetPackedMovementMaxBits, TEXT("ServerMovePacked_ServerReceive: NumBits (%d) exceeds CharacterMovementCVars::NetPackedMovementMaxBits (%d)"), NumBits, CharacterMovementCVars::NetPackedMovementMaxBits))
+	if (NumBits > CharacterMovementCVars::NetPackedMovementMaxBits)
 	{
 		// Protect against bad data that could cause server to allocate way too much memory.
-		devCode(UE_LOG(LogNetPlayerMovement, Error, TEXT("ServerMovePacked_ServerReceive: NumBits (%d) exceeds allowable limit!"), NumBits));
+		UE_LOG(LogNetPlayerMovement, Error, TEXT("ServerMovePacked_ServerReceive (%s): Dropping move due to NumBits (%d) exceeding allowable limit (%d). See NetPackedMovementMaxBits."), *GetNameSafe(GetOwner()), NumBits, CharacterMovementCVars::NetPackedMovementMaxBits);
 		return;
 	}
 
