#ue5 - Make AActor::SetActorTickEnabled() virtual (matching already-virtual UActorComponent::SetComponentTickEnabled()).

diff --git a/Engine/Source/Runtime/Engine/Classes/GameFramework/Actor.h b/Engine/Source/Runtime/Engine/Classes/GameFramework/Actor.h
--- a/Engine/Source/Runtime/Engine/Classes/GameFramework/Actor.h
+++ b/Engine/Source/Runtime/Engine/Classes/GameFramework/Actor.h
@@ -2562,7 +2562,7 @@ public:
 	 * @param	bEnabled	Whether it should be enabled or not
 	 */
 	UFUNCTION(BlueprintCallable, Category="Actor|Tick")
-	void SetActorTickEnabled(bool bEnabled);
+	virtual void SetActorTickEnabled(bool bEnabled);
 
 	/**  Returns whether this actor has tick enabled or not	 */
 	UFUNCTION(BlueprintCallable, Category="Actor|Tick")
