Fix some Texture CompressionSettings not getting R to Gray editor visualization
Factor out to ShouldUseGreyScaleEditorVisualization
Fixes TC_DistanceFieldFont as used in font preview for Offline fonts

#10104

diff --git a/Engine/Source/Runtime/Engine/Classes/Engine/TextureDefines.h b/Engine/Source/Runtime/Engine/Classes/Engine/TextureDefines.h
--- a/Engine/Source/Runtime/Engine/Classes/Engine/TextureDefines.h
+++ b/Engine/Source/Runtime/Engine/Classes/Engine/TextureDefines.h
@@ -499,6 +499,20 @@ static FORCEINLINE bool IsUncompressed(TextureCompressionSettings CompressionSet
 		);
 }
 
+static FORCEINLINE bool ShouldUseGreyScaleEditorVisualization(TextureCompressionSettings CompressionSettings)
+{
+	// these formats should do R -> RGB red to gray replication in Editor viz
+	return (CompressionSettings == TC_Grayscale ||
+			CompressionSettings == TC_Alpha ||
+			CompressionSettings == TC_Displacementmap ||
+			CompressionSettings == TC_DistanceFieldFont );
+
+	// ?? maybe these too ??
+	//		CompressionSettings == TC_HalfFloat ||
+	//		CompressionSettings == TC_SingleFloat
+}
+
+
 } // TextureDefines
 } // UE
 
diff --git a/Engine/Source/Runtime/Engine/Private/Rendering/StreamableTextureResource.cpp b/Engine/Source/Runtime/Engine/Private/Rendering/StreamableTextureResource.cpp
--- a/Engine/Source/Runtime/Engine/Private/Rendering/StreamableTextureResource.cpp
+++ b/Engine/Source/Runtime/Engine/Private/Rendering/StreamableTextureResource.cpp
@@ -96,7 +96,7 @@ FStreamableTextureResource::FStreamableTextureResource(UTexture* InOwner, const
 	// HDR images are stored in linear but still require gamma correction to display correctly.
 	bIgnoreGammaConversions = !InOwner->SRGB && !IsHDR(PixelFormat);
 	bSRGB = InOwner->SRGB;
-	bGreyScaleFormat = ( InOwner->CompressionSettings == TC_Grayscale || InOwner->CompressionSettings == TC_Alpha );
+	bGreyScaleFormat = UE::TextureDefines::ShouldUseGreyScaleEditorVisualization( InOwner->CompressionSettings );
 
 	Filter = (ESamplerFilter)UDeviceProfileManager::Get().GetActiveProfile()->GetTextureLODSettings()->GetSamplerFilter(InOwner);
 	AddressU = InOwner->GetTextureAddressX() == TA_Wrap ? AM_Wrap : (InOwner->GetTextureAddressX() == TA_Clamp ? AM_Clamp : AM_Mirror);
diff --git a/Engine/Source/Runtime/Engine/Private/Texture2D.cpp b/Engine/Source/Runtime/Engine/Private/Texture2D.cpp
--- a/Engine/Source/Runtime/Engine/Private/Texture2D.cpp
+++ b/Engine/Source/Runtime/Engine/Private/Texture2D.cpp
@@ -1261,7 +1261,7 @@ FVirtualTexture2DResource::FVirtualTexture2DResource(const UTexture2D* InOwner,
 	FirstMipToUse = FMath::Min((int32)MaxMip, InFirstMipToUse);
 
 	bSRGB = InOwner->SRGB;
-	bGreyScaleFormat = ( InOwner->CompressionSettings == TC_Grayscale || InOwner->CompressionSettings == TC_Alpha );
+	bGreyScaleFormat = UE::TextureDefines::ShouldUseGreyScaleEditorVisualization( InOwner->CompressionSettings );
 
 	// Initialize this resource FeatureLevel, so it gets re-created on FeatureLevel changes
 	SetFeatureLevel(GMaxRHIFeatureLevel);
diff --git a/Engine/Source/Runtime/Engine/Private/TextureCube.cpp b/Engine/Source/Runtime/Engine/Private/TextureCube.cpp
--- a/Engine/Source/Runtime/Engine/Private/TextureCube.cpp
+++ b/Engine/Source/Runtime/Engine/Private/TextureCube.cpp
@@ -491,8 +491,7 @@ public:
 		);
 		SamplerStateRHI = GetOrCreateSamplerState(SamplerStateInitializer);
 
-		// Set the greyscale format flag appropriately.
-		bGreyScaleFormat = ( Owner->CompressionSettings == TC_Grayscale || Owner->CompressionSettings == TC_Alpha );
+		bGreyScaleFormat = UE::TextureDefines::ShouldUseGreyScaleEditorVisualization( Owner->CompressionSettings );
 	}
 
 	virtual void ReleaseRHI() override
diff --git a/Engine/Source/Runtime/Engine/Private/TextureCubeArray.cpp b/Engine/Source/Runtime/Engine/Private/TextureCubeArray.cpp
--- a/Engine/Source/Runtime/Engine/Private/TextureCubeArray.cpp
+++ b/Engine/Source/Runtime/Engine/Private/TextureCubeArray.cpp
@@ -170,9 +170,8 @@ public:
 			AM_Clamp
 		);
 		SamplerStateRHI = GetOrCreateSamplerState(SamplerStateInitializer);
-
-		// Set the greyscale format flag appropriately.
-		bGreyScaleFormat = ( Owner->CompressionSettings == TC_Grayscale || Owner->CompressionSettings == TC_Alpha );
+		
+		bGreyScaleFormat = UE::TextureDefines::ShouldUseGreyScaleEditorVisualization( Owner->CompressionSettings );
 	}
 
 	virtual void ReleaseRHI() override
