Moved API macro from types to members to fix compiling for Windows with Clang

An optimization to UE_LOG exposed an issue in Clang when targeting Windows. Log statements in inline functions of imported types fail to compile.

#rnx

diff --git a/Engine/Plugins/Online/OnlineSubsystemUtils/Source/OnlineSubsystemUtils/Public/OnlineBeacon.h b/Engine/Plugins/Online/OnlineSubsystemUtils/Source/OnlineSubsystemUtils/Public/OnlineBeacon.h
--- a/Engine/Plugins/Online/OnlineSubsystemUtils/Source/OnlineSubsystemUtils/Public/OnlineBeacon.h
+++ b/Engine/Plugins/Online/OnlineSubsystemUtils/Source/OnlineSubsystemUtils/Public/OnlineBeacon.h
@@ -30,25 +30,25 @@ namespace EBeaconState
  * Base class for beacon communication (Unreal Networking, but outside normal gameplay traffic)
  */
 UCLASS(transient, config=Engine, notplaceable)
-class ONLINESUBSYSTEMUTILS_API AOnlineBeacon : public AActor, public FNetworkNotify
+class AOnlineBeacon : public AActor, public FNetworkNotify
 {
 	GENERATED_UCLASS_BODY()
 
 	//~ Begin AActor Interface
-	virtual void OnActorChannelOpen(FInBunch& InBunch, UNetConnection* Connection) override;
-	virtual bool IsRelevancyOwnerFor(const AActor* ReplicatedActor, const AActor* ActorOwner, const AActor* ConnectionActor) const override;
-	virtual bool IsNetRelevantFor(const AActor* RealViewer, const AActor* ViewTarget, const FVector& SrcLocation) const override;
+	ONLINESUBSYSTEMUTILS_API virtual void OnActorChannelOpen(FInBunch& InBunch, UNetConnection* Connection) override;
+	ONLINESUBSYSTEMUTILS_API virtual bool IsRelevancyOwnerFor(const AActor* ReplicatedActor, const AActor* ActorOwner, const AActor* ConnectionActor) const override;
+	ONLINESUBSYSTEMUTILS_API virtual bool IsNetRelevantFor(const AActor* RealViewer, const AActor* ViewTarget, const FVector& SrcLocation) const override;
 	virtual const AActor* GetNetOwner() const override { return nullptr; }
 	virtual UNetConnection* GetNetConnection() const override { return nullptr; }
 	virtual bool IsLevelBoundsRelevant() const override { return false; }
-	virtual void EndPlay(const EEndPlayReason::Type EndPlayReason);
+	ONLINESUBSYSTEMUTILS_API virtual void EndPlay(const EEndPlayReason::Type EndPlayReason);
 	//~ End AActor Interface
 
 	//~ Begin FNetworkNotify Interface
-	virtual EAcceptConnection::Type NotifyAcceptingConnection() override;
-	virtual void NotifyAcceptedConnection(UNetConnection* Connection) override;
-	virtual bool NotifyAcceptingChannel(UChannel* Channel) override;
-	virtual void NotifyControlMessage(UNetConnection* Connection, uint8 MessageType, FInBunch& Bunch) override;
+	ONLINESUBSYSTEMUTILS_API virtual EAcceptConnection::Type NotifyAcceptingConnection() override;
+	ONLINESUBSYSTEMUTILS_API virtual void NotifyAcceptedConnection(UNetConnection* Connection) override;
+	ONLINESUBSYSTEMUTILS_API virtual bool NotifyAcceptingChannel(UChannel* Channel) override;
+	ONLINESUBSYSTEMUTILS_API virtual void NotifyControlMessage(UNetConnection* Connection, uint8 MessageType, FInBunch& Bunch) override;
 	//~ End FNetworkNotify Interface
 	
     /**
@@ -66,7 +66,7 @@ class ONLINESUBSYSTEMUTILS_API AOnlineBeacon : public AActor, public FNetworkNot
 	 * @param	FailureType	the type of error
 	 * @param	ErrorString	additional string detailing the error
 	 */
-	virtual void HandleNetworkFailure(UWorld* World, UNetDriver *NetDriver, ENetworkFailure::Type FailureType, const FString& ErrorString);
+	ONLINESUBSYSTEMUTILS_API virtual void HandleNetworkFailure(UWorld* World, UNetDriver *NetDriver, ENetworkFailure::Type FailureType, const FString& ErrorString);
 
 	/**
 	 * Set Beacon state 
@@ -92,11 +92,11 @@ class ONLINESUBSYSTEMUTILS_API AOnlineBeacon : public AActor, public FNetworkNot
 	}
 
 	/** Beacon cleanup and net driver destruction */
-	virtual void DestroyBeacon();
+	ONLINESUBSYSTEMUTILS_API virtual void DestroyBeacon();
 
 protected:
 
-	void CleanupNetDriver();
+	ONLINESUBSYSTEMUTILS_API void CleanupNetDriver();
 
 	/** Time beacon will wait to establish a connection with the beacon host */
 	UPROPERTY(Config)
@@ -118,13 +118,13 @@ protected:
 	FName NetDriverDefinitionName;
 
 	/** Common initialization for all beacon types */
-	virtual bool InitBase();
+	ONLINESUBSYSTEMUTILS_API virtual bool InitBase();
 
 	/** Notification that failure needs to be handled */
-	virtual void OnFailure();
+	ONLINESUBSYSTEMUTILS_API virtual void OnFailure();
 
 	/** overridden to return that player controllers are capable of RPCs */
-	virtual bool HasNetOwner() const override;
+	ONLINESUBSYSTEMUTILS_API virtual bool HasNetOwner() const override;
 };
 
 #if UE_ENABLE_INCLUDE_ORDER_DEPRECATED_IN_5_2
diff --git a/Engine/Plugins/Runtime/Metasound/Source/MetasoundFrontend/Public/MetasoundAudioBuffer.h b/Engine/Plugins/Runtime/Metasound/Source/MetasoundFrontend/Public/MetasoundAudioBuffer.h
--- a/Engine/Plugins/Runtime/Metasound/Source/MetasoundFrontend/Public/MetasoundAudioBuffer.h
+++ b/Engine/Plugins/Runtime/Metasound/Source/MetasoundFrontend/Public/MetasoundAudioBuffer.h
@@ -20,7 +20,7 @@ namespace Metasound
 	 * arguments via an implicit conversion operator which exposes the underlying FAlignedFloatBuffer container.
 	 *
 	 */
-	class METASOUNDFRONTEND_API FAudioBuffer
+	class FAudioBuffer
 	{
 		public:
 			/** Create an FAudioBuffer with a specific number of samples.
diff --git a/Engine/Plugins/Runtime/ZoneGraph/Source/ZoneGraph/Public/ZoneGraphAStar.h b/Engine/Plugins/Runtime/ZoneGraph/Source/ZoneGraph/Public/ZoneGraphAStar.h
--- a/Engine/Plugins/Runtime/ZoneGraph/Source/ZoneGraph/Public/ZoneGraphAStar.h
+++ b/Engine/Plugins/Runtime/ZoneGraph/Source/ZoneGraph/Public/ZoneGraphAStar.h
@@ -53,7 +53,7 @@ struct ZONEGRAPH_API FZoneGraphAStarNode : public FGraphAStarDefaultNode<FZoneGr
 };
 
 /** Context for FGraphAStar::FindPath() */
-struct ZONEGRAPH_API FZoneGraphPathFilter
+struct FZoneGraphPathFilter
 {
 	// @todo: rename FZoneGraphPathfindContext?
 
diff --git a/Engine/Plugins/VirtualProduction/DMX/DMXProtocol/Source/DMXProtocol/Public/IO/DMXInputPortReference.h b/Engine/Plugins/VirtualProduction/DMX/DMXProtocol/Source/DMXProtocol/Public/IO/DMXInputPortReference.h
--- a/Engine/Plugins/VirtualProduction/DMX/DMXProtocol/Source/DMXProtocol/Public/IO/DMXInputPortReference.h
+++ b/Engine/Plugins/VirtualProduction/DMX/DMXProtocol/Source/DMXProtocol/Public/IO/DMXInputPortReference.h
@@ -12,7 +12,7 @@
 
 /** Reference of an input port */
 USTRUCT(BlueprintType)
-struct DMXPROTOCOL_API FDMXInputPortReference
+struct FDMXInputPortReference
 {
 	GENERATED_BODY()
 
diff --git a/Engine/Plugins/VirtualProduction/DMX/DMXProtocol/Source/DMXProtocol/Public/IO/DMXOutputPortReference.h b/Engine/Plugins/VirtualProduction/DMX/DMXProtocol/Source/DMXProtocol/Public/IO/DMXOutputPortReference.h
--- a/Engine/Plugins/VirtualProduction/DMX/DMXProtocol/Source/DMXProtocol/Public/IO/DMXOutputPortReference.h
+++ b/Engine/Plugins/VirtualProduction/DMX/DMXProtocol/Source/DMXProtocol/Public/IO/DMXOutputPortReference.h
@@ -12,7 +12,7 @@
 
 /** Reference of an input port */
 USTRUCT(BlueprintType)
-struct DMXPROTOCOL_API FDMXOutputPortReference
+struct FDMXOutputPortReference
 {
 	GENERATED_BODY()
 
diff --git a/Engine/Source/Editor/PropertyEditor/Public/NumericPropertyParams.h b/Engine/Source/Editor/PropertyEditor/Public/NumericPropertyParams.h
--- a/Engine/Source/Editor/PropertyEditor/Public/NumericPropertyParams.h
+++ b/Engine/Source/Editor/PropertyEditor/Public/NumericPropertyParams.h
@@ -5,7 +5,7 @@
 #include "CoreMinimal.h"
 
 template<typename NumericType>
-struct PROPERTYEDITOR_API TNumericPropertyParams
+struct TNumericPropertyParams
 {
 public:
 	DECLARE_DELEGATE_RetVal_OneParam(FString, FMetaDataGetter, FName);
diff --git a/Engine/Source/Runtime/AIModule/Classes/EnvironmentQuery/EnvQueryTypes.h b/Engine/Source/Runtime/AIModule/Classes/EnvironmentQuery/EnvQueryTypes.h
--- a/Engine/Source/Runtime/AIModule/Classes/EnvironmentQuery/EnvQueryTypes.h
+++ b/Engine/Source/Runtime/AIModule/Classes/EnvironmentQuery/EnvQueryTypes.h
@@ -1045,9 +1045,9 @@ public:
 		int32 CurrentItem;
 	};
 
-	struct AIMODULE_API FItemIterator : public FConstItemIterator
+	struct FItemIterator : public FConstItemIterator
 	{
-		FItemIterator(const UEnvQueryTest* QueryTest, FEnvQueryInstance& QueryInstance, int32 StartingItemIndex = INDEX_NONE);
+		AIMODULE_API FItemIterator(const UEnvQueryTest* QueryTest, FEnvQueryInstance& QueryInstance, int32 StartingItemIndex = INDEX_NONE);
 
 		~FItemIterator()
 		{
@@ -1212,8 +1212,8 @@ public:
 			bForced = false;
 		}
 
-		void HandleFailedTestResult();
-		void StoreTestResult();
+		AIMODULE_API void HandleFailedTestResult();
+		AIMODULE_API void StoreTestResult();
 
 		FORCEINLINE void SetScoreInternal(float Score)
 		{
diff --git a/Engine/Source/Runtime/AudioLink/AudioMixerPlatformAudioLink/AudioMixerPlatformAudioLink.Build.cs b/Engine/Source/Runtime/AudioLink/AudioMixerPlatformAudioLink/AudioMixerPlatformAudioLink.Build.cs
--- a/Engine/Source/Runtime/AudioLink/AudioMixerPlatformAudioLink/AudioMixerPlatformAudioLink.Build.cs
+++ b/Engine/Source/Runtime/AudioLink/AudioMixerPlatformAudioLink/AudioMixerPlatformAudioLink.Build.cs
@@ -16,7 +16,8 @@ public class AudioMixerPlatformAudioLink: ModuleRules
 				"Engine",
 				"AudioMixer",
 				"AudioMixerCore",
-				"AudioLinkEngine"
+				"AudioLinkEngine",
+				"SignalProcessing",
 		});
 					
 		if (Target.bCompileAgainstEngine) 
diff --git a/Engine/Source/Runtime/AudioMixerCore/Public/AudioMixer.h b/Engine/Source/Runtime/AudioMixerCore/Public/AudioMixer.h
--- a/Engine/Source/Runtime/AudioMixerCore/Public/AudioMixer.h
+++ b/Engine/Source/Runtime/AudioMixerCore/Public/AudioMixer.h
@@ -432,7 +432,7 @@ namespace Audio
 
 
 	/** Abstract interface for mixer platform. */
-	class AUDIOMIXERCORE_API IAudioMixerPlatformInterface : public FRunnable,
+	class IAudioMixerPlatformInterface : public FRunnable,
 														public FSingleThreadRunnable,
 														public IAudioMixerDeviceChangedListener
 	{
@@ -440,7 +440,7 @@ namespace Audio
 	public: // Virtual functions
 		
 		/** Virtual destructor. */
-		virtual ~IAudioMixerPlatformInterface();
+		AUDIOMIXERCORE_API virtual ~IAudioMixerPlatformInterface();
 
 		/** Returns the platform API name. */
 		virtual FString GetPlatformApi() const = 0;
@@ -476,7 +476,7 @@ namespace Audio
 		 * On most platforms, this index may be invalidated if any devices are added or removed.
 		 * Returns INDEX_NONE if no mapping is found
 		 */
-		virtual int32 GetIndexForDevice(const FString& InDeviceName);
+		AUDIOMIXERCORE_API virtual int32 GetIndexForDevice(const FString& InDeviceName);
 
 		/** Gets the platform specific audio settings. */
 		virtual FAudioPlatformSettings GetPlatformSettings() const = 0;
@@ -532,7 +532,7 @@ namespace Audio
 		virtual FString GetDefaultDeviceName() = 0;
 
 		// Helper function to gets the channel map type at the given index.
-		static bool GetChannelTypeAtIndex(const int32 Index, EAudioMixerChannel::Type& OutType);
+		AUDIOMIXERCORE_API static bool GetChannelTypeAtIndex(const int32 Index, EAudioMixerChannel::Type& OutType);
 
         // Function to stop all audio from rendering. Used on mobile platforms which can suspend the application.
         virtual void SuspendContext() {}
@@ -548,7 +548,7 @@ namespace Audio
 
 	public: // Public Functions
 		//~ Begin FRunnable
-		uint32 Run() override;
+		AUDIOMIXERCORE_API uint32 Run() override;
 		//~ End FRunnable
 
 		/**
@@ -558,34 +558,34 @@ namespace Audio
 		virtual class FSingleThreadRunnable* GetSingleThreadInterface() override { return this; }
 
 		//~ Begin FSingleThreadRunnable Interface
-		virtual void Tick() override;
+		AUDIOMIXERCORE_API virtual void Tick() override;
 		//~ End FSingleThreadRunnable Interface
 
 		/** Constructor. */
-		IAudioMixerPlatformInterface();
+		AUDIOMIXERCORE_API IAudioMixerPlatformInterface();
 
 		/** Retrieves the next generated buffer and feeds it to the platform mixer output stream. */
-		void ReadNextBuffer();
+		AUDIOMIXERCORE_API void ReadNextBuffer();
 
 		/** Reset the fade state (use if reusing audio platform interface, e.g. in main audio device. */
-		virtual void FadeIn();
+		AUDIOMIXERCORE_API virtual void FadeIn();
 
 		/** Start a fadeout. Prevents pops during shutdown. */
-		virtual void FadeOut();
+		AUDIOMIXERCORE_API virtual void FadeOut();
 
 		/** Returns the last error generated. */
 		FString GetLastError() const { return LastError; }
 
 		/** This is called after InitializeHardware() is called. */
-		void PostInitializeHardware();
+		AUDIOMIXERCORE_API void PostInitializeHardware();
 
 	protected:
 		
 		// Run the "main" audio device
-		uint32 MainAudioDeviceRun();
+		AUDIOMIXERCORE_API uint32 MainAudioDeviceRun();
 		
 		// Wrapper around the thread Run. This is virtualized so a platform can fundamentally override the render function.
-		virtual uint32 RunInternal();
+		AUDIOMIXERCORE_API virtual uint32 RunInternal();
 
 		/** Is called when an error, warning or log is generated. */
 		inline void AudioMixerPlatformLogOnce(const FString& LogDetails, const FString& FileName, int32 LineNumber, ELogVerbosity::Type InVerbosity = ELogVerbosity::Error)
@@ -653,29 +653,29 @@ namespace Audio
 
 
 		/** Start generating audio from our mixer. */
-		void BeginGeneratingAudio();
+		AUDIOMIXERCORE_API void BeginGeneratingAudio();
 
 		/** Stops the render thread from generating audio. */
-		void StopGeneratingAudio();
+		AUDIOMIXERCORE_API void StopGeneratingAudio();
 
 		// Deprecated - use ApplyPrimaryAttenuation
 		UE_DEPRECATED(5.1, "ApplyMasterAttenuation is deprecated, please use ApplyPrimaryAttenuation instead.")
-		void ApplyMasterAttenuation(TArrayView<const uint8>& InOutPoppedAudio);
+		AUDIOMIXERCORE_API void ApplyMasterAttenuation(TArrayView<const uint8>& InOutPoppedAudio);
 
 		/** Performs buffer fades for shutdown/startup of audio mixer. */
-		void ApplyPrimaryAttenuation(TArrayView<const uint8>& InOutPoppedAudio);
+		AUDIOMIXERCORE_API void ApplyPrimaryAttenuation(TArrayView<const uint8>& InOutPoppedAudio);
 
 		template<typename BufferType>
-		void ApplyAttenuationInternal(TArrayView<BufferType>& InOutBuffer);
+		AUDIOMIXERCORE_API void ApplyAttenuationInternal(TArrayView<BufferType>& InOutBuffer);
 
 		/** When called, spins up a thread to start consuming output when no audio device is available. */
-		void StartRunningNullDevice();
+		AUDIOMIXERCORE_API void StartRunningNullDevice();
 
 		/** When called, terminates the null device. */
-		void StopRunningNullDevice();
+		AUDIOMIXERCORE_API void StopRunningNullDevice();
 		
 		/** Called by platform specific logic to pre-create or create the null renderer thread  */
-		void CreateNullDeviceThread(const TFunction<void()> InCallback, float InBufferDuration, bool bShouldPauseOnStart);
+		AUDIOMIXERCORE_API void CreateNullDeviceThread(const TFunction<void()> InCallback, float InBufferDuration, bool bShouldPauseOnStart);
 
 	protected:
 
@@ -732,7 +732,7 @@ namespace Audio
 		FThreadSafeBool bIsGeneratingAudio;
 
 		/** A Counter to provide the next unique id. */
-		static FThreadSafeCounter NextInstanceID;
+		AUDIOMIXERCORE_API static FThreadSafeCounter NextInstanceID;
 
 		/** A Unique ID Identifying this instance. Mostly used for logging. */ 
 		const int32 InstanceID{ -1 };
diff --git a/Engine/Source/Runtime/Core/Public/Async/AsyncFileHandle.h b/Engine/Source/Runtime/Core/Public/Async/AsyncFileHandle.h
--- a/Engine/Source/Runtime/Core/Public/Async/AsyncFileHandle.h
+++ b/Engine/Source/Runtime/Core/Public/Async/AsyncFileHandle.h
@@ -20,7 +20,7 @@ DECLARE_DWORD_ACCUMULATOR_STAT_EXTERN(TEXT("Num Async File Requests"), STAT_Asyn
 class IAsyncReadRequest;
 typedef TFunction<void(bool bWasCancelled, IAsyncReadRequest*)> FAsyncFileCallBack;
 
-class CORE_API IAsyncReadRequest
+class IAsyncReadRequest
 {
 protected:
 	union
diff --git a/Engine/Source/Runtime/Core/Public/HAL/IPlatformFileManagedStorageWrapper.h b/Engine/Source/Runtime/Core/Public/HAL/IPlatformFileManagedStorageWrapper.h
--- a/Engine/Source/Runtime/Core/Public/HAL/IPlatformFileManagedStorageWrapper.h
+++ b/Engine/Source/Runtime/Core/Public/HAL/IPlatformFileManagedStorageWrapper.h
@@ -366,7 +366,7 @@ public:
 
 	FPersistentStorageManager();
 
-	void CORE_API Initialize()
+	void Initialize()
 	{
 		if (bInitialized)
 		{
@@ -655,7 +655,7 @@ private:
 };
 
 // Only write handle 
-class CORE_API FManagedStorageFileWriteHandle : public IFileHandle
+class FManagedStorageFileWriteHandle : public IFileHandle
 {
 private:
 	static bool IsReady()
diff --git a/Engine/Source/Runtime/Core/Public/Stats/StatsFile.h b/Engine/Source/Runtime/Core/Public/Stats/StatsFile.h
--- a/Engine/Source/Runtime/Core/Public/Stats/StatsFile.h
+++ b/Engine/Source/Runtime/Core/Public/Stats/StatsFile.h
@@ -617,7 +617,7 @@ protected:
 /**
  * Class for maintaining state of receiving a stream of stat messages
  */
-struct CORE_API FStatsReadStream
+struct FStatsReadStream
 {
 public:
 	/** Stats stream header. */
diff --git a/Engine/Source/Runtime/CoreUObject/Public/UObject/Class.h b/Engine/Source/Runtime/CoreUObject/Public/UObject/Class.h
--- a/Engine/Source/Runtime/CoreUObject/Public/UObject/Class.h
+++ b/Engine/Source/Runtime/CoreUObject/Public/UObject/Class.h
@@ -1936,10 +1936,10 @@ ENUM_CLASS_FLAGS(EGetByNameFlags)
 //
 // Reflection data for an enumeration.
 //
-class COREUOBJECT_API UEnum : public UField
+class UEnum : public UField
 {
-	DECLARE_CASTED_CLASS_INTRINSIC_NO_CTOR(UEnum, UField, 0, TEXT("/Script/CoreUObject"), CASTCLASS_UEnum, NO_API)
-	UEnum(const FObjectInitializer& ObjectInitialzer);
+	DECLARE_CASTED_CLASS_INTRINSIC_NO_CTOR(UEnum, UField, 0, TEXT("/Script/CoreUObject"), CASTCLASS_UEnum, COREUOBJECT_API)
+	COREUOBJECT_API UEnum(const FObjectInitializer& ObjectInitialzer);
 
 public:
 	/** How this enum is declared in C++, affects the internal naming of enum values */
@@ -1979,31 +1979,31 @@ public:
 	}
 
 	/** Gets enum name by index in Names array. Returns NAME_None if Index is not valid. */
-	FName GetNameByIndex(int32 Index) const;
+	COREUOBJECT_API FName GetNameByIndex(int32 Index) const;
 
 	/** Gets index of name in enum, returns INDEX_NONE and optionally errors when name is not found. This is faster than ByNameString if the FName is exact, but will fall back if needed */
-	int32 GetIndexByName(FName InName, EGetByNameFlags Flags = EGetByNameFlags::None) const;
+	COREUOBJECT_API int32 GetIndexByName(FName InName, EGetByNameFlags Flags = EGetByNameFlags::None) const;
 
 	/** Gets enum name by value. Returns NAME_None if value is not found. */
-	FName GetNameByValue(int64 InValue) const;
+	COREUOBJECT_API FName GetNameByValue(int64 InValue) const;
 
 	/** Gets enum value by name, returns INDEX_NONE and optionally errors when name is not found. This is faster than ByNameString if the FName is exact, but will fall back if needed */
-	int64 GetValueByName(FName InName, EGetByNameFlags Flags = EGetByNameFlags::None) const;
+	COREUOBJECT_API int64 GetValueByName(FName InName, EGetByNameFlags Flags = EGetByNameFlags::None) const;
 
 	/** Returns the short name at the enum index, returns empty string if invalid */
-	FString GetNameStringByIndex(int32 InIndex) const;
+	COREUOBJECT_API FString GetNameStringByIndex(int32 InIndex) const;
 
 	/** Gets index of name in enum, returns INDEX_NONE and optionally errors when name is not found. Handles full or short names. */
-	int32 GetIndexByNameString(const FString& SearchString, EGetByNameFlags Flags = EGetByNameFlags::None) const;
+	COREUOBJECT_API int32 GetIndexByNameString(const FString& SearchString, EGetByNameFlags Flags = EGetByNameFlags::None) const;
 
 	/** Returns the short name matching the enum Value, returns empty string if invalid */
-	FString GetNameStringByValue(int64 InValue) const;
+	COREUOBJECT_API FString GetNameStringByValue(int64 InValue) const;
 
 	/** Looks for a name with a given value and returns true and writes the name to Out if one was found */
-	bool FindNameStringByValue(FString& Out, int64 InValue) const;
+	COREUOBJECT_API bool FindNameStringByValue(FString& Out, int64 InValue) const;
 
 	/** Gets enum value by name, returns INDEX_NONE and optionally errors when name is not found. Handles full or short names */
-	int64 GetValueByNameString(const FString& SearchString, EGetByNameFlags Flags = EGetByNameFlags::None) const;
+	COREUOBJECT_API int64 GetValueByNameString(const FString& SearchString, EGetByNameFlags Flags = EGetByNameFlags::None) const;
 
 	/**
 	 * Finds the localized display name or native display name as a fallback.
@@ -2013,13 +2013,13 @@ public:
 	 *
 	 * @return The display name for this object, or an empty text if Index is invalid
 	 */
-	virtual FText GetDisplayNameTextByIndex(int32 InIndex) const;
+	COREUOBJECT_API virtual FText GetDisplayNameTextByIndex(int32 InIndex) const;
 
 	/** Version of GetDisplayNameTextByIndex that takes a value instead */
-	FText GetDisplayNameTextByValue(int64 InValue) const;
+	COREUOBJECT_API FText GetDisplayNameTextByValue(int64 InValue) const;
 
 	/** Looks for a display name with a given value and returns true and writes the name to Out if one was found */
-	bool FindDisplayNameTextByValue(FText& Out, int64 InValue) const;
+	COREUOBJECT_API bool FindDisplayNameTextByValue(FText& Out, int64 InValue) const;
 
 	/**
 	 * Returns the unlocalized logical name originally assigned to the enum at creation.
@@ -2030,22 +2030,22 @@ public:
 	 *
 	 * @return The author-specified name, or an empty string if Index is invalid
 	 */
-	virtual FString GetAuthoredNameStringByIndex(int32 InIndex) const;
+	COREUOBJECT_API virtual FString GetAuthoredNameStringByIndex(int32 InIndex) const;
 
 	/** Version of GetAuthoredNameByIndex that takes a value instead */
-	FString GetAuthoredNameStringByValue(int64 InValue) const;
+	COREUOBJECT_API FString GetAuthoredNameStringByValue(int64 InValue) const;
 
 	/** Looks for a display name with a given value and returns true and writes the unlocalized logical name to Out if one was found */
-	bool FindAuthoredNameStringByValue(FString& Out, int64 InValue) const;
+	COREUOBJECT_API bool FindAuthoredNameStringByValue(FString& Out, int64 InValue) const;
 
 	/** Gets max value of Enum. Defaults to zero if there are no entries. */
-	int64 GetMaxEnumValue() const;
+	COREUOBJECT_API int64 GetMaxEnumValue() const;
 
 	/** Checks if enum has entry with given value. Includes autogenerated _MAX entry. */
-	bool IsValidEnumValue(int64 InValue) const;
+	COREUOBJECT_API bool IsValidEnumValue(int64 InValue) const;
 
 	/** Checks if enum has entry with given name. Includes autogenerated _MAX entry. */
-	bool IsValidEnumName(FName InName) const;
+	COREUOBJECT_API bool IsValidEnumName(FName InName) const;
 
 	/** Removes the Names in this enum from the primary AllEnumNames list */
 	UE_DEPRECATED(5.1, "RemoveNamesFromMasterList is deprecated, please use RemoveNamesFromPrimaryList instead.")
@@ -2055,10 +2055,10 @@ public:
 	}
 
 	/** Removes the Names in this enum from the primary AllEnumNames list */
-	void RemoveNamesFromPrimaryList();
+	COREUOBJECT_API void RemoveNamesFromPrimaryList();
 
 	/** Try to update an out-of-date enum index after an enum changes at runtime */
-	virtual int64 ResolveEnumerator(FArchive& Ar, int64 EnumeratorIndex) const;
+	COREUOBJECT_API virtual int64 ResolveEnumerator(FArchive& Ar, int64 EnumeratorIndex) const;
 
 	/** Associate a function for looking up Enum display names by index, only intended for use by generated code */
 	void SetEnumDisplayNameFn(FEnumDisplayNameFn InEnumDisplayNameFn)
@@ -2103,7 +2103,7 @@ public:
 	 * @param InEnumName Enum name.
 	 * @return Full enum name.
 	 */
-	virtual FString GenerateFullEnumName(const TCHAR* InEnumName) const;
+	COREUOBJECT_API virtual FString GenerateFullEnumName(const TCHAR* InEnumName) const;
 
 	/**
 	 * Searches the list of all enum value names for the specified name
@@ -2113,7 +2113,7 @@ public:
 	 * @param OutFoundEnum Optional address of a variable where the resulting UEnum object should be stored
 	 * @return The value the specified name represents if found, otherwise INDEX_NONE
 	 */
-	static int64 LookupEnumName(FName PackageName, FName TestName, EFindFirstObjectOptions Options = EFindFirstObjectOptions::None, UEnum** OutFoundEnum = nullptr);
+	COREUOBJECT_API static int64 LookupEnumName(FName PackageName, FName TestName, EFindFirstObjectOptions Options = EFindFirstObjectOptions::None, UEnum** OutFoundEnum = nullptr);
 
 	/** searches the list of all enum value names for the specified name
 	 * @return the value the specified name represents if found, otherwise INDEX_NONE
@@ -2132,7 +2132,7 @@ public:
 	 * @param OutFoundEnum Optional address of a variable where the resulting UEnum object should be stored
 	 * @return The value the specified name represents if found, otherwise INDEX_NONE
 	 */
-	static int64 LookupEnumNameSlow(FName PackageName, const TCHAR* InTestShortName, EFindFirstObjectOptions Options = EFindFirstObjectOptions::None, UEnum** OutFoundEnum = nullptr);
+	COREUOBJECT_API static int64 LookupEnumNameSlow(FName PackageName, const TCHAR* InTestShortName, EFindFirstObjectOptions Options = EFindFirstObjectOptions::None, UEnum** OutFoundEnum = nullptr);
 
 	/** searches the list of all enum value names for the specified name
 	 * @return the value the specified name represents if found, otherwise INDEX_NONE
@@ -2147,14 +2147,14 @@ public:
 	 * @param Str	pointer to string to parse; if we successfully find an enum, this pointer is advanced past the name found
 	 * @return index of the value the parsed enum name matches, or INDEX_NONE if no matches
 	 */
-	static int64 ParseEnum(const TCHAR*& Str);
+	COREUOBJECT_API static int64 ParseEnum(const TCHAR*& Str);
 
 	/**
 	 * Tests if the enum contains a MAX value
 	 *
 	 * @return	true if the enum contains a MAX enum, false otherwise.
 	 */
-	bool ContainsExistingMax() const;
+	COREUOBJECT_API bool ContainsExistingMax() const;
 
 	/**
 	 * Sets the array of enums.
@@ -2164,7 +2164,7 @@ public:
 	 * @param bAddMaxKeyIfMissing Should a default Max item be added.
 	 * @return	true unless the MAX enum already exists and isn't the last enum.
 	 */
-	virtual bool SetEnums(TArray<TPair<FName, int64>>& InNames, ECppForm InCppForm, EEnumFlags InFlags = EEnumFlags::None, bool bAddMaxKeyIfMissing = true);
+	COREUOBJECT_API virtual bool SetEnums(TArray<TPair<FName, int64>>& InNames, ECppForm InCppForm, EEnumFlags InFlags = EEnumFlags::None, bool bAddMaxKeyIfMissing = true);
 
 	/**
 	 * @return	 The number of enum names.
@@ -2180,7 +2180,7 @@ public:
 	 * @return	the longest common prefix between all items in the enum.  If a common prefix
 	 *			cannot be found, returns the full name of the enum.
 	 */
-	FString GenerateEnumPrefix() const;
+	COREUOBJECT_API FString GenerateEnumPrefix() const;
 
 #if WITH_EDITOR
 	/**
@@ -2190,7 +2190,7 @@ public:
 	 *
 	 * @return The tooltip for this object.
 	 */
-	FText GetToolTipTextByIndex(int32 NameIndex) const;
+	COREUOBJECT_API FText GetToolTipTextByIndex(int32 NameIndex) const;
 #endif
 
 #if WITH_EDITORONLY_DATA
@@ -2202,7 +2202,7 @@ public:
 	 *
 	 * @return true if the specified key exists in the list of metadata for this enum, even if the value of that key is empty
 	 */
-	bool HasMetaData( const TCHAR* Key, int32 NameIndex=INDEX_NONE ) const;
+	COREUOBJECT_API bool HasMetaData( const TCHAR* Key, int32 NameIndex=INDEX_NONE ) const;
 
 	/**
 	 * Return the metadata value associated with the specified key.
@@ -2213,7 +2213,7 @@ public:
 	 *
 	 * @return	the value for the key specified, or an empty string if the key wasn't found or had no value.
 	 */
-	FString GetMetaData( const TCHAR* Key, int32 NameIndex=INDEX_NONE, bool bAllowRemap=true ) const;
+	COREUOBJECT_API FString GetMetaData( const TCHAR* Key, int32 NameIndex=INDEX_NONE, bool bAllowRemap=true ) const;
 
 	/**
 	 * Set the metadata value associated with the specified key.
@@ -2223,7 +2223,7 @@ public:
 	 * @param	InValue		Value of the metadata for the key
 	 *
 	 */
-	void SetMetaData( const TCHAR* Key, const TCHAR* InValue, int32 NameIndex=INDEX_NONE) const;
+	COREUOBJECT_API void SetMetaData( const TCHAR* Key, const TCHAR* InValue, int32 NameIndex=INDEX_NONE) const;
 	
 	/**
 	 * Remove given key meta data
@@ -2232,7 +2232,7 @@ public:
 	 * @param	NameIndex	if specified, will search the metadata linked for that enum value; otherwise, searches the metadata for the enum itself
 	 *
 	 */
-	void RemoveMetaData( const TCHAR* Key, int32 NameIndex=INDEX_NONE ) const;
+	COREUOBJECT_API void RemoveMetaData( const TCHAR* Key, int32 NameIndex=INDEX_NONE ) const;
 #endif // WITH_EDITORONLY_DATA
 	
 	/**
@@ -2368,8 +2368,8 @@ public:
 	}
 
 	// UObject interface.
-	virtual void Serialize(FArchive& Ar) override;
-	virtual void BeginDestroy() override;
+	COREUOBJECT_API virtual void Serialize(FArchive& Ar) override;
+	COREUOBJECT_API virtual void BeginDestroy() override;
 	// End of UObject interface.
 
 protected:
@@ -2402,7 +2402,7 @@ protected:
 	}
 
 	/** adds the Names in this enum to the primary AllEnumNames list */
-	void AddNamesToPrimaryList();
+	COREUOBJECT_API void AddNamesToPrimaryList();
 
 private:
 
diff --git a/Engine/Source/Runtime/CoreUObject/Public/UObject/UObjectThreadContext.h b/Engine/Source/Runtime/CoreUObject/Public/UObject/UObjectThreadContext.h
--- a/Engine/Source/Runtime/CoreUObject/Public/UObject/UObjectThreadContext.h
+++ b/Engine/Source/Runtime/CoreUObject/Public/UObject/UObjectThreadContext.h
@@ -31,12 +31,12 @@ struct FUObjectSerializeContext;
 
 COREUOBJECT_API DECLARE_LOG_CATEGORY_EXTERN(LogUObjectThreadContext, Log, All);
 
-class COREUOBJECT_API FUObjectThreadContext : public TThreadSingleton<FUObjectThreadContext>
+class FUObjectThreadContext : public TThreadSingleton<FUObjectThreadContext>
 {
 	friend TThreadSingleton<FUObjectThreadContext>;
 
-	FUObjectThreadContext();
-	virtual ~FUObjectThreadContext();
+	COREUOBJECT_API FUObjectThreadContext();
+	COREUOBJECT_API virtual ~FUObjectThreadContext();
 
 	/** Stack of currently used FObjectInitializers for this thread */
 	TArray<FObjectInitializer*> InitializerStack;
diff --git a/Engine/Source/Runtime/Engine/Public/AudioDevice.h b/Engine/Source/Runtime/Engine/Public/AudioDevice.h
--- a/Engine/Source/Runtime/Engine/Public/AudioDevice.h
+++ b/Engine/Source/Runtime/Engine/Public/AudioDevice.h
@@ -436,19 +436,19 @@ struct FAudioDeviceRenderInfo
 DECLARE_MULTICAST_DELEGATE_OneParam(FOnAudioDevicePreRender, const FAudioDeviceRenderInfo&);
 DECLARE_MULTICAST_DELEGATE_OneParam(FOnAudioDevicePostRender, const FAudioDeviceRenderInfo&);
 
-class ENGINE_API FAudioDevice : public FExec
+class FAudioDevice : public FExec
 {
 public:
 
 	//Begin FExec Interface
 #if UE_ALLOW_EXEC_COMMANDS
-	virtual bool Exec(UWorld* InWorld, const TCHAR* Cmd, FOutputDevice& Ar = *GLog) override;
+	ENGINE_API virtual bool Exec(UWorld* InWorld, const TCHAR* Cmd, FOutputDevice& Ar = *GLog) override;
 #endif
 	//End FExec Interface
 
 #if !UE_BUILD_SHIPPING
-	UAudioComponent* GetTestComponent(UWorld* InWorld);
-	void StopTestComponent();
+	ENGINE_API UAudioComponent* GetTestComponent(UWorld* InWorld);
+	ENGINE_API void StopTestComponent();
 
 private:
 	/**
@@ -522,7 +522,7 @@ public:
 	/**
 	 * Constructor
 	 */
-	FAudioDevice();
+	ENGINE_API FAudioDevice();
 
 	virtual ~FAudioDevice();
 
@@ -532,37 +532,37 @@ public:
 	}
 
 	/** Returns the quality settings used by the default audio settings. */
-	static const FAudioQualitySettings& GetQualityLevelSettings();
+	ENGINE_API static const FAudioQualitySettings& GetQualityLevelSettings();
 
 	/**
 	 * Basic initialization of the platform agnostic layer of the audio system
 	 */
-	bool Init(Audio::FDeviceId InDeviceID, int32 InMaxSources, int32 BufferSizeOverride = INDEX_NONE, int32 NumBuffersOverride = INDEX_NONE);
+	ENGINE_API bool Init(Audio::FDeviceId InDeviceID, int32 InMaxSources, int32 BufferSizeOverride = INDEX_NONE, int32 NumBuffersOverride = INDEX_NONE);
 
 	/**
 	 * Called after FAudioDevice creation and init
 	 */
-	void OnDeviceCreated(Audio::FDeviceId InDeviceID);
+	ENGINE_API void OnDeviceCreated(Audio::FDeviceId InDeviceID);
 
 	/**
 	 * Called before FAudioDevice teardown
 	 */
-	void OnDeviceDestroyed(Audio::FDeviceId InDeviceID);
+	ENGINE_API void OnDeviceDestroyed(Audio::FDeviceId InDeviceID);
 
 	/**
 	 * Tears down the audio device
 	 */
-	void Teardown();
+	ENGINE_API void Teardown();
 
 	/**
 	 * De-initialization step that occurs after device destroyed broadcast, but before removal from the device manager
 	 */
-	void Deinitialize();
+	ENGINE_API void Deinitialize();
 
 	/**
 	 * The audio system's main "Tick" function
 	 */
-	void Update(bool bGameTicking);
+	ENGINE_API void Update(bool bGameTicking);
 
 	/** Update called on game thread. */
 	virtual void UpdateGameThread() {}
@@ -572,27 +572,27 @@ public:
 	 *
 	 * @param bGameTicking Whether the game is still ticking at the time of suspend
 	 */
-	void Suspend(bool bGameTicking);
+	ENGINE_API void Suspend(bool bGameTicking);
 
 	/**
 	 * Counts the bytes for the structures used in this class
 	 */
-	virtual void CountBytes(FArchive& Ar);
+	ENGINE_API virtual void CountBytes(FArchive& Ar);
 
 	/**
 	 * Track references to UObjects
 	 */
-	void AddReferencedObjects(FReferenceCollector& Collector);
+	ENGINE_API void AddReferencedObjects(FReferenceCollector& Collector);
 
 	/**
 	 * Iterate over the active AudioComponents for wave instances that could be playing.
 	 *
 	 * @return Index of first wave instance that can have a source attached
 	 */
-	int32 GetSortedActiveWaveInstances(TArray<FWaveInstance*>& WaveInstances, const ESortedActiveWaveGetType::Type GetType);
+	ENGINE_API int32 GetSortedActiveWaveInstances(TArray<FWaveInstance*>& WaveInstances, const ESortedActiveWaveGetType::Type GetType);
 
 	/** Update the active sound playback time. This is done here to do after all audio is updated. */
-	void UpdateActiveSoundPlaybackTime(bool bIsTimeTicking);
+	ENGINE_API void UpdateActiveSoundPlaybackTime(bool bIsTimeTicking);
 
 	/** Optional fadeout and fade in of audio to avoid clicks when closing or opening/reusing audio device. */
 	virtual void FadeOut() {}
@@ -601,12 +601,12 @@ public:
 	/**
 	 * Stop all the audio components and sources attached to the world. nullptr world means all components.
 	 */
-	void Flush(UWorld* WorldToFlush, bool bClearActivatedReverb = true);
+	ENGINE_API void Flush(UWorld* WorldToFlush, bool bClearActivatedReverb = true);
 
 	/*
 	* Derived classes can override this method to do their own cleanup. Called at the end of FAudioDevice::Flush()
 	*/
-	virtual void FlushExtended(UWorld* WorldToFlush, bool bClearActivatedReverb);
+	ENGINE_API virtual void FlushExtended(UWorld* WorldToFlush, bool bClearActivatedReverb);
 
 	/**
 	 * Allows audio rendering command queue to flush during audio device flush.
@@ -614,7 +614,7 @@ public:
 	 */
 	virtual void FlushAudioRenderingCommands(bool bPumpSynchronously = false) {}
 
-	void OnPreGarbageCollect();
+	ENGINE_API void OnPreGarbageCollect();
 
 	/**
 	 * Stop any playing sounds that are using a particular SoundWave
@@ -622,16 +622,16 @@ public:
 	 * @param SoundWave					Resource to stop any sounds that are using it
 	 * @param[out] StoppedComponents	List of Audio Components that were stopped
 	 */
-	void StopSoundsUsingResource(USoundWave* SoundWave, TArray<UAudioComponent*>* StoppedComponents = nullptr);
+	ENGINE_API void StopSoundsUsingResource(USoundWave* SoundWave, TArray<UAudioComponent*>* StoppedComponents = nullptr);
 
-	static bool LegacyReverbDisabled();
+	ENGINE_API static bool LegacyReverbDisabled();
 
 #if WITH_EDITOR
 	/** Deals with anything audio related that should happen when PIE starts */
-	void OnBeginPIE(const bool bIsSimulating);
+	ENGINE_API void OnBeginPIE(const bool bIsSimulating);
 
 	/** Deals with anything audio related that should happen when PIE ends */
-	void OnEndPIE(const bool bIsSimulating);
+	ENGINE_API void OnEndPIE(const bool bIsSimulating);
 #endif
 
 	/**
@@ -642,55 +642,55 @@ public:
 	 * @param	bTrackMemory	If true, the audio mem stats will be updated
 	 * @param   bForceFullDecompression If true, the sound wave will be fully decompressed regardless of size.
 	 */
-	virtual void Precache(USoundWave* SoundWave, bool bSynchronous = false, bool bTrackMemory = true, bool bForceFullDecompression = false);
+	ENGINE_API virtual void Precache(USoundWave* SoundWave, bool bSynchronous = false, bool bTrackMemory = true, bool bForceFullDecompression = false);
 
-	float GetCompressionDurationThreshold(const FSoundGroup &SoundGroup);
+	ENGINE_API float GetCompressionDurationThreshold(const FSoundGroup &SoundGroup);
 
 	/**
 	 * Returns true if a sound wave should be decompressed.
 	 */
-	bool ShouldUseRealtimeDecompression(bool bForceFullDecompression, const FSoundGroup &SoundGroup, USoundWave* SoundWave, float CompressedDurationThreshold) const;
+	ENGINE_API bool ShouldUseRealtimeDecompression(bool bForceFullDecompression, const FSoundGroup &SoundGroup, USoundWave* SoundWave, float CompressedDurationThreshold) const;
 
 	/**
 	 * Precaches all existing sounds. Called when audio setup is complete
 	 */
-	void PrecacheStartupSounds();
+	ENGINE_API void PrecacheStartupSounds();
 
 	/**
 	 * Sets the maximum number of channels dynamically. Can't raise the cap over the initial value but can lower it
 	 */
-	void SetMaxChannels(int32 InMaxChannels);
+	ENGINE_API void SetMaxChannels(int32 InMaxChannels);
 
 	/**
 	 * Sets the maximum number of channels dynamically by scaled percentage.
 	 */
-	void SetMaxChannelsScaled(float InScaledChannelCount);
+	ENGINE_API void SetMaxChannelsScaled(float InScaledChannelCount);
 
 	/** Returns the max channels used by the audio device. */
-	int32 GetMaxChannels() const;
+	ENGINE_API int32 GetMaxChannels() const;
 
 	/** Returns the maximum sources used by the audio device set on initialization,
 	  * including the number of stopping voices reserved. */
-	int32 GetMaxSources() const;
+	ENGINE_API int32 GetMaxSources() const;
 
 	/**
 	 * Returns global pitch range
 	 */
-	TRange<float> GetGlobalPitchRange() const;
+	ENGINE_API TRange<float> GetGlobalPitchRange() const;
 
 	/**
 	* Stops any sound sources which are using the given buffer.
 	*
 	* @param	FSoundBuffer	Buffer to check against
 	*/
-	void StopSourcesUsingBuffer(FSoundBuffer * SoundBuffer);
+	ENGINE_API void StopSourcesUsingBuffer(FSoundBuffer * SoundBuffer);
 
 	/**
 	 * Stops all game sounds (and possibly UI) sounds
 	 *
 	 * @param bShouldStopUISounds If true, this function will stop UI sounds as well
 	 */
-	virtual void StopAllSounds(bool bShouldStopUISounds = false);
+	ENGINE_API virtual void StopAllSounds(bool bShouldStopUISounds = false);
 
 	/**
 	 * Sets the details about the listener
@@ -699,21 +699,21 @@ public:
 	 * @param   ListenerTransform   The listener's world transform
 	 * @param   DeltaSeconds		The amount of time over which velocity should be calculated.  If 0, then velocity will not be calculated.
 	 */
-	void SetListener(UWorld* World, int32 InListenerIndex, const FTransform& ListenerTransform, float InDeltaSeconds);
+	ENGINE_API void SetListener(UWorld* World, int32 InListenerIndex, const FTransform& ListenerTransform, float InDeltaSeconds);
 
 	/** Sets an override for the listener to do attenuation calculations. */
 	UE_DEPRECATED(4.25, "Use SetListenerAttenuationOverride that passes a ListenerIndex instead")
 	void SetListenerAttenuationOverride(const FVector AttenuationPosition) { SetListenerAttenuationOverride(0, AttenuationPosition); }
 
 	/** Sets an override position for the specified listener to do attenuation calculations. */
-	void SetListenerAttenuationOverride(int32 ListenerIndex, const FVector AttenuationPosition);
+	ENGINE_API void SetListenerAttenuationOverride(int32 ListenerIndex, const FVector AttenuationPosition);
 
 	/** Removes a listener attenuation override. */
 	UE_DEPRECATED(4.25, "Use ClearListenerAttenuationOverride that passes a ListenerIndex instead")
 	void ClearListenerAttenuationOverride() { ClearListenerAttenuationOverride(0); }
 
 	/** Removes a listener attenuation override for the specified listener. */
-	void ClearListenerAttenuationOverride(int32 ListenerIndex);
+	ENGINE_API void ClearListenerAttenuationOverride(int32 ListenerIndex);
 
 	const TArray<FListener>& GetListeners() const { check(IsInAudioThread()); return Listeners; }
 
@@ -762,7 +762,7 @@ public:
 	 * @param   Params				The parameter block of properties to create the component based on
 	 * @return	The created audio component if the function successfully created one or a nullptr if not successful. Note: if audio is disabled or if there were no hardware audio devices available, this will return nullptr.
 	 */
-	static UAudioComponent* CreateComponent(USoundBase* Sound, const FCreateComponentParams& Params = FCreateComponentParams());
+	ENGINE_API static UAudioComponent* CreateComponent(USoundBase* Sound, const FCreateComponentParams& Params = FCreateComponentParams());
 
 	/**
 	 * Plays a sound at the given location without creating an audio component.
@@ -777,40 +777,40 @@ public:
 	 * @param	USoundConcurrency	The sound's sound concurrency settings to use (optional). Will use the USoundBase's USoundConcurrency if not specified.
 	 * @param	Params				An optional list of audio component params to immediately apply to a sound.
 	 */
-	void PlaySoundAtLocation(USoundBase* Sound, UWorld* World, float VolumeMultiplier, float PitchMultiplier, float StartTime, const FVector& Location, const FRotator& Rotation, USoundAttenuation* AttenuationSettings = nullptr, USoundConcurrency* ConcurrencySettings = nullptr, const TArray<FAudioParameter>* Params = nullptr, const AActor* OwningActor = nullptr);
+	ENGINE_API void PlaySoundAtLocation(USoundBase* Sound, UWorld* World, float VolumeMultiplier, float PitchMultiplier, float StartTime, const FVector& Location, const FRotator& Rotation, USoundAttenuation* AttenuationSettings = nullptr, USoundConcurrency* ConcurrencySettings = nullptr, const TArray<FAudioParameter>* Params = nullptr, const AActor* OwningActor = nullptr);
 
 	/**
 	 * Adds an active sound to the audio device
 	 */
-	void AddNewActiveSound(const FActiveSound& ActiveSound, const TArray<FAudioParameter>* InDefaultParams = nullptr);
+	ENGINE_API void AddNewActiveSound(const FActiveSound& ActiveSound, const TArray<FAudioParameter>* InDefaultParams = nullptr);
 
 	/**
 	 * Adds an active sound to the audio device
 	 */
-	void AddNewActiveSound(const FActiveSound& ActiveSound, TArray<FAudioParameter>&& InDefaultParams);
+	ENGINE_API void AddNewActiveSound(const FActiveSound& ActiveSound, TArray<FAudioParameter>&& InDefaultParams);
 
 	/**
 	 * Attempts to retrigger a provided loop
 	 */
-	void RetriggerVirtualLoop(FAudioVirtualLoop& VirtualLoop);
+	ENGINE_API void RetriggerVirtualLoop(FAudioVirtualLoop& VirtualLoop);
 
 	/**
 	 * Removes the active sound for the specified audio component
 	 */
-	void StopActiveSound(uint64 AudioComponentID);
+	ENGINE_API void StopActiveSound(uint64 AudioComponentID);
 
 	/**
 	* (Deprecated in favor of AddSoundToStop). Stops the active sound
 	*/
-	void StopActiveSound(FActiveSound* ActiveSound);
+	ENGINE_API void StopActiveSound(FActiveSound* ActiveSound);
 
 	/**
 	* Pauses the active sound for the specified audio component
 	*/
-	void PauseActiveSound(uint64 AudioComponentID, const bool bInIsPaused);
+	ENGINE_API void PauseActiveSound(uint64 AudioComponentID, const bool bInIsPaused);
 
 	/** Notify that a pending async occlusion trace finished on the active sound. */
-	void NotifyActiveSoundOcclusionTraceDone(FActiveSound* ActiveSound, bool bIsOccluded);
+	ENGINE_API void NotifyActiveSoundOcclusionTraceDone(FActiveSound* ActiveSound, bool bIsOccluded);
 
 	/**
 	* Finds the active sound for the specified audio component ID.
@@ -818,28 +818,28 @@ public:
 	* multiple ActiveSound instances at once.  Use 'SendCommandToActiveSounds' instead.
 	*/
 	UE_DEPRECATED(5.0, "This call cannot be called from the GameThread & AudioComponents can now execute multiple ActiveSound instances at once.  Use 'SendCommandToActiveSounds' instead.")
-	FActiveSound* FindActiveSound(uint64 AudioComponentID);
+	ENGINE_API FActiveSound* FindActiveSound(uint64 AudioComponentID);
 
 	/**
 	 * Whether a given Audio Component ID should be allowed to have multiple
 	 * associated Active Sounds
 	 */
-	bool CanHaveMultipleActiveSounds(uint64 AudioComponentID) const;
+	ENGINE_API bool CanHaveMultipleActiveSounds(uint64 AudioComponentID) const;
 
 	/**
 	 * Set whether a given Audio Component ID should be allowed to have multiple
 	 * associated Active Sounds
 	 */
-	void SetCanHaveMultipleActiveSounds(uint64 AudioComponentID, bool InCanHaveMultipleActiveSounds);
+	ENGINE_API void SetCanHaveMultipleActiveSounds(uint64 AudioComponentID, bool InCanHaveMultipleActiveSounds);
 
 	/**
 	 * Removes an active sound from the active sounds array
 	 */
-	void RemoveActiveSound(FActiveSound* ActiveSound);
+	ENGINE_API void RemoveActiveSound(FActiveSound* ActiveSound);
 
-	void AddAudioVolumeProxy(const FAudioVolumeProxy& Proxy);
-	void RemoveAudioVolumeProxy(uint32 AudioVolumeID);
-	void UpdateAudioVolumeProxy(const FAudioVolumeProxy& Proxy);
+	ENGINE_API void AddAudioVolumeProxy(const FAudioVolumeProxy& Proxy);
+	ENGINE_API void RemoveAudioVolumeProxy(uint32 AudioVolumeID);
+	ENGINE_API void UpdateAudioVolumeProxy(const FAudioVolumeProxy& Proxy);
 
 	struct FAudioVolumeSettings
 	{
@@ -852,18 +852,18 @@ public:
 		bool bChanged = false;
 	};
 
-	void GetAudioVolumeSettings(const uint32 WorldID, const FVector& Location, FAudioVolumeSettings& OutSettings) const;
-	void ResetAudioVolumeProxyChangedState();
+	ENGINE_API void GetAudioVolumeSettings(const uint32 WorldID, const FVector& Location, FAudioVolumeSettings& OutSettings) const;
+	ENGINE_API void ResetAudioVolumeProxyChangedState();
 
 	/**
 	 * Gathers data about interior volumes affecting the active sound (called on audio thread)
 	 */
-	void GatherInteriorData(FActiveSound& ActiveSound, FSoundParseParameters& ParseParams) const;
+	ENGINE_API void GatherInteriorData(FActiveSound& ActiveSound, FSoundParseParameters& ParseParams) const;
 
 	/**
 	 * Applies interior settings from affecting volumes to the active sound (called on audio thread)
 	 */
-	void ApplyInteriorSettings(FActiveSound& ActiveSound, FSoundParseParameters& ParseParams) const;
+	ENGINE_API void ApplyInteriorSettings(FActiveSound& ActiveSound, FSoundParseParameters& ParseParams) const;
 
 	/**
 	 * Notifies subsystems an active sound is about to be deleted (called on audio thread) - Deprecated, see NotifyPendingDeleteInternal
@@ -876,46 +876,46 @@ protected:
 	/**
 	 * Notifies subsystems an active sound has been added (called on audio thread)
 	 */
-	void NotifyAddActiveSound(FActiveSound& ActiveSound) const;
+	ENGINE_API void NotifyAddActiveSound(FActiveSound& ActiveSound) const;
 
 	/**
 	 * Notifies subsystems an active sound is about to be deleted (called on audio thread)
 	 */
-	void NotifyPendingDeleteInternal(FActiveSound& ActiveSound) const;
+	ENGINE_API void NotifyPendingDeleteInternal(FActiveSound& ActiveSound) const;
 
 public:
 
 	/**
 	 * Gets the default reverb and interior settings for the provided world.  Returns true if settings with WorldID were located
 	 */
-	bool GetDefaultAudioSettings(uint32 WorldID, FReverbSettings& OutReverbSettings, FInteriorSettings& OutInteriorSettings) const;
+	ENGINE_API bool GetDefaultAudioSettings(uint32 WorldID, FReverbSettings& OutReverbSettings, FInteriorSettings& OutInteriorSettings) const;
 
 	/**
 	 * Sets the default reverb and interior settings for the provided world
 	 */
-	void SetDefaultAudioSettings(UWorld* World, const FReverbSettings& DefaultReverbSettings, const FInteriorSettings& DefaultInteriorSettings);
+	ENGINE_API void SetDefaultAudioSettings(UWorld* World, const FReverbSettings& DefaultReverbSettings, const FInteriorSettings& DefaultInteriorSettings);
 
 	/**
 	 * Gets the current audio debug state
 	 */
 	EDebugState GetMixDebugState() const { return((EDebugState)DebugState); }
 
-	void SetMixDebugState(EDebugState DebugState);
+	ENGINE_API void SetMixDebugState(EDebugState DebugState);
 
 	/**
 	 * Set up the sound class hierarchy
 	 */
-	void InitSoundClasses();
+	ENGINE_API void InitSoundClasses();
 
 protected:
 	/**
 	 * Set up the initial sound sources
 	 * Allows us to initialize sound source early on, allowing for render callback hookups for iOS Audio.
 	 */
-	void InitSoundSources();
+	ENGINE_API void InitSoundSources();
 
 	/** Create our subsystem collection root object and initialize subsystems */
-	void InitializeSubsystemCollection();
+	ENGINE_API void InitializeSubsystemCollection();
 
 	// Handle for our device created delegate
 	FDelegateHandle DeviceCreatedHandle;
@@ -938,12 +938,12 @@ public:
 	 * @param	SoundClassName	name of sound class to retrieve
 	 * @return	sound class properties if it exists
 	 */
-	void RegisterSoundClass(USoundClass* InSoundClass);
+	ENGINE_API void RegisterSoundClass(USoundClass* InSoundClass);
 
 	/**
 	* Unregisters a sound class
 	*/
-	void UnregisterSoundClass(USoundClass* SoundClass);
+	ENGINE_API void UnregisterSoundClass(USoundClass* SoundClass);
 
 	/* Initialized audio buses marked as default that are to be enabled for the duration of the application. */
 	virtual void InitDefaultAudioBuses() {}
@@ -964,15 +964,15 @@ public:
 	 * Registers the submix buffer listener with the given submix.
 	 * A nullptr for SoundSubmix will register the listener with the master submix.
 	*/
-	virtual void RegisterSubmixBufferListener(ISubmixBufferListener* InSubmixBufferListener, USoundSubmix* SoundSubmix = nullptr);
+	ENGINE_API virtual void RegisterSubmixBufferListener(ISubmixBufferListener* InSubmixBufferListener, USoundSubmix* SoundSubmix = nullptr);
 
 	/**
 	 * Unregisters the submix buffer listener with the given submix.
 	 * A nullptr for SoundSubmix will unregister the listener with the master submix.
 	*/
-	virtual void UnregisterSubmixBufferListener(ISubmixBufferListener* InSubmixBufferListener, USoundSubmix* SoundSubmix = nullptr);
+	ENGINE_API virtual void UnregisterSubmixBufferListener(ISubmixBufferListener* InSubmixBufferListener, USoundSubmix* SoundSubmix = nullptr);
 
-	virtual Audio::FPatchOutputStrongPtr AddPatchForSubmix(uint32 InObjectId, float InPatchGain);
+	ENGINE_API virtual Audio::FPatchOutputStrongPtr AddPatchForSubmix(uint32 InObjectId, float InPatchGain);
 
 	virtual void StartAudioBus(uint32 InAudioBusId, int32 InNumChannels, bool bInIsAutomatic)
 	{
@@ -988,10 +988,10 @@ public:
 	}
 
 	UE_DEPRECATED(5.2, "AddPatchForAudioBus is deprecated.  Use AddPatchOutputForAudioBus.")
-	virtual Audio::FPatchOutputStrongPtr AddPatchForAudioBus(uint32 InAudioBusId, float InPatchGain = 1.0f);
+	ENGINE_API virtual Audio::FPatchOutputStrongPtr AddPatchForAudioBus(uint32 InAudioBusId, float InPatchGain = 1.0f);
 
 	UE_DEPRECATED(5.2, "AddPatchForAudioBus_GameThread is deprecated.  Use AddPatchOutputForAudioBus.")
-	virtual Audio::FPatchOutputStrongPtr AddPatchForAudioBus_GameThread(uint32 InAudioBusId, float InPatchGain = 1.0f);
+	ENGINE_API virtual Audio::FPatchOutputStrongPtr AddPatchForAudioBus_GameThread(uint32 InAudioBusId, float InPatchGain = 1.0f);
 
 	UE_DEPRECATED(5.2, "This overload of AddPatchInputForAudioBus is deprecated.  Use the overload that takes the number of frames and channels as parameters.")
 	virtual void AddPatchInputForAudioBus(const Audio::FPatchInput& InPatchInput, uint32 InAudioBusId, float InPatchGain = 1.0f)
@@ -1003,9 +1003,9 @@ public:
 	{
 	}
 
-	virtual Audio::FPatchInput AddPatchInputForAudioBus(uint32 InAudioBusId, int32 InFrames, int32 InChannels, float InGain = 1.f);
+	ENGINE_API virtual Audio::FPatchInput AddPatchInputForAudioBus(uint32 InAudioBusId, int32 InFrames, int32 InChannels, float InGain = 1.f);
 
-	virtual Audio::FPatchOutputStrongPtr AddPatchOutputForAudioBus(uint32 InAudioBusId, int32 InFrames, int32 InChannels, float InGain = 1.f);
+	ENGINE_API virtual Audio::FPatchOutputStrongPtr AddPatchOutputForAudioBus(uint32 InAudioBusId, int32 InFrames, int32 InChannels, float InGain = 1.f);
 
 	virtual void InitSoundEffectPresets() {}
 
@@ -1015,53 +1015,53 @@ public:
 	* @param	SoundClassName	name of sound class to retrieve
 	* @return	sound class properties if it exists
 	*/
-	FSoundClassProperties* GetSoundClassCurrentProperties(USoundClass* InSoundClass);
+	ENGINE_API FSoundClassProperties* GetSoundClassCurrentProperties(USoundClass* InSoundClass);
 
 	/** 
 	* Returns the parameters which are dynamic from the given sound class. 
 	*/
-	FSoundClassDynamicProperties* GetSoundClassDynamicProperties(USoundClass* InSoundClass);
+	ENGINE_API FSoundClassDynamicProperties* GetSoundClassDynamicProperties(USoundClass* InSoundClass);
 
 	/**
 	* Checks to see if a coordinate is within a distance of any listener
 	*/
-	bool LocationIsAudible(const FVector& Location, const float MaxDistance) const;
+	ENGINE_API bool LocationIsAudible(const FVector& Location, const float MaxDistance) const;
 
 	/**
 	* Checks to see if a coordinate is within a distance of the given listener
 	*/
 	UE_DEPRECATED(4.25, "Use LocationIsAudible that passes a ListenerIndex to check against a specific Listener")
-	bool LocationIsAudible(const FVector& Location, const FTransform& ListenerTransform, const float MaxDistance) const;
+	ENGINE_API bool LocationIsAudible(const FVector& Location, const FTransform& ListenerTransform, const float MaxDistance) const;
 
 	/**
 	* Checks to see if a coordinate is within a distance of a specific listener
 	*/
-	bool LocationIsAudible(const FVector& Location, int32 ListenerIndex, const float MaxDistance) const;
+	ENGINE_API bool LocationIsAudible(const FVector& Location, int32 ListenerIndex, const float MaxDistance) const;
 
 	/**
 	* Returns the distance to the nearest listener from the given location
 	*/
-	float GetDistanceToNearestListener(const FVector& Location) const;
+	ENGINE_API float GetDistanceToNearestListener(const FVector& Location) const;
 
 	UE_DEPRECATED(4.25, "Use GetDistanceSquaredToListener to check against a specific Listener")
-	float GetSquaredDistanceToListener(const FVector& Location, const FTransform& ListenerTransform) const;
+	ENGINE_API float GetSquaredDistanceToListener(const FVector& Location, const FTransform& ListenerTransform) const;
 
 	/**
 	* Sets OutSqDistance to the distance from location to the appropriate listener representation, depending on calling thread.
 	* Returns true if listener position is valid, false if not (in which case, OutSqDistance is undefined).
 	*/
-	bool GetDistanceSquaredToListener(const FVector& Location, int32 ListenerIndex, float& OutSqDistance) const;
+	ENGINE_API bool GetDistanceSquaredToListener(const FVector& Location, int32 ListenerIndex, float& OutSqDistance) const;
 
 	/**
 	* Sets OutSqDistance to the distance from location the closest listener, depending on calling thread.
 	* Returns true if listener position is valid, false if not (in which case, OutSqDistance is undefined).
 	*/
-	bool GetDistanceSquaredToNearestListener(const FVector& Location, float& OutSqDistance) const;
+	ENGINE_API bool GetDistanceSquaredToNearestListener(const FVector& Location, float& OutSqDistance) const;
 		
 	/**
 	* Returns the global maximum distance used in the audio engine.
 	*/
-	static float GetMaxWorldDistance();
+	ENGINE_API static float GetMaxWorldDistance();
 
 	/**
 	* Returns a position from the appropriate listener representation, depending on calling thread.
@@ -1071,41 +1071,41 @@ public:
 	* @param	bAllowOverride	if true we will use the attenuation override for position, if set
 	* @return	true if successful
 	*/
-	bool GetListenerPosition(int32 ListenerIndex, FVector& OutPosition, bool bAllowOverride) const;
+	ENGINE_API bool GetListenerPosition(int32 ListenerIndex, FVector& OutPosition, bool bAllowOverride) const;
 
 	/**
 	* Returns the transform of the appropriate listener representation, depending on calling thread
 	*/
-	bool GetListenerTransform(int32 ListenerIndex, FTransform& OutTransform) const;
+	ENGINE_API bool GetListenerTransform(int32 ListenerIndex, FTransform& OutTransform) const;
 
 	/**
 	 * Returns the WorldID of the appropriate listener representation, depending on calling thread
 	 */
-	bool GetListenerWorldID(int32 ListenerIndex, uint32& OutWorldID) const;
+	ENGINE_API bool GetListenerWorldID(int32 ListenerIndex, uint32& OutWorldID) const;
 
 	/**
 	 * Sets the Sound Mix that should be active by default
 	 */
-	void SetDefaultBaseSoundMix(USoundMix* SoundMix);
+	ENGINE_API void SetDefaultBaseSoundMix(USoundMix* SoundMix);
 
 	/**
 	 * Removes a sound mix - called when SoundMix is unloaded
 	 */
-	void RemoveSoundMix(USoundMix* SoundMix);
+	ENGINE_API void RemoveSoundMix(USoundMix* SoundMix);
 
 	/**
 	 * Resets all interpolating values to defaults.
 	 */
-	void ResetInterpolation();
+	ENGINE_API void ResetInterpolation();
 
 	/** Enables or Disables the radio effect. */
-	void EnableRadioEffect(bool bEnable = false);
+	ENGINE_API void EnableRadioEffect(bool bEnable = false);
 
 	friend class FAudioEffectsManager;
 	/**
 	 * Sets a new sound mix and applies it to all appropriate sound classes
 	 */
-	void SetBaseSoundMix(USoundMix* SoundMix);
+	ENGINE_API void SetBaseSoundMix(USoundMix* SoundMix);
 
 	/**
 	 * Push a SoundMix onto the Audio Device's list.
@@ -1113,17 +1113,17 @@ public:
 	 * @param SoundMix The SoundMix to push.
 	 * @param bIsPassive Whether this is a passive push from a playing sound.
 	 */
-	void PushSoundMixModifier(USoundMix* SoundMix, bool bIsPassive = false, bool bIsRetrigger = false);
+	ENGINE_API void PushSoundMixModifier(USoundMix* SoundMix, bool bIsPassive = false, bool bIsRetrigger = false);
 
 	/**
 	 * Sets a sound class override in the given sound mix.
 	 */
-	void SetSoundMixClassOverride(USoundMix* InSoundMix, USoundClass* InSoundClass, float Volume, float Pitch, float FadeInTime, bool bApplyToChildren);
+	ENGINE_API void SetSoundMixClassOverride(USoundMix* InSoundMix, USoundClass* InSoundClass, float Volume, float Pitch, float FadeInTime, bool bApplyToChildren);
 
 	/**
 	* Clears a sound class override in the given sound mix.
 	*/
-	void ClearSoundMixClassOverride(USoundMix* InSoundMix, USoundClass* InSoundClass, float FadeOutTime);
+	ENGINE_API void ClearSoundMixClassOverride(USoundMix* InSoundMix, USoundClass* InSoundClass, float FadeOutTime);
 
 	/**
 	 * Pop a SoundMix from the Audio Device's list.
@@ -1131,19 +1131,19 @@ public:
 	 * @param SoundMix The SoundMix to pop.
 	 * @param bIsPassive Whether this is a passive pop from a sound finishing.
 	 */
-	void PopSoundMixModifier(USoundMix* SoundMix, bool bIsPassive = false);
+	ENGINE_API void PopSoundMixModifier(USoundMix* SoundMix, bool bIsPassive = false);
 
 	/**
 	 * Clear the effect of one SoundMix completely.
 	 *
 	 * @param SoundMix The SoundMix to clear.
 	 */
-	void ClearSoundMixModifier(USoundMix* SoundMix);
+	ENGINE_API void ClearSoundMixModifier(USoundMix* SoundMix);
 
 	/**
 	 * Clear the effect of all SoundMix modifiers.
 	 */
-	void ClearSoundMixModifiers();
+	ENGINE_API void ClearSoundMixModifiers();
 
 	/** Activates a Reverb Effect without the need for a volume
 	 * @param ReverbEffect Reverb Effect to use
@@ -1152,14 +1152,14 @@ public:
 	 * @param Volume Volume level of Reverb Effect
 	 * @param FadeTime Time before Reverb Effect is fully active
 	 */
-	void ActivateReverbEffect(UReverbEffect* ReverbEffect, FName TagName, float Priority, float Volume, float FadeTime);
+	ENGINE_API void ActivateReverbEffect(UReverbEffect* ReverbEffect, FName TagName, float Priority, float Volume, float FadeTime);
 
 	/**
 	 * Deactivates a Reverb Effect not applied by a volume
 	 *
 	 * @param TagName Tag associated with Reverb Effect to remove
 	 */
-	void DeactivateReverbEffect(FName TagName);
+	ENGINE_API void DeactivateReverbEffect(FName TagName);
 
 	UE_DEPRECATED(5.2, "GetRuntimeFormat is now deprecated. Please call SoundWave::GetRuntimeFormat() instead")
 	virtual FName GetRuntimeFormat(const USoundWave* SoundWave) const = 0;
@@ -1201,7 +1201,7 @@ public:
 	}
 
 	/** When the set of Audio volumes have changed invalidate the cached values of active sounds */
-	void InvalidateCachedInteriorVolumes() const;
+	ENGINE_API void InvalidateCachedInteriorVolumes() const;
 
 	/** Suspend any context related objects */
 	virtual void SuspendContext() {}
@@ -1213,7 +1213,7 @@ public:
 	virtual bool IsExernalBackgroundSoundActive() { return false; }
 
 	/** Whether or not HRTF spatialization is enabled for all. */
-	bool IsHRTFEnabledForAll() const;
+	ENGINE_API bool IsHRTFEnabledForAll() const;
 
 	void SetHRTFEnabledForAll(bool InbHRTFEnabledForAll)
 	{
@@ -1232,7 +1232,7 @@ public:
 	}
 
 	/** Whether or not HRTF is disabled. */
-	bool IsHRTFDisabled() const;
+	ENGINE_API bool IsHRTFDisabled() const;
 
 	void SetHRTFDisabled(bool InIsHRTFDisabled)
 	{
@@ -1260,23 +1260,23 @@ public:
 	}
 
 	/** Registers a third party listener-observer to this audio device. */
-	void RegisterPluginListener(const TAudioPluginListenerPtr PluginListener);
+	ENGINE_API void RegisterPluginListener(const TAudioPluginListenerPtr PluginListener);
 
 	/** Unregisters a third party listener-observer to this audio device. */
-	void UnregisterPluginListener(const TAudioPluginListenerPtr PluginListener);
+	ENGINE_API void UnregisterPluginListener(const TAudioPluginListenerPtr PluginListener);
 
-	bool IsAudioDeviceMuted() const;
+	ENGINE_API bool IsAudioDeviceMuted() const;
 
-	void SetDeviceMuted(bool bMuted);
+	ENGINE_API void SetDeviceMuted(bool bMuted);
 
 	/** Returns the azimuth angle of the sound relative to the sound's nearest listener. Used for 3d audio calculations. */
-	void GetAzimuth(const FAttenuationListenerData& OutListenerData, float& OutAzimuth, float& AbsoluteAzimuth) const;
+	ENGINE_API void GetAzimuth(const FAttenuationListenerData& OutListenerData, float& OutAzimuth, float& AbsoluteAzimuth) const;
 
 	/** Returns the focus factor of a sound based on its position and listener data. */
-	float GetFocusFactor(const float Azimuth, const FSoundAttenuationSettings& AttenuationSettings) const;
+	ENGINE_API float GetFocusFactor(const float Azimuth, const FSoundAttenuationSettings& AttenuationSettings) const;
 
 	/** Gets the max distance and focus factor of a sound. */
-	void GetMaxDistanceAndFocusFactor(USoundBase* Sound, const UWorld* World, const FVector& Location, const FSoundAttenuationSettings* AttenuationSettingsToApply, float& OutMaxDistance, float& OutFocusFactor);
+	ENGINE_API void GetMaxDistanceAndFocusFactor(USoundBase* Sound, const UWorld* World, const FVector& Location, const FSoundAttenuationSettings* AttenuationSettingsToApply, float& OutMaxDistance, float& OutFocusFactor);
 
 	/**
 	* Checks if the given sound would be audible.
@@ -1289,17 +1289,17 @@ public:
 	*
 	* @return Returns true if the sound is audible, false otherwise.
 	*/
-	bool SoundIsAudible(USoundBase* Sound, const UWorld* World, const FVector& Location, const FSoundAttenuationSettings* AttenuationSettingsToApply, float MaxDistance, float FocusFactor) const;
+	ENGINE_API bool SoundIsAudible(USoundBase* Sound, const UWorld* World, const FVector& Location, const FSoundAttenuationSettings* AttenuationSettingsToApply, float MaxDistance, float FocusFactor) const;
 
 	/** Returns the index of the listener closest to the given sound transform */
-	static int32 FindClosestListenerIndex(const FTransform& SoundTransform, const TArray<FListener>& InListeners);
+	ENGINE_API static int32 FindClosestListenerIndex(const FTransform& SoundTransform, const TArray<FListener>& InListeners);
 
 	/** Returns the index of the listener closest to the given sound transform */
-	int32 FindClosestListenerIndex(const FTransform& SoundTransform) const;
-	int32 FindClosestListenerIndex(const FVector& Position, float& OutSqDistance, bool AllowAttenuationOverrides) const;
+	ENGINE_API int32 FindClosestListenerIndex(const FTransform& SoundTransform) const;
+	ENGINE_API int32 FindClosestListenerIndex(const FVector& Position, float& OutSqDistance, bool AllowAttenuationOverrides) const;
 
 	/** Disables ActiveSound from responding to calls from its associated AudioComponent. */
-	void UnlinkActiveSoundFromComponent(const FActiveSound& InActiveSound);
+	ENGINE_API void UnlinkActiveSoundFromComponent(const FActiveSound& InActiveSound);
 
 	/** Return the audio stream time */
 	virtual double GetAudioTime() const
@@ -1313,13 +1313,13 @@ public:
 	}
 
 	/** Returns the main audio device of the engine */
-	static FAudioDeviceHandle GetMainAudioDevice();
+	ENGINE_API static FAudioDeviceHandle GetMainAudioDevice();
 
 	/** Returns the audio device manager */
-	static FAudioDeviceManager* GetAudioDeviceManager();
+	ENGINE_API static FAudioDeviceManager* GetAudioDeviceManager();
 
 	/** Low pass filter OneOverQ value */
-	float GetLowPassFilterResonance() const;
+	ENGINE_API float GetLowPassFilterResonance() const;
 
 	/** Returns the number of active sound sources */
 	virtual int32 GetNumActiveSources() const { return 0; }
@@ -1343,10 +1343,10 @@ public:
 	}
 
 	/** Set and initialize the current Spatialization plugin (by name) */
-	bool SetCurrentSpatializationPlugin(FName PluginName);
+	ENGINE_API bool SetCurrentSpatializationPlugin(FName PluginName);
 
 	/** Get the current list of available spatialization plugins */
-	TArray<FName> GetAvailableSpatializationPluginNames() const;
+	ENGINE_API TArray<FName> GetAvailableSpatializationPluginNames() const;
 
 	/** Return the spatialization plugin interface. */
 	TAudioSpatializationPtr GetSpatializationPluginInterface()
@@ -1356,9 +1356,9 @@ public:
 
 	/** Return the spatialization plugin interface info (requested by name). */
 	struct FAudioSpatializationInterfaceInfo;
-	FAudioSpatializationInterfaceInfo GetCurrentSpatializationPluginInterfaceInfo();
+	ENGINE_API FAudioSpatializationInterfaceInfo GetCurrentSpatializationPluginInterfaceInfo();
 
-	bool SpatializationPluginInterfacesAvailable();
+	ENGINE_API bool SpatializationPluginInterfacesAvailable();
 
 	/** Whether or not there's a modulation plugin enabled. */
 	bool IsModulationPluginEnabled() const
@@ -1372,7 +1372,7 @@ public:
 		return bOcclusionInterfaceEnabled;
 	}
 
-	static bool IsOcclusionPluginLoaded();
+	ENGINE_API static bool IsOcclusionPluginLoaded();
 
 	/** Whether or not there's a reverb plugin enabled. */
 	bool IsReverbPluginEnabled() const
@@ -1380,14 +1380,14 @@ public:
 		return bReverbInterfaceEnabled;
 	}
 
-	static bool IsReverbPluginLoaded();
+	ENGINE_API static bool IsReverbPluginLoaded();
 
 	bool IsSourceDataOverridePluginEnabled() const
 	{
 		return bSourceDataOverrideInterfaceEnabled;
 	}
 
-	static bool IsSourceDataOverridePluginLoaded();
+	ENGINE_API static bool IsSourceDataOverridePluginLoaded();
 
 	/** Returns if this is the multi-platform audio mixer. */
 	bool IsAudioMixerEnabled() const
@@ -1408,7 +1408,7 @@ public:
 	}
 
 	/** Performs an operation on all active sounds requested to execute by an audio component */
-	void SendCommandToActiveSounds(uint64 InAudioComponentID, TUniqueFunction<void(FActiveSound&)> InFunc, const TStatId InStatId = TStatId());
+	ENGINE_API void SendCommandToActiveSounds(uint64 InAudioComponentID, TUniqueFunction<void(FActiveSound&)> InFunc, const TStatId InStatId = TStatId());
 
 	virtual bool IsNonRealtime() const
 	{
@@ -1519,14 +1519,14 @@ public:
 	}
 
 	/** Adds an envelope follower delegate to the submix for this audio device. */
-	virtual void AddEnvelopeFollowerDelegate(USoundSubmix* InSubmix, const FOnSubmixEnvelopeBP& OnSubmixEnvelopeBP);
+	ENGINE_API virtual void AddEnvelopeFollowerDelegate(USoundSubmix* InSubmix, const FOnSubmixEnvelopeBP& OnSubmixEnvelopeBP);
 
-	virtual void StartSpectrumAnalysis(USoundSubmix* InSubmix, const FSoundSpectrumAnalyzerSettings& InSettings);
-	virtual void StopSpectrumAnalysis(USoundSubmix* InSubmix);
-	virtual void GetMagnitudesForFrequencies(USoundSubmix* InSubmix, const TArray<float>& InFrequencies, TArray<float>& OutMagnitudes);
-	virtual void GetPhasesForFrequencies(USoundSubmix* InSubmix, const TArray<float>& InFrequencies, TArray<float>& OutPhases);
-	virtual void AddSpectralAnalysisDelegate(USoundSubmix* InSubmix, const FSoundSpectrumAnalyzerDelegateSettings& InDelegateSettings, const FOnSubmixSpectralAnalysisBP& OnSubmixSpectralAnalysisBP);
-	virtual void RemoveSpectralAnalysisDelegate(USoundSubmix* InSubmix, const FOnSubmixSpectralAnalysisBP& OnSubmixSpectralAnalysisBP);
+	ENGINE_API virtual void StartSpectrumAnalysis(USoundSubmix* InSubmix, const FSoundSpectrumAnalyzerSettings& InSettings);
+	ENGINE_API virtual void StopSpectrumAnalysis(USoundSubmix* InSubmix);
+	ENGINE_API virtual void GetMagnitudesForFrequencies(USoundSubmix* InSubmix, const TArray<float>& InFrequencies, TArray<float>& OutMagnitudes);
+	ENGINE_API virtual void GetPhasesForFrequencies(USoundSubmix* InSubmix, const TArray<float>& InFrequencies, TArray<float>& OutPhases);
+	ENGINE_API virtual void AddSpectralAnalysisDelegate(USoundSubmix* InSubmix, const FSoundSpectrumAnalyzerDelegateSettings& InDelegateSettings, const FOnSubmixSpectralAnalysisBP& OnSubmixSpectralAnalysisBP);
+	ENGINE_API virtual void RemoveSpectralAnalysisDelegate(USoundSubmix* InSubmix, const FOnSubmixSpectralAnalysisBP& OnSubmixSpectralAnalysisBP);
 
 
 protected:
@@ -1535,27 +1535,27 @@ protected:
 	/**
 	 * Handle pausing/unpausing of sources when entering or leaving pause mode, or global pause (like device suspend)
 	 */
-	void HandlePause(bool bGameTicking, bool bGlobalPause = false);
+	ENGINE_API void HandlePause(bool bGameTicking, bool bGlobalPause = false);
 
 	/**
 	 * Stop sources that need to be stopped, and touch the ones that need to be kept alive
 	 * Stop sounds that are too low in priority to be played
 	 */
-	void StopSources(TArray<FWaveInstance*>& WaveInstances, int32 FirstActiveIndex);
+	ENGINE_API void StopSources(TArray<FWaveInstance*>& WaveInstances, int32 FirstActiveIndex);
 
 	/**
 	 * Start and/or update any sources that have a high enough priority to play
 	 */
-	void StartSources(TArray<FWaveInstance*>& WaveInstances, int32 FirstActiveIndex, bool bGameTicking);
+	ENGINE_API void StartSources(TArray<FWaveInstance*>& WaveInstances, int32 FirstActiveIndex, bool bGameTicking);
 
 	/**
 	 * This is overridden in Audio::FMixerDevice to propogate listener information to the audio thread.
 	 */
 	virtual void OnListenerUpdated(const TArray<FListener>& InListeners) {};
 
-	void NotifyAudioDevicePreRender(const FAudioDeviceRenderInfo& InInfo);
+	ENGINE_API void NotifyAudioDevicePreRender(const FAudioDeviceRenderInfo& InInfo);
 
-	void NotifyAudioDevicePostRender(const FAudioDeviceRenderInfo& InInfo);
+	ENGINE_API void NotifyAudioDevicePostRender(const FAudioDeviceRenderInfo& InInfo);
 
 private:
 
@@ -1712,18 +1712,18 @@ public:
 	/**
 	 * Platform dependent call to init effect data on a sound source
 	 */
-	void* InitEffect(FSoundSource* Source);
+	ENGINE_API void* InitEffect(FSoundSource* Source);
 
 	/**
 	 * Platform dependent call to update the sound output with new parameters
 	 * The audio system's main "Tick" function
 	 */
-	void* UpdateEffect(FSoundSource* Source);
+	ENGINE_API void* UpdateEffect(FSoundSource* Source);
 
 	/**
 	 * Platform dependent call to destroy any effect related data
 	 */
-	void DestroyEffect(FSoundSource* Source);
+	ENGINE_API void DestroyEffect(FSoundSource* Source);
 
 	/**
 	 * Return the pointer to the sound effects handler
@@ -1765,13 +1765,13 @@ public:
 		DefaultBaseSoundMix = InDefaultBaseSoundMix;
 	}
 
-	FDelegateHandle AddPreRenderDelegate(const FOnAudioDevicePreRender::FDelegate& InDelegate);
+	ENGINE_API FDelegateHandle AddPreRenderDelegate(const FOnAudioDevicePreRender::FDelegate& InDelegate);
 
-	bool RemovePreRenderDelegate(const FDelegateHandle& InHandle);
+	ENGINE_API bool RemovePreRenderDelegate(const FDelegateHandle& InHandle);
 
-	FDelegateHandle AddPostRenderDelegate(const FOnAudioDevicePostRender::FDelegate& InDelegate);
+	ENGINE_API FDelegateHandle AddPostRenderDelegate(const FOnAudioDevicePostRender::FDelegate& InDelegate);
 
-	bool RemovePostRenderDelegate(const FDelegateHandle& InHandle);
+	ENGINE_API bool RemovePostRenderDelegate(const FDelegateHandle& InHandle);
 
 private:
 	/**
@@ -1815,7 +1815,7 @@ public:
 	}
 
 	/** Creates a new platform specific sound source */
-	virtual FAudioEffectsManager* CreateEffectsManager();
+	ENGINE_API virtual FAudioEffectsManager* CreateEffectsManager();
 
 	/** Creates a new platform specific sound source */
 	virtual FSoundSource* CreateSoundSource() = 0;
@@ -1824,12 +1824,12 @@ public:
 	 * Marks a sound to be stopped.  Returns true if added to stop,
 	 * false if already pending stop.
 	 */
-	void AddSoundToStop(FActiveSound* SoundToStop);
+	ENGINE_API void AddSoundToStop(FActiveSound* SoundToStop);
 
 	/**
 	 * Whether the provided ActiveSound is currently pending to stop
 	 */
-	bool IsPendingStop(FActiveSound* ActiveSound);
+	ENGINE_API bool IsPendingStop(FActiveSound* ActiveSound);
 
 	/**
 	* Gets the direction of the given position vector transformed relative to listener.
@@ -1837,29 +1837,29 @@ public:
 	* @param OutDistance			Optional output of distance from position to listener
 	* @return The input position relative to the listener.
 	*/
-	FVector GetListenerTransformedDirection(const FVector& Position, float* OutDistance);
+	ENGINE_API FVector GetListenerTransformedDirection(const FVector& Position, float* OutDistance);
 
 	/** Returns the current audio device update delta time. */
-	float GetDeviceDeltaTime() const;
+	ENGINE_API float GetDeviceDeltaTime() const;
 
 	/** Returns the game's delta time */
-	float GetGameDeltaTime() const;
+	ENGINE_API float GetGameDeltaTime() const;
 
 	/** Whether device is using listener attenuation override or not. */
 	UE_DEPRECATED(4.25, "Use ParseAttenuation that passes a ListenerIndex instead")
 	bool IsUsingListenerAttenuationOverride() const { return IsUsingListenerAttenuationOverride(0); }
 
 	/** Returns if the specific listener is using an attenuation override position. */
-	bool IsUsingListenerAttenuationOverride(int32 ListenerIndex) const;
+	ENGINE_API bool IsUsingListenerAttenuationOverride(int32 ListenerIndex) const;
 
 	/** Returns the listener attenuation override */
 	UE_DEPRECATED(4.25, "Use ParseAttenuation that passes a ListenerIndex instead")
 	const FVector& GetListenerAttenuationOverride() const { return GetListenerAttenuationOverride(0); }
 
 	/** Returns the listener attenuation override for the specified listener */
-	const FVector& GetListenerAttenuationOverride(int32 ListenerIndex) const;
+	ENGINE_API const FVector& GetListenerAttenuationOverride(int32 ListenerIndex) const;
 
-	void UpdateVirtualLoops(bool bForceUpdate);
+	ENGINE_API void UpdateVirtualLoops(bool bForceUpdate);
 
 	/** Sets the update delta time for the audio frame */
 	virtual void UpdateDeviceDeltaTime()
@@ -1886,12 +1886,12 @@ private:
 public:
 
 	/** Query if the editor is in VR Preview for the current play world. Returns false for non-editor builds */
-	static bool CanUseVRAudioDevice();
+	ENGINE_API static bool CanUseVRAudioDevice();
 
 	/** Returns the audio clock of the audio device. Not supported on all platforms. */
 	double GetAudioClock() const { return AudioClock; }
 
-	void AddVirtualLoop(const FAudioVirtualLoop& InVirtualLoop);
+	ENGINE_API void AddVirtualLoop(const FAudioVirtualLoop& InVirtualLoop);
 
 	bool AreStartupSoundsPreCached() const { return bStartupSoundsPreCached; }
 
@@ -1899,33 +1899,33 @@ public:
 	float GetTransientMasterVolume() const { check(IsInAudioThread()); return TransientPrimaryVolume; }
 
 	UE_DEPRECATED(5.1, "SetTransientMasterVolume has been deprecated. Please use SetTransientPrimaryVolume instead.")
-	void SetTransientMasterVolume(float TransientPrimaryVolume);
+	ENGINE_API void SetTransientMasterVolume(float TransientPrimaryVolume);
 
 	UE_DEPRECATED(5.1, "GetMasterVolume has been deprecated. Please use GetPrimaryVolume instead.")
 	float GetMasterVolume() const { return PrimaryVolume; }
 
 	float GetTransientPrimaryVolume() const { check(IsInAudioThread()); return TransientPrimaryVolume; }
-	void SetTransientPrimaryVolume(float TransientPrimaryVolume);
+	ENGINE_API void SetTransientPrimaryVolume(float TransientPrimaryVolume);
 
 	/** Returns the volume that combines transient master volume and the FApp::GetVolumeMultiplier() value */
 	float GetPrimaryVolume() const { return PrimaryVolume; }
 
-	FSoundSource* GetSoundSource(FWaveInstance* WaveInstance) const;
+	ENGINE_API FSoundSource* GetSoundSource(FWaveInstance* WaveInstance) const;
 
-	const FGlobalFocusSettings& GetGlobalFocusSettings() const;
-	void SetGlobalFocusSettings(const FGlobalFocusSettings& NewFocusSettings);
+	ENGINE_API const FGlobalFocusSettings& GetGlobalFocusSettings() const;
+	ENGINE_API void SetGlobalFocusSettings(const FGlobalFocusSettings& NewFocusSettings);
 
 	const FDynamicParameter& GetGlobalPitchScale() const { check(IsInAudioThread()); return GlobalPitchScale; }
-	void SetGlobalPitchModulation(float PitchScale, float TimeSec);
-	float ClampPitch(float InPitchScale) const;
+	ENGINE_API void SetGlobalPitchModulation(float PitchScale, float TimeSec);
+	ENGINE_API float ClampPitch(float InPitchScale) const;
 
 	/** Overrides the attenuation scale used on a sound class. */
 	void SetSoundClassDistanceScale(USoundClass* InSoundClass, float DistanceScale, float TimeSec);
 
 	float GetPlatformAudioHeadroom() const { check(IsInAudioThread()); return PlatformAudioHeadroom; }
-	void SetPlatformAudioHeadroom(float PlatformHeadRoom);
+	ENGINE_API void SetPlatformAudioHeadroom(float PlatformHeadRoom);
 
-	const TMap<FName, FActivatedReverb>& GetActiveReverb() const;
+	ENGINE_API const TMap<FName, FActivatedReverb>& GetActiveReverb() const;
 
 	UE_DEPRECATED(4.13, "Direct access of SoundClasses is no longer allowed. Instead you should use the SoundMixClassOverride system")
 	const TMap<USoundClass*, FSoundClassProperties>& GetSoundClassPropertyMap() const
@@ -1937,7 +1937,7 @@ public:
 	/** Whether play when silent is enabled for all sounds associated with this audio device*/
 	bool PlayWhenSilentEnabled() const { return bAllowPlayWhenSilent; }
 
-	bool IsMainAudioDevice() const;
+	ENGINE_API bool IsMainAudioDevice() const;
 
 	/** Set whether or not we force the use of attenuation for non-game worlds (as by default we only care about game worlds) */
 	void SetUseAttenuationForNonGameWorlds(bool bInUseAttenuationForNonGameWorlds)
@@ -1945,14 +1945,14 @@ public:
 		bUseAttenuationForNonGameWorlds = bInUseAttenuationForNonGameWorlds;
 	}
 
-	const TArray<FWaveInstance*>& GetActiveWaveInstances() const;
+	ENGINE_API const TArray<FWaveInstance*>& GetActiveWaveInstances() const;
 
 	/** Returns the default reverb send level used for sources which have reverb applied but no attenuation settings. */
 	float GetDefaultReverbSendLevel() const { return DefaultReverbSendLevel; }
 
-	const TMap<FWaveInstance*, FSoundSource*>& GetWaveInstanceSourceMap() const;
+	ENGINE_API const TMap<FWaveInstance*, FSoundSource*>& GetWaveInstanceSourceMap() const;
 
-	FName GetAudioStateProperty(const FName& PropertyName) const;
+	ENGINE_API FName GetAudioStateProperty(const FName& PropertyName) const;
 	void SetAudioStateProperty(const FName& PropertyName, const FName& PropertyValue);
 
 	/** Get a Subsystem of specified type */
diff --git a/Engine/Source/Runtime/Experimental/Chaos/Public/Chaos/ErrorReporter.h b/Engine/Source/Runtime/Experimental/Chaos/Public/Chaos/ErrorReporter.h
--- a/Engine/Source/Runtime/Experimental/Chaos/Public/Chaos/ErrorReporter.h
+++ b/Engine/Source/Runtime/Experimental/Chaos/Public/Chaos/ErrorReporter.h
@@ -5,7 +5,7 @@
 
 namespace Chaos
 {
-class CHAOS_API FErrorReporter
+class FErrorReporter
 {
 public:
 	FErrorReporter(FString ErrorPrefix = "")
diff --git a/Engine/Source/Runtime/Experimental/Dataflow/Core/Public/Dataflow/DataflowNodeFactory.h b/Engine/Source/Runtime/Experimental/Dataflow/Core/Public/Dataflow/DataflowNodeFactory.h
--- a/Engine/Source/Runtime/Experimental/Dataflow/Core/Public/Dataflow/DataflowNodeFactory.h
+++ b/Engine/Source/Runtime/Experimental/Dataflow/Core/Public/Dataflow/DataflowNodeFactory.h
@@ -38,7 +38,7 @@ namespace Dataflow
 	//
 	//
 	//
-	class DATAFLOWCORE_API FNodeFactory
+	class FNodeFactory
 	{
 		typedef TFunction<TUniquePtr<FDataflowNode> (const FNewNodeParameters&)> FNewNodeFunction;
 
@@ -47,7 +47,7 @@ namespace Dataflow
 		TMap<FName, FFactoryParameters > ParametersMap;	// [TypeName] -> Parameters
 		TMap<FName, FName > DisplayMap;					// [DisplayName] -> TypeName
 
-		static FNodeFactory* Instance;
+		DATAFLOWCORE_API static FNodeFactory* Instance;
 		FNodeFactory() {}
 
 	public:
@@ -120,7 +120,7 @@ namespace Dataflow
 			return FFactoryParameters();
 		}
 
-		TSharedPtr<FDataflowNode> NewNodeFromRegisteredType(FGraph& Graph, const FNewNodeParameters& Param);
+		DATAFLOWCORE_API TSharedPtr<FDataflowNode> NewNodeFromRegisteredType(FGraph& Graph, const FNewNodeParameters& Param);
 
 		template<class T> TSharedPtr<T> NewNode(FGraph& Graph, const FNewNodeParameters& Param)
 		{
diff --git a/Engine/Source/Runtime/Experimental/Dataflow/Engine/Public/Dataflow/DataflowRenderingFactory.h b/Engine/Source/Runtime/Experimental/Dataflow/Engine/Public/Dataflow/DataflowRenderingFactory.h
--- a/Engine/Source/Runtime/Experimental/Dataflow/Engine/Public/Dataflow/DataflowRenderingFactory.h
+++ b/Engine/Source/Runtime/Experimental/Dataflow/Engine/Public/Dataflow/DataflowRenderingFactory.h
@@ -49,13 +49,13 @@ namespace Dataflow
 	//
 	//
 	//
-	class DATAFLOWENGINE_API FRenderingFactory
+	class FRenderingFactory
 	{
 		typedef TFunction<void(GeometryCollection::Facades::FRenderingFacade& RenderData, const FGraphRenderingState& State)> FOutputRenderingFunction;
 
 		// All Maps indexed by TypeName
 		TMap<FName, FOutputRenderingFunction > RenderMap;		// [TypeName] -> NewNodeFunction
-		static FRenderingFactory* Instance;
+		DATAFLOWENGINE_API static FRenderingFactory* Instance;
 		FRenderingFactory() {}
 
 	public:
@@ -84,7 +84,7 @@ namespace Dataflow
 			}
 		}
 
-		void RenderNodeOutput(GeometryCollection::Facades::FRenderingFacade& RenderData, const FGraphRenderingState& State);
+		DATAFLOWENGINE_API void RenderNodeOutput(GeometryCollection::Facades::FRenderingFacade& RenderData, const FGraphRenderingState& State);
 
 		bool Contains(FName InType) const { return RenderMap.Contains(InType); }
 
diff --git a/Engine/Source/Runtime/MovieSceneCapture/Public/Protocols/AudioCaptureProtocol.h b/Engine/Source/Runtime/MovieSceneCapture/Public/Protocols/AudioCaptureProtocol.h
--- a/Engine/Source/Runtime/MovieSceneCapture/Public/Protocols/AudioCaptureProtocol.h
+++ b/Engine/Source/Runtime/MovieSceneCapture/Public/Protocols/AudioCaptureProtocol.h
@@ -27,13 +27,13 @@ class MOVIESCENECAPTURE_API UNullAudioCaptureProtocol : public UMovieSceneAudioC
 * If the sequence evaluation hitches the audio will become desynchronized due to their being more time passed 
 * in real time (platform time) than in the sequence itself.
 */
-UCLASS(meta = (DisplayName = "Master Audio Submix", CommandLineID = "MasterAudioSubmix"))
-class MOVIESCENECAPTURE_API UMasterAudioSubmixCaptureProtocol : public UMovieSceneAudioCaptureProtocolBase
+UCLASS(MinimalApi, meta = (DisplayName = "Master Audio Submix", CommandLineID = "MasterAudioSubmix"))
+class UMasterAudioSubmixCaptureProtocol : public UMovieSceneAudioCaptureProtocolBase
 {
 public:
 	GENERATED_BODY()
 
-		UMasterAudioSubmixCaptureProtocol(const FObjectInitializer& Init)
+	UMasterAudioSubmixCaptureProtocol(const FObjectInitializer& Init)
 		: UMovieSceneAudioCaptureProtocolBase(Init)
 		, TotalGameRecordingTime(0.0)
 		, TotalPlatformRecordingTime(0.0)
diff --git a/Engine/Source/Runtime/PakFile/Public/IPlatformFilePak.h b/Engine/Source/Runtime/PakFile/Public/IPlatformFilePak.h
--- a/Engine/Source/Runtime/PakFile/Public/IPlatformFilePak.h
+++ b/Engine/Source/Runtime/PakFile/Public/IPlatformFilePak.h
@@ -1994,7 +1994,7 @@ public:
 /**
  * Platform file wrapper to be able to use pak files.
  **/
-class PAKFILE_API FPakPlatformFile : public IPlatformFile
+class FPakPlatformFile : public IPlatformFile
 {
 	struct FPakListEntry
 	{
@@ -2098,7 +2098,7 @@ class PAKFILE_API FPakPlatformFile : public IPlatformFile
 	 * @param BufferSize Sizeof of the buffer.
 	 * @return true if the operation was successfull, false otherwise.
 	 */
-	bool BufferedCopyFile(IFileHandle& Dest, IFileHandle& Source, const int64 FileSize, uint8* Buffer, const int64 BufferSize) const;
+	PAKFILE_API bool BufferedCopyFile(IFileHandle& Dest, IFileHandle& Source, const int64 FileSize, uint8* Buffer, const int64 BufferSize) const;
 
 	/**
 	 * Creates file handle to read from Pak file.
@@ -2108,29 +2108,29 @@ class PAKFILE_API FPakPlatformFile : public IPlatformFile
 	 * @param FileEntry File entry to create the handle for.
 	 * @return Pointer to the new handle.
 	 */
-	IFileHandle* CreatePakFileHandle(const TCHAR* Filename, const TRefCountPtr<FPakFile>& PakFile, const FPakEntry* FileEntry);
+	PAKFILE_API IFileHandle* CreatePakFileHandle(const TCHAR* Filename, const TRefCountPtr<FPakFile>& PakFile, const FPakEntry* FileEntry);
 
 	/**
 	* Hardcode default load ordering of game main pak -> game content -> engine content -> saved dir
 	* would be better to make this config but not even the config system is initialized here so we can't do that
 	*/
-	static int32 GetPakOrderFromPakFilePath(const FString& PakFilePath);
+	PAKFILE_API static int32 GetPakOrderFromPakFilePath(const FString& PakFilePath);
 
 	/**
 	 * Handler for device delegate to prompt us to load a new pak.	 
 	 */
-	IPakFile* HandleMountPakDelegate(const FString& PakFilePath, int32 PakOrder);
+	PAKFILE_API IPakFile* HandleMountPakDelegate(const FString& PakFilePath, int32 PakOrder);
 
 	/**
 	 * Handler for device delegate to prompt us to load a new pak.
 	 */
 	UE_DEPRECATED(4.26, "Use HandleMountPakDelegate instead")
-	bool HandleOnMountPakDelegate(const FString& PakFilePath, int32 PakOrder, IPlatformFile::FDirectoryVisitor* Visitor);
+	PAKFILE_API bool HandleOnMountPakDelegate(const FString& PakFilePath, int32 PakOrder, IPlatformFile::FDirectoryVisitor* Visitor);
 
 	/**
 	 * Handler for device delegate to prompt us to unload a pak.
 	 */
-	bool HandleUnmountPakDelegate(const FString& PakFilePath);
+	PAKFILE_API bool HandleUnmountPakDelegate(const FString& PakFilePath);
 
 	/**
 	 * Finds all pak files in the given directory.
@@ -2138,14 +2138,14 @@ class PAKFILE_API FPakPlatformFile : public IPlatformFile
 	 * @param Directory Directory to (recursively) look for pak files in
 	 * @param OutPakFiles List of pak files
 	 */
-	static void FindPakFilesInDirectory(IPlatformFile* LowLevelFile, const TCHAR* Directory, const FString& WildCard, TArray<FString>& OutPakFiles);
+	PAKFILE_API static void FindPakFilesInDirectory(IPlatformFile* LowLevelFile, const TCHAR* Directory, const FString& WildCard, TArray<FString>& OutPakFiles);
 
 	/**
 	 * Finds all pak files in the known pak folders
 	 *
 	 * @param OutPakFiles List of all found pak files
 	 */
-	static void FindAllPakFiles(IPlatformFile* LowLevelFile, const TArray<FString>& PakFolders, const FString& WildCard, TArray<FString>& OutPakFiles);
+	PAKFILE_API static void FindAllPakFiles(IPlatformFile* LowLevelFile, const TArray<FString>& PakFolders, const FString& WildCard, TArray<FString>& OutPakFiles);
 
 	/**
 	 * When security is enabled, determine if this filename can be looked for in the lower level file system
@@ -2153,7 +2153,7 @@ class PAKFILE_API FPakPlatformFile : public IPlatformFile
 	 * @param InFilename			Filename to check
 	 * @param bAllowDirectories		Consider directories as valid filepaths?
 	 */
-	bool IsNonPakFilenameAllowed(const FString& InFilename);
+	PAKFILE_API bool IsNonPakFilenameAllowed(const FString& InFilename);
 
 	/**
 	 * Registers a new AES key with the given guid. Triggers the mounting of any pak files that we encountered that use that key
@@ -2161,7 +2161,7 @@ class PAKFILE_API FPakPlatformFile : public IPlatformFile
 	 * @param InEncryptionKeyGuid	Guid for this encryption key
 	 * @param InKey					Encryption key
 	 */
-	void RegisterEncryptionKey(const FGuid& InEncryptionKeyGuid, const FAES::FAESKey& InKey);
+	PAKFILE_API void RegisterEncryptionKey(const FGuid& InEncryptionKeyGuid, const FAES::FAESKey& InKey);
 
 	/**
 	 * Checks with any current chunk installation system if the given pak file is installed
@@ -2169,7 +2169,7 @@ class PAKFILE_API FPakPlatformFile : public IPlatformFile
 	 * @param InFilename  the pak filename to check
 	 * @return whether the pak file is installed
 	 */
-	static bool IsPakFileInstalled(const FString& InFilename);
+	PAKFILE_API static bool IsPakFileInstalled(const FString& InFilename);
 
 public:
 
@@ -2190,22 +2190,22 @@ public:
 	/**
 	 * Get the wild card pattern used to identify paks to load on startup
 	 */
-	static const TCHAR* GetMountStartupPaksWildCard();
+	PAKFILE_API static const TCHAR* GetMountStartupPaksWildCard();
 
 	/**
 	 * Overrides the wildcard used for searching paks. Call before initialization
 	 */
-	static void SetMountStartupPaksWildCard(const FString& WildCard);
+	PAKFILE_API static void SetMountStartupPaksWildCard(const FString& WildCard);
 
 	/**
 	* Determine location information for a given pakchunk index. Will be DoesNotExist if the pak file wasn't detected, NotAvailable if it exists but hasn't been mounted due to a missing encryption key, or LocalFast if it exists and has been mounted
 	*/
-	EChunkLocation::Type GetPakChunkLocation(int32 InPakchunkIndex) const;
+	PAKFILE_API EChunkLocation::Type GetPakChunkLocation(int32 InPakchunkIndex) const;
 
 	/**
 	* Returns true if any of the mounted or pending pak files are chunks (filenames starting pakchunkN)
 	*/
-	bool AnyChunksAvailable() const;
+	PAKFILE_API bool AnyChunksAvailable() const;
 
 	/**
 	* Get a list of all pak files which have been successfully mounted
@@ -2223,47 +2223,47 @@ public:
 	/**
 	 * Checks if pak files exist in any of the known pak file locations.
 	 */
-	static bool CheckIfPakFilesExist(IPlatformFile* LowLevelFile, const TArray<FString>& PakFolders);
+	PAKFILE_API static bool CheckIfPakFilesExist(IPlatformFile* LowLevelFile, const TArray<FString>& PakFolders);
 
 	/**
 	 * Gets all pak file locations.
 	 */
-	static void GetPakFolders(const TCHAR* CmdLine, TArray<FString>& OutPakFolders);
+	PAKFILE_API static void GetPakFolders(const TCHAR* CmdLine, TArray<FString>& OutPakFolders);
 
 	/**
 	* Helper function for accessing pak encryption key
 	*/
-	static void GetPakEncryptionKey(FAES::FAESKey& OutKey, const FGuid& InEncryptionKeyGuid);
+	PAKFILE_API static void GetPakEncryptionKey(FAES::FAESKey& OutKey, const FGuid& InEncryptionKeyGuid);
 
 	/**
 	* Load a pak signature file. Validates the contents by comparing a SHA hash of the chunk table against and encrypted version that
 	* is stored within the file. Returns nullptr if the data is missing or fails the signature check. This function also calls
 	* the generic pak signature failure delegates if anything is wrong.
 	*/
-	static TSharedPtr<const struct FPakSignatureFile, ESPMode::ThreadSafe> GetPakSignatureFile(const TCHAR* InFilename);
+	PAKFILE_API static TSharedPtr<const struct FPakSignatureFile, ESPMode::ThreadSafe> GetPakSignatureFile(const TCHAR* InFilename);
 
 	/**
 	 * Remove the intenrally cached pointer to the signature file for the specified pak
 	 */
-	static void RemoveCachedPakSignaturesFile(const TCHAR* InFilename);
+	PAKFILE_API static void RemoveCachedPakSignaturesFile(const TCHAR* InFilename);
 
 	/**
 	 * Constructor.
 	 * 
 	 * @param InLowerLevel Wrapper platform file.
 	 */
-	FPakPlatformFile();
+	PAKFILE_API FPakPlatformFile();
 
 	/**
 	 * Destructor.
 	 */
-	virtual ~FPakPlatformFile();
+	PAKFILE_API virtual ~FPakPlatformFile();
 
-	virtual bool ShouldBeUsed(IPlatformFile* Inner, const TCHAR* CmdLine) const override;
-	virtual bool Initialize(IPlatformFile* Inner, const TCHAR* CommandLineParam) override;
-	virtual void InitializeNewAsyncIO() override;
+	PAKFILE_API virtual bool ShouldBeUsed(IPlatformFile* Inner, const TCHAR* CmdLine) const override;
+	PAKFILE_API virtual bool Initialize(IPlatformFile* Inner, const TCHAR* CommandLineParam) override;
+	PAKFILE_API virtual void InitializeNewAsyncIO() override;
 
-	void OptimizeMemoryUsageForMountedPaks();
+	PAKFILE_API void OptimizeMemoryUsageForMountedPaks();
 
 	virtual IPlatformFile* GetLowerLevel() override
 	{
@@ -2279,7 +2279,7 @@ public:
 		return FPakPlatformFile::GetTypeName();
 	}
 
-	void Tick() override;
+	PAKFILE_API void Tick() override;
 
 	/**
 	 * Mounts a pak file at the specified path.
@@ -2287,22 +2287,22 @@ public:
 	 * @param InPakFilename Pak filename.
 	 * @param InPath Path to mount the pak at.
 	 */
-	bool Mount(const TCHAR* InPakFilename, uint32 PakOrder, const TCHAR* InPath = NULL, bool bLoadIndex = true);
+	PAKFILE_API bool Mount(const TCHAR* InPakFilename, uint32 PakOrder, const TCHAR* InPath = NULL, bool bLoadIndex = true);
 
-	bool Unmount(const TCHAR* InPakFilename);
+	PAKFILE_API bool Unmount(const TCHAR* InPakFilename);
 
-	int32 MountAllPakFiles(const TArray<FString>& PakFolders);
-	int32 MountAllPakFiles(const TArray<FString>& PakFolders, const FString& WildCard);
+	PAKFILE_API int32 MountAllPakFiles(const TArray<FString>& PakFolders);
+	PAKFILE_API int32 MountAllPakFiles(const TArray<FString>& PakFolders, const FString& WildCard);
 
 	/**
 	 * Re-creates all the pak readers
 	 */
-	bool ReloadPakReaders();
+	PAKFILE_API bool ReloadPakReaders();
 
 	/**
 	 * Make unique in memory pak files from a list of named files
 	 */
-	virtual void MakeUniquePakFilesForTheseFiles(const TArray<TArray<FString>>& InFiles);
+	PAKFILE_API virtual void MakeUniquePakFilesForTheseFiles(const TArray<TArray<FString>>& InFiles);
 
 
 	/** Overload needed for deprecation; remove this when removing the version with a FPakFile** OutPakFile */
@@ -2631,7 +2631,7 @@ public:
 		return LowerLevel->IsSymlink(Filename);
 	}
 
-	virtual IFileHandle* OpenRead(const TCHAR* Filename, bool bAllowWrite = false) override;
+	PAKFILE_API virtual IFileHandle* OpenRead(const TCHAR* Filename, bool bAllowWrite = false) override;
 
 	virtual IFileHandle* OpenWrite(const TCHAR* Filename, bool bAppend = false, bool bAllowRead = false) override
 	{
@@ -3016,10 +3016,10 @@ public:
 		return LowerLevel->CreateDirectoryTree(Directory);
 	}
 
-	virtual bool CopyFile(const TCHAR* To, const TCHAR* From, EPlatformFileRead ReadFlags = EPlatformFileRead::None, EPlatformFileWrite WriteFlags = EPlatformFileWrite::None) override;
+	PAKFILE_API virtual bool CopyFile(const TCHAR* To, const TCHAR* From, EPlatformFileRead ReadFlags = EPlatformFileRead::None, EPlatformFileWrite WriteFlags = EPlatformFileWrite::None) override;
 
-	virtual IAsyncReadFileHandle* OpenAsyncRead(const TCHAR* Filename) override;
-	virtual void SetAsyncMinimumPriority(EAsyncIOPriorityAndFlags Priority) override;
+	PAKFILE_API virtual IAsyncReadFileHandle* OpenAsyncRead(const TCHAR* Filename) override;
+	PAKFILE_API virtual void SetAsyncMinimumPriority(EAsyncIOPriorityAndFlags Priority) override;
 
 	virtual IMappedFileHandle* OpenMapped(const TCHAR* Filename) override;
 	/**
@@ -3065,10 +3065,10 @@ public:
 	//~ End IPlatformFile Interface
 
 	// Access static delegate for loose file security
-	static FFilenameSecurityDelegate& GetFilenameSecurityDelegate();
+	PAKFILE_API static FFilenameSecurityDelegate& GetFilenameSecurityDelegate();
 
 	// Access static delegate for custom encryption
-	static FPakCustomEncryptionDelegate& GetPakCustomEncryptionDelegate();
+	PAKFILE_API static FPakCustomEncryptionDelegate& GetPakCustomEncryptionDelegate();
 
 	struct FPakSigningFailureHandlerData
 	{
@@ -3089,56 +3089,56 @@ public:
 	};
 
 	// Access static delegate for handling a Pak signature check failure
-	static FPakSigningFailureHandlerData& GetPakSigningFailureHandlerData();
+	PAKFILE_API static FPakSigningFailureHandlerData& GetPakSigningFailureHandlerData();
 	
 	// Broadcast a signature check failure through any registered delegates in a thread safe way
-	static void BroadcastPakChunkSignatureCheckFailure(const FPakChunkSignatureCheckFailedData& InData);
+	PAKFILE_API static void BroadcastPakChunkSignatureCheckFailure(const FPakChunkSignatureCheckFailedData& InData);
 
 	// Broadcast a principal signature table failure through any registered delegates in a thread safe way
-	static void BroadcastPakPrincipalSignatureTableCheckFailure(const FString& InFilename);
+	PAKFILE_API static void BroadcastPakPrincipalSignatureTableCheckFailure(const FString& InFilename);
 
 	UE_DEPRECATED("5.1", "Use BroadcastPakPrincipalSignatureTableCheckFailure instead")
-	static void BroadcastPakMasterSignatureTableCheckFailure(const FString& InFilename);
+	PAKFILE_API static void BroadcastPakMasterSignatureTableCheckFailure(const FString& InFilename);
 
 	// Access static delegate for setting PakIndex settings.
-	static FPakSetIndexSettings& GetPakSetIndexSettingsDelegate();
+	PAKFILE_API static FPakSetIndexSettings& GetPakSetIndexSettingsDelegate();
 
 	/* Get a list of RelativePathFromMount for every file in the given Pak that lives in any of the given chunks.  Only searches the Pruned DirectoryIndex */
-	void GetPrunedFilenamesInChunk(const FString& InPakFilename, const TArray<int32>& InChunkIDs, TArray<FString>& OutFileList);
+	PAKFILE_API void GetPrunedFilenamesInChunk(const FString& InPakFilename, const TArray<int32>& InChunkIDs, TArray<FString>& OutFileList);
 
 	/** Gets a list of FullPaths (includes Mount directory) for every File in the given Pak's Pruned DirectoryIndex */
-	void GetPrunedFilenamesInPakFile(const FString& InPakFilename, TArray<FString>& OutFileList);
+	PAKFILE_API void GetPrunedFilenamesInPakFile(const FString& InPakFilename, TArray<FString>& OutFileList);
 
 	/** Returns the RelativePathFromMount Filename for every file found in the given Iostore Container */
-	static void GetFilenamesFromIostoreContainer(const FString& InContainerName, TArray<FString>& OutFileList);
+	PAKFILE_API static void GetFilenamesFromIostoreContainer(const FString& InContainerName, TArray<FString>& OutFileList);
 
 	/** Returns the RelativePathFromMount Filename for every Filename found in the Iostore Container that relates to the provided block indexes */
-	static void GetFilenamesFromIostoreByBlockIndex(const FString& InContainerName, const TArray<int32>& InBlockIndex, TArray<FString>& OutFileList);
-	
+	PAKFILE_API static void GetFilenamesFromIostoreByBlockIndex(const FString& InContainerName, const TArray<int32>& InBlockIndex, TArray<FString>& OutFileList);
+
 	/** Iterates Iostore Container while Predicate returns true */
-	static void ForeachPackageInIostoreWhile(TFunctionRef<bool(FName)> Predicate);
+	PAKFILE_API static void ForeachPackageInIostoreWhile(TFunctionRef<bool(FName)> Predicate);
 
-	void ReleaseOldReaders();
+	PAKFILE_API void ReleaseOldReaders();
 
 	// BEGIN Console commands
 #if !UE_BUILD_SHIPPING
-	void HandlePakListCommand(const TCHAR* Cmd, FOutputDevice& Ar);
-	void HandleMountCommand(const TCHAR* Cmd, FOutputDevice& Ar);
-	void HandleUnmountCommand(const TCHAR* Cmd, FOutputDevice& Ar);
-	void HandlePakCorruptCommand(const TCHAR* Cmd, FOutputDevice& Ar);
-	void HandleReloadPakReadersCommand(const TCHAR* Cmd, FOutputDevice& Ar);
+	PAKFILE_API void HandlePakListCommand(const TCHAR* Cmd, FOutputDevice& Ar);
+	PAKFILE_API void HandleMountCommand(const TCHAR* Cmd, FOutputDevice& Ar);
+	PAKFILE_API void HandleUnmountCommand(const TCHAR* Cmd, FOutputDevice& Ar);
+	PAKFILE_API void HandlePakCorruptCommand(const TCHAR* Cmd, FOutputDevice& Ar);
+	PAKFILE_API void HandleReloadPakReadersCommand(const TCHAR* Cmd, FOutputDevice& Ar);
 #endif
 	// END Console commands
 	
 #if PAK_TRACKER
-	static TMap<FString, int32> GPakSizeMap;
-	static void TrackPak(const TCHAR* Filename, const FPakEntry* PakEntry);
+	PAKFILE_API static TMap<FString, int32> GPakSizeMap;
+	PAKFILE_API static void TrackPak(const TCHAR* Filename, const FPakEntry* PakEntry);
 	static TMap<FString, int32>& GetPakMap() { return GPakSizeMap; }
 #endif
 
 	// Internal cache of pak signature files
-	static TMap<FName, TSharedPtr<const struct FPakSignatureFile, ESPMode::ThreadSafe>> PakSignatureFileCache;
-	static FCriticalSection PakSignatureFileCacheLock;
+	PAKFILE_API static TMap<FName, TSharedPtr<const struct FPakSignatureFile, ESPMode::ThreadSafe>> PakSignatureFileCache;
+	PAKFILE_API static FCriticalSection PakSignatureFileCacheLock;
 };
 
 /**
diff --git a/Engine/Source/Runtime/RigVM/Public/RigVMCore/RigVMDispatchFactory.h b/Engine/Source/Runtime/RigVM/Public/RigVMCore/RigVMDispatchFactory.h
--- a/Engine/Source/Runtime/RigVM/Public/RigVMCore/RigVMDispatchFactory.h
+++ b/Engine/Source/Runtime/RigVM/Public/RigVMCore/RigVMDispatchFactory.h
@@ -42,7 +42,7 @@ struct RIGVM_API FRigVMDispatchContext
  * A factory to generate a template and its dispatch functions
  */
 USTRUCT()
-struct RIGVM_API FRigVMDispatchFactory
+struct FRigVMDispatchFactory
 {
 	GENERATED_BODY()
 
@@ -56,7 +56,7 @@ public:
 	virtual ~FRigVMDispatchFactory() {}
 
 	// returns the name of the factory template
-	FName GetFactoryName() const;
+	RIGVM_API FName GetFactoryName() const;
 
 	// returns the struct for this dispatch factory
 	UScriptStruct* GetScriptStruct() const { return FactoryScriptStruct; };
@@ -67,13 +67,13 @@ public:
 	virtual FString GetNodeTitle(const FRigVMTemplateTypeMap& InTypes) const { return GetScriptStruct()->GetDisplayNameText().ToString(); }
 
 	// returns the color of the node for a given type set.
-	virtual FLinearColor GetNodeColor() const;
+	RIGVM_API virtual FLinearColor GetNodeColor() const;
 
 	// returns the tooltip for the node
-	virtual FText GetNodeTooltip(const FRigVMTemplateTypeMap& InTypes) const;
+	RIGVM_API virtual FText GetNodeTooltip(const FRigVMTemplateTypeMap& InTypes) const;
 
 	// returns the default value for an argument
-	virtual FString GetArgumentDefaultValue(const FName& InArgumentName, TRigVMTypeIndex InTypeIndex) const;
+	RIGVM_API virtual FString GetArgumentDefaultValue(const FName& InArgumentName, TRigVMTypeIndex InTypeIndex) const;
 
 	// returns the tooltip for an argument
 	virtual FText GetArgumentTooltip(const FName& InArgumentName, TRigVMTypeIndex InTypeIndex) const { return FText(); }
@@ -85,13 +85,13 @@ public:
 	virtual TArray<FName> GetAggregateOutputArguments() const { return TArray<FName>(); }
 
 	// returns the next name to be used for an aggregate pin
-	virtual FName GetNextAggregateName(const FName& InLastAggregatePinName) const;
+	RIGVM_API virtual FName GetNextAggregateName(const FName& InLastAggregatePinName) const;
 
 	// Returns the display name text for an argument 
-	virtual FName GetDisplayNameForArgument(const FName& InArgumentName) const;
+	RIGVM_API virtual FName GetDisplayNameForArgument(const FName& InArgumentName) const;
 
 	// Returns meta data on the property of the permutations 
-	virtual FString GetArgumentMetaData(const FName& InArgumentName, const FName& InMetaDataKey) const;
+	RIGVM_API virtual FString GetArgumentMetaData(const FName& InArgumentName, const FName& InMetaDataKey) const;
 
 	// Returns true if the factory provides metadata for a given argument
 	bool HasArgumentMetaData(const FName& InArgumentName, const FName& InMetaDataKey) const
@@ -100,26 +100,26 @@ public:
 	}
 
 	// Returns the category this factory is under
-	virtual FString GetCategory() const;
+	RIGVM_API virtual FString GetCategory() const;
 
 	// Returns the keywords used for looking up this factory
-	virtual FString GetKeywords() const;
+	RIGVM_API virtual FString GetKeywords() const;
 
 	// Returns true if the argument is lazy
-	bool IsLazyInputArgument(const FName& InArgumentName) const;
+	RIGVM_API bool IsLazyInputArgument(const FName& InArgumentName) const;
 
 #endif
 
 	// Returns the name to use for the branch info / argument based on the operand index.
 	// operand index count may be different than the number of arguments, since fixed size array
 	// arguments get unrolled.
-	virtual FName GetArgumentNameForOperandIndex(int32 InOperandIndex, int32 InTotalOperands) const;
+	RIGVM_API virtual FName GetArgumentNameForOperandIndex(int32 InOperandIndex, int32 InTotalOperands) const;
 
 	// returns true if the dispatch is a control flow dispatch
 	bool IsControlFlowDispatch(const FRigVMDispatchContext& InContext) const { return !GetControlFlowBlocks(InContext).IsEmpty(); }
 
 	// returns the control flow blocks of this dispatch
-	const TArray<FName>& GetControlFlowBlocks(const FRigVMDispatchContext& InContext) const;
+	RIGVM_API const TArray<FName>& GetControlFlowBlocks(const FRigVMDispatchContext& InContext) const;
 
 	// returns true if a given control flow block needs to be sliced
 	virtual const bool IsControlFlowBlockSliced(const FName& InBlockName) const { return false; }
@@ -128,7 +128,7 @@ public:
 	virtual UScriptStruct* GetExecuteContextStruct() const { return FRigVMExecuteContext::StaticStruct(); }
 
 	// Returns truf if this factory supports a given executecontext struct
-	bool SupportsExecuteContextStruct(const UScriptStruct* InExecuteContextStruct) const;
+	RIGVM_API bool SupportsExecuteContextStruct(const UScriptStruct* InExecuteContextStruct) const;
 
 	// registered needed types during registration of the factory
 	virtual void RegisterDependencyTypes() const {}
@@ -137,7 +137,7 @@ public:
 	virtual TArray<FRigVMTemplateArgument> GetArguments() const { return TArray<FRigVMTemplateArgument>(); }
 
 	// returns the execute arguments of the template
-	TArray<FRigVMExecuteArgument> GetExecuteArguments(const FRigVMDispatchContext& InContext) const;
+	RIGVM_API TArray<FRigVMExecuteArgument> GetExecuteArguments(const FRigVMDispatchContext& InContext) const;
 
 	// returns the delegate to react to new types being added to an argument.
 	// this happens if types are being loaded later after this factory has already been deployed
@@ -147,13 +147,13 @@ public:
 	virtual FRigVMStructUpgradeInfo GetUpgradeInfo(const FRigVMTemplateTypeMap& InTypes, const FRigVMDispatchContext& InContext) const { return FRigVMStructUpgradeInfo(); }
 
 	// returns the dispatch function for a given type set
-	FRigVMFunctionPtr GetDispatchFunction(const FRigVMTemplateTypeMap& InTypes) const;
+	RIGVM_API FRigVMFunctionPtr GetDispatchFunction(const FRigVMTemplateTypeMap& InTypes) const;
 
 	// builds and returns the template
-	const FRigVMTemplate* GetTemplate() const;
+	RIGVM_API const FRigVMTemplate* GetTemplate() const;
 
 	// returns the name of the permutation for a given set of types
-	FString GetPermutationName(const FRigVMTemplateTypeMap& InTypes) const;
+	RIGVM_API FString GetPermutationName(const FRigVMTemplateTypeMap& InTypes) const;
 
 	// returns true if the dispatch uses the same function ptr for all permutations
 	virtual bool IsSingleton() const { return false; } 
@@ -161,13 +161,13 @@ public:
 protected:
 
 	// returns the name of the permutation for a given set of types
-	FString GetPermutationNameImpl(const FRigVMTemplateTypeMap& InTypes) const;
+	RIGVM_API FString GetPermutationNameImpl(const FRigVMTemplateTypeMap& InTypes) const;
 
 	virtual FRigVMFunctionPtr GetDispatchFunctionImpl(const FRigVMTemplateTypeMap& InTypes) const { return nullptr; }
 
 	virtual TArray<FRigVMExecuteArgument> GetExecuteArguments_Impl(const FRigVMDispatchContext& InContext) const { return TArray<FRigVMExecuteArgument>(); }
 
-	static bool CopyProperty(
+	RIGVM_API static bool CopyProperty(
 		const FProperty* InTargetProperty,
 		uint8* InTargetPtr,
 		const FProperty* InSourceProperty,
@@ -212,10 +212,10 @@ protected:
 	}
 #endif
 
-	virtual const TArray<FName>& GetControlFlowBlocks_Impl(const FRigVMDispatchContext& InContext) const;
+	RIGVM_API virtual const TArray<FName>& GetControlFlowBlocks_Impl(const FRigVMDispatchContext& InContext) const;
 
-	static const FString DispatchPrefix;
-	static const FString TrueString;
+	RIGVM_API static const FString DispatchPrefix;
+	RIGVM_API static const FString TrueString;
 
 	UScriptStruct* FactoryScriptStruct;
 	mutable const FRigVMTemplate* CachedTemplate;
diff --git a/Engine/Source/Runtime/RigVM/Public/RigVMCore/RigVMExecuteContext.h b/Engine/Source/Runtime/RigVM/Public/RigVMCore/RigVMExecuteContext.h
--- a/Engine/Source/Runtime/RigVM/Public/RigVMCore/RigVMExecuteContext.h
+++ b/Engine/Source/Runtime/RigVM/Public/RigVMCore/RigVMExecuteContext.h
@@ -221,7 +221,7 @@ struct RIGVM_API FRigVMRuntimeSettings
  * indicate execution order.
  */
 USTRUCT(BlueprintType, meta=(DisplayName="Execute Context"))
-struct RIGVM_API FRigVMExecuteContext
+struct FRigVMExecuteContext
 {
 	GENERATED_BODY()
 
@@ -307,11 +307,11 @@ struct RIGVM_API FRigVMExecuteContext
 	/** The world this VM is running in */
 	const UWorld* GetWorld() const { return World; }
 
-	void SetOwningComponent(const USceneComponent* InOwningComponent);
+	RIGVM_API void SetOwningComponent(const USceneComponent* InOwningComponent);
 
-	void SetOwningActor(const AActor* InActor);
+	RIGVM_API void SetOwningActor(const AActor* InActor);
 
-	void SetWorld(const UWorld* InWorld);
+	RIGVM_API void SetWorld(const UWorld* InWorld);
 
 	void SetToWorldSpaceTransform(const FTransform& InToWorldSpaceTransform) { ToWorldSpaceTransform = InToWorldSpaceTransform; }
 
@@ -398,7 +398,7 @@ struct RIGVM_API FRigVMExecuteContext
 	/**
 	 * Serialize this type from another
 	 */
-	bool SerializeFromMismatchedTag(const FPropertyTag& Tag, FStructuredArchive::FSlot Slot);
+	RIGVM_API bool SerializeFromMismatchedTag(const FPropertyTag& Tag, FStructuredArchive::FSlot Slot);
 
 protected:
 
diff --git a/Engine/Source/Runtime/SlateCore/Public/Layout/ChildrenBase.h b/Engine/Source/Runtime/SlateCore/Public/Layout/ChildrenBase.h
--- a/Engine/Source/Runtime/SlateCore/Public/Layout/ChildrenBase.h
+++ b/Engine/Source/Runtime/SlateCore/Public/Layout/ChildrenBase.h
@@ -24,7 +24,7 @@ class FSlotBase;
  * FChildren is intended to be returned by the GetChildren() method.
  */
 PRAGMA_DISABLE_DEPRECATION_WARNINGS
-class SLATECORE_API FChildren
+class FChildren
 {
 public:
 	FChildren(SWidget* InOwner)
diff --git a/Engine/Source/Runtime/SlateCore/Public/Styling/SlateWidgetStyleAsset.h b/Engine/Source/Runtime/SlateCore/Public/Styling/SlateWidgetStyleAsset.h
--- a/Engine/Source/Runtime/SlateCore/Public/Styling/SlateWidgetStyleAsset.h
+++ b/Engine/Source/Runtime/SlateCore/Public/Styling/SlateWidgetStyleAsset.h
@@ -13,8 +13,8 @@ class Error;
 /**
  * Just a wrapper for the struct with real data in it.
  */
-UCLASS(hidecategories=Object)
-class SLATECORE_API USlateWidgetStyleAsset : public UObject
+UCLASS(MinimalApi, hidecategories=Object)
+class USlateWidgetStyleAsset : public UObject
 {
 	GENERATED_BODY()
 		  
