Fix build break from missing _API specifier in CL24392821.

#rnx

diff --git a/Engine/Source/Runtime/Engine/Public/AudioDevice.h b/Engine/Source/Runtime/Engine/Public/AudioDevice.h
--- a/Engine/Source/Runtime/Engine/Public/AudioDevice.h
+++ b/Engine/Source/Runtime/Engine/Public/AudioDevice.h
@@ -1953,7 +1953,7 @@ public:
 	ENGINE_API const TMap<FWaveInstance*, FSoundSource*>& GetWaveInstanceSourceMap() const;
 
 	ENGINE_API FName GetAudioStateProperty(const FName& PropertyName) const;
-	void SetAudioStateProperty(const FName& PropertyName, const FName& PropertyValue);
+	ENGINE_API void SetAudioStateProperty(const FName& PropertyName, const FName& PropertyValue);
 
 	/** Get a Subsystem of specified type */
 	UAudioEngineSubsystem* GetSubsystemBase(TSubclassOf<UAudioEngineSubsystem> SubsystemClass) const
