Improve shader compilation stats output and clamp percentage for non-completed jobs to 99.99%.

#rnx

diff --git a/Engine/Source/Runtime/Engine/Private/ShaderCompiler/ShaderCompiler.cpp b/Engine/Source/Runtime/Engine/Private/ShaderCompiler/ShaderCompiler.cpp
--- a/Engine/Source/Runtime/Engine/Private/ShaderCompiler/ShaderCompiler.cpp
+++ b/Engine/Source/Runtime/Engine/Private/ShaderCompiler/ShaderCompiler.cpp
@@ -1319,7 +1319,7 @@ namespace ShaderCompileWorkerError
 	{
 		const FString AvailableMemoryStr = FString::Printf(TEXT("%d MB"), FUnitConversion::Convert(AvailableMemory, EUnit::Bytes, EUnit::Megabytes));
 		const FString UsedMemoryStr = FString::Printf(TEXT("%d MB"), FUnitConversion::Convert(UsedMemory, EUnit::Bytes, EUnit::Megabytes));
-		const FString ErrorReport = FString::Printf(TEXT("ShaderCompileWorker failed with out-of-memory exception on machine \"%s\"\n%s; Used physical memory %s of %s\n"), Hostname, ExceptionInfo, *UsedMemoryStr, *AvailableMemoryStr);
+		const FString ErrorReport = FString::Printf(TEXT("ShaderCompileWorker failed with out-of-memory (OOM) exception on machine \"%s\"\n%s; Used physical memory %s of %s\n"), Hostname, (ExceptionInfo[0] == TEXT('\0') ? TEXT("No exception information") : ExceptionInfo), *UsedMemoryStr, *AvailableMemoryStr);
 
 		if (TryReissueShaderCompileJobs(QueuedJobs))
 		{
@@ -3306,6 +3306,23 @@ static FString FormatNumber(T Number)
 	return FText::AsNumber(Number, &FormattingOptions).ToString();
 }
 
+static FString PrintJobsCompletedPercentageToString(int64 JobsAssigned, int64 JobsCompleted)
+{
+	if (JobsAssigned == 0)
+	{
+		return TEXT("0%");
+	}
+	if (JobsAssigned == JobsCompleted)
+	{
+		return TEXT("100%");
+	}
+
+	// With more than a million compile jobs but only a small number that didn't complete,
+	// the output might be rounded up to 100%. To avoid a misleading output, we clamp this value to 99.99%
+	double JobsCompletedPercentage = 100.0 * (double)JobsCompleted / (double)JobsAssigned;
+	return FString::Printf(TEXT("%.2f%%"), FMath::Min(JobsCompletedPercentage, 99.99));
+}
+
 static FString WriteShaderCompilerInvocationStats(const FShaderCompilerStats::FCompilerInvocations& Invocations)
 {
 	FString Stats;
@@ -3348,10 +3365,10 @@ void FShaderCompilerStats::WriteStatSummary()
 	UE_LOG(LogShaderCompilers, Display, TEXT("Shaders Compiled: %s"), *FormatNumber(TotalCompiled));
 
 	FScopeLock Lock(&CompileStatsLock);	// make a local copy for all the stats?
-	UE_LOG(LogShaderCompilers, Display, TEXT("Jobs assigned %s, completed %s (%.2f%%)"), 
+	UE_LOG(LogShaderCompilers, Display, TEXT("Jobs assigned %s, completed %s (%s)"), 
 		*FormatNumber(JobsAssigned),
 		*FormatNumber(JobsCompleted),
-		(JobsAssigned > 0.0) ? 100.0 * JobsCompleted / JobsAssigned : 0.0);
+		*PrintJobsCompletedPercentageToString(JobsAssigned, JobsCompleted));
 
 	UE_LOG(LogShaderCompilers, Display, TEXT("Compiler invocations: %s"), *WriteShaderCompilerInvocationStats(CompilerInvocations));
 
@@ -3360,15 +3377,15 @@ void FShaderCompilerStats::WriteStatSummary()
 		UE_LOG(LogShaderCompilers, Display, TEXT("Average time worker was idle: %.2f s"), AccumulatedLocalWorkerIdleTime / TimesLocalWorkersWereIdle);
 	}
 
-	if (JobsAssigned > 0.0)
+	if (JobsAssigned > 0)
 	{
-		UE_LOG(LogShaderCompilers, Display, TEXT("Time job spent in pending queue: average %.2f s, longest %.2f s"), AccumulatedPendingTime / JobsAssigned, MaxPendingTime);
+		UE_LOG(LogShaderCompilers, Display, TEXT("Time job spent in pending queue: average %.2f s, longest %.2f s"), AccumulatedPendingTime / (double)JobsAssigned, MaxPendingTime);
 	}
 
-	if (JobsCompleted > 0.0)
+	if (JobsCompleted > 0)
 	{
-		UE_LOG(LogShaderCompilers, Display, TEXT("Job execution time: average %.2f s, max %.2f s"), AccumulatedJobExecutionTime / JobsCompleted, MaxJobExecutionTime);
-		UE_LOG(LogShaderCompilers, Display, TEXT("Job life time (pending + execution): average %.2f s, max %.2f"), AccumulatedJobLifeTime / JobsCompleted, MaxJobLifeTime);
+		UE_LOG(LogShaderCompilers, Display, TEXT("Job execution time: average %.2f s, max %.2f s"), AccumulatedJobExecutionTime / (double)JobsCompleted, MaxJobExecutionTime);
+		UE_LOG(LogShaderCompilers, Display, TEXT("Job life time (pending + execution): average %.2f s, max %.2f"), AccumulatedJobLifeTime / (double)JobsCompleted, MaxJobLifeTime);
 	}
 
 	if (NumAccumulatedShaderCodes > 0)
@@ -3407,7 +3424,7 @@ void FShaderCompilerStats::WriteStatSummary()
 
 	if (TotalTimeAtLeastOneJobWasInFlight > 0.0)
 	{
-		UE_LOG(LogShaderCompilers, Display, TEXT("Average processing rate: %.2f jobs/sec"), JobsCompleted / TotalTimeAtLeastOneJobWasInFlight);
+		UE_LOG(LogShaderCompilers, Display, TEXT("Average processing rate: %.2f jobs/sec"), (double)JobsCompleted / TotalTimeAtLeastOneJobWasInFlight);
 	}
 
 	if (ShaderTimings.Num())
@@ -3480,7 +3497,7 @@ void FShaderCompilerStats::WriteStatSummary()
 uint32 FShaderCompilerStats::GetTotalShadersCompiled()
 {
 	FScopeLock Lock(&CompileStatsLock);
-	return JobsCompleted;
+	return (uint32)FMath::Max(0ll, JobsCompleted);
 }
 
 void FShaderCompilerStats::RegisterLocalWorkerIdleTime(double IdleTime)
diff --git a/Engine/Source/Runtime/Engine/Public/ShaderCompiler.h b/Engine/Source/Runtime/Engine/Public/ShaderCompiler.h
--- a/Engine/Source/Runtime/Engine/Public/ShaderCompiler.h
+++ b/Engine/Source/Runtime/Engine/Public/ShaderCompiler.h
@@ -525,7 +525,10 @@ private:
 	double TimesLocalWorkersWereIdle = 0;
 
 	/** Number of jobs assigned to workers, no matter if they completed or not - used to average pending time. */
-	double JobsAssigned = 0;
+	int64 JobsAssigned = 0;
+
+	/** Total number jobs completed. */
+	int64 JobsCompleted = 0;
 
 	/** Amount of time a job had to spent in pending queue (i.e. waiting to be assigned to a worker). */
 	double AccumulatedPendingTime = 0;
@@ -533,9 +536,6 @@ private:
 	/** Max amount of time any single job was pending (waiting to be assigned to a worker). */
 	double MaxPendingTime = 0;
 
-	/** Total number jobs completed. */
-	double JobsCompleted = 0;
-
 	/** Amount of time job spent being processed by the worker. */
 	double AccumulatedJobExecutionTime = 0;
 
