Fix error LNK2019: unresolved external symbol "public: virtual __cdecl FAudioDevice::~FAudioDevice(void)
#rnx

diff --git a/Engine/Source/Runtime/Engine/Public/AudioDevice.h b/Engine/Source/Runtime/Engine/Public/AudioDevice.h
--- a/Engine/Source/Runtime/Engine/Public/AudioDevice.h
+++ b/Engine/Source/Runtime/Engine/Public/AudioDevice.h
@@ -524,7 +524,7 @@ public:
 	 */
 	ENGINE_API FAudioDevice();
 
-	virtual ~FAudioDevice();
+	ENGINE_API virtual ~FAudioDevice() = default;
 
 	/** Returns an array of available audio devices names for the platform */
 	virtual void GetAudioDeviceList(TArray<FString>& OutAudioDeviceNames) const
