Fixed ENGINE_API usage in FAudioDevice after merge from 5.2

#rnx

diff --git a/Engine/Source/Runtime/Engine/Public/AudioDevice.h b/Engine/Source/Runtime/Engine/Public/AudioDevice.h
--- a/Engine/Source/Runtime/Engine/Public/AudioDevice.h
+++ b/Engine/Source/Runtime/Engine/Public/AudioDevice.h
@@ -524,7 +524,7 @@ public:
 	 */
 	ENGINE_API FAudioDevice();
 
-	ENGINE_API virtual ~FAudioDevice() = default;
+	ENGINE_API virtual ~FAudioDevice();
 
 	/** Returns an array of available audio devices names for the platform */
 	virtual void GetAudioDeviceList(TArray<FString>& OutAudioDeviceNames) const
@@ -1920,7 +1920,7 @@ public:
 	ENGINE_API float ClampPitch(float InPitchScale) const;
 
 	/** Overrides the attenuation scale used on a sound class. */
-	void SetSoundClassDistanceScale(USoundClass* InSoundClass, float DistanceScale, float TimeSec);
+	ENGINE_API void SetSoundClassDistanceScale(USoundClass* InSoundClass, float DistanceScale, float TimeSec);
 
 	float GetPlatformAudioHeadroom() const { check(IsInAudioThread()); return PlatformAudioHeadroom; }
 	ENGINE_API void SetPlatformAudioHeadroom(float PlatformHeadRoom);
