Fixed crash caused by "SetCustomData" instance update commands being sent to the renderer for a static mesh that was still compiling.  Fix is to test "IsCompiling", and skip sending the incremental update commands, since the instance tree is fully rebuilt when compilation finishes.  Logic and comment cut and pasted from a couple similar incremental update commands in the same file.

#rnx

diff --git a/Engine/Source/Runtime/Engine/Private/HierarchicalInstancedStaticMesh.cpp b/Engine/Source/Runtime/Engine/Private/HierarchicalInstancedStaticMesh.cpp
--- a/Engine/Source/Runtime/Engine/Private/HierarchicalInstancedStaticMesh.cpp
+++ b/Engine/Source/Runtime/Engine/Private/HierarchicalInstancedStaticMesh.cpp
@@ -2397,7 +2397,8 @@ bool UHierarchicalInstancedStaticMeshComponent::SetCustomDataValue(int32 Instanc
 
 	bool Result = Super::SetCustomDataValue(InstanceIndex, CustomDataIndex, CustomDataValue, bMarkRenderStateDirty);
 
-	if (Result && GetStaticMesh() != nullptr)
+	// The tree will be fully rebuilt once the static mesh compilation is finished, no need for incremental update in that case.
+	if (Result && GetStaticMesh() && !GetStaticMesh()->IsCompiling() && GetStaticMesh()->HasValidRenderData())
 	{
 		if (!bIsOmittedInstance)
 		{
@@ -2428,7 +2429,8 @@ bool UHierarchicalInstancedStaticMeshComponent::SetCustomData(int32 InstanceInde
 
 	bool Result = Super::SetCustomData(InstanceIndex, InCustomData, bMarkRenderStateDirty);
 
-	if (Result && GetStaticMesh() != nullptr)
+	// The tree will be fully rebuilt once the static mesh compilation is finished, no need for incremental update in that case.
+	if (Result && GetStaticMesh() && !GetStaticMesh()->IsCompiling() && GetStaticMesh()->HasValidRenderData())
 	{
 		if (!bIsOmittedInstance)
 		{
