Check if there is any work before running the compute graph worker.
Avoids unnecessary stats appearing.
Also avoids a null RDG graph execution.
Integrated from UE5 Main

diff --git a/Engine/Source/Runtime/Engine/Private/LevelTick.cpp b/Engine/Source/Runtime/Engine/Private/LevelTick.cpp
--- a/Engine/Source/Runtime/Engine/Private/LevelTick.cpp
+++ b/Engine/Source/Runtime/Engine/Private/LevelTick.cpp
@@ -988,7 +988,6 @@ void BeginSendEndOfFrameUpdatesDrawEvent(FSendAllEndOfFrameUpdates& SendAllEndOf
 
 DECLARE_GPU_STAT(EndOfFrameUpdates);
 DECLARE_GPU_STAT(GPUSkinCacheRayTracingGeometry);
-DECLARE_GPU_STAT(ComputeTaskWorkerUpdates);
 void EndSendEndOfFrameUpdatesDrawEvent(FSendAllEndOfFrameUpdates& SendAllEndOfFrameUpdates)
 {
 	ENQUEUE_RENDER_COMMAND(EndDrawEventCommand)(
@@ -1004,17 +1003,13 @@ void EndSendEndOfFrameUpdatesDrawEvent(FSendAllEndOfFrameUpdates& SendAllEndOfFr
 				GPUSkinCache->EndBatchDispatch(RHICmdList);
 			}
 
-			if (ComputeTaskWorkers.Num() > 0)
+			for (IComputeTaskWorker* ComputeTaskWorker : ComputeTaskWorkers)
 			{
-				SCOPED_GPU_STAT(RHICmdList, ComputeTaskWorkerUpdates);
-				for (IComputeTaskWorker* ComputeTaskWorker : ComputeTaskWorkers)
+				if (ComputeTaskWorker->HasWork(ComputeTaskExecutionGroup::EndOfFrameUpdate))
 				{
-					if (ComputeTaskWorker->HasWork(ComputeTaskExecutionGroup::EndOfFrameUpdate))
-					{
-						FRDGBuilder GraphBuilder(RHICmdList, RDG_EVENT_NAME("ComputeTaskWorker"));
-						ComputeTaskWorker->SubmitWork(GraphBuilder, ComputeTaskExecutionGroup::EndOfFrameUpdate, FeatureLevel);
-						GraphBuilder.Execute();
-					}
+					FRDGBuilder GraphBuilder(RHICmdList, RDG_EVENT_NAME("ComputeTaskWorker"));
+					ComputeTaskWorker->SubmitWork(GraphBuilder, ComputeTaskExecutionGroup::EndOfFrameUpdate, FeatureLevel);
+					GraphBuilder.Execute();
 				}
 			}
 		});
