Unshelved from pending changelist '24339224':

Added checks for invalid element counts in various allocators' ResizeAllocation() functions.

#localization none
#tests preflight, local runs to ensure everything is still fine, runs with deliberately engineered huge reallocations to show that bad resizes are caught.

diff --git a/Engine/Source/Runtime/Core/Private/Containers/ContainerAllocationPolicies.cpp b/Engine/Source/Runtime/Core/Private/Containers/ContainerAllocationPolicies.cpp
new file mode 100644
--- /dev/null
+++ b/Engine/Source/Runtime/Core/Private/Containers/ContainerAllocationPolicies.cpp
@@ -0,0 +1,20 @@
+// Copyright Epic Games, Inc. All Rights Reserved.
+
+#include "Containers/ContainerAllocationPolicies.h"
+
+FORCENOINLINE void UE::Core::Private::OnInvalidAlignedHeapAllocatorNum(int32 NewNum, SIZE_T NumBytesPerElement)
+{
+	UE_LOG(LogCore, Fatal, TEXT("Trying to resize TAlignedHeapAllocator to an invalid size of %d with element size %" SIZE_T_FMT), NewNum, NumBytesPerElement);
+	for (;;);
+}
+
+FORCENOINLINE void UE::Core::Private::OnInvalidSizedHeapAllocatorNum(int32 IndexSize, int64 NewNum, SIZE_T NumBytesPerElement)
+{
+	UE_LOG(LogCore, Fatal, TEXT("Trying to resize TSizedHeapAllocator<%d> to an invalid size of %" INT64_FMT " with element size %" SIZE_T_FMT), IndexSize, NewNum, NumBytesPerElement);
+	for (;;);
+}
+
+template void TSizedHeapAllocator<32, FMemory>::ForAnyElementType::ResizeAllocation(SizeType PreviousNumElements, SizeType NumElements, SIZE_T NumBytesPerElement);
+template void TSizedHeapAllocator<32, FMemory>::ForAnyElementType::ResizeAllocation(SizeType PreviousNumElements, SizeType NumElements, SIZE_T NumBytesPerElement, uint32 AlignmentOfElement);
+template void TSizedHeapAllocator<64, FMemory>::ForAnyElementType::ResizeAllocation(SizeType PreviousNumElements, SizeType NumElements, SIZE_T NumBytesPerElement);
+template void TSizedHeapAllocator<64, FMemory>::ForAnyElementType::ResizeAllocation(SizeType PreviousNumElements, SizeType NumElements, SIZE_T NumBytesPerElement, uint32 AlignmentOfElement);
diff --git a/Engine/Source/Runtime/Core/Private/Containers/ContainerHelpers.cpp b/Engine/Source/Runtime/Core/Private/Containers/ContainerHelpers.cpp
--- a/Engine/Source/Runtime/Core/Private/Containers/ContainerHelpers.cpp
+++ b/Engine/Source/Runtime/Core/Private/Containers/ContainerHelpers.cpp
@@ -6,9 +6,11 @@
 void UE::Core::Private::OnInvalidArrayNum(unsigned long long NewNum)
 {
 	UE_LOG(LogCore, Fatal, TEXT("Trying to resize TArray to an invalid size of %llu"), NewNum);
+	for (;;);
 }
 
 void UE::Core::Private::OnInvalidSetNum(unsigned long long NewNum)
 {
 	UE_LOG(LogCore, Fatal, TEXT("Trying to resize TSet to an invalid size of %llu"), NewNum);
+	for (;;);
 }
diff --git a/Engine/Source/Runtime/Core/Private/Containers/MemoryImageAllocator.cpp b/Engine/Source/Runtime/Core/Private/Containers/MemoryImageAllocator.cpp
--- a/Engine/Source/Runtime/Core/Private/Containers/MemoryImageAllocator.cpp
+++ b/Engine/Source/Runtime/Core/Private/Containers/MemoryImageAllocator.cpp
@@ -9,6 +9,15 @@
 
 IMPLEMENT_ABSTRACT_TYPE_LAYOUT(FResourceArrayInterface);
 
+namespace UE::Core::Private
+{
+	[[noreturn]] FORCENOINLINE void OnInvalidMemoryImageAllocatorNum(int32 NewNum, SIZE_T NumBytesPerElement)
+	{
+		UE_LOG(LogCore, Fatal, TEXT("Trying to resize TMemoryImageAllocator to an invalid size of %d with element size %" SIZE_T_FMT), NewNum, NumBytesPerElement);
+		for (;;);
+	}
+}
+
 FMemoryImageAllocatorBase::~FMemoryImageAllocatorBase()
 {
 	if (!Data.IsFrozen())
@@ -42,11 +51,22 @@ void FMemoryImageAllocatorBase::ResizeAllocation(int32 PreviousNumElements, int3
 	// Avoid calling FMemory::Realloc( nullptr, 0 ) as ANSI C mandates returning a valid pointer which is not what we want.
 	if (Data.IsFrozen())
 	{
-		// Can't grow a frozen array
-		check(NumElements <= PreviousNumElements);
+		// Can't grow a frozen array or shrink below zero
+		if (UNLIKELY((uint32)NumElements > (uint32)PreviousNumElements))
+		{
+			UE::Core::Private::OnInvalidMemoryImageAllocatorNum(NumElements, NumBytesPerElement);
+		}
 	}
 	else if (LocalData || NumElements > 0)
 	{
+		static_assert(sizeof(int32) <= sizeof(SIZE_T), "SIZE_T is expected to be larger than int32");
+
+		// Check for under/overflow
+		if (UNLIKELY(NumElements < 0 || NumBytesPerElement < 1 || NumBytesPerElement > (SIZE_T)MAX_int32))
+		{
+			UE::Core::Private::OnInvalidMemoryImageAllocatorNum(NumElements, NumBytesPerElement);
+		}
+
 		Data = (FScriptContainerElement*)FMemory::Realloc(LocalData, NumElements*NumBytesPerElement, Alignment);
 	}
 }
diff --git a/Engine/Source/Runtime/Core/Private/Experimental/ConcurrentLinearAllocator.cpp b/Engine/Source/Runtime/Core/Private/Experimental/ConcurrentLinearAllocator.cpp
new file mode 100644
--- /dev/null
+++ b/Engine/Source/Runtime/Core/Private/Experimental/ConcurrentLinearAllocator.cpp
@@ -0,0 +1,9 @@
+// Copyright Epic Games, Inc. All Rights Reserved.
+
+#include "Experimental/ConcurrentLinearAllocator.h"
+
+FORCENOINLINE void UE::Core::Private::OnInvalidConcurrentLinearArrayAllocatorNum(int32 NewNum, SIZE_T NumBytesPerElement)
+{
+	UE_LOG(LogCore, Fatal, TEXT("Trying to resize TConcurrentLinearArrayAllocator to an invalid size of %d with element size %" SIZE_T_FMT), NewNum, NumBytesPerElement);
+	for (;;);
+}
diff --git a/Engine/Source/Runtime/Core/Private/HAL/Allocators/AnsiAllocator.cpp b/Engine/Source/Runtime/Core/Private/HAL/Allocators/AnsiAllocator.cpp
new file mode 100644
--- /dev/null
+++ b/Engine/Source/Runtime/Core/Private/HAL/Allocators/AnsiAllocator.cpp
@@ -0,0 +1,32 @@
+// Copyright Epic Games, Inc. All Rights Reserved.
+
+#include "HAL/Allocators/AnsiAllocator.h"
+
+FORCENOINLINE void UE::Core::Private::OnInvalidAnsiAllocatorNum(int32 NewNum, SIZE_T NumBytesPerElement)
+{
+	UE_LOG(LogCore, Fatal, TEXT("Trying to resize FAnsiAllocator to an invalid size of %d with element size %" SIZE_T_FMT), NewNum, NumBytesPerElement);
+	for (;;);
+}
+
+void FAnsiAllocator::ForAnyElementType::ResizeAllocation(SizeType PreviousNumElements, SizeType NumElements, SIZE_T NumBytesPerElement)
+{
+	// Avoid calling FMemory::Realloc( nullptr, 0 ) as ANSI C mandates returning a valid pointer which is not what we want.
+	if (NumElements)
+	{
+		static_assert(sizeof(int32) <= sizeof(SIZE_T), "SIZE_T is expected to be larger than int32");
+
+		// Check for under/overflow
+		if (UNLIKELY(NumElements < 0 || NumBytesPerElement < 1 || NumBytesPerElement > (SIZE_T)MAX_int32))
+		{
+			UE::Core::Private::OnInvalidAnsiAllocatorNum(NumElements, NumBytesPerElement);
+		}
+
+		void* NewRealloc = ::realloc(Data, NumElements*NumBytesPerElement);
+		Data = (FScriptContainerElement*)NewRealloc;
+	}
+	else
+	{
+		::free(Data);
+		Data = nullptr;
+	}
+}
diff --git a/Engine/Source/Runtime/Core/Private/HAL/Allocators/AnsiAllocator.h b/Engine/Source/Runtime/Core/Private/HAL/Allocators/AnsiAllocator.h
--- a/Engine/Source/Runtime/Core/Private/HAL/Allocators/AnsiAllocator.h
+++ b/Engine/Source/Runtime/Core/Private/HAL/Allocators/AnsiAllocator.h
@@ -10,6 +10,11 @@
 #include "Misc/AssertionMacros.h"
 #include "Containers/ContainerAllocationPolicies.h"
 
+namespace UE::Core::Private
+{
+	[[noreturn]] CORE_API void OnInvalidAnsiAllocatorNum(int32 NewNum, SIZE_T NumBytesPerElement);
+}
+
 /** Allocator that allocates memory using standard library functions. */
 class CORE_API FAnsiAllocator
 {
@@ -22,7 +27,7 @@ public:
 	typedef FAnsiAllocator ElementAllocator;
 	typedef FAnsiAllocator BitArrayAllocator;
 
-	class CORE_API ForAnyElementType
+	class ForAnyElementType
 	{
 	public:
 		/** Default constructor. */
@@ -62,21 +67,7 @@ public:
 		{
 			return Data;
 		}
-		void ResizeAllocation(SizeType PreviousNumElements, SizeType NumElements, SIZE_T NumBytesPerElement)
-		{
-			// Avoid calling FMemory::Realloc( nullptr, 0 ) as ANSI C mandates returning a valid pointer which is not what we want.
-			if (NumElements)
-			{
-				//checkSlow(((uint64)NumElements*(uint64)ElementTypeInfo.GetSize() < (uint64)INT_MAX));
-				void* NewRealloc = ::realloc(Data, NumElements*NumBytesPerElement);
-				Data = (FScriptContainerElement*)NewRealloc;
-			}
-			else
-			{
-				::free(Data);
-				Data = nullptr;
-			}
-		}
+		CORE_API void ResizeAllocation(SizeType PreviousNumElements, SizeType NumElements, SIZE_T NumBytesPerElement);
 		SizeType CalculateSlackReserve(SizeType NumElements, SIZE_T NumBytesPerElement) const
 		{
 			return DefaultCalculateSlackReserve(NumElements, NumBytesPerElement, false);
diff --git a/Engine/Source/Runtime/Core/Private/HAL/LowLevelMemoryUtils.cpp b/Engine/Source/Runtime/Core/Private/HAL/LowLevelMemoryUtils.cpp
new file mode 100644
--- /dev/null
+++ b/Engine/Source/Runtime/Core/Private/HAL/LowLevelMemoryUtils.cpp
@@ -0,0 +1,16 @@
+// Copyright Epic Games, Inc. All Rights Reserved.
+
+#include "LowLevelMemoryUtils.h"
+
+#if ENABLE_LOW_LEVEL_MEM_TRACKER
+
+FORCENOINLINE void UE::Core::Private::OnInvalidLLMAllocatorNum(int32 IndexSize, int64 NewNum, SIZE_T NumBytesPerElement)
+{
+	UE_LOG(LogCore, Fatal, TEXT("Trying to resize TSizedLLMAllocator<%d> to an invalid size of % " INT64_FMT " with element size %" SIZE_T_FMT), IndexSize, NewNum, NumBytesPerElement);
+	for (;;);
+}
+
+template void TSizedLLMAllocator<32>::ForAnyElementType::ResizeAllocation(SizeType PreviousNumElements, SizeType NumElements, SIZE_T NumBytesPerElement);
+template void TSizedLLMAllocator<64>::ForAnyElementType::ResizeAllocation(SizeType PreviousNumElements, SizeType NumElements, SIZE_T NumBytesPerElement);
+
+#endif
diff --git a/Engine/Source/Runtime/Core/Private/HAL/LowLevelMemoryUtils.h b/Engine/Source/Runtime/Core/Private/HAL/LowLevelMemoryUtils.h
--- a/Engine/Source/Runtime/Core/Private/HAL/LowLevelMemoryUtils.h
+++ b/Engine/Source/Runtime/Core/Private/HAL/LowLevelMemoryUtils.h
@@ -10,8 +10,10 @@
 #include "Containers/ContainerAllocationPolicies.h"
 #include "Containers/Map.h"
 #include "Containers/Set.h"
+#include "Math/NumericLimits.h"
 #include "Misc/ScopeLock.h"
 #include "Templates/Tuple.h"
+#include <type_traits>
 
 #define LLM_PAGE_SIZE (16*1024)
 
@@ -882,6 +884,11 @@ private:
 	}
 };
 
+namespace UE::Core::Private
+{
+	[[noreturn]] CORE_API void OnInvalidLLMAllocatorNum(int32 IndexSize, int64 NewNum, SIZE_T NumBytesPerElement);
+}
+
 /**
  * An allocator usable in Core containers. It is based on TSizedHeapAllocator, but instead of allocating from FMemory it allocates from FLLMAllocator.
  * Because FLLMAllocator::Free requires the Size, this allocator also has a Size field that TSizedHeapAllocator does not.
@@ -892,6 +899,10 @@ class TSizedLLMAllocator
 public:
 	using SizeType = typename TBitsToSizeType<IndexSize>::Type;
 
+private:
+	using USizeType = std::make_unsigned_t<SizeType>;
+
+public:
 	enum { NeedsElementType = false };
 	enum { RequireRangeCheck = true };
 
@@ -960,6 +971,19 @@ public:
 			// Avoid calling FMemory::Realloc( nullptr, 0 ) as ANSI C mandates returning a valid pointer which is not what we want.
 			if (Data || NumElements)
 			{
+				static_assert(sizeof(SizeType) <= sizeof(SIZE_T), "SIZE_T is expected to handle all possible sizes");
+
+				// Check for under/overflow
+				bool bInvalidResize = NumElements < 0 || NumBytesPerElement < 1 || NumBytesPerElement > (SIZE_T)MAX_int32;
+				if constexpr (sizeof(SizeType) == sizeof(SIZE_T))
+				{
+					bInvalidResize = bInvalidResize || (SIZE_T)(USizeType)NumElements > (SIZE_T)TNumericLimits<SizeType>::Max() / NumBytesPerElement;
+				}
+				if (UNLIKELY(bInvalidResize))
+				{
+					UE::Core::Private::OnInvalidLLMAllocatorNum(IndexSize, NumElements, NumBytesPerElement);
+				}
+
 				//checkSlow(((uint64)NumElements*(uint64)ElementTypeInfo.GetSize() < (uint64)INT_MAX));
 				size_t NewSize = NumElements * NumBytesPerElement;
 				Data = UE::LLMPrivate::FLLMAllocator::Get()->Realloc(Data, Size, NewSize);
@@ -1020,6 +1044,10 @@ public:
 	};
 };
 
+// Define the ResizeAllocation functions with the regular allocators as exported to avoid bloat
+extern template CORE_API FORCENOINLINE void TSizedLLMAllocator<32>::ForAnyElementType::ResizeAllocation(SizeType PreviousNumElements, SizeType NumElements, SIZE_T NumBytesPerElement);
+extern template CORE_API FORCENOINLINE void TSizedLLMAllocator<64>::ForAnyElementType::ResizeAllocation(SizeType PreviousNumElements, SizeType NumElements, SIZE_T NumBytesPerElement);
+
 // The standard container-specific allocators based on TSizedLLMAllocator; these are copied from ContainerAllocationPolicies.h
 using FDefaultLLMAllocator = TSizedLLMAllocator<32>;
 using FDefaultLLMAllocator64 = TSizedLLMAllocator<64>;
diff --git a/Engine/Source/Runtime/Core/Private/Misc/MemStack.cpp b/Engine/Source/Runtime/Core/Private/Misc/MemStack.cpp
--- a/Engine/Source/Runtime/Core/Private/Misc/MemStack.cpp
+++ b/Engine/Source/Runtime/Core/Private/Misc/MemStack.cpp
@@ -414,3 +414,9 @@ bool FMemStackBase::ContainsPointer(const void* Pointer) const
 
 	return false;
 }
+
+FORCENOINLINE void UE::Core::Private::OnInvalidMemStackAllocatorNum(int32 NewNum, SIZE_T NumBytesPerElement)
+{
+	UE_LOG(LogCore, Fatal, TEXT("Trying to resize TMemStackAllocator to an invalid size of %d with element size %" SIZE_T_FMT), NewNum, NumBytesPerElement);
+	for (;;);
+}
\ No newline at end of file
diff --git a/Engine/Source/Runtime/Core/Public/Containers/Array.h b/Engine/Source/Runtime/Core/Public/Containers/Array.h
--- a/Engine/Source/Runtime/Core/Public/Containers/Array.h
+++ b/Engine/Source/Runtime/Core/Public/Containers/Array.h
@@ -336,7 +336,7 @@ namespace UE4Array_Private
 
 namespace UE::Core::Private
 {
-	CORE_API void OnInvalidArrayNum(unsigned long long NewNum);
+	[[noreturn]] CORE_API void OnInvalidArrayNum(unsigned long long NewNum);
 }
 
 
diff --git a/Engine/Source/Runtime/Core/Public/Containers/ContainerAllocationPolicies.h b/Engine/Source/Runtime/Core/Public/Containers/ContainerAllocationPolicies.h
--- a/Engine/Source/Runtime/Core/Public/Containers/ContainerAllocationPolicies.h
+++ b/Engine/Source/Runtime/Core/Public/Containers/ContainerAllocationPolicies.h
@@ -3,13 +3,15 @@
 #pragma once
 
 #include "CoreTypes.h"
-#include "Misc/AssertionMacros.h"
-#include "HAL/UnrealMemory.h"
-#include "Templates/TypeCompatibleBytes.h"
+#include "Containers/ContainerHelpers.h"
 #include "HAL/PlatformMath.h"
-#include "Templates/MemoryOps.h"
+#include "HAL/UnrealMemory.h"
 #include "Math/NumericLimits.h"
+#include "Misc/AssertionMacros.h"
 #include "Templates/IsPolymorphic.h"
+#include "Templates/MemoryOps.h"
+#include "Templates/TypeCompatibleBytes.h"
+#include <type_traits>
 
 // This option disables array slack for initial allocations, e.g where TArray::SetNum 
 // is called. This tends to save a lot of memory with almost no measured performance cost.
@@ -330,6 +332,12 @@ public:
 	typedef ForElementType<FScriptContainerElement> ForAnyElementType;
 };
 
+namespace UE::Core::Private
+{
+	[[noreturn]] CORE_API void OnInvalidAlignedHeapAllocatorNum(int32 NewNum, SIZE_T NumBytesPerElement);
+	[[noreturn]] CORE_API void OnInvalidSizedHeapAllocatorNum(int32 IndexSize, int64 NewNum, SIZE_T NumBytesPerElement);
+}
+
 /** The indirect allocation policy always allocates the elements indirectly. */
 template<uint32 Alignment = DEFAULT_ALIGNMENT>
 class TAlignedHeapAllocator
@@ -390,7 +398,14 @@ public:
 			// Avoid calling FMemory::Realloc( nullptr, 0 ) as ANSI C mandates returning a valid pointer which is not what we want.
 			if (Data || NumElements)
 			{
-				//checkSlow(((uint64)NumElements*(uint64)ElementTypeInfo.GetSize() < (uint64)INT_MAX));
+				static_assert(sizeof(int32) <= sizeof(SIZE_T), "SIZE_T is expected to be larger than int32");
+
+				// Check for under/overflow
+				if (UNLIKELY(NumElements < 0 || NumBytesPerElement < 1 || NumBytesPerElement > (SIZE_T)MAX_int32))
+				{
+					UE::Core::Private::OnInvalidAlignedHeapAllocatorNum(NumElements, NumBytesPerElement);
+				}
+
 				Data = (FScriptContainerElement*)FMemory::Realloc( Data, NumElements*NumBytesPerElement, Alignment );
 			}
 		}
@@ -473,6 +488,10 @@ class TSizedHeapAllocator
 public:
 	using SizeType = typename TBitsToSizeType<IndexSize>::Type;
 
+private:
+	using USizeType = std::make_unsigned_t<SizeType>;
+
+public:
 	enum { NeedsElementType = true };
 	enum { RequireRangeCheck = true };
 
@@ -533,21 +552,45 @@ public:
 		{
 			return Data;
 		}
-		FORCEINLINE void ResizeAllocation(SizeType PreviousNumElements, SizeType NumElements, SIZE_T NumBytesPerElement)
+		void ResizeAllocation(SizeType PreviousNumElements, SizeType NumElements, SIZE_T NumBytesPerElement)
 		{
 			// Avoid calling FMemory::Realloc( nullptr, 0 ) as ANSI C mandates returning a valid pointer which is not what we want.
 			if (Data || NumElements)
 			{
-				//checkSlow(((uint64)NumElements*(uint64)ElementTypeInfo.GetSize() < (uint64)INT_MAX));
+				static_assert(sizeof(SizeType) <= sizeof(SIZE_T), "SIZE_T is expected to handle all possible sizes");
+
+				// Check for under/overflow
+				bool bInvalidResize = NumElements < 0 || NumBytesPerElement < 1 || NumBytesPerElement > (SIZE_T)MAX_int32;
+				if constexpr (sizeof(SizeType) == sizeof(SIZE_T))
+				{
+					bInvalidResize = bInvalidResize || (SIZE_T)(USizeType)NumElements > (SIZE_T)TNumericLimits<SizeType>::Max() / NumBytesPerElement;
+				}
+				if (UNLIKELY(bInvalidResize))
+				{
+					UE::Core::Private::OnInvalidSizedHeapAllocatorNum(IndexSize, NumElements, NumBytesPerElement);
+				}
+
 				Data = (FScriptContainerElement*)BaseMallocType::Realloc( Data, NumElements*NumBytesPerElement );
 			}
 		}
-		FORCEINLINE void ResizeAllocation(SizeType PreviousNumElements, SizeType NumElements, SIZE_T NumBytesPerElement, uint32 AlignmentOfElement)
+		void ResizeAllocation(SizeType PreviousNumElements, SizeType NumElements, SIZE_T NumBytesPerElement, uint32 AlignmentOfElement)
 		{
 			// Avoid calling FMemory::Realloc( nullptr, 0 ) as ANSI C mandates returning a valid pointer which is not what we want.
 			if (Data || NumElements)
 			{
-				//checkSlow(((uint64)NumElements*(uint64)ElementTypeInfo.GetSize() < (uint64)INT_MAX));
+				static_assert(sizeof(SizeType) <= sizeof(SIZE_T), "SIZE_T is expected to handle all possible sizes");
+
+				// Check for under/overflow
+				bool bInvalidResize = NumElements < 0 || NumBytesPerElement < 1 || NumBytesPerElement > (SIZE_T)MAX_int32;
+				if constexpr (sizeof(SizeType) == sizeof(SIZE_T))
+				{
+					bInvalidResize = bInvalidResize || ((SIZE_T)(USizeType)NumElements > (SIZE_T)TNumericLimits<SizeType>::Max() / NumBytesPerElement);
+				}
+				if (UNLIKELY(bInvalidResize))
+				{
+					UE::Core::Private::OnInvalidSizedHeapAllocatorNum(IndexSize, NumElements, NumBytesPerElement);
+				}
+
 				Data = (FScriptContainerElement*)BaseMallocType::Realloc( Data, NumElements*NumBytesPerElement, AlignmentOfElement );
 			}
 		}
@@ -615,6 +658,12 @@ public:
 	};
 };
 
+// Define the ResizeAllocation functions with the regular allocator as exported to avoid bloat
+extern template CORE_API FORCENOINLINE void TSizedHeapAllocator<32, FMemory>::ForAnyElementType::ResizeAllocation(SizeType PreviousNumElements, SizeType NumElements, SIZE_T NumBytesPerElement);
+extern template CORE_API FORCENOINLINE void TSizedHeapAllocator<32, FMemory>::ForAnyElementType::ResizeAllocation(SizeType PreviousNumElements, SizeType NumElements, SIZE_T NumBytesPerElement, uint32 AlignmentOfElement);
+extern template CORE_API FORCENOINLINE void TSizedHeapAllocator<64, FMemory>::ForAnyElementType::ResizeAllocation(SizeType PreviousNumElements, SizeType NumElements, SIZE_T NumBytesPerElement);
+extern template CORE_API FORCENOINLINE void TSizedHeapAllocator<64, FMemory>::ForAnyElementType::ResizeAllocation(SizeType PreviousNumElements, SizeType NumElements, SIZE_T NumBytesPerElement, uint32 AlignmentOfElement);
+
 template <uint8 IndexSize>
 struct TAllocatorTraits<TSizedHeapAllocator<IndexSize>> : TAllocatorTraitsBase<TSizedHeapAllocator<IndexSize>>
 {
@@ -979,7 +1028,7 @@ public:
 		void ResizeAllocation(SizeType PreviousNumElements, SizeType NumElements,SIZE_T NumBytesPerElement)
 		{
 			// Ensure the requested allocation will fit in the inline data area.
-			check(NumElements <= NumInlineElements);
+			check(NumElements >= 0 && NumElements <= NumInlineElements);
 		}
 
 		FORCEINLINE SizeType CalculateSlackReserve(SizeType NumElements, SIZE_T NumBytesPerElement) const
diff --git a/Engine/Source/Runtime/Core/Public/Containers/ContainerHelpers.h b/Engine/Source/Runtime/Core/Public/Containers/ContainerHelpers.h
--- a/Engine/Source/Runtime/Core/Public/Containers/ContainerHelpers.h
+++ b/Engine/Source/Runtime/Core/Public/Containers/ContainerHelpers.h
@@ -6,6 +6,6 @@
 
 namespace UE::Core::Private
 {
-	CORE_API void OnInvalidArrayNum(unsigned long long NewNum);
-	CORE_API void OnInvalidSetNum  (unsigned long long NewNum);
+	[[noreturn]] CORE_API void OnInvalidArrayNum(unsigned long long NewNum);
+	[[noreturn]] CORE_API void OnInvalidSetNum  (unsigned long long NewNum);
 }
diff --git a/Engine/Source/Runtime/Core/Public/Containers/Set.h b/Engine/Source/Runtime/Core/Public/Containers/Set.h
--- a/Engine/Source/Runtime/Core/Public/Containers/Set.h
+++ b/Engine/Source/Runtime/Core/Public/Containers/Set.h
@@ -260,7 +260,7 @@ public:
 
 namespace UE::Core::Private
 {
-	CORE_API void OnInvalidSetNum(unsigned long long NewNum);
+	[[noreturn]] CORE_API void OnInvalidSetNum(unsigned long long NewNum);
 }
 
 /**
diff --git a/Engine/Source/Runtime/Core/Public/Experimental/ConcurrentLinearAllocator.h b/Engine/Source/Runtime/Core/Public/Experimental/ConcurrentLinearAllocator.h
--- a/Engine/Source/Runtime/Core/Public/Experimental/ConcurrentLinearAllocator.h
+++ b/Engine/Source/Runtime/Core/Public/Experimental/ConcurrentLinearAllocator.h
@@ -553,6 +553,11 @@ public:
 	}
 };
 
+namespace UE::Core::Private
+{
+	[[noreturn]] CORE_API void OnInvalidConcurrentLinearArrayAllocatorNum(int32 NewNum, SIZE_T NumBytesPerElement);
+}
+
 template<typename BlockAllocationTag>
 class TConcurrentLinearArrayAllocator
 {
@@ -597,6 +602,14 @@ public:
 		}
 		void ResizeAllocation(SizeType PreviousNumElements, SizeType NumElements, SIZE_T NumBytesPerElement)
 		{
+			static_assert(sizeof(int32) <= sizeof(SIZE_T), "SIZE_T is expected to be larger than int32");
+
+			// Check for under/overflow
+			if (UNLIKELY(NumElements < 0 || NumBytesPerElement < 1 || NumBytesPerElement > (SIZE_T)MAX_int32))
+			{
+				UE::Core::Private::OnInvalidConcurrentLinearArrayAllocatorNum(NumElements, NumBytesPerElement);
+			}
+
 			Data = (ElementType*)TConcurrentLinearAllocator<BlockAllocationTag>::Realloc(Data, NumElements * NumBytesPerElement, alignof(ElementType));
 		}
 		SizeType CalculateSlackReserve(SizeType NumElements, SIZE_T NumBytesPerElement) const
diff --git a/Engine/Source/Runtime/Core/Public/Misc/MemStack.h b/Engine/Source/Runtime/Core/Public/Misc/MemStack.h
--- a/Engine/Source/Runtime/Core/Public/Misc/MemStack.h
+++ b/Engine/Source/Runtime/Core/Public/Misc/MemStack.h
@@ -356,6 +356,10 @@ inline void* operator new[](size_t Size, std::align_val_t Align, FMemStackBase&
 	return Result;
 }
 
+namespace UE::Core::Private
+{
+	[[noreturn]] CORE_API void OnInvalidMemStackAllocatorNum(int32 NewNum, SIZE_T NumBytesPerElement);
+}
 
 /** A container allocator that allocates from a mem-stack. */
 template<uint32 Alignment = DEFAULT_ALIGNMENT>
@@ -402,6 +406,14 @@ public:
 			void* OldData = Data;
 			if( NumElements )
 			{
+				static_assert(sizeof(int32) <= sizeof(SIZE_T), "SIZE_T is expected to be larger than int32");
+
+				// Check for under/overflow
+				if (UNLIKELY(NumElements < 0 || NumBytesPerElement < 1 || NumBytesPerElement > (SIZE_T)MAX_int32))
+				{
+					UE::Core::Private::OnInvalidMemStackAllocatorNum(NumElements, NumBytesPerElement);
+				}
+
 				// Allocate memory from the stack.
 				Data = (ElementType*)FMemStack::Get().PushBytes(
 					(int32)(NumElements * NumBytesPerElement),
diff --git a/Engine/Source/Runtime/Engine/Private/Animation/AnimCompressionTypes.cpp b/Engine/Source/Runtime/Engine/Private/Animation/AnimCompressionTypes.cpp
--- a/Engine/Source/Runtime/Engine/Private/Animation/AnimCompressionTypes.cpp
+++ b/Engine/Source/Runtime/Engine/Private/Animation/AnimCompressionTypes.cpp
@@ -738,7 +738,7 @@ void FCompressibleAnimData::FetchData(const ITargetPlatform* InPlatform)
 	checkf(AnimSequence, TEXT("Invalid animation sequence while trying to fetch to-compress animation data"));
 	USkeleton* Skeleton = AnimSequence->GetSkeleton();
 	checkf(Skeleton, TEXT("Invalid Skeleton while trying to fetch to-compress animation data, %s"), *Skeleton->GetPathName());
-	
+
 	FAnimationUtils::BuildSkeletonMetaData(Skeleton, BoneData);
 
 	const FFrameRate DefaultSamplingFrameRate = AnimSequence->GetSamplingFrameRate();
@@ -783,7 +783,7 @@ void FCompressibleAnimData::FetchData(const ITargetPlatform* InPlatform)
 		UE::Anim::Compression::SanitizeRawAnimSequenceTrack(Track);
 		TrackToSkeletonMapTable.Add(AnimTrack.BoneTreeIndex);
 		OriginalTrackNames.Add(AnimTrack.Name);
-    } 
+	}
 
 	// Apply any key reduction if possible
 	if (RawAnimationData.Num())
@@ -817,7 +817,7 @@ void FCompressibleAnimData::FetchData(const ITargetPlatform* InPlatform)
 		TArray<FTrackToSkeletonMap> TempTrackToSkeletonMapTable;
 		TempTrackToSkeletonMapTable.Reserve(OriginalTrackNames.Num());
 		TempRawAnimationData.Reserve(OriginalTrackNames.Num());
-		FinalTrackNames.Reserve(ResampledTrackData.Num());
+	FinalTrackNames.Reserve(ResampledTrackData.Num());	
 		TempAdditiveBaseAnimationData.Reserve(AdditiveBaseAnimationData.Num() ? AdditiveBaseAnimationData.Num() : 0);
 
 		// Include root bone track
@@ -831,17 +831,17 @@ void FCompressibleAnimData::FetchData(const ITargetPlatform* InPlatform)
 
 		const int32 NumTracks = RawAnimationData.Num();
 		for (int32 TrackIndex = 1; TrackIndex < NumTracks; ++TrackIndex)
-		{
-			const FRawAnimSequenceTrack& Track = RawAnimationData[TrackIndex];
-			// Try find correct bone index
-			const int32 BoneIndex = RefSkeleton.FindBoneIndex(OriginalTrackNames[TrackIndex]);
+	{
+		const FRawAnimSequenceTrack& Track = RawAnimationData[TrackIndex];
+		// Try find correct bone index
+		const int32 BoneIndex = RefSkeleton.FindBoneIndex(OriginalTrackNames[TrackIndex]);
 
 			const bool bValidBoneIndex = BoneIndex != INDEX_NONE;
 			const bool bValidAdditiveTrack = !IsRawTrackZeroAdditive(Track);
 
 			// Only include track if it contains valid (additive) data and its name corresponds to a bone on the skeleton
 			if ((!bIsAdditiveAnimation || bValidAdditiveTrack) && bValidBoneIndex)
-			{
+		{
 				FinalTrackNames.Add(OriginalTrackNames[TrackIndex]);
 				TempTrackToSkeletonMapTable.Add(TrackToSkeletonMapTable[TrackIndex]);
 				TempRawAnimationData.Add(RawAnimationData[TrackIndex]);
@@ -850,7 +850,7 @@ void FCompressibleAnimData::FetchData(const ITargetPlatform* InPlatform)
 				{
 					TempAdditiveBaseAnimationData.Add(AdditiveBaseAnimationData[TrackIndex]);
 				}
-			}
+		}
 		}
 
 		// Swap out maintained track data
@@ -858,9 +858,9 @@ void FCompressibleAnimData::FetchData(const ITargetPlatform* InPlatform)
 		Swap(TrackToSkeletonMapTable, TempTrackToSkeletonMapTable);
 
 		if (AdditiveBaseAnimationData.Num())
-		{
+        {
 			Swap(AdditiveBaseAnimationData, TempAdditiveBaseAnimationData);
-		}
+        }
 	}
 
 	// Find or add curve names on skeleton
@@ -1496,58 +1496,58 @@ void DecompressPose(FCompactPose& OutPose, const FCompressedAnimSequence& Compre
 }
 
 FArchive& operator<<(FArchive& Ar, FCompressedOffsetData& D)
-{
+	{
 	Ar << D.OffsetData << D.StripSize;
 	return Ar;
 }
 
 FArchive& operator<<(FArchive& Ar, FAnimationErrorStats& ErrorStats)
-{
+		{
 	Ar << ErrorStats.AverageError;
 	Ar << ErrorStats.MaxError;
 	Ar << ErrorStats.MaxErrorTime;
 	Ar << ErrorStats.MaxErrorBone;
 	return Ar;
-}
+						}
 
 #if WITH_EDITORONLY_DATA
 UE::Anim::Compression::FAnimDDCKeyArgs::FAnimDDCKeyArgs(const UAnimSequenceBase& AnimSequence)
 	: AnimSequence(AnimSequence)
 	, TargetPlatform(nullptr)
-{
-}
+						{
+	}
 
 UE::Anim::Compression::FAnimDDCKeyArgs::FAnimDDCKeyArgs(const UAnimSequenceBase& AnimSequence, const ITargetPlatform* TargetPlatform)
 	: AnimSequence(AnimSequence)
 	, TargetPlatform(TargetPlatform)
-{
-}
+	{
+			}
 #endif // WITH_EDITORONLY_DATA
 
 #if WITH_EDITOR
 template <typename ArrayType>
 void UpdateSHAWithArray(FSHA1& Sha, const TArray<ArrayType>& Array)
-{
+		{
 	Sha.Update((uint8*)Array.GetData(), Array.Num() * Array.GetTypeSize());
-}
+	}
 
 void UpdateSHAWithRawTrack(FSHA1& Sha, const FRawAnimSequenceTrack& RawTrack)
-{
+	{
 	UpdateSHAWithArray(Sha, RawTrack.PosKeys);
 	UpdateSHAWithArray(Sha, RawTrack.RotKeys);
 	UpdateSHAWithArray(Sha, RawTrack.ScaleKeys);
-}
+	}
 
 template<class DataType>
 void UpdateWithData(FSHA1& Sha, const DataType& Data)
-{
+	{
 	Sha.Update((uint8*)(&Data), sizeof(DataType));
-}
+	}
 
 void UpdateSHAWithCurves(FSHA1& Sha, const FRawCurveTracks& InRawCurveData) 
-{
-	auto UpdateWithFloatCurve = [&Sha](const FRichCurve& Curve)
 	{
+	auto UpdateWithFloatCurve = [&Sha](const FRichCurve& Curve)
+		{
 		UpdateWithData(Sha, Curve.DefaultValue);
 		UpdateSHAWithArray(Sha, Curve.GetConstRefOfKeys());
 		UpdateWithData(Sha, Curve.PreInfinityExtrap);
@@ -1555,7 +1555,7 @@ void UpdateSHAWithCurves(FSHA1& Sha, const FRawCurveTracks& InRawCurveData)
 	};
 
 	for (const FFloatCurve& Curve : InRawCurveData.FloatCurves)
-	{
+			{
 		UpdateWithData(Sha, Curve.Name.UID);
 		UpdateWithFloatCurve(Curve.FloatCurve);
 	}
@@ -1571,19 +1571,19 @@ void UpdateSHAWithCurves(FSHA1& Sha, const FRawCurveTracks& InRawCurveData)
 				UpdateWithFloatCurve(VectorCurve.FloatCurves[ChannelIndex]);
 			}
 		};
-		
+
 		UpdateWithComponent(Curve.TranslationCurve);
 		UpdateWithComponent(Curve.RotationCurve);
 		UpdateWithComponent(Curve.ScaleCurve);
+		}
 	}
-}
 
 FGuid GenerateGuidFromRawAnimData(const TArray<FRawAnimSequenceTrack>& RawAnimationData, const FRawCurveTracks& RawCurveData)
-{
+	{
 	FSHA1 Sha;
 
 	for (const FRawAnimSequenceTrack& Track : RawAnimationData)
-	{
+		{
 		UpdateSHAWithRawTrack(Sha, Track);
 	}
 
@@ -1600,3 +1600,10 @@ FGuid GenerateGuidFromRawAnimData(const TArray<FRawAnimSequenceTrack>& RawAnimat
 #endif // WITH_EDITOR
 
 
+FORCENOINLINE void UE::Animation::Private::OnInvalidMaybeMappedAllocatorNum(int32 NewNum, SIZE_T NumBytesPerElement)
+{
+	UE_LOG(LogAnimation, Fatal, TEXT("Trying to resize TMaybeMappedAllocator to an invalid size of %d with element size %" SIZE_T_FMT), NewNum, NumBytesPerElement);
+	for (;;);
+}
+
+template void TMaybeMappedAllocator<DEFAULT_ALIGNMENT>::ForAnyElementType::ResizeAllocation(SizeType PreviousNumElements, SizeType NumElements, SIZE_T NumBytesPerElement);
diff --git a/Engine/Source/Runtime/Engine/Public/Animation/AnimCompressionTypes.h b/Engine/Source/Runtime/Engine/Public/Animation/AnimCompressionTypes.h
--- a/Engine/Source/Runtime/Engine/Public/Animation/AnimCompressionTypes.h
+++ b/Engine/Source/Runtime/Engine/Public/Animation/AnimCompressionTypes.h
@@ -517,6 +517,11 @@ struct ENGINE_API FUECompressedAnimData : public ICompressedAnimData, public FCo
 	virtual bool IsValid() const override { return CompressedByteStream.Num() > 0 || (TranslationCompressionFormat == ACF_Identity && RotationCompressionFormat == ACF_Identity && ScaleCompressionFormat == ACF_Identity); }
 };
 
+namespace UE::Animation::Private
+{
+	[[noreturn]] ENGINE_API void OnInvalidMaybeMappedAllocatorNum(int32 NewNum, SIZE_T NumBytesPerElement);
+}
+
 template<uint32 Alignment = DEFAULT_ALIGNMENT>
 class TMaybeMappedAllocator
 {
@@ -587,6 +592,14 @@ public:
 			}
 			else if (Data || NumElements)
 			{
+				static_assert(sizeof(int32) <= sizeof(SIZE_T), "SIZE_T is expected to be larger than int32");
+
+				// Check for under/overflow
+				if (UNLIKELY(NumElements < 0 || NumBytesPerElement < 1 || NumBytesPerElement > (SIZE_T)MAX_int32))
+				{
+				    UE::Animation::Private::OnInvalidMaybeMappedAllocatorNum(NumElements, NumBytesPerElement);
+				}
+
 				// Avoid calling FMemory::Realloc( nullptr, 0 ) as ANSI C mandates returning a valid pointer which is not what we want.
 				//checkSlow(((uint64)NumElements*(uint64)ElementTypeInfo.GetSize() < (uint64)INT_MAX));
 				Data = (FScriptContainerElement*)FMemory::Realloc(Data, NumElements*NumBytesPerElement, Alignment);
@@ -688,6 +701,9 @@ public:
 	};
 };
 
+// Define the ResizeAllocation function with the regular alignment as exported to avoid bloat
+extern template ENGINE_API FORCENOINLINE void TMaybeMappedAllocator<DEFAULT_ALIGNMENT>::ForAnyElementType::ResizeAllocation(SizeType PreviousNumElements, SizeType NumElements, SIZE_T NumBytesPerElement);
+
 template<typename T, uint32 Alignment = DEFAULT_ALIGNMENT>
 class TMaybeMappedArray : public TArray<T, TMaybeMappedAllocator<Alignment>>
 {
diff --git a/Engine/Source/Runtime/RenderCore/Private/RenderGraphAllocator.cpp b/Engine/Source/Runtime/RenderCore/Private/RenderGraphAllocator.cpp
--- a/Engine/Source/Runtime/RenderCore/Private/RenderGraphAllocator.cpp
+++ b/Engine/Source/Runtime/RenderCore/Private/RenderGraphAllocator.cpp
@@ -2,6 +2,7 @@
 
 #include "RenderGraphAllocator.h"
 #include "Async/TaskGraphInterfaces.h"
+#include "RenderCore.h"
 #include "RenderGraphPrivate.h"
 
 FRDGAllocator& FRDGAllocator::Get()
@@ -65,3 +66,9 @@ FRDGAllocatorScope::~FRDGAllocatorScope()
 		Allocator.ReleaseAll();
 	}
 }
+
+FORCENOINLINE void UE::RenderCore::Private::OnInvalidRDGAllocatorNum(int32 NewNum, SIZE_T NumBytesPerElement)
+{
+	UE_LOG(LogRendererCore, Fatal, TEXT("Trying to resize TRDGArrayAllocator to an invalid size of %d with element size %" SIZE_T_FMT), NewNum, NumBytesPerElement);
+	for (;;);
+}
\ No newline at end of file
diff --git a/Engine/Source/Runtime/RenderCore/Public/RenderGraphAllocator.h b/Engine/Source/Runtime/RenderCore/Public/RenderGraphAllocator.h
--- a/Engine/Source/Runtime/RenderCore/Public/RenderGraphAllocator.h
+++ b/Engine/Source/Runtime/RenderCore/Public/RenderGraphAllocator.h
@@ -153,6 +153,11 @@ private:
 	TUniqueFunction<void()> AsyncDeleteFunction;
 };
 
+namespace UE::RenderCore::Private
+{
+	[[noreturn]] RENDERCORE_API void OnInvalidRDGAllocatorNum(int32 NewNum, SIZE_T NumBytesPerElement);
+}
+
 /** A container allocator that allocates from a global RDG allocator instance. */
 template<uint32 Alignment = DEFAULT_ALIGNMENT>
 class TRDGArrayAllocator
@@ -187,6 +192,14 @@ public:
 			void* OldData = Data;
 			if (NumElements)
 			{
+				static_assert(sizeof(int32) <= sizeof(SIZE_T), "SIZE_T is expected to be larger than int32");
+
+				// Check for under/overflow
+				if (UNLIKELY(NumElements < 0 || NumBytesPerElement < 1 || NumBytesPerElement > (SIZE_T)MAX_int32))
+				{
+					UE::RenderCore::Private::OnInvalidRDGAllocatorNum(NumElements, NumBytesPerElement);
+				}
+
 				// Allocate memory from the allocator.
 				const int32 AllocSize = (int32)(NumElements * NumBytesPerElement);
 				const int32 AllocAlignment = FMath::Max(Alignment, (uint32)alignof(ElementType));
