Minor UpdateStreamingState optimizations.

diff --git a/Engine/Source/Runtime/Engine/Private/WorldPartition/WorldPartitionRuntimeCellData.cpp b/Engine/Source/Runtime/Engine/Private/WorldPartition/WorldPartitionRuntimeCellData.cpp
--- a/Engine/Source/Runtime/Engine/Private/WorldPartition/WorldPartitionRuntimeCellData.cpp
+++ b/Engine/Source/Runtime/Engine/Private/WorldPartition/WorldPartitionRuntimeCellData.cpp
@@ -11,11 +11,6 @@ UWorldPartitionRuntimeCellData::UWorldPartitionRuntimeCellData(const FObjectInit
 	, ContentBounds(ForceInit)
 {}
 
-bool UWorldPartitionRuntimeCellData::ShouldResetStreamingSourceInfo() const
-{
-	return CachedSourceInfoEpoch != StreamingSourceCacheEpoch;
-}
-
 void UWorldPartitionRuntimeCellData::ResetStreamingSourceInfo() const
 {
 	check(ShouldResetStreamingSourceInfo());
diff --git a/Engine/Source/Runtime/Engine/Public/WorldPartition/WorldPartitionRuntimeCell.h b/Engine/Source/Runtime/Engine/Public/WorldPartition/WorldPartitionRuntimeCell.h
--- a/Engine/Source/Runtime/Engine/Public/WorldPartition/WorldPartitionRuntimeCell.h
+++ b/Engine/Source/Runtime/Engine/Public/WorldPartition/WorldPartitionRuntimeCell.h
@@ -290,13 +290,13 @@ protected:
 
 public:
 	//~Begin UWorldPartitionRuntimeCellData Proxy
-	bool ShouldResetStreamingSourceInfo() const { return RuntimeCellData->ShouldResetStreamingSourceInfo(); }
-	void ResetStreamingSourceInfo() const { RuntimeCellData->ResetStreamingSourceInfo(); }
-	void AppendStreamingSourceInfo(const FWorldPartitionStreamingSource& Source, const FSphericalSector& SourceShape) const { RuntimeCellData->AppendStreamingSourceInfo(Source, SourceShape); }
-	void MergeStreamingSourceInfo() const { RuntimeCellData->MergeStreamingSourceInfo(); }
+	inline bool ShouldResetStreamingSourceInfo() const { return RuntimeCellData->ShouldResetStreamingSourceInfo(); }
+	inline void ResetStreamingSourceInfo() const { RuntimeCellData->ResetStreamingSourceInfo(); }
+	inline void AppendStreamingSourceInfo(const FWorldPartitionStreamingSource& Source, const FSphericalSector& SourceShape) const { RuntimeCellData->AppendStreamingSourceInfo(Source, SourceShape); }
+	inline void MergeStreamingSourceInfo() const { RuntimeCellData->MergeStreamingSourceInfo(); }
 	int32 SortCompare(const UWorldPartitionRuntimeCell* Other, bool bCanUseSortingCache = true) const;
-	const FBox& GetContentBounds() const { return RuntimeCellData->GetContentBounds(); }
-	FBox GetCellBounds() const { return RuntimeCellData->GetCellBounds(); }
+	inline const FBox& GetContentBounds() const { return RuntimeCellData->GetContentBounds(); }
+	inline FBox GetCellBounds() const { return RuntimeCellData->GetCellBounds(); }
 	virtual bool IsDebugShown() const;
 	//~End UWorldPartitionRuntimeCellData Proxy
 
diff --git a/Engine/Source/Runtime/Engine/Public/WorldPartition/WorldPartitionRuntimeCellData.h b/Engine/Source/Runtime/Engine/Public/WorldPartition/WorldPartitionRuntimeCellData.h
--- a/Engine/Source/Runtime/Engine/Public/WorldPartition/WorldPartitionRuntimeCellData.h
+++ b/Engine/Source/Runtime/Engine/Public/WorldPartition/WorldPartitionRuntimeCellData.h
@@ -15,7 +15,8 @@ class ENGINE_API UWorldPartitionRuntimeCellData : public UObject
 
 	virtual ~UWorldPartitionRuntimeCellData() {}
 
-	bool ShouldResetStreamingSourceInfo() const;
+	inline bool ShouldResetStreamingSourceInfo() const { return CachedSourceInfoEpoch != StreamingSourceCacheEpoch; }
+
 	virtual void ResetStreamingSourceInfo() const;
 	virtual void AppendStreamingSourceInfo(const FWorldPartitionStreamingSource& Source, const FSphericalSector& SourceShape) const;
 	virtual void MergeStreamingSourceInfo() const;
@@ -28,7 +29,7 @@ class ENGINE_API UWorldPartitionRuntimeCellData : public UObject
 	virtual FString GetDebugName() const { return GetName(); }
 
 	static int32 StreamingSourceCacheEpoch;
-	static void DirtyStreamingSourceCacheEpoch() { ++StreamingSourceCacheEpoch; }
+	static inline void DirtyStreamingSourceCacheEpoch() { ++StreamingSourceCacheEpoch; }
 
 	// Source Priority
 	mutable uint8 CachedMinSourcePriority;
