Moved selection code from Level into editor code.

diff --git a/Engine/Source/Editor/UnrealEd/Private/EditorEngine.cpp b/Engine/Source/Editor/UnrealEd/Private/EditorEngine.cpp
--- a/Engine/Source/Editor/UnrealEd/Private/EditorEngine.cpp
+++ b/Engine/Source/Editor/UnrealEd/Private/EditorEngine.cpp
@@ -144,6 +144,8 @@
 #include "Editor/ActorPositioning.h"
 
 #include "Elements/Framework/TypedElementSelectionSet.h"
+#include "Elements/Framework/EngineElementsLibrary.h"
+#include "Elements/Framework/TypedElementList.h"
 #include "Elements/Actor/ActorElementEditorViewportInteractionCustomization.h"
 #include "Elements/Component/ComponentElementEditorViewportInteractionCustomization.h"
 
@@ -277,6 +279,77 @@ static USelection* GActorSelection = nullptr;
 static USelection* GComponentSelection = nullptr;
 static USelection* GObjectSelection = nullptr;
 
+class FLevelEvents
+{
+public:
+	FLevelEvents()
+		: bSelectionChanged(false)
+		, ElementSelectionSet(nullptr)
+	{
+		ULevel::OnLoadedActorRemovedFromLevelPreEvent.AddRaw(this, &FLevelEvents::OnLoadedActorRemovedToLevelPreEvent);
+		ULevel::OnLoadedActorRemovedFromLevelPostEvent.AddRaw(this, &FLevelEvents::OnLoadedActorRemovedToLevelPostEvent);
+	}
+
+	~FLevelEvents()
+	{
+		ULevel::OnLoadedActorRemovedFromLevelPreEvent.RemoveAll(this);
+		ULevel::OnLoadedActorRemovedFromLevelPostEvent.RemoveAll(this);
+	}
+
+private:
+	void OnLoadedActorRemovedToLevelPreEvent(ULevel* InLevel)
+	{
+		check(!bSelectionChanged);
+		check(!ElementSelectionSet);
+		check(!LegacySyncBatch);
+
+		ElementSelectionSet = (GEditor && GEditor->GetSelectedActors()) ? GEditor->GetSelectedActors()->GetElementSelectionSet() : nullptr;
+		LegacySyncBatch = ElementSelectionSet ? MakeUnique<FTypedElementList::FLegacySyncScopedBatch>(*ElementSelectionSet->GetElementList()) : nullptr;
+
+		InLevel->OnLoadedActorRemovedFromLevelEvent.AddRaw(this, &FLevelEvents::OnLoadedActorRemoved);
+	}
+
+	void OnLoadedActorRemoved(AActor& InActor)
+	{
+		if (ElementSelectionSet)
+		{
+			if (FTypedElementHandle ActorHandle = UEngineElementsLibrary::AcquireEditorActorElementHandle(&InActor, /*bAllowCreate*/false))
+			{
+				static const FTypedElementSelectionOptions SelectionOptions = FTypedElementSelectionOptions()
+					.SetAllowHidden(true)
+					.SetAllowGroups(false)
+					.SetWarnIfLocked(false)
+					.SetChildElementInclusionMethod(ETypedElementChildInclusionMethod::Recursive);
+
+				bSelectionChanged |= ElementSelectionSet->DeselectElement(ActorHandle, SelectionOptions);
+			}
+		}
+	}
+
+	void OnLoadedActorRemovedToLevelPostEvent(ULevel* InLevel)
+	{
+		if (ElementSelectionSet)
+		{
+			if (bSelectionChanged)
+			{
+				GEditor->NoteSelectionChange();
+				bSelectionChanged = false;
+			}
+
+			LegacySyncBatch.Reset();
+			ElementSelectionSet = nullptr;
+		}
+
+		InLevel->OnLoadedActorRemovedFromLevelEvent.RemoveAll(this);
+	}
+
+	bool bSelectionChanged;
+	UTypedElementSelectionSet* ElementSelectionSet;
+	TUniquePtr<FTypedElementList::FLegacySyncScopedBatch> LegacySyncBatch;
+};
+
+static FLevelEvents* GLevelEvents = nullptr;
+
 void InitSelectionSets()
 {
 	// Note: The actor and component typed element selection set is set and owned by the level editor, so it is deliberately left null here
@@ -304,6 +377,8 @@ void InitSelectionSets()
 	{
 		return GObjectSelection->IsSelected(InObject);
 	};
+
+	GLevelEvents = new FLevelEvents;
 }
 
 void DestroySelectionSets()
@@ -323,6 +398,8 @@ void DestroySelectionSets()
 				ObjectSelectionSet->ClearSelection(FTypedElementSelectionOptions());
 			}
 		}
+
+		delete GLevelEvents;
 	}
 
 	GIsActorSelectedInEditor = nullptr;
@@ -332,6 +409,7 @@ void DestroySelectionSets()
 	GActorSelection = nullptr;
 	GComponentSelection = nullptr;
 	GObjectSelection = nullptr;
+	GLevelEvents = nullptr;
 }
 
 } // namespace PrivateEditorSelection
diff --git a/Engine/Source/Runtime/Engine/Classes/Engine/Level.h b/Engine/Source/Runtime/Engine/Classes/Engine/Level.h
--- a/Engine/Source/Runtime/Engine/Classes/Engine/Level.h
+++ b/Engine/Source/Runtime/Engine/Classes/Engine/Level.h
@@ -948,13 +948,29 @@ public:
 	ENGINE_API void RemoveLoadedActor(AActor* Actor, const FTransform* TransformToRemove = nullptr);
 	ENGINE_API void RemoveLoadedActors(const TArray<AActor*>& ActorList, const FTransform* TransformToRemove = nullptr);
 
+	/** Called when dynamically loaded actors are being added to this level */
+	DECLARE_MULTICAST_DELEGATE_OneParam(FLoadedActorAddedToLevelPreEvent, ULevel*);
+	static ENGINE_API FLoadedActorAddedToLevelPreEvent OnLoadedActorAddedToLevelPreEvent;
+
 	/** Called when dynamically loaded actor is added to this level */
 	DECLARE_EVENT_OneParam(ULevel, FLoadedActorAddedToLevelEvent, AActor&);
 	FLoadedActorAddedToLevelEvent OnLoadedActorAddedToLevelEvent;
 
+	/** Called when dynamically loaded actors were added to this level */
+	DECLARE_MULTICAST_DELEGATE_OneParam(FLoadedActorAddedToLevelPostEvent, ULevel*);
+	static ENGINE_API FLoadedActorAddedToLevelPostEvent OnLoadedActorAddedToLevelPostEvent;
+
+	/** Called when dynamically loaded actors are being removed from this level */
+	DECLARE_MULTICAST_DELEGATE_OneParam(FLoadedActorRemovedFromLevelPreEvent, ULevel*);
+	static ENGINE_API FLoadedActorRemovedFromLevelPreEvent OnLoadedActorRemovedFromLevelPreEvent;
+
 	/** Called when dynamically loaded actor is removed from this level */
 	DECLARE_EVENT_OneParam(ULevel, FLoadedActorRemovedFromLevelEvent, AActor&);
 	FLoadedActorRemovedFromLevelEvent OnLoadedActorRemovedFromLevelEvent;
+
+	/** Called when dynamically loaded actors were removed from this level */
+	DECLARE_MULTICAST_DELEGATE_OneParam(FLoadedActorRemovedFromLevelPostEvent, ULevel*);
+	static ENGINE_API FLoadedActorRemovedFromLevelPostEvent OnLoadedActorRemovedFromLevelPostEvent;
 #endif
 
 	/* Called when level is loaded. */
diff --git a/Engine/Source/Runtime/Engine/Private/Level.cpp b/Engine/Source/Runtime/Engine/Private/Level.cpp
--- a/Engine/Source/Runtime/Engine/Private/Level.cpp
+++ b/Engine/Source/Runtime/Engine/Private/Level.cpp
@@ -59,9 +59,6 @@ Level.cpp: Level-related functions
 #include "WorldPartition/WorldPartitionHelpers.h"
 #include "HAL/PlatformFileManager.h"
 #include "Misc/PathViews.h"
-#include "Selection.h"
-#include "Elements/Framework/EngineElementsLibrary.h"
-#include "Elements/Framework/TypedElementList.h"
 #endif
 #include "WorldPartition/WorldPartition.h"
 #include "WorldPartition/WorldPartitionLog.h"
@@ -196,6 +193,13 @@ FPendingAutoReceiveInputActor::~FPendingAutoReceiveInputActor() = default;
 ULevel implementation.
 -----------------------------------------------------------------------------*/
 
+#if WITH_EDITOR
+ULevel::FLoadedActorAddedToLevelPreEvent ULevel::OnLoadedActorAddedToLevelPreEvent;
+ULevel::FLoadedActorAddedToLevelPostEvent ULevel::OnLoadedActorAddedToLevelPostEvent;
+ULevel::FLoadedActorRemovedFromLevelPreEvent ULevel::OnLoadedActorRemovedFromLevelPreEvent;
+ULevel::FLoadedActorRemovedFromLevelPostEvent ULevel::OnLoadedActorRemovedFromLevelPostEvent;
+#endif
+
 /** Called when a level package has been dirtied. */
 FSimpleMulticastDelegate ULevel::LevelDirtiedEvent;
 
@@ -827,6 +831,8 @@ void ULevel::AddLoadedActors(const TArray<AActor*>& ActorList, const FTransform*
 	FScopedSlowTask SlowTask(ActorsQueue.Num() * 3, LOCTEXT("RegisteringActors", "Registering actors..."));
 	SlowTask.MakeDialogDelayed(1.0f);
 
+	OnLoadedActorAddedToLevelPreEvent.Broadcast(this);
+
 	// Register all components
 	for (AActor* Actor : ActorsQueue)
 	{
@@ -893,6 +899,8 @@ void ULevel::AddLoadedActors(const TArray<AActor*>& ActorList, const FTransform*
 
 		SlowTask.EnterProgressFrame(1);
 	}
+
+	OnLoadedActorAddedToLevelPostEvent.Broadcast(this);
 }
 
 void ULevel::RemoveLoadedActor(AActor* Actor, const FTransform* TransformToRemove)
@@ -939,55 +947,33 @@ void ULevel::RemoveLoadedActors(const TArray<AActor*>& ActorList, const FTransfo
 	FScopedSlowTask SlowTask(ActorsQueue.Num(), LOCTEXT("UnregisteringActors", "Unregistering actors..."));
 	SlowTask.MakeDialogDelayed(1.0f);
 			
-	bool bSelectionChanged = false;
+	OnLoadedActorRemovedFromLevelPreEvent.Broadcast(this);
 
+	for (AActor* Actor : ActorsQueue)
 	{
-		UTypedElementSelectionSet* ElementSelectionSet = (GEditor && GEditor->GetSelectedActors()) ? GEditor->GetSelectedActors()->GetElementSelectionSet() : nullptr;
-		TUniquePtr<FTypedElementList::FLegacySyncScopedBatch> LegacySyncBatch = ElementSelectionSet ? MakeUnique<FTypedElementList::FLegacySyncScopedBatch>(*ElementSelectionSet->GetElementList()) : nullptr;
-
-		const FTypedElementSelectionOptions SelectionOptions = FTypedElementSelectionOptions()
-			.SetAllowHidden(true)
-			.SetAllowGroups(false)
-			.SetWarnIfLocked(false)
-			.SetChildElementInclusionMethod(ETypedElementChildInclusionMethod::Recursive);
-
-		for (AActor* Actor : ActorsQueue)
+		// UnregisterAllComponents can destroy child actors
+		if (!IsValid(Actor))
 		{
-			// UnregisterAllComponents can destroy child actors
-			if (!IsValid(Actor))
-			{
-				continue;
-			}
-
-			if (ElementSelectionSet)
-			{
-				if (FTypedElementHandle ActorHandle = UEngineElementsLibrary::AcquireEditorActorElementHandle(Actor, /*bAllowCreate*/false))
-				{
-					bSelectionChanged |= ElementSelectionSet->DeselectElement(ActorHandle, SelectionOptions);
-				}
-			}
+			continue;
+		}
 
-			Actor->UnregisterAllComponents();
-			Actor->RegisterAllActorTickFunctions(false, true);
+		Actor->UnregisterAllComponents();
+		Actor->RegisterAllActorTickFunctions(false, true);
 
-			if (TransformToRemove)
-			{
-				FLevelUtils::FApplyLevelTransformParams TransformParams(this, TransformToRemove->Inverse());
-				TransformParams.Actor = Actor;
-				TransformParams.bDoPostEditMove = true;
-				FLevelUtils::ApplyLevelTransform(TransformParams);
-			}
+		if (TransformToRemove)
+		{
+			FLevelUtils::FApplyLevelTransformParams TransformParams(this, TransformToRemove->Inverse());
+			TransformParams.Actor = Actor;
+			TransformParams.bDoPostEditMove = true;
+			FLevelUtils::ApplyLevelTransform(TransformParams);
+		}
 
-			OnLoadedActorRemovedFromLevelEvent.Broadcast(*Actor);
+		OnLoadedActorRemovedFromLevelEvent.Broadcast(*Actor);
 
-			SlowTask.EnterProgressFrame(1);
-		}
+		SlowTask.EnterProgressFrame(1);
 	}
 
-	if (bSelectionChanged)
-	{
-		GEditor->NoteSelectionChange();
-	}
+	OnLoadedActorRemovedFromLevelPostEvent.Broadcast(this);
 }
 #endif
 
