Ground work for Local/Embedded DataLayerAssets:
Change TObjectPtr<UDataLayerAsset> to TSoftObjectPtr<UDataLAyerAsset>

diff --git a/Engine/Source/Editor/UnrealEd/Private/Commandlets/WorldPartitionDataLayerToAssetCommandlet.cpp b/Engine/Source/Editor/UnrealEd/Private/Commandlets/WorldPartitionDataLayerToAssetCommandlet.cpp
--- a/Engine/Source/Editor/UnrealEd/Private/Commandlets/WorldPartitionDataLayerToAssetCommandlet.cpp
+++ b/Engine/Source/Editor/UnrealEd/Private/Commandlets/WorldPartitionDataLayerToAssetCommandlet.cpp
@@ -573,7 +573,7 @@ uint32 UDataLayerToAssetCommandlet::RemapDataLayersAssetsFromPreviousConversions
 {
 	uint32 ErrorCount = 0;
 
-	const TArray<TObjectPtr<const UDataLayerAsset>>& ActorDataLayerAssets = Actor->GetDataLayerAssets();
+	TArray<const UDataLayerAsset*> ActorDataLayerAssets = Actor->GetDataLayerAssets();
 	for (int32 i = ActorDataLayerAssets.Num() - 1; i >= 0; --i)
 	{
 		const TObjectPtr<const UDataLayerAsset>& ActorDataLayerAsset = ActorDataLayerAssets[i];
diff --git a/Engine/Source/Runtime/Engine/Classes/GameFramework/Actor.h b/Engine/Source/Runtime/Engine/Classes/GameFramework/Actor.h
--- a/Engine/Source/Runtime/Engine/Classes/GameFramework/Actor.h
+++ b/Engine/Source/Runtime/Engine/Classes/GameFramework/Actor.h
@@ -899,9 +899,9 @@ protected:
 	TArray<FActorDataLayer> DataLayers;
 
 	UPROPERTY(EditAnywhere, AdvancedDisplay, Category = DataLayers)
-	TArray<TObjectPtr<const UDataLayerAsset>> DataLayerAssets;
+	TArray<TSoftObjectPtr<UDataLayerAsset>> DataLayerAssets;
 
-	TArray<TObjectPtr<const UDataLayerAsset>> PreEditChangeDataLayers;
+	TArray<TSoftObjectPtr<UDataLayerAsset>> PreEditChangeDataLayers;
 
 public:
 	/** The copy/paste id used to remap actors during copy operations */
@@ -1250,7 +1250,7 @@ public:
 	static const FName GetDataLayerAssetsPropertyName() { return GET_MEMBER_NAME_CHECKED(AActor, DataLayerAssets); }
 	static const FName GetDataLayerPropertyName() { return GET_MEMBER_NAME_CHECKED(AActor, DataLayers); }
 
-	const TArray<TObjectPtr<const UDataLayerAsset>>& GetDataLayerAssets() const { return DataLayerAssets; }
+	TArray<const UDataLayerAsset*> GetDataLayerAssets() const;
 
 	//~ Begin Deprecated
 
diff --git a/Engine/Source/Runtime/Engine/Private/Actor.cpp b/Engine/Source/Runtime/Engine/Private/Actor.cpp
--- a/Engine/Source/Runtime/Engine/Private/Actor.cpp
+++ b/Engine/Source/Runtime/Engine/Private/Actor.cpp
@@ -6143,7 +6143,7 @@ TArray<const UDataLayerInstance*> AActor::GetDataLayerInstancesInternal(bool bUs
 		TArray<const UDataLayerInstance*> DataLayerInstances;
 		if (UDataLayerManager* DataLayerManager = bUseLevelContext ? UDataLayerManager::GetDataLayerManager(this) : UDataLayerManager::GetDataLayerManager(GetWorld()))
 		{
-			DataLayerInstances += DataLayerManager->GetDataLayerInstances(DataLayerAssets);
+			DataLayerInstances += DataLayerManager->GetDataLayerInstances(GetDataLayerAssets());
 
 			PRAGMA_DISABLE_DEPRECATION_WARNINGS
 			DataLayerInstances += DataLayerManager->GetDataLayerInstances(DataLayers);
diff --git a/Engine/Source/Runtime/Engine/Private/ActorEditor.cpp b/Engine/Source/Runtime/Engine/Private/ActorEditor.cpp
--- a/Engine/Source/Runtime/Engine/Private/ActorEditor.cpp
+++ b/Engine/Source/Runtime/Engine/Private/ActorEditor.cpp
@@ -32,6 +32,9 @@
 #include "Modules/ModuleManager.h"
 #include "Engine/SimpleConstructionScript.h"
 #include "StaticMeshResources.h"
+#include "WorldPartition/DataLayer/DataLayerAsset.h"
+#include "UObject/SoftObjectPtr.h"
+#include "Algo/Transform.h"
 
 #define LOCTEXT_NAMESPACE "ErrorChecking"
 
@@ -1620,6 +1623,13 @@ bool AActor::RemoveAllDataLayers()
 	return false;
 }
 
+TArray<const UDataLayerAsset*> AActor::GetDataLayerAssets() const
+{
+	TArray<const UDataLayerAsset*> ResolvedAssets;
+	Algo::TransformIf(DataLayerAssets, ResolvedAssets, [](const TSoftObjectPtr<UDataLayerAsset> DataLayerAsset) { return DataLayerAsset.IsValid(); }, [](const TSoftObjectPtr<UDataLayerAsset> DataLayerAsset) { return DataLayerAsset.Get(); });
+	return ResolvedAssets;
+}
+
 TArray<FName> AActor::GetDataLayerInstanceNames() const
 {
 	TArray<FName> DataLayerInstanceNames;
@@ -1670,18 +1680,21 @@ void AActor::FixupDataLayers(bool bRevertChangesOnLockedDataLayer /*= false*/)
 	if (bRevertChangesOnLockedDataLayer)
 	{
 		// Since it's not possible to prevent changes of particular elements of an array, rollback change on locked DataLayers.
-		TSet<const UDataLayerAsset*> PreEdit(PreEditChangeDataLayers);
-		TSet<const UDataLayerAsset*> PostEdit(DataLayerAssets);
+		TSet<TSoftObjectPtr<UDataLayerAsset>> PreEdit(PreEditChangeDataLayers);
+		TSet<TSoftObjectPtr<UDataLayerAsset>> PostEdit(DataLayerAssets);
 
-		auto DifferenceContainsLockedDataLayers = [DataLayerManager](const TSet<const UDataLayerAsset*>& A, const TSet<const UDataLayerAsset*>& B)
+		auto DifferenceContainsLockedDataLayers = [DataLayerManager](const TSet<TSoftObjectPtr<UDataLayerAsset>>& A, const TSet<TSoftObjectPtr<UDataLayerAsset>>& B)
 		{
-			TSet<const UDataLayerAsset*> Diff = A.Difference(B);
-			for (const UDataLayerAsset* DataLayerAsset : Diff)
+			TSet<TSoftObjectPtr<UDataLayerAsset>> Diff = A.Difference(B);
+			for (const TSoftObjectPtr<UDataLayerAsset>& DataLayerAsset : Diff)
 			{
-				const UDataLayerInstance* DataLayerInstance = DataLayerManager->GetDataLayerInstance(DataLayerAsset);
-				if (DataLayerInstance && DataLayerInstance->IsLocked())
+				if (UDataLayerAsset* ResolvedAsset = DataLayerAsset.Get())
 				{
-					return true;
+					const UDataLayerInstance* DataLayerInstance = DataLayerManager->GetDataLayerInstance(ResolvedAsset);
+					if (DataLayerInstance && DataLayerInstance->IsLocked())
+					{
+						return true;
+					}
 				}
 			}
 			return false;
@@ -1718,7 +1731,10 @@ void AActor::FixupDataLayers(bool bRevertChangesOnLockedDataLayer /*= false*/)
 	// Only invalidate DataLayerAssets on cook. In Editor we want to be able to re-resolve if the asset gets readded to the WorldDataLayers actor
 	if (IsRunningCookCommandlet())
 	{
-		CleanupDataLayers(DataLayerAssets);
+		TArray<const UDataLayerAsset*> ResolvedAssets = GetDataLayerAssets();
+		CleanupDataLayers(ResolvedAssets);
+		DataLayerAssets.Empty(ResolvedAssets.Num());
+		Algo::Transform(ResolvedAssets, DataLayerAssets, [](const UDataLayerAsset* DataLayerAsset) { return TSoftObjectPtr<UDataLayerAsset>(DataLayerAsset); });
 	}
 	CleanupDataLayers(DataLayers);
 	PRAGMA_ENABLE_DEPRECATION_WARNINGS
diff --git a/Engine/Source/Runtime/Engine/Private/LevelInstance/LevelInstanceSubsystem.cpp b/Engine/Source/Runtime/Engine/Private/LevelInstance/LevelInstanceSubsystem.cpp
--- a/Engine/Source/Runtime/Engine/Private/LevelInstance/LevelInstanceSubsystem.cpp
+++ b/Engine/Source/Runtime/Engine/Private/LevelInstance/LevelInstanceSubsystem.cpp
@@ -1041,7 +1041,7 @@ ILevelInstanceInterface* ULevelInstanceSubsystem::CreateLevelInstanceFrom(const
 					if (UDataLayerManager* DataLayerManager = UDataLayerManager::GetDataLayerManager(ActorToMove))
 					{
 						// Use the raw asset list as we don't want parent DataLayers
-						for (const TObjectPtr<const UDataLayerAsset>& DataLayerAsset : ActorToMove->GetDataLayerAssets())
+						for (const UDataLayerAsset* DataLayerAsset : ActorToMove->GetDataLayerAssets())
 						{
 							// Validate that there's a valid Data Layer Instance for this asset in the source level
 							if (DataLayerManager->GetDataLayerInstanceFromAsset(DataLayerAsset))
diff --git a/Engine/Source/Runtime/Engine/Private/WorldPartition/WorldPartitionActorDesc.cpp b/Engine/Source/Runtime/Engine/Private/WorldPartition/WorldPartitionActorDesc.cpp
--- a/Engine/Source/Runtime/Engine/Private/WorldPartition/WorldPartitionActorDesc.cpp
+++ b/Engine/Source/Runtime/Engine/Private/WorldPartition/WorldPartitionActorDesc.cpp
@@ -76,7 +76,7 @@ void FWorldPartitionActorDesc::Init(const AActor* InActor)
 		if (UWorldPartition* WorldPartition = FWorldPartitionHelpers::GetWorldPartition(InActor))
 		{
 			LocalDataLayerAssetPaths.Reserve(InActor->GetDataLayerAssets().Num());
-			for (const TObjectPtr<const UDataLayerAsset>& DataLayerAsset : InActor->GetDataLayerAssets())
+			for (const UDataLayerAsset* DataLayerAsset : InActor->GetDataLayerAssets())
 			{
 				if (DataLayerAsset)
 				{
