[BugFix] Changing root component in C++ does not apply correctly to map placed actor instances
- Because deserialization happens after construction, if a native actor was constructed with a specific root component, serialized, and then the root component is changed in the constructor, the deserializer will attempt to load the serialized root component overwriting the constructor's change. To combat this, we're now detecting circular attatchments and fixing them up using the archetype

diff --git a/Engine/Source/Runtime/Engine/Classes/GameFramework/Actor.h b/Engine/Source/Runtime/Engine/Classes/GameFramework/Actor.h
--- a/Engine/Source/Runtime/Engine/Classes/GameFramework/Actor.h
+++ b/Engine/Source/Runtime/Engine/Classes/GameFramework/Actor.h
@@ -3562,12 +3562,6 @@ public:
 		ForEachComponent_Internal<UActorComponent>(UActorComponent::StaticClass(), bIncludeFromChildActors, InFunc);
 	}
 
-	UE_DEPRECATED(4.24, "Use one of the GetComponents implementations as appropriate")
-	TArray<UActorComponent*> GetComponentsByClass(TSubclassOf<UActorComponent> ComponentClass) const
-	{
-		return K2_GetComponentsByClass(ComponentClass);
-	}
-
 	/**
 	 * Get all components derived from specified ComponentClass and fill in the OutComponents array with the result.
 	 * It's recommended to use TArrays with a TInlineAllocator to potentially avoid memory allocation costs.
@@ -3578,12 +3572,12 @@ public:
 	 *
 	 * @param ComponentClass            The component class to find all components of a class derived from
 	 * @param bIncludeFromChildActors   If true then recurse in to ChildActor components and find components of the appropriate type in those Actors as well
-	 */
-	template<class AllocatorType>
-	void GetComponents(TSubclassOf<UActorComponent> ComponentClass, TArray<UActorComponent*, AllocatorType>& OutComponents, bool bIncludeFromChildActors = false) const
+	*/
+	template<class AllocatorType, class ComponentType>
+	void GetComponents(TSubclassOf<UActorComponent> ComponentClass, TArray<ComponentType*, AllocatorType>& OutComponents, bool bIncludeFromChildActors = false) const
 	{
 		OutComponents.Reset();
-		ForEachComponent_Internal<UActorComponent>(ComponentClass, bIncludeFromChildActors, [&](UActorComponent* InComp)
+		ForEachComponent_Internal<ComponentType>(ComponentClass, bIncludeFromChildActors, [&](ComponentType* InComp)
 		{
 			OutComponents.Add(InComp);
 		});
diff --git a/Engine/Source/Runtime/Engine/Private/Actor.cpp b/Engine/Source/Runtime/Engine/Private/Actor.cpp
--- a/Engine/Source/Runtime/Engine/Private/Actor.cpp
+++ b/Engine/Source/Runtime/Engine/Private/Actor.cpp
@@ -976,6 +976,24 @@ void AActor::PostLoadSubobjects(FObjectInstancingGraph* OuterInstanceGraph)
 
 	ResetOwnedComponents();
 
+	// If a constructor changed the component hierarchy of a serialized actor, it's possible to end up with a circular
+	// attachment in the hierarchy. Fix that using the archetype as a guideline.
+	if (GetAttachParentActor() == this)
+	{
+		// detected circular attachment. Fixup using archetype
+		const USceneComponent* ArchetypeRoot = CastChecked<AActor>(GetArchetype())->GetRootComponent();
+		TInlineComponentArray<USceneComponent*> Applicants;
+		GetComponents(ArchetypeRoot->GetClass(), Applicants);
+		for (USceneComponent* Applicant : Applicants)
+		{
+			if (Applicant->GetFName() == ArchetypeRoot->GetFName())
+			{
+				SetRootComponent(Applicant);
+				break;
+			}
+		}
+	}
+
 	// Redirect the root component away from a Blueprint-added instance if the native ctor now constructs a default subobject as the root.
 	if (RootComponent && !RootComponent->HasAnyFlags(RF_DefaultSubObject))
 	{
