Fix for Strata node note handling pbad input correctly

diff --git a/Engine/Source/Runtime/Engine/Private/Materials/MaterialExpressions.cpp b/Engine/Source/Runtime/Engine/Private/Materials/MaterialExpressions.cpp
--- a/Engine/Source/Runtime/Engine/Private/Materials/MaterialExpressions.cpp
+++ b/Engine/Source/Runtime/Engine/Private/Materials/MaterialExpressions.cpp
@@ -25096,9 +25096,19 @@ UMaterialExpressionStrataTransmittanceToMFP::UMaterialExpressionStrataTransmitta
 #if WITH_EDITOR
 int32 UMaterialExpressionStrataTransmittanceToMFP::Compile(class FMaterialCompiler* Compiler, int32 OutputIndex)
 {
+	int32 TransmittanceColorCodeChunk = TransmittanceColor.GetTracedInput().Expression ? TransmittanceColor.Compile(Compiler) : Compiler->Constant(0.5f);
+	int32 ThicknessCodeChunk = Thickness.GetTracedInput().Expression ? Thickness.Compile(Compiler) : Compiler->Constant(STRATA_LAYER_DEFAULT_THICKNESS_CM);
+	if (TransmittanceColorCodeChunk == INDEX_NONE)
+	{
+		return Compiler->Errorf(TEXT("TransmittanceColor input graph could not be evaluated for TransmittanceToMFP."));
+	}
+	if (ThicknessCodeChunk == INDEX_NONE)
+	{
+		return Compiler->Errorf(TEXT("ThicknessCodeChunk input graph could not be evaluated for TransmittanceToMFP."));
+	}
 	return Compiler->StrataTransmittanceToMFP(
-		TransmittanceColor.GetTracedInput().Expression ? TransmittanceColor.Compile(Compiler) : Compiler->Constant(0.5f), 
-		Thickness.GetTracedInput().Expression ? Thickness.Compile(Compiler) : INDEX_NONE,
+		TransmittanceColorCodeChunk,
+		ThicknessCodeChunk,
 		OutputIndex);
 }
 
@@ -25196,10 +25206,25 @@ UMaterialExpressionStrataMetalnessToDiffuseAlbedoF0::UMaterialExpressionStrataMe
 #if WITH_EDITOR
 int32 UMaterialExpressionStrataMetalnessToDiffuseAlbedoF0::Compile(class FMaterialCompiler* Compiler, int32 OutputIndex)
 {
+	int32 BaseColorCodeChunk = BaseColor.GetTracedInput().Expression ? BaseColor.Compile(Compiler) : Compiler->Constant(0.18f);
+	int32 SpecularCodeChunk = Specular.GetTracedInput().Expression	 ? Specular.Compile(Compiler) : Compiler->Constant(0.5f);
+	int32 MetallicCodeChunk = Metallic.GetTracedInput().Expression ? Metallic.Compile(Compiler) : Compiler->Constant(0.0f);
+	if (BaseColorCodeChunk == INDEX_NONE)
+	{
+		return Compiler->Errorf(TEXT("BaseColor input graph could not be evaluated for MetalnessToDiffuseAlbedoF0."));
+	}
+	if (SpecularCodeChunk == INDEX_NONE)
+	{
+		return Compiler->Errorf(TEXT("Specular input graph could not be evaluated for MetalnessToDiffuseAlbedoF0."));
+	}
+	if (MetallicCodeChunk == INDEX_NONE)
+	{
+		return Compiler->Errorf(TEXT("Metallic input graph could not be evaluated for MetalnessToDiffuseAlbedoF0."));
+	}
 	return Compiler->StrataMetalnessToDiffuseAlbedoF0(
-		BaseColor.GetTracedInput().Expression ? BaseColor.Compile(Compiler) : Compiler->Constant(0.18f),
-		Specular.GetTracedInput().Expression ? Specular.Compile(Compiler) : Compiler->Constant(0.5f),
-		Metallic.GetTracedInput().Expression ? Metallic.Compile(Compiler) : Compiler->Constant(0.f),
+		BaseColorCodeChunk,
+		SpecularCodeChunk,
+		MetallicCodeChunk,
 		OutputIndex);
 }
 
@@ -25274,9 +25299,19 @@ UMaterialExpressionStrataHazinessToSecondaryRoughness::UMaterialExpressionStrata
 #if WITH_EDITOR
 int32 UMaterialExpressionStrataHazinessToSecondaryRoughness::Compile(class FMaterialCompiler* Compiler, int32 OutputIndex)
 {
+	int32 BaseRoughnessCodeChunk= BaseRoughness.GetTracedInput().Expression ? BaseRoughness.Compile(Compiler) : Compiler->Constant(0.1f);
+	int32 HazinessCodeChunk		= Haziness.GetTracedInput().Expression      ? Haziness.Compile(Compiler)      : Compiler->Constant(0.5f);
+	if (BaseRoughnessCodeChunk == INDEX_NONE)
+	{
+		return Compiler->Errorf(TEXT("BaseRoughness input graph could not be evaluated for HazinessToSecondaryRoughness."));
+	}
+	if (HazinessCodeChunk == INDEX_NONE)
+	{
+		return Compiler->Errorf(TEXT("Haziness input graph could not be evaluated for HazinessToSecondaryRoughness."));
+	}
 	return Compiler->StrataHazinessToSecondaryRoughness(
-		BaseRoughness.GetTracedInput().Expression ? BaseRoughness.Compile(Compiler) : Compiler->Constant(0.1f),
-		Haziness.GetTracedInput().Expression ? Haziness.Compile(Compiler) : Compiler->Constant(0.5f),
+		BaseRoughnessCodeChunk,
+		HazinessCodeChunk,
 		OutputIndex);
 }
 
@@ -25384,6 +25419,27 @@ int32 UMaterialExpressionStrataThinFilm::Compile(class FMaterialCompiler* Compil
 	int32 ThicknessCodeChunk	= Thickness.GetTracedInput().Expression	? Thickness.Compile(Compiler)	: Compiler->Constant(1.0f);
 	int32 IORCodeChunk			= IOR.GetTracedInput().Expression		? IOR.Compile(Compiler)			: Compiler->Constant(1.44f);
 
+	if (NormalCodeChunk == INDEX_NONE)
+	{
+		return Compiler->Errorf(TEXT("NormalCode input graph could not be evaluated for ThinFilm."));
+	}
+	if (F0CodeChunk == INDEX_NONE)
+	{
+		return Compiler->Errorf(TEXT("F0 input graph could not be evaluated for ThinFilm."));
+	}
+	if (F90CodeChunk == INDEX_NONE)
+	{
+		return Compiler->Errorf(TEXT("F90 input graph could not be evaluated for ThinFilm."));
+	}
+	if (ThicknessCodeChunk == INDEX_NONE)
+	{
+		return Compiler->Errorf(TEXT("Thickness input graph could not be evaluated for ThinFilm."));
+	}
+	if (IORCodeChunk == INDEX_NONE)
+	{
+		return Compiler->Errorf(TEXT("IOR input graph could not be evaluated for ThinFilm."));
+	}
+
 	return Compiler->StrataThinFilm(NormalCodeChunk, F0CodeChunk, F90CodeChunk, ThicknessCodeChunk, IORCodeChunk, OutputIndex);
 }
 
@@ -25827,6 +25883,15 @@ int32 UMaterialExpressionSkyLightEnvMapSample::Compile(class FMaterialCompiler*
 	int32 DirectionCodeChunk = Direction.GetTracedInput().Expression ? Direction.Compile(Compiler) : Compiler->Constant3(0.0f, 0.0f, 1.0f);
 	int32 RoughnessCodeChunk = Roughness.GetTracedInput().Expression ? Roughness.Compile(Compiler) : Compiler->Constant(0.0f);
 
+	if (DirectionCodeChunk == INDEX_NONE)
+	{
+		return Compiler->Errorf(TEXT("Direction input graph could not be evaluated for SkyLightEnvMapSample."));
+	}
+	if (RoughnessCodeChunk == INDEX_NONE)
+	{
+		return Compiler->Errorf(TEXT("Roughness input graph could not be evaluated for SkyLightEnvMapSample."));
+	}
+
 	return Compiler->SkyLightEnvMapSample(DirectionCodeChunk, RoughnessCodeChunk);
 }
 
