Fixed bug when unloading a data layer while editing a level instance : Wrong DataLayerManager used to resolve data layers for Level Instance's WorldPartitionActorDescs.

diff --git a/Engine/Source/Runtime/Engine/Private/WorldPartition/WorldPartitionActorDesc.cpp b/Engine/Source/Runtime/Engine/Private/WorldPartition/WorldPartitionActorDesc.cpp
--- a/Engine/Source/Runtime/Engine/Private/WorldPartition/WorldPartitionActorDesc.cpp
+++ b/Engine/Source/Runtime/Engine/Private/WorldPartition/WorldPartitionActorDesc.cpp
@@ -72,8 +72,8 @@ void FWorldPartitionActorDesc::Init(const AActor* InActor)
 	{
 		TArray<FName> LocalDataLayerAssetPaths;
 		TArray<FName> LocalDataLayerInstanceNames;
-		// Use Actor's DataLayerManager since the fixup is relative to this level
-		if (UWorldPartition* WorldPartition = FWorldPartitionHelpers::GetWorldPartition(InActor))
+
+		if (UWorldPartition* ActorWorldPartition = FWorldPartitionHelpers::GetWorldPartition(InActor))
 		{
 			LocalDataLayerAssetPaths.Reserve(InActor->GetDataLayerAssets().Num());
 			for (const UDataLayerAsset* DataLayerAsset : InActor->GetDataLayerAssets())
@@ -85,9 +85,10 @@ void FWorldPartitionActorDesc::Init(const AActor* InActor)
 			}
 			bIsUsingDataLayerAsset = LocalDataLayerAssetPaths.Num() > 0;
 
-			UDataLayerManager* DataLayerManager = WorldPartition->GetDataLayerManager();
 			if (!bIsUsingDataLayerAsset)
 			{
+				// Use Actor's DataLayerManager since the fixup is relative to this level
+				const UDataLayerManager* DataLayerManager = ActorWorldPartition->GetDataLayerManager();
 				if (ensure(DataLayerManager))
 				{
 					PRAGMA_DISABLE_DEPRECATION_WARNINGS
@@ -101,15 +102,18 @@ void FWorldPartitionActorDesc::Init(const AActor* InActor)
 			DataLayers = bIsUsingDataLayerAsset ? MoveTemp(LocalDataLayerAssetPaths) : MoveTemp(LocalDataLayerInstanceNames);
 
 			// Init DataLayers transient info
-			if (ensure(DataLayerManager))
+			UWorld* ActorOwningWorld = InActor->GetWorld();
+			// In Editor, always use owning world to resolve data layers.
+			// In game (PIE), always use actor's world partition DataLayerManager to resolve data layers.
+			const UDataLayerManager* ResolvingDataLayerManager = !ActorOwningWorld->IsGameWorld() ? ActorOwningWorld->GetDataLayerManager() : ActorWorldPartition->GetDataLayerManager();
+			// The only case where it's notmal not to have a valid ResolvingDataLayerManager is if OwningWorld is not game and not partitioned
+			check(ResolvingDataLayerManager || (!ActorOwningWorld->IsGameWorld() && !ActorOwningWorld->IsPartitionedWorld()));
+			if (ResolvingDataLayerManager)
 			{
-				// There are valid cases where the DataLayerManager can't resolve transient data layer instance names.
-				// One case is when committing changes on a level instance, FWorldPartitionActorDesc::Init is called when saving 
-				// the child actors.
-				if (DataLayerManager->CanResolveDataLayers())
+				if (ensure(ResolvingDataLayerManager->CanResolveDataLayers()))
 				{
 					// Here we process a loaded actor, so resolving makes sense as long as the actordesc is not reused as a template
-					ResolvedDataLayerInstanceNames = FDataLayerUtils::ResolvedDataLayerInstanceNames(DataLayerManager, this);
+					ResolvedDataLayerInstanceNames = FDataLayerUtils::ResolvedDataLayerInstanceNames(ResolvingDataLayerManager, this);
 				}
 				else
 				{
