Fixed PIE/Game using level instance referencing a partitioned level instances with data layers where data layers were removed from the sublevel's WorldDataLayers

diff --git a/Engine/Source/Runtime/Engine/Private/WorldPartition/DataLayer/WorldDataLayers.cpp b/Engine/Source/Runtime/Engine/Private/WorldPartition/DataLayer/WorldDataLayers.cpp
--- a/Engine/Source/Runtime/Engine/Private/WorldPartition/DataLayer/WorldDataLayers.cpp
+++ b/Engine/Source/Runtime/Engine/Private/WorldPartition/DataLayer/WorldDataLayers.cpp
@@ -859,10 +859,17 @@ bool AWorldDataLayers::IsRuntimeRelevant() const
 	UWorld* OuterWorld = GetTypedOuter<UWorld>();
 	if (OuterWorld && (OuterWorld->GetWorldDataLayers() == this))
 	{
-		check(OuterWorld->IsPartitionedWorld());
-		ULevelInstanceSubsystem* LevelInstanceSubsystem = UWorld::GetSubsystem<ULevelInstanceSubsystem>(GetWorld());
-		ILevelInstanceInterface* LevelInstance = LevelInstanceSubsystem ? LevelInstanceSubsystem->GetOwningLevelInstance(OuterWorld->PersistentLevel) : nullptr;
-		return (LevelInstance == nullptr);
+		if (!UWorld::IsPartitionedWorld(GetWorld()))
+		{
+			return true;
+		}
+		else
+		{
+			check(OuterWorld->IsPartitionedWorld());
+			ULevelInstanceSubsystem* LevelInstanceSubsystem = UWorld::GetSubsystem<ULevelInstanceSubsystem>(GetWorld());
+			ILevelInstanceInterface* LevelInstance = LevelInstanceSubsystem ? LevelInstanceSubsystem->GetOwningLevelInstance(OuterWorld->PersistentLevel) : nullptr;
+			return (LevelInstance == nullptr);
+		}
 	}
 
 	// @todo_ow: revisit that logic for content bundle data layers
