Remove ensure around exceeding the maximum number of bits in FCharacterNetworkSerializationPackedBits::NetSerialize
Also removing the ensure in FCharacterNetworkSerializationPackedBits::NetSerialize in favor of logging
#rnx

diff --git a/Engine/Source/Runtime/Engine/Private/Components/CharacterMovementComponent.cpp b/Engine/Source/Runtime/Engine/Private/Components/CharacterMovementComponent.cpp
--- a/Engine/Source/Runtime/Engine/Private/Components/CharacterMovementComponent.cpp
+++ b/Engine/Source/Runtime/Engine/Private/Components/CharacterMovementComponent.cpp
@@ -8930,10 +8930,10 @@ bool FCharacterNetworkSerializationPackedBits::NetSerialize(FArchive& Ar, class
 	uint32 NumBits = DataBits.Num();
 	Ar.SerializeIntPacked(NumBits);
 
-	if (!ensureMsgf(NumBits <= (uint32)CharacterMovementCVars::NetPackedMovementMaxBits, TEXT("FCharacterNetworkSerializationPackedBits::NetSerialize: NumBits (%d) exceeds CharacterMovementCVars::NetPackedMovementMaxBits (%d)"), NumBits, (uint32)CharacterMovementCVars::NetPackedMovementMaxBits))
+	if (NumBits > static_cast<uint32>(CharacterMovementCVars::NetPackedMovementMaxBits))
 	{
 		// Protect against bad data that could cause server to allocate way too much memory.
-		devCode(UE_LOG(LogNetPlayerMovement, Error, TEXT("FCharacterNetworkSerializationPackedBits::NetSerialize: NumBits (%d) exceeds allowable limit!"), NumBits));
+		UE_LOG(LogNetPlayerMovement, Error, TEXT("FCharacterNetworkSerializationPackedBits::NetSerialize: Dropping move due to NumBits (%d) exceeding allowable limit (%d). See NetPackedMovementMaxBits."), NumBits, CharacterMovementCVars::NetPackedMovementMaxBits);
 		return false;
 	}
 
