Remove spam log.

diff --git a/Engine/Source/Runtime/Engine/Private/WorldPartition/WorldPartitionRuntimeHash.cpp b/Engine/Source/Runtime/Engine/Private/WorldPartition/WorldPartitionRuntimeHash.cpp
--- a/Engine/Source/Runtime/Engine/Private/WorldPartition/WorldPartitionRuntimeHash.cpp
+++ b/Engine/Source/Runtime/Engine/Private/WorldPartition/WorldPartitionRuntimeHash.cpp
@@ -256,8 +256,6 @@ void UWorldPartitionRuntimeHash::PopulateRuntimeCell(UWorldPartitionRuntimeCell*
 
 		// Map relative package to StreamingCell for PopulateGeneratedPackageForCook/PopulateGeneratorPackageForCook/GetCellForPackage
 		PackagesToGenerateForCook.Add(PackageRelativePath, RuntimeCell);
-
-		UE_CLOG(IsRunningCookCommandlet(), LogWorldPartition, Log, TEXT("Creating runtime streaming cells %s."), *RuntimeCell->GetName());
 	}
 }
 
