World Partition: append world and plugin name to streaming generation logs.

diff --git a/Engine/Source/Runtime/Engine/Private/WorldPartition/WorldPartitionStreamingGeneration.cpp b/Engine/Source/Runtime/Engine/Private/WorldPartition/WorldPartitionStreamingGeneration.cpp
--- a/Engine/Source/Runtime/Engine/Private/WorldPartition/WorldPartitionStreamingGeneration.cpp
+++ b/Engine/Source/Runtime/Engine/Private/WorldPartition/WorldPartitionStreamingGeneration.cpp
@@ -1026,8 +1026,21 @@ bool UWorldPartition::GenerateContainerStreaming(const UActorDescContainer* InAc
 	}
 
 	// Dump state log
-	const TCHAR* StateLogSuffix = bIsPIE ? TEXT("PIE") : (IsRunningGame() ? TEXT("Game") : (IsRunningCookCommandlet() ? TEXT("Cook") : TEXT("Manual")));
-	TUniquePtr<FArchive> LogFileAr = FWorldPartitionStreamingGenerator::CreateDumpStateLogArchive(StateLogSuffix);
+	TStringBuilder<256> StateLogSuffix;
+	StateLogSuffix += bIsPIE ? TEXT("PIE") : (IsRunningGame() ? TEXT("Game") : (IsRunningCookCommandlet() ? TEXT("Cook") : TEXT("Manual")));
+	StateLogSuffix += TEXT("_");
+	FString ContainerPackageName = InActorDescContainer->ContainerPackageName.ToString();
+	StateLogSuffix += FPackageName::GetShortName(ContainerPackageName);
+	if (!ContainerPackageName.StartsWith(TEXT("/Game/")))
+	{
+		TArray<FString> SplitContainerPath;
+		if (ContainerPackageName.ParseIntoArray(SplitContainerPath, TEXT("/")))
+		{
+			StateLogSuffix += TEXT("_");
+			StateLogSuffix += SplitContainerPath[0];
+		}
+	}
+	TUniquePtr<FArchive> LogFileAr = FWorldPartitionStreamingGenerator::CreateDumpStateLogArchive(*StateLogSuffix);
 	FHierarchicalLogArchive HierarchicalLogAr(*LogFileAr);
 
 	FWorldPartitionStreamingGenerator::FWorldPartitionStreamingGeneratorParams StreamingGeneratorParams;
