World Partition: expose native class to actor descriptor blueprint library.

diff --git a/Engine/Source/Runtime/Engine/Private/WorldPartition/LevelInstance/LevelInstanceActorDesc.cpp b/Engine/Source/Runtime/Engine/Private/WorldPartition/LevelInstance/LevelInstanceActorDesc.cpp
--- a/Engine/Source/Runtime/Engine/Private/WorldPartition/LevelInstance/LevelInstanceActorDesc.cpp
+++ b/Engine/Source/Runtime/Engine/Private/WorldPartition/LevelInstance/LevelInstanceActorDesc.cpp
@@ -158,7 +158,7 @@ bool FLevelInstanceActorDesc::IsContainerInstance() const
 	return ULevelInstanceSubsystem::CanUsePackage(LevelPackage);
 }
 
-bool FLevelInstanceActorDesc::GetContainerInstance(FWorldPartitionActorDesc::FContainerInstance& OutContainerInstance) const
+bool FLevelInstanceActorDesc::GetContainerInstance(FContainerInstance& OutContainerInstance) const
 {
 	if (LevelInstanceContainer.IsValid())
 	{
diff --git a/Engine/Source/Runtime/Engine/Private/WorldPartition/WorldPartitionActorDesc.cpp b/Engine/Source/Runtime/Engine/Private/WorldPartition/WorldPartitionActorDesc.cpp
--- a/Engine/Source/Runtime/Engine/Private/WorldPartition/WorldPartitionActorDesc.cpp
+++ b/Engine/Source/Runtime/Engine/Private/WorldPartition/WorldPartitionActorDesc.cpp
@@ -588,20 +588,6 @@ FName FWorldPartitionActorDesc::GetDisplayClassName() const
 	return BaseClass.IsNull() ? GetCleanClassName(NativeClass) : GetCleanClassName(BaseClass);
 }
 
-bool FWorldPartitionActorDesc::GetContainerInstance(const UActorDescContainer*& OutLevelContainer, FTransform& OutLevelTransform, EContainerClusterMode& OutClusterMode) const
-{
-	FContainerInstance ContainerInstance;
-	if (GetContainerInstance(ContainerInstance))
-	{
-		OutLevelContainer = ContainerInstance.Container;
-		OutLevelTransform = ContainerInstance.Transform;
-		OutClusterMode = ContainerInstance.ClusterMode;
-		return true;
-	}
-
-	return false;
-}
-
 FGuid FWorldPartitionActorDesc::GetContentBundleGuid() const
 {
 	return ContentBundleGuid;
diff --git a/Engine/Source/Runtime/Engine/Private/WorldPartitionBlueprintLibrary.cpp b/Engine/Source/Runtime/Engine/Private/WorldPartitionBlueprintLibrary.cpp
--- a/Engine/Source/Runtime/Engine/Private/WorldPartitionBlueprintLibrary.cpp
+++ b/Engine/Source/Runtime/Engine/Private/WorldPartitionBlueprintLibrary.cpp
@@ -14,7 +14,8 @@
 #endif
 
 FActorDesc::FActorDesc()
-	: Bounds(ForceInit)
+	: NativeClass(nullptr)
+	, Bounds(ForceInit)
 	, bIsSpatiallyLoaded(false)
 	, bActorIsEditorOnly(false)
 {}
@@ -24,6 +25,8 @@ FActorDesc::FActorDesc(const FWorldPartitionActorDesc& InActorDesc, const FTrans
 {
 	Guid = InActorDesc.GetGuid();
 
+	NativeClass = InActorDesc.GetActorNativeClass();
+
 	if (!InActorDesc.GetBaseClass().IsNull())
 	{
 		Class = FSoftObjectPath(InActorDesc.GetBaseClass(), {});
@@ -90,15 +93,12 @@ bool UWorldPartitionBlueprintLibrary::GetActorDescs(const UActorDescContainer* I
 
 	for (FActorDescList::TConstIterator<> ActorDescIt(InContainer); ActorDescIt; ++ActorDescIt)
 	{
-		const UActorDescContainer* SubContainer;
-		EContainerClusterMode SubClusterMode;
-		FTransform SubTransform;
-
 		if (ActorDescIt->IsContainerInstance())
 		{
-			if (ActorDescIt->GetContainerInstance(SubContainer, SubTransform, SubClusterMode))
+			FWorldPartitionActorDesc::FContainerInstance ContainerInstance;
+			if (ActorDescIt->GetContainerInstance(ContainerInstance))
 			{
-				bResult &= GetActorDescs(SubContainer, SubTransform * InTransform, OutActorDescs);
+				bResult &= GetActorDescs(ContainerInstance.Container, ContainerInstance.Transform * InTransform, OutActorDescs);
 			}
 			else
 			{
@@ -117,15 +117,13 @@ bool UWorldPartitionBlueprintLibrary::GetActorDescs(const UActorDescContainer* I
 bool UWorldPartitionBlueprintLibrary::HandleIntersectingActorDesc(const FWorldPartitionActorDesc* ActorDesc, const FBox& InBox, const FTransform& InTransform, TArray<FActorDesc>& OutActorDescs)
 {
 	bool bResult = true;
-	const UActorDescContainer* SubContainer;
-	EContainerClusterMode SubClusterMode;
-	FTransform SubTransform;
 
 	if (ActorDesc->IsContainerInstance())
 	{
-		if (ActorDesc->GetContainerInstance(SubContainer, SubTransform, SubClusterMode))
+		FWorldPartitionActorDesc::FContainerInstance ContainerInstance;
+		if (ActorDesc->GetContainerInstance(ContainerInstance))
 		{
-			bResult &= GetIntersectingActorDescs(SubContainer, InBox, SubTransform * InTransform, OutActorDescs);
+			bResult &= GetIntersectingActorDescs(ContainerInstance.Container, InBox, ContainerInstance.Transform * InTransform, OutActorDescs);
 		}
 		else
 		{
diff --git a/Engine/Source/Runtime/Engine/Public/WorldPartition/LevelInstance/LevelInstanceActorDesc.h b/Engine/Source/Runtime/Engine/Public/WorldPartition/LevelInstance/LevelInstanceActorDesc.h
--- a/Engine/Source/Runtime/Engine/Public/WorldPartition/LevelInstance/LevelInstanceActorDesc.h
+++ b/Engine/Source/Runtime/Engine/Public/WorldPartition/LevelInstance/LevelInstanceActorDesc.h
@@ -26,7 +26,7 @@ public:
 
 	virtual bool IsContainerInstance() const override;
 	virtual FName GetLevelPackage() const override { return LevelPackage; }
-	virtual bool GetContainerInstance(FWorldPartitionActorDesc::FContainerInstance& OutContainerInstance) const override;
+	virtual bool GetContainerInstance(FContainerInstance& OutContainerInstance) const override;
 	virtual void CheckForErrors(IStreamingGenerationErrorHandler* ErrorHandler) const override;
 
 protected:
diff --git a/Engine/Source/Runtime/Engine/Public/WorldPartition/WorldPartitionActorDesc.h b/Engine/Source/Runtime/Engine/Public/WorldPartition/WorldPartitionActorDesc.h
--- a/Engine/Source/Runtime/Engine/Public/WorldPartition/WorldPartitionActorDesc.h
+++ b/Engine/Source/Runtime/Engine/Public/WorldPartition/WorldPartitionActorDesc.h
@@ -153,7 +153,6 @@ PRAGMA_ENABLE_DEPRECATION_WARNINGS
 
 	virtual bool IsContainerInstance() const { return false; }
 	virtual FName GetLevelPackage() const { return NAME_None; }
-	bool GetContainerInstance(const UActorDescContainer*& OutLevelContainer, FTransform& OutLevelTransform, EContainerClusterMode& OutClusterMode) const;
 	virtual bool GetContainerInstance(FContainerInstance& OutContainerInstance) const { return false; }
 
 	FGuid GetContentBundleGuid() const;
diff --git a/Engine/Source/Runtime/Engine/Public/WorldPartition/WorldPartitionBlueprintLibrary.h b/Engine/Source/Runtime/Engine/Public/WorldPartition/WorldPartitionBlueprintLibrary.h
--- a/Engine/Source/Runtime/Engine/Public/WorldPartition/WorldPartitionBlueprintLibrary.h
+++ b/Engine/Source/Runtime/Engine/Public/WorldPartition/WorldPartitionBlueprintLibrary.h
@@ -31,6 +31,10 @@ struct FActorDesc
 	UPROPERTY(BlueprintReadOnly, VisibleAnywhere, Category=Actor)
 	FGuid Guid;
 
+	/** Actor first native class. */
+	UPROPERTY(BlueprintReadOnly, VisibleAnywhere, Category=Actor)
+	TObjectPtr<UClass> NativeClass;
+
 	/** Actor class, can point to a native or Blueprint class and may be redirected. */
 	UPROPERTY(BlueprintReadOnly, VisibleAnywhere, Category=Actor)
 	FSoftObjectPath Class;
