Animation Curve Runtime & Editor Improvements

diff --git a/Engine/Source/Runtime/AnimGraphRuntime/Private/AnimNodes/AnimNode_PoseDriver.cpp b/Engine/Source/Runtime/AnimGraphRuntime/Private/AnimNodes/AnimNode_PoseDriver.cpp
--- a/Engine/Source/Runtime/AnimGraphRuntime/Private/AnimNodes/AnimNode_PoseDriver.cpp
+++ b/Engine/Source/Runtime/AnimGraphRuntime/Private/AnimNodes/AnimNode_PoseDriver.cpp
@@ -42,31 +42,14 @@ void FAnimNode_PoseDriver::RebuildPoseList(const FBoneContainer& InBoneContainer
 	const USkeleton* Skeleton = InPoseAsset->GetSkeleton();
 	if (Skeleton)
 	{
-		const TArray<FSmartName>& PoseNames = InPoseAsset->GetPoseNames();
+		const TArray<FName>& PoseNames = InPoseAsset->GetPoseFNames();
 		for (FPoseDriverTarget& PoseTarget : PoseTargets)
 		{
-			if (DriveOutput == EPoseDriverOutput::DriveCurves)
-			{
-				PoseTarget.DrivenUID = Skeleton->GetUIDByName(USkeleton::AnimCurveMappingName, PoseTarget.DrivenName);
-			}
-			else
-			{
-				PoseTarget.DrivenUID = SmartName::MaxUID;
-			}
-
 			const int32 PoseIndex = InPoseAsset->GetPoseIndexByName(PoseTarget.DrivenName);
 			if (PoseIndex != INDEX_NONE)
 			{
-				TArray<uint16> const& LUTIndex = InBoneContainer.GetUIDToArrayLookupTable();
-				if (LUTIndex.IsValidIndex(PoseNames[PoseIndex].UID) && LUTIndex[PoseNames[PoseIndex].UID] != MAX_uint16)
-				{
-					// we keep pose index as that is the fastest way to search when extracting pose asset
-					PoseTarget.PoseCurveIndex = PoseExtractContext.PoseCurves.Add(FPoseCurve(PoseIndex, PoseNames[PoseIndex].UID, 0.f));
-				}
-				else
-				{
-					PoseTarget.PoseCurveIndex = INDEX_NONE;
-				}
+				// we keep pose index as that is the fastest way to search when extracting pose asset
+				PoseTarget.PoseCurveIndex = PoseExtractContext.PoseCurves.Add(FPoseCurve(PoseIndex, PoseNames[PoseIndex], 0.f));
 			}
 			else
 			{
@@ -453,9 +436,9 @@ void FAnimNode_PoseDriver::Evaluate_AnyThread(FPoseContext& Output)
 				for (const FRBFOutputWeight& Weight : OutputWeights)
 				{
 					FPoseDriverTarget& PoseTarget = PoseTargets[Weight.TargetIndex];
-					if (PoseTarget.DrivenUID != SmartName::MaxUID)
+					if (PoseTarget.DrivenName != NAME_None)
 					{
-						Output.Curve.Set(PoseTarget.DrivenUID, Weight.TargetWeight);
+						Output.Curve.Set(PoseTarget.DrivenName, Weight.TargetWeight);
 					}
 				}
 
diff --git a/Engine/Source/Runtime/AnimGraphRuntime/Private/AnimNodes/AnimNode_PoseHandler.cpp b/Engine/Source/Runtime/AnimGraphRuntime/Private/AnimNodes/AnimNode_PoseHandler.cpp
--- a/Engine/Source/Runtime/AnimGraphRuntime/Private/AnimNodes/AnimNode_PoseHandler.cpp
+++ b/Engine/Source/Runtime/AnimGraphRuntime/Private/AnimNodes/AnimNode_PoseHandler.cpp
@@ -60,19 +60,16 @@ void FAnimNode_PoseHandler::CacheBones_AnyThread(const FAnimationCacheBonesConte
 void FAnimNode_PoseHandler::RebuildPoseList(const FBoneContainer& InBoneContainer, const UPoseAsset* InPoseAsset)
 {
 	PoseExtractContext.PoseCurves.Reset();
-	const TArray<FSmartName>& PoseNames = InPoseAsset->GetPoseNames();
+	const TArray<FName>& PoseNames = InPoseAsset->GetPoseFNames();
 	const int32 TotalPoseNum = PoseNames.Num();
 	if (TotalPoseNum > 0)
 	{
-		TArray<uint16> const& LUTIndex = InBoneContainer.GetUIDToArrayLookupTable();
 		for (int32 PoseIndex = 0; PoseIndex < PoseNames.Num(); ++PoseIndex)
 		{
-			const FSmartName& PoseName = PoseNames[PoseIndex];
-			if (LUTIndex.IsValidIndex(PoseName.UID) && LUTIndex[PoseName.UID] != MAX_uint16)
-			{
-				// we keep pose index as that is the fastest way to search when extracting pose asset
-				PoseExtractContext.PoseCurves.Add(FPoseCurve(PoseIndex, PoseName.UID, 0.f));
-			}
+			const FName& PoseName = PoseNames[PoseIndex];
+
+			// we keep pose index as that is the fastest way to search when extracting pose asset
+			PoseExtractContext.PoseCurves.Add(FPoseCurve(PoseIndex, PoseName, 0.f));
 		}
 	}
 }
diff --git a/Engine/Source/Runtime/AnimGraphRuntime/Private/BoneControllers/AnimNode_BoneDrivenController.cpp b/Engine/Source/Runtime/AnimGraphRuntime/Private/BoneControllers/AnimNode_BoneDrivenController.cpp
--- a/Engine/Source/Runtime/AnimGraphRuntime/Private/BoneControllers/AnimNode_BoneDrivenController.cpp
+++ b/Engine/Source/Runtime/AnimGraphRuntime/Private/BoneControllers/AnimNode_BoneDrivenController.cpp
@@ -201,17 +201,12 @@ void FAnimNode_BoneDrivenController::EvaluateComponentSpaceInternal(FComponentSp
 	const FTransform& SourceRefPoseBoneTransform = BoneContainer.GetRefPoseArray()[SourceBone.BoneIndex];
 	const FTransform SourceCurrentBoneTransform = Context.Pose.GetLocalSpaceTransform(SourceBone.GetCompactPoseIndex(BoneContainer));
 
-	const double FinalDriverValue = ExtractSourceValue(SourceCurrentBoneTransform, SourceRefPoseBoneTransform);
-
 	if (DestinationMode == EDrivenDestinationMode::MorphTarget || DestinationMode == EDrivenDestinationMode::MaterialParameter)
 	{
+		const double FinalDriverValue = ExtractSourceValue(SourceCurrentBoneTransform, SourceRefPoseBoneTransform);
+		
 		//	Morph target and Material parameter curves
-		USkeleton* Skeleton = Context.AnimInstanceProxy->GetSkeleton();
-		USkeleton::AnimCurveUID NameUID = Skeleton->GetUIDByName(USkeleton::AnimCurveMappingName, ParameterName);
-		if (NameUID != SmartName::MaxUID)
-		{
-			Context.Curve.Set(NameUID, static_cast<float>(FinalDriverValue));
-		}
+		Context.Curve.Set(ParameterName, static_cast<float>(FinalDriverValue));
 	}
 }
 
diff --git a/Engine/Source/Runtime/AnimGraphRuntime/Private/BoneControllers/AnimNode_TwistCorrectiveNode.cpp b/Engine/Source/Runtime/AnimGraphRuntime/Private/BoneControllers/AnimNode_TwistCorrectiveNode.cpp
--- a/Engine/Source/Runtime/AnimGraphRuntime/Private/BoneControllers/AnimNode_TwistCorrectiveNode.cpp
+++ b/Engine/Source/Runtime/AnimGraphRuntime/Private/BoneControllers/AnimNode_TwistCorrectiveNode.cpp
@@ -4,6 +4,7 @@
 #include "AnimationRuntime.h"
 #include "Animation/AnimInstanceProxy.h"
 #include "Components/SkeletalMeshComponent.h"
+#include "UObject/UE5MainStreamObjectVersion.h"
 
 #include UE_INLINE_GENERATED_CPP_BY_NAME(AnimNode_TwistCorrectiveNode)
 
@@ -49,14 +50,13 @@ void FAnimNode_TwistCorrectiveNode::EvaluateComponentSpaceInternal(FComponentSpa
 	const float FinalMappedValue = (FMath::Clamp(CurAngle, ReferenceAngle, RangeMaxInRadian) - ReferenceAngle) * (RemappedMax - RemappedMin) / (RangeMaxInRadian - ReferenceAngle);
 
 	//	set Curve Value
-	Context.Curve.Set(Curve.UID, FinalMappedValue*Alpha);
+	Context.Curve.Set(CurveName, FinalMappedValue*Alpha);
 }
 
 bool FAnimNode_TwistCorrectiveNode::IsValidToEvaluate(const USkeleton* Skeleton, const FBoneContainer& RequiredBones)
 {
 	return BaseFrame.Bone.IsValidToEvaluate(RequiredBones) && BaseFrame.Axis.IsValid() &&
-		TwistFrame.Bone.IsValidToEvaluate(RequiredBones) && TwistFrame.Axis.IsValid() &&
-		Curve.IsValidToEvaluate();
+		TwistFrame.Bone.IsValidToEvaluate(RequiredBones) && TwistFrame.Axis.IsValid();
 }
 
 void FAnimNode_TwistCorrectiveNode::InitializeBoneReferences(const FBoneContainer& RequiredBones)
@@ -75,7 +75,6 @@ void FAnimNode_TwistCorrectiveNode::Initialize_AnyThread(const FAnimationInitial
 {
 	DECLARE_SCOPE_HIERARCHICAL_COUNTER_ANIMNODE(Initialize_AnyThread)
 	FAnimNode_SkeletalControlBase::Initialize_AnyThread(Context);
-	Curve.Initialize(Context.AnimInstanceProxy->GetSkeleton());
 }
 
 void FAnimNode_TwistCorrectiveNode::CacheBones_AnyThread(const FAnimationCacheBonesContext& Context)
@@ -127,6 +126,22 @@ float FAnimNode_TwistCorrectiveNode::GetAngle(const FVector& Base, const FVector
 	return TwistAngle - BaseAngle;
 }
 
+bool FAnimNode_TwistCorrectiveNode::Serialize(FArchive& Ar)
+{
+	Ar.UsingCustomVersion(FUE5MainStreamObjectVersion::GUID);
+	return false;
+}
+
+void FAnimNode_TwistCorrectiveNode::PostSerialize(const FArchive& Ar)
+{
+#if WITH_EDITORONLY_DATA
+	if(Ar.CustomVer(FUE5MainStreamObjectVersion::GUID) < FUE5MainStreamObjectVersion::AnimationRemoveSmartNames)
+	{
+		CurveName = Curve_DEPRECATED.Name;
+	}
+#endif
+}
+
 /////////////////////////////////////////////////////////
 
 
diff --git a/Engine/Source/Runtime/AnimGraphRuntime/Public/AnimNodes/AnimNode_CopyPoseFromMesh.h b/Engine/Source/Runtime/AnimGraphRuntime/Public/AnimNodes/AnimNode_CopyPoseFromMesh.h
--- a/Engine/Source/Runtime/AnimGraphRuntime/Public/AnimNodes/AnimNode_CopyPoseFromMesh.h
+++ b/Engine/Source/Runtime/AnimGraphRuntime/Public/AnimNodes/AnimNode_CopyPoseFromMesh.h
@@ -65,14 +65,13 @@ private:
 	TWeakObjectPtr<USkeletalMesh>			CurrentlyUsedTargetMesh;
 	// cache of target space bases to source space bases
 	TMap<int32, int32> BoneMapToSource;
-	TMap<FName, SmartName::UID_Type> CurveNameToUIDMap;
 	TMap<int32, int32> SourceBoneToTarget;
 
 	// Cached transforms, copied on the game thread
 	TArray<FTransform> SourceMeshTransformArray;
 
 	// Cached curves, copied on the game thread
-	TMap<FName, float> SourceCurveList;
+	FBlendedHeapCurve SourceCurves;
 
 	// Cached attributes, copied on the game thread
 	UE::Anim::FMeshAttributeContainer SourceCustomAttributes;
diff --git a/Engine/Source/Runtime/AnimGraphRuntime/Public/AnimNodes/AnimNode_LayeredBoneBlend.h b/Engine/Source/Runtime/AnimGraphRuntime/Public/AnimNodes/AnimNode_LayeredBoneBlend.h
--- a/Engine/Source/Runtime/AnimGraphRuntime/Public/AnimNodes/AnimNode_LayeredBoneBlend.h
+++ b/Engine/Source/Runtime/AnimGraphRuntime/Public/AnimNodes/AnimNode_LayeredBoneBlend.h
@@ -99,7 +99,9 @@ protected:
 	// this array changes based on required bones
 	TArray<FPerBoneBlendWeight> DesiredBoneBlendWeights;
 	TArray<FPerBoneBlendWeight> CurrentBoneBlendWeights;
-	TArray<uint8> CurvePoseSourceIndices;
+
+	// Per-curve source pose index
+	TBaseBlendedCurve<FDefaultAllocator, UE::Anim::FCurveElementIndexed> CurvePoseSourceIndices;
 
 	// Serial number of the required bones container
 	uint16 RequiredBonesSerialNumber;
diff --git a/Engine/Source/Runtime/AnimGraphRuntime/Public/AnimNodes/AnimNode_ModifyCurve.h b/Engine/Source/Runtime/AnimGraphRuntime/Public/AnimNodes/AnimNode_ModifyCurve.h
--- a/Engine/Source/Runtime/AnimGraphRuntime/Public/AnimNodes/AnimNode_ModifyCurve.h
+++ b/Engine/Source/Runtime/AnimGraphRuntime/Public/AnimNodes/AnimNode_ModifyCurve.h
@@ -44,10 +44,17 @@ struct ANIMGRAPHRUNTIME_API FAnimNode_ModifyCurve : public FAnimNode_Base
 	UPROPERTY(meta = (BlueprintCompilerGeneratedDefaults))
 	TArray<FName> CurveNames;
 
+	UE_DEPRECATED(5.3, "LastCurveValues is no longer used.")
 	TArray<float> LastCurveValues;
+
+	UE_DEPRECATED(5.3, "LastCurveMapValues is no longer used.")
 	TMap<FName, float> LastCurveMapValues;
+
+	UE_DEPRECATED(5.3, "bInitializeLastValuesMap is no longer used.")
 	bool bInitializeLastValuesMap;
 
+	FBlendedHeapCurve LastCurve;
+
 	UPROPERTY(EditAnywhere, BlueprintReadWrite, Category = ModifyCurve, meta = (PinShownByDefault))
 	float Alpha;
 
@@ -55,6 +62,10 @@ struct ANIMGRAPHRUNTIME_API FAnimNode_ModifyCurve : public FAnimNode_Base
 	EModifyCurveApplyMode ApplyMode;
 
 	FAnimNode_ModifyCurve();
+	PRAGMA_DISABLE_DEPRECATION_WARNINGS
+	FAnimNode_ModifyCurve(const FAnimNode_ModifyCurve&) = default;
+	FAnimNode_ModifyCurve& operator=(const FAnimNode_ModifyCurve&) = default;
+	PRAGMA_ENABLE_DEPRECATION_WARNINGS
 
 	// FAnimNode_Base interface
 	virtual void Initialize_AnyThread(const FAnimationInitializeContext& Context) override;
@@ -72,6 +83,6 @@ struct ANIMGRAPHRUNTIME_API FAnimNode_ModifyCurve : public FAnimNode_Base
 #endif // WITH_EDITOR
 
 private:
-	void ProcessCurveOperation(const EModifyCurveApplyMode& InApplyMode, FPoseContext& Output, const SmartName::UID_Type& NameUID, float CurrentValue, float NewValue);
-	void ProcessCurveWMAOperation(FPoseContext& Output, const SmartName::UID_Type& NameUID, float CurrentValue, float NewValue, float& InOutLastValue);
+	float ProcessCurveOperation(float CurrentValue, float NewValue) const;
+	float ProcessCurveWMAOperation(float CurrentValue, float LastValue) const;
 };
diff --git a/Engine/Source/Runtime/AnimGraphRuntime/Public/AnimNodes/AnimNode_PoseBlendNode.h b/Engine/Source/Runtime/AnimGraphRuntime/Public/AnimNodes/AnimNode_PoseBlendNode.h
--- a/Engine/Source/Runtime/AnimGraphRuntime/Public/AnimNodes/AnimNode_PoseBlendNode.h
+++ b/Engine/Source/Runtime/AnimGraphRuntime/Public/AnimNodes/AnimNode_PoseBlendNode.h
@@ -6,6 +6,7 @@
 #include "UObject/ObjectMacros.h"
 #include "AnimNodes/AnimNode_PoseHandler.h"
 #include "AlphaBlend.h"
+#include "Animation/AnimBulkCurves.h"
 #include "AnimNode_PoseBlendNode.generated.h"
 
 // Evaluates a point in an anim sequence, using a specific time input rather than advancing time internally.
@@ -27,6 +28,10 @@ public:
 	UPROPERTY(EditAnywhere, Category = "Blend")
 	TObjectPtr<UCurveFloat> CustomCurve;
 
+private:
+	// Cached curves to extract
+	UE::Anim::TNamedValueArray<FDefaultAllocator, UE::Anim::FNamedIndexElement> BulkCurves;
+
 public:	
 	FAnimNode_PoseBlendNode();
 
diff --git a/Engine/Source/Runtime/AnimGraphRuntime/Public/AnimNodes/AnimNode_PoseDriver.h b/Engine/Source/Runtime/AnimGraphRuntime/Public/AnimNodes/AnimNode_PoseDriver.h
--- a/Engine/Source/Runtime/AnimGraphRuntime/Public/AnimNodes/AnimNode_PoseDriver.h
+++ b/Engine/Source/Runtime/AnimGraphRuntime/Public/AnimNodes/AnimNode_PoseDriver.h
@@ -103,7 +103,7 @@ struct ANIMGRAPHRUNTIME_API FPoseDriverTarget
 	UPROPERTY(EditAnywhere, Category = PoseDriver)
 	FName DrivenName;
 
-	/** Cached curve UID, if DriveOutput is set to DriveCurves */
+	UE_DEPRECATED(5.3, "DrivenUID is no longer used")
 	USkeleton::AnimCurveUID DrivenUID = SmartName::MaxUID;
 
 	/** Pose buffer index */
@@ -113,6 +113,12 @@ struct ANIMGRAPHRUNTIME_API FPoseDriverTarget
 	UPROPERTY(EditAnywhere, Category = PoseDriver)
 	bool bIsHidden;
 
+	// removing deprecation for default copy operator/constructor to avoid deprecation warnings
+	PRAGMA_DISABLE_DEPRECATION_WARNINGS
+	FPoseDriverTarget(const FPoseDriverTarget&) = default;
+	FPoseDriverTarget& operator=(const FPoseDriverTarget&) = default;
+	PRAGMA_ENABLE_DEPRECATION_WARNINGS
+
 	FPoseDriverTarget()
 		: TargetRotation(ForceInitToZero)
 		, TargetScale(1.f)
@@ -120,7 +126,6 @@ struct ANIMGRAPHRUNTIME_API FPoseDriverTarget
 		, FunctionType(ERBFFunctionType::DefaultFunction)
 		, bApplyCustomCurve(false)
 		, DrivenName(NAME_None)
-		, DrivenUID(SmartName::MaxUID)
 		, PoseCurveIndex(INDEX_NONE)
 		, bIsHidden(false)
 	{}
diff --git a/Engine/Source/Runtime/AnimGraphRuntime/Public/BoneControllers/AnimNode_TwistCorrectiveNode.h b/Engine/Source/Runtime/AnimGraphRuntime/Public/BoneControllers/AnimNode_TwistCorrectiveNode.h
--- a/Engine/Source/Runtime/AnimGraphRuntime/Public/BoneControllers/AnimNode_TwistCorrectiveNode.h
+++ b/Engine/Source/Runtime/AnimGraphRuntime/Public/BoneControllers/AnimNode_TwistCorrectiveNode.h
@@ -67,8 +67,13 @@ struct ANIMGRAPHRUNTIME_API FAnimNode_TwistCorrectiveNode : public FAnimNode_Ske
 	UPROPERTY(EditAnywhere, Category = "Mapping", meta = (EditCondition = bUseRange, DisplayName = "Mapped Range Max"))
  	float RemappedMax;
 
+#if WITH_EDITORONLY_DATA
+	UPROPERTY()
+	FAnimCurveParam Curve_DEPRECATED;
+#endif
+
 	UPROPERTY(EditAnywhere, Category = "Output Curve")
-	FAnimCurveParam Curve;
+	FName CurveName;
 
 public:
 	FAnimNode_TwistCorrectiveNode();
@@ -84,6 +89,10 @@ public:
 	virtual bool IsValidToEvaluate(const USkeleton* Skeleton, const FBoneContainer& RequiredBones) override;
 	// End of FAnimNode_SkeletalControlBase interface
 
+	// Type traits support
+	bool Serialize(FArchive& Ar);
+	void PostSerialize(const FArchive& Ar);
+
 protected:
 	
 	// FAnimNode_SkeletalControlBase protected interface
@@ -100,3 +109,13 @@ private:
 	// Get Angle of Base, and Twist from Reference Bone Transform
 	float	GetAngle(const FVector& Base, const FVector& Twist, const FTransform& ReferencetBoneTransform) const;
 };
+
+template<>
+struct TStructOpsTypeTraits<FAnimNode_TwistCorrectiveNode> : public TStructOpsTypeTraitsBase2<FAnimNode_TwistCorrectiveNode>
+{
+	enum
+	{
+		WithSerializer = true,
+		WithPostSerialize = true
+	};
+};
\ No newline at end of file
diff --git a/Engine/Source/Runtime/Core/Public/UObject/UE5MainStreamObjectVersion.h b/Engine/Source/Runtime/Core/Public/UObject/UE5MainStreamObjectVersion.h
--- a/Engine/Source/Runtime/Core/Public/UObject/UE5MainStreamObjectVersion.h
+++ b/Engine/Source/Runtime/Core/Public/UObject/UE5MainStreamObjectVersion.h
@@ -337,6 +337,9 @@ struct CORE_API FUE5MainStreamObjectVersion
 		// Explicit override of Groom's hair width
 		GroomAssetWidthOverride,
 
+		// Smart names removed from animation systems
+		AnimationRemoveSmartNames,
+
 		// -----<new versions can be added above this line>-------------------------------------------------
 		VersionPlusOne,
 		LatestVersion = VersionPlusOne - 1
diff --git a/Engine/Source/Runtime/Engine/Classes/Animation/AnimCompress.h b/Engine/Source/Runtime/Engine/Classes/Animation/AnimCompress.h
--- a/Engine/Source/Runtime/Engine/Classes/Animation/AnimCompress.h
+++ b/Engine/Source/Runtime/Engine/Classes/Animation/AnimCompress.h
@@ -269,7 +269,7 @@ namespace UE
 			// This is a version string that mimics the old versioning scheme. If you
 			// want to bump this version, generate a new guid using VS->Tools->Create GUID and
 			// return it here. Ex.
-			static FString AnimationCompressionVersionString = TEXT("FB8D0FF1ED0C4848951601EB6BB028F6");
+			static FString AnimationCompressionVersionString = TEXT("E64909A2AA124539905EBA5557DDEE58");
 			
 			struct ENGINE_API FAnimationCompressionMemorySummaryScope
 			{
diff --git a/Engine/Source/Runtime/Engine/Classes/Animation/AnimCurveCompressionCodec.h b/Engine/Source/Runtime/Engine/Classes/Animation/AnimCurveCompressionCodec.h
--- a/Engine/Source/Runtime/Engine/Classes/Animation/AnimCurveCompressionCodec.h
+++ b/Engine/Source/Runtime/Engine/Classes/Animation/AnimCurveCompressionCodec.h
@@ -75,5 +75,8 @@ class ENGINE_API UAnimCurveCompressionCodec : public UObject
 	 * Note: Codecs should _NOT_ rely on any member properties during decompression. Decompression
 	 * behavior should entirely be driven by code and the compressed data.
 	 */
-	virtual float DecompressCurve(const FCompressedAnimSequence& AnimSeq, SmartName::UID_Type CurveUID, float CurrentTime) const PURE_VIRTUAL(UAnimCurveCompressionCodec::DecompressCurve, return 0.0f;);
+	virtual float DecompressCurve(const FCompressedAnimSequence& AnimSeq, FName CurveName, float CurrentTime) const PURE_VIRTUAL(UAnimCurveCompressionCodec::DecompressCurve, return 0.0f;);
+	
+	UE_DEPRECATED(5.3, "Please use DecompressCurve that takes an FName.")
+    virtual float DecompressCurve(const FCompressedAnimSequence& AnimSeq, SmartName::UID_Type CurveUID, float CurrentTime) const { return 0.0f; }
 };
diff --git a/Engine/Source/Runtime/Engine/Classes/Animation/AnimCurveCompressionCodec_CompressedRichCurve.h b/Engine/Source/Runtime/Engine/Classes/Animation/AnimCurveCompressionCodec_CompressedRichCurve.h
--- a/Engine/Source/Runtime/Engine/Classes/Animation/AnimCurveCompressionCodec_CompressedRichCurve.h
+++ b/Engine/Source/Runtime/Engine/Classes/Animation/AnimCurveCompressionCodec_CompressedRichCurve.h
@@ -36,7 +36,7 @@ class ENGINE_API UAnimCurveCompressionCodec_CompressedRichCurve : public UAnimCu
 	virtual bool Compress(const FCompressibleAnimData& AnimSeq, FAnimCurveCompressionResult& OutResult) override;
 	virtual void PopulateDDCKey(FArchive& Ar) override;
 #endif
-
+	
 	virtual void DecompressCurves(const FCompressedAnimSequence& AnimSeq, FBlendedCurve& Curves, float CurrentTime) const override;
-	virtual float DecompressCurve(const FCompressedAnimSequence& AnimSeq, SmartName::UID_Type CurveUID, float CurrentTime) const override;
+	virtual float DecompressCurve(const FCompressedAnimSequence& AnimSeq, FName CurveName, float CurrentTime) const override;
 };
diff --git a/Engine/Source/Runtime/Engine/Classes/Animation/AnimCurveCompressionCodec_UniformIndexable.h b/Engine/Source/Runtime/Engine/Classes/Animation/AnimCurveCompressionCodec_UniformIndexable.h
--- a/Engine/Source/Runtime/Engine/Classes/Animation/AnimCurveCompressionCodec_UniformIndexable.h
+++ b/Engine/Source/Runtime/Engine/Classes/Animation/AnimCurveCompressionCodec_UniformIndexable.h
@@ -23,7 +23,10 @@ private:
 	bool bUseCompressedData;
 public:
 
-	FAnimCurveBufferAccess(const UAnimSequenceBase* InSequence, USkeleton::AnimCurveUID InUID);
+	UE_DEPRECATED(5.3, "Please use the constructor that takes a FName")
+	FAnimCurveBufferAccess(const UAnimSequenceBase* InSequence, SmartName::UID_Type InUID);
+
+	FAnimCurveBufferAccess(const UAnimSequenceBase* InSequence, FName InName);
 
 	bool IsValid() const
 	{
@@ -49,9 +52,11 @@ class ENGINE_API UAnimCurveCompressionCodec_UniformIndexable : public UAnimCurve
 	virtual bool Compress(const FCompressibleAnimData& AnimSeq, FAnimCurveCompressionResult& OutResult) override;
 	virtual void PopulateDDCKey(FArchive& Ar) override;
 #endif
-
+	
 	virtual void DecompressCurves(const FCompressedAnimSequence& AnimSeq, FBlendedCurve& Curves, float CurrentTime) const override;
-	virtual float DecompressCurve(const FCompressedAnimSequence& AnimSeq, SmartName::UID_Type CurveUID, float CurrentTime) const override;
+	virtual float DecompressCurve(const FCompressedAnimSequence& AnimSeq, FName CurveName, float CurrentTime) const override;
 
-	bool GetCurveBufferAndSamples(const FCompressedAnimSequence& AnimSeq, SmartName::UID_Type CurveUID, const float*& OutCurveBuffer, int32& OutSamples, float& OutSampleRate);
+	UE_DEPRECATED(5.3, "Please use GetCurveBufferAndSamples that takes a curve's FName")
+	bool GetCurveBufferAndSamples(const FCompressedAnimSequence& AnimSeq, SmartName::UID_Type CurveUID, const float*& OutCurveBuffer, int32& OutSamples, float& OutSampleRate) { return false; }
+	bool GetCurveBufferAndSamples(const FCompressedAnimSequence& AnimSeq, FName CurveName, const float*& OutCurveBuffer, int32& OutSamples, float& OutSampleRate);
 };
diff --git a/Engine/Source/Runtime/Engine/Classes/Animation/AnimCurveCompressionCodec_UniformlySampled.h b/Engine/Source/Runtime/Engine/Classes/Animation/AnimCurveCompressionCodec_UniformlySampled.h
--- a/Engine/Source/Runtime/Engine/Classes/Animation/AnimCurveCompressionCodec_UniformlySampled.h
+++ b/Engine/Source/Runtime/Engine/Classes/Animation/AnimCurveCompressionCodec_UniformlySampled.h
@@ -32,7 +32,7 @@ class ENGINE_API UAnimCurveCompressionCodec_UniformlySampled : public UAnimCurve
 	virtual bool Compress(const FCompressibleAnimData& AnimSeq, FAnimCurveCompressionResult& OutResult) override;
 	virtual void PopulateDDCKey(FArchive& Ar) override;
 #endif
-
+	
 	virtual void DecompressCurves(const FCompressedAnimSequence& AnimSeq, FBlendedCurve& Curves, float CurrentTime) const override;
-	virtual float DecompressCurve(const FCompressedAnimSequence& AnimSeq, SmartName::UID_Type CurveUID, float CurrentTime) const override;
+	virtual float DecompressCurve(const FCompressedAnimSequence& AnimSeq, FName CurveName, float CurrentTime) const override;
 };
diff --git a/Engine/Source/Runtime/Engine/Classes/Animation/AnimData/CurveIdentifier.h b/Engine/Source/Runtime/Engine/Classes/Animation/AnimData/CurveIdentifier.h
--- a/Engine/Source/Runtime/Engine/Classes/Animation/AnimData/CurveIdentifier.h
+++ b/Engine/Source/Runtime/Engine/Classes/Animation/AnimData/CurveIdentifier.h
@@ -35,21 +35,36 @@ enum class EVectorCurveChannel : uint8
 USTRUCT(BlueprintType)
 struct ENGINE_API FAnimationCurveIdentifier
 {
-	GENERATED_USTRUCT_BODY();
+	GENERATED_BODY();
 
-	FAnimationCurveIdentifier() : CurveType(ERawCurveTrackTypes::RCT_MAX), Channel(ETransformCurveChannel::Invalid), Axis(EVectorCurveChannel::Invalid) {}
-	FAnimationCurveIdentifier(const FSmartName& InSmartName, ERawCurveTrackTypes InCurveType) : InternalName(InSmartName), CurveType(InCurveType), Channel(ETransformCurveChannel::Invalid), Axis(EVectorCurveChannel::Invalid) {}
-	
-	explicit FAnimationCurveIdentifier(const SmartName::UID_Type& InUID, ERawCurveTrackTypes InCurveType) : CurveType(InCurveType)
+	FAnimationCurveIdentifier() = default;
+
+	UE_DEPRECATED(5.3, "Please use the constructor that takes an FName")
+	FAnimationCurveIdentifier(const FSmartName& InSmartName, ERawCurveTrackTypes InCurveType)
+		: CurveName(InSmartName.DisplayName)
+		, CurveType(InCurveType)
+		, Channel(ETransformCurveChannel::Invalid)
+		, Axis(EVectorCurveChannel::Invalid)
+	{}
+
+	UE_DEPRECATED(5.3, "Please use the constructor that takes an FName")
+	explicit FAnimationCurveIdentifier(const SmartName::UID_Type& InUID, ERawCurveTrackTypes InCurveType)
+		: CurveType(InCurveType)
 	{
-		InternalName.UID = InUID;
 	}
 
+	FAnimationCurveIdentifier(const FName& InName, ERawCurveTrackTypes InCurveType)
+		: CurveName(InName)
+		, CurveType(InCurveType)
+	{
+	}
+	
 	bool operator==(const FAnimationCurveIdentifier& Other) const
 	{
-		return (InternalName == Other.InternalName && CurveType == Other.CurveType 
+		return (CurveName == Other.CurveName && CurveType == Other.CurveType 
 			&& Channel == Other.Channel && Axis == Other.Axis);
 	}
+
 	bool operator!=(const FAnimationCurveIdentifier& Other) const
 	{
 		return !(*this == Other);
@@ -57,12 +72,12 @@ struct ENGINE_API FAnimationCurveIdentifier
 
 	bool IsValid() const
 	{
-		return InternalName.IsValid() && CurveType != ERawCurveTrackTypes::RCT_MAX;
+		return CurveName != NAME_None && CurveType != ERawCurveTrackTypes::RCT_MAX;
 	}
 
 	FAnimationCurveIdentifier& operator=(const FAnimationCurveIdentifier& Other)
 	{
-		InternalName = Other.InternalName;
+		CurveName = Other.CurveName;
 		CurveType = Other.CurveType;
 		Channel = Other.Channel;
 		Axis = Other.Axis;
@@ -72,20 +87,27 @@ struct ENGINE_API FAnimationCurveIdentifier
 	
 	friend uint32 GetTypeHash(const FAnimationCurveIdentifier& CurveId)
 	{
-		return HashCombine(HashCombine(HashCombine((uint32)CurveId.Channel, (uint32)CurveId.Axis), (uint32)CurveId.CurveType), GetTypeHash(CurveId.InternalName.DisplayName));
+		return HashCombine(HashCombine(HashCombine((uint32)CurveId.Channel, (uint32)CurveId.Axis), (uint32)CurveId.CurveType), GetTypeHash(CurveId.CurveName));
 	}
+	
+	void PostSerialize(const FArchive& Ar);
 
+#if WITH_EDITORONLY_DATA
 	UPROPERTY()
-	FSmartName InternalName;
-	
+	FSmartName InternalName_DEPRECATED;
+#endif
+
 	UPROPERTY()
-	ERawCurveTrackTypes CurveType;
+	FName CurveName = NAME_None;
 
 	UPROPERTY()
-	ETransformCurveChannel Channel;
+	ERawCurveTrackTypes CurveType = ERawCurveTrackTypes::RCT_MAX;
+
+	UPROPERTY()
+	ETransformCurveChannel Channel = ETransformCurveChannel::Invalid;
 	
 	UPROPERTY()
-	EVectorCurveChannel Axis;
+	EVectorCurveChannel Axis = EVectorCurveChannel::Invalid;
 };
 
 template<>
@@ -94,7 +116,8 @@ struct TStructOpsTypeTraits<FAnimationCurveIdentifier> : public TStructOpsTypeTr
 	enum
 	{
 		WithCopy = true,
-		WithIdenticalViaEquality = true
+		WithIdenticalViaEquality = true,
+		WithPostSerialize = true,
 	};
 };
 
@@ -119,7 +142,7 @@ public:
 	UFUNCTION(BlueprintPure, Category = Curve, meta = (ScriptMethod))
 	static FName GetName(UPARAM(ref) FAnimationCurveIdentifier& Identifier)
 	{
-		return Identifier.InternalName.DisplayName;
+		return Identifier.CurveName;
 	}
 
 	/**
@@ -133,18 +156,18 @@ public:
 
 #if WITH_EDITOR
 	/**
-	* Constructs a valid FAnimationCurveIdentifier instance. Ensuring that the underlying SmartName exists on the provided Skeleton for the provided curve type.
-	* If it is not found initially it will add it to the Skeleton thus modifying it.
+	* Constructs a valid FAnimationCurveIdentifier instance.
 	*
-	* @param	InSkeleton			Skeleton on which to look for or add the curve name
+	* @param	InOutIdentifier		The identifier to set up
 	* @param	Name				Name of the curve to find or add
 	* @param	CurveType			Type of the curve to find or add
-	*
-	* @return	Valid FAnimationCurveIdentifier if the operation was successful
 	*/
 	UFUNCTION(BlueprintCallable, Category = Curves, meta = (ScriptMethod))
-	static FAnimationCurveIdentifier GetCurveIdentifier(USkeleton* InSkeleton, FName Name, ERawCurveTrackTypes CurveType);
+	static void SetCurveIdentifier(UPARAM(ref) FAnimationCurveIdentifier& InOutIdentifier, FName Name, ERawCurveTrackTypes CurveType);
 
+	UFUNCTION(BlueprintCallable, Category = Curves, meta = (ScriptMethod, DeprecatedFunction, DeprecationMessage="Please use SetCurveIdentifier."))
+	static FAnimationCurveIdentifier GetCurveIdentifier(USkeleton* InSkeleton, FName Name, ERawCurveTrackTypes CurveType);
+	
 	/**
 	* Tries to construct a valid FAnimationCurveIdentifier instance. It tries to find the underlying SmartName on the provided Skeleton for the provided curve type.
 	*
@@ -154,7 +177,7 @@ public:
 	*
 	* @return	Valid FAnimationCurveIdentifier if the name exists on the skeleton and the operation was successful, invalid otherwise
 	*/
-	UFUNCTION(BlueprintCallable, Category = Curves, meta = (ScriptMethod))
+	UFUNCTION(BlueprintCallable, Category = Curves, meta = (ScriptMethod, DeprecatedFunction, DeprecationMessage="Curve identifiers are no longer retrievable globally from the skeleton, they are specified per-animation."))
 	static FAnimationCurveIdentifier FindCurveIdentifier(const USkeleton* InSkeleton, FName Name, ERawCurveTrackTypes CurveType);
 
 	/**
@@ -165,7 +188,7 @@ public:
 	*
 	* @return	Array of FAnimationCurveIdentifier instances each representing a unique curve if the operation was successful, empyty array otherwise
 	*/
-	UFUNCTION(BlueprintCallable, Category = Curves, meta = (ScriptMethod))
+	UFUNCTION(BlueprintCallable, Category = Curves, meta = (ScriptMethod, DeprecatedFunction, DeprecationMessage="Curve identifiers are no longer retrievable globally from the skeleton, they are specified per-animation."))
 	static TArray<FAnimationCurveIdentifier> GetCurveIdentifiers(USkeleton* InSkeleton, ERawCurveTrackTypes CurveType);
 
 	/**
diff --git a/Engine/Source/Runtime/Engine/Classes/Animation/AnimData/IAnimationDataController.h b/Engine/Source/Runtime/Engine/Classes/Animation/AnimData/IAnimationDataController.h
--- a/Engine/Source/Runtime/Engine/Classes/Animation/AnimData/IAnimationDataController.h
+++ b/Engine/Source/Runtime/Engine/Classes/Animation/AnimData/IAnimationDataController.h
@@ -488,8 +488,9 @@ public:
 	* @param	SupportedCurveType	Curve type for which the names should be updated
 	* @param	bShouldTransact		Whether or not any undo-redo changes should be generated
 	*/
-	UFUNCTION(BlueprintCallable, Category = CurveData)
-	virtual void UpdateCurveNamesFromSkeleton(const USkeleton* Skeleton, ERawCurveTrackTypes SupportedCurveType, bool bShouldTransact = true) = 0;
+	UE_DEPRECATED(5.3, "This function is no longer used.")
+	UFUNCTION(BlueprintCallable, Category = CurveData, meta=(DeprecatedFunction, DeprecationMessage="This function is no longer used."))
+	virtual void UpdateCurveNamesFromSkeleton(const USkeleton* Skeleton, ERawCurveTrackTypes SupportedCurveType, bool bShouldTransact = true) {}
 
 	/**
 	* Updates the curve names with the provided skeleton, if a display name is not found it will be added thus modifying the skeleton. Broadcasts a EAnimDataModelNotifyType::CurveRenamed for each curve name for which the UID was different or if it was added as a new smart-name, wrapped within BracketOpened/BracketClosed notifies.
@@ -498,8 +499,9 @@ public:
 	* @param	SupportedCurveType	Curve type for which the names should be updated
 	* @param	bShouldTransact		Whether or not any undo-redo changes should be generated
 	*/
-	UFUNCTION(BlueprintCallable, Category = CurveData)
-	virtual void FindOrAddCurveNamesOnSkeleton(USkeleton* Skeleton, ERawCurveTrackTypes SupportedCurveType, bool bShouldTransact = true) = 0;
+	UE_DEPRECATED(5.3, "This function is no longer used.")
+	UFUNCTION(BlueprintCallable, Category = CurveData, meta=(DeprecatedFunction, DeprecationMessage="This function is no longer used."))
+	virtual void FindOrAddCurveNamesOnSkeleton(USkeleton* Skeleton, ERawCurveTrackTypes SupportedCurveType, bool bShouldTransact = true) {}
 
 	/**
 	* Removes any bone track for which the name was not found in the provided skeleton. Broadcasts a EAnimDataModelNotifyType::TrackRemoved for each track which was not found in the skeleton, wrapped within BracketOpened/BracketClosed notifies.
diff --git a/Engine/Source/Runtime/Engine/Classes/Animation/AnimInstance.h b/Engine/Source/Runtime/Engine/Classes/Animation/AnimInstance.h
--- a/Engine/Source/Runtime/Engine/Classes/Animation/AnimInstance.h
+++ b/Engine/Source/Runtime/Engine/Classes/Animation/AnimInstance.h
@@ -51,6 +51,7 @@ struct FAnimNode_LinkedAnimLayer;
 struct FNodeDebugData;
 enum class ETransitionRequestQueueMode : uint8;
 enum class ETransitionRequestOverwriteMode : uint8;
+class UAnimMontage;
 
 typedef TArray<FTransform> FTransformArrayA2;
 
@@ -58,6 +59,7 @@ namespace UE::Anim
 {
 	struct FHeapAttributeContainer;
 	using FSlotInertializationRequest = TPair<float, const UBlendProfile*>;
+	struct FCurveFilterSettings;
 }	// namespace UE::Anim
 
 struct FParallelEvaluationData
@@ -1103,7 +1105,7 @@ public:
 	UFUNCTION(BlueprintPure, Category = "Animation|Curves", meta=(BlueprintThreadSafe))
 	void GetActiveCurveNames(EAnimCurveType CurveType, TArray<FName>& OutNames) const;
 
-	/* This returns all curve names */
+	/* This returns all curve names. This is the same as calling GetActiveCurveNames with CurveType == AttributeCurve */
 	UFUNCTION(BlueprintPure, Category = "Animation|Curves", meta=(BlueprintThreadSafe))
 	void GetAllCurveNames(TArray<FName>& OutNames) const;
 
@@ -1403,6 +1405,9 @@ public:
 	/**
 	* Recalculate Required Curves based on Required Bones [RequiredBones]
 	*/
+	void RecalcRequiredCurves(const UE::Anim::FCurveFilterSettings& CurveFilterSettings);
+	
+	UE_DEPRECATED(5.3, "Please use RecalcRequiredCurves that takes a FCurveFilterSettings.")
 	void RecalcRequiredCurves(const FCurveEvaluationOption& CurveEvalOption);
 
 	// @todo document
@@ -1422,19 +1427,19 @@ public:
 	/** Triggers end on active notify states and clears the array */
 	void EndNotifyStates();
 
-	/** Add curve float data using a curve Uid, the name of the curve will be resolved from the skeleton **/
-	void AddCurveValue(const SmartName::UID_Type Uid, float Value);
+	UE_DEPRECATED(5.3, "Please use AddCurveValue that takes a CurveName")
+	void AddCurveValue(const SmartName::UID_Type Uid, float Value) {}
 
-	/** Add curve float data using a curve Uid, the name of the curve will be resolved from the skeleton. This uses an already-resolved proxy and mapping table for efficency **/
-	void AddCurveValue(const FSmartNameMapping& Mapping, const FName& CurveName, float Value);
+	UE_DEPRECATED(5.3, "Please use AddCurveValue that takes a CurveName")
+	void AddCurveValue(const FSmartNameMapping& Mapping, const FName& CurveName, float Value) {}
 
 	/** Given a machine index, record a state machine weight for this frame */
 	void RecordMachineWeight(const int32 InMachineClassIndex, const float InMachineWeight);
+	
 	/** 
-	 * Add curve float data, using a curve name. External values should all be added using
-	 * The curve UID to the public version of this method
+	 * Add curve float data, using a curve name.
 	 */
-	void AddCurveValue(const FName& CurveName, float Value);
+	void AddCurveValue(const FName& CurveName, float Value, bool bMorphtarget = false, bool bMaterial = false);
 
 	/** Given a machine and state index, record a state weight for this frame */
 	void RecordStateWeight(const int32 InMachineClassIndex, const int32 InStateIndex, const float InStateWeight, const float InElapsedTime);
diff --git a/Engine/Source/Runtime/Engine/Classes/Animation/AnimMontage.h b/Engine/Source/Runtime/Engine/Classes/Animation/AnimMontage.h
--- a/Engine/Source/Runtime/Engine/Classes/Animation/AnimMontage.h
+++ b/Engine/Source/Runtime/Engine/Classes/Animation/AnimMontage.h
@@ -82,7 +82,7 @@ public:
  * contains slot name, and animation data 
  */
 USTRUCT()
-struct FSlotAnimationTrack
+struct ENGINE_API FSlotAnimationTrack
 {
 	GENERATED_USTRUCT_BODY()
 
@@ -92,7 +92,7 @@ struct FSlotAnimationTrack
 	UPROPERTY(EditAnywhere, Category=Slot)
 	FAnimTrack AnimTrack;
 
-	ENGINE_API FSlotAnimationTrack();
+	FSlotAnimationTrack();
 };
 
 /** 
diff --git a/Engine/Source/Runtime/Engine/Classes/Animation/AnimNode_Inertialization.h b/Engine/Source/Runtime/Engine/Classes/Animation/AnimNode_Inertialization.h
--- a/Engine/Source/Runtime/Engine/Classes/Animation/AnimNode_Inertialization.h
+++ b/Engine/Source/Runtime/Engine/Classes/Animation/AnimNode_Inertialization.h
@@ -67,6 +67,8 @@ enum class EInertializationSpace : uint8
 struct ENGINE_API FInertializationCurve
 {
 	FBlendedHeapCurve BlendedCurve;
+
+	UE_DEPRECATED(5.3, "CurveUIDToArrayIndexLUT is no longer used.")
 	TArray<uint16> CurveUIDToArrayIndexLUT;
 
 	FInertializationCurve() = default;
@@ -84,31 +86,19 @@ struct ENGINE_API FInertializationCurve
 	FInertializationCurve& operator=(const FInertializationCurve& Other)
 	{
 		BlendedCurve.CopyFrom(Other.BlendedCurve);
-		BlendedCurve.UIDToArrayIndexLUT = &CurveUIDToArrayIndexLUT;
-		CurveUIDToArrayIndexLUT = Other.CurveUIDToArrayIndexLUT;
 		return *this;
 	}
 
 	FInertializationCurve& operator=(FInertializationCurve&& Other)
 	{
 		BlendedCurve.MoveFrom(Other.BlendedCurve);
-		BlendedCurve.UIDToArrayIndexLUT = &CurveUIDToArrayIndexLUT;
-		CurveUIDToArrayIndexLUT = MoveTemp(Other.CurveUIDToArrayIndexLUT);
 		return *this;
 	}
 
 	template <typename OtherAllocator>
-	void InitFrom(const FBaseBlendedCurve<OtherAllocator>& Other)
+	void InitFrom(const TBaseBlendedCurve<OtherAllocator>& Other)
 	{
-		CurveUIDToArrayIndexLUT.Reset();
-
 		BlendedCurve.CopyFrom(Other);
-		BlendedCurve.UIDToArrayIndexLUT = &CurveUIDToArrayIndexLUT;
-
-		if (Other.UIDToArrayIndexLUT)
-		{
-			CurveUIDToArrayIndexLUT = *Other.UIDToArrayIndexLUT;
-		}
 	}
 };
 
@@ -246,22 +236,17 @@ struct FInertializationBoneDiff
 	}
 };
 
-USTRUCT()
-struct FInertializationCurveDiff
+struct FInertializationCurveDiffElement : public UE::Anim::FCurveElement
 {
-	GENERATED_BODY();
+	UE_DEPRECATED(5.3, "Please use Value instead.")
+	static float Delta;
+	float Derivative = 0.0f;
 
-	float Delta;
-	float Derivative;
-
-	FInertializationCurveDiff()
-		: Delta(0.0f)
-		, Derivative(0.0f)
-	{}
+	FInertializationCurveDiffElement() = default;
 
 	void Clear()
 	{
-		Delta = 0.0f;
+		Value = 0.0f;
 		Derivative = 0.0f;
 	}
 };
@@ -291,9 +276,9 @@ struct FInertializationPoseDiff
 	// AttachParentName					the current frame's attach parent name (for checking if the attachment has changed)
 	// Prev1							the previous frame's pose
 	// Prev2							the pose from two frames before
-	// FilteredCurvesUIDs				list of curves we don't want to inertialize
+	// CurveFilter						filter for curves we don't want to inertialize
 	//
-	void InitFrom(const FCompactPose& Pose, const FBlendedCurve& Curves, const FTransform& ComponentTransform, const FName& AttachParentName, const FInertializationPose& Prev1, const FInertializationPose& Prev2, const TSet<SmartName::UID_Type>& FilteredCurvesUIDs);
+	void InitFrom(const FCompactPose& Pose, const FBlendedCurve& Curves, const FTransform& ComponentTransform, const FName& AttachParentName, const FInertializationPose& Prev1, const FInertializationPose& Prev2, const UE::Anim::FCurveFilter& CurveFilter);
 
 	// Apply this difference to a pose, decaying over time as InertializationElapsedTime approaches InertializationDuration
 	//
@@ -319,8 +304,8 @@ private:
 	// Bone differences indexed by skeleton bone index
 	TArray<FInertializationBoneDiff> BoneDiffs;
 
-	// Curve differences indexed by CurveID
-	TArray<FInertializationCurveDiff> CurveDiffs;
+	// Curve differences
+	TBaseBlendedCurve<FDefaultAllocator, FInertializationCurveDiffElement> CurveDiffs;
 
 	// Inertialization space (local vs world for situations where we wish to correct a world-space discontinuity such as an abrupt orientation change)
 	EInertializationSpace InertializationSpace;
@@ -389,9 +374,6 @@ protected:
 
 private:
 
-	// Set of UIDs for valid curves in FilteredCurves
-	TSet<SmartName::UID_Type> CachedFilteredCurvesUIDs;
-
 	// Snapshots of the actor pose from past frames
 	TArray<FInertializationPose> PoseSnapshots;
 
@@ -424,6 +406,9 @@ private:
 	// Inertialization pose differences
 	FInertializationPoseDiff InertializationPoseDiff;
 
+	// Cached curve filter built from FilteredCurves
+	UE::Anim::FCurveFilter CurveFilter;
+	
 	// Reset inertialization timing and state
 	void Deactivate();
 };
diff --git a/Engine/Source/Runtime/Engine/Classes/Animation/AnimSequence.h b/Engine/Source/Runtime/Engine/Classes/Animation/AnimSequence.h
--- a/Engine/Source/Runtime/Engine/Classes/Animation/AnimSequence.h
+++ b/Engine/Source/Runtime/Engine/Classes/Animation/AnimSequence.h
@@ -415,8 +415,15 @@ public:
 	virtual int32 GetNumberOfSampledKeys() const override;
 	virtual FFrameRate GetSamplingFrameRate() const override { return PlatformTargetFrameRate.Default; }
 	virtual void EvaluateCurveData(FBlendedCurve& OutCurve, float CurrentTime, bool bForceUseRawData = false) const override;
-	virtual float EvaluateCurveData(SmartName::UID_Type CurveUID, float CurrentTime, bool bForceUseRawData = false) const override;
-	virtual bool HasCurveData(SmartName::UID_Type CurveUID, bool bForceUseRawData) const override;
+
+	UE_DEPRECATED(5.3, "Please use EvaluateCurveData that takes an FName.")
+	virtual float EvaluateCurveData(SmartName::UID_Type CurveUID, float CurrentTime, bool bForceUseRawData = false) const override { return 0.0f; }
+	virtual float EvaluateCurveData(FName CurveName, float CurrentTime, bool bForceUseRawData = false) const override;
+
+	UE_DEPRECATED(5.3, "Please use HasCurveData that takes an FName.")
+	virtual bool HasCurveData(SmartName::UID_Type CurveUID, bool bForceUseRawData) const override { return false; }
+	virtual bool HasCurveData(FName CurveName, bool bForceUseRawData) const override;
+
 	//~ End UAnimSequenceBase Interface
 
 	// Extract Root Motion transform from the animation
@@ -452,11 +459,18 @@ public:
 	void GetBonePose(struct FAnimationPoseData& OutAnimationPoseData, const FAnimExtractContext& ExtractionContext, bool bForceUseRawData = false) const;
 
 	const TArray<FTrackToSkeletonMap>& GetCompressedTrackToSkeletonMapTable() const { return CompressedData.CompressedTrackToSkeletonMapTable; }
-	const TArray<struct FSmartName>& GetCompressedCurveNames() const { return CompressedData.CompressedCurveNames; }
+
+	UE_DEPRECATED(5.3, "Please use GetCompressedCurveIndexedNames")
+	const TArray<struct FSmartName>& GetCompressedCurveNames() const;
+
+	const TArray<FAnimCompressedCurveIndexedName>& GetCompressedCurveIndexedNames() const { return CompressedData.IndexedCurveNames; }
 
 #if WITH_EDITORONLY_DATA
 protected:
-	void UpdateCompressedCurveName(SmartName::UID_Type CurveUID, const struct FSmartName& NewCurveName);
+	void UpdateCompressedCurveName(const FName& OldCurveName, const FName& NewCurveName);
+	
+	UE_DEPRECATED(5.3, "Please use UpdateCompressedCurveName that takes FNames.")
+	void UpdateCompressedCurveName(SmartName::UID_Type CurveUID, const struct FSmartName& NewCurveName) {}
 private:
 	void UpdateRetargetSourceAsset();
 
diff --git a/Engine/Source/Runtime/Engine/Classes/Animation/AnimSequenceBase.h b/Engine/Source/Runtime/Engine/Classes/Animation/AnimSequenceBase.h
--- a/Engine/Source/Runtime/Engine/Classes/Animation/AnimSequenceBase.h
+++ b/Engine/Source/Runtime/Engine/Classes/Animation/AnimSequenceBase.h
@@ -93,6 +93,11 @@ public:
 	/** Remove all notifies */
 	void RemoveNotifies();
 
+#if WITH_EDITOR
+	/** Renames all named notifies with InOldName to InNewName */
+	void RenameNotifies(FName InOldName, FName InNewName);
+#endif
+	
 	/** 
 	 * Retrieves AnimNotifies given a StartTime and a DeltaTime.
 	 * Time will be advanced and support looping if bAllowLooping is true.
@@ -134,10 +139,18 @@ public:
 
 	/** Evaluate curve data to Instance at the time of CurrentTime **/
 	virtual void EvaluateCurveData(FBlendedCurve& OutCurve, float CurrentTime, bool bForceUseRawData = false) const;
-	virtual float EvaluateCurveData(SmartName::UID_Type CurveUID, float CurrentTime, bool bForceUseRawData = false) const;
 
+	UE_DEPRECATED(5.3, "Please use EvaluateCurveData that takes a FName.")
+	virtual float EvaluateCurveData(SmartName::UID_Type CurveUID, float CurrentTime, bool bForceUseRawData = false) const { return 0.0f; }
+	
+	virtual float EvaluateCurveData(FName CurveName, float CurrentTime, bool bForceUseRawData = false) const;
+	
 	virtual const FRawCurveTracks& GetCurveData() const;
-	virtual bool HasCurveData(SmartName::UID_Type CurveUID, bool bForceUseRawData = false) const;
+	
+	UE_DEPRECATED(5.3, "Please use HasCurveData that takes a FName.")
+	virtual bool HasCurveData(SmartName::UID_Type CurveUID, bool bForceUseRawData = false) const { return false; }
+	
+	virtual bool HasCurveData(FName CurveName, bool bForceUseRawData = false) const;
 
 	/** Return Number of Keys **/
 	UE_DEPRECATED(4.19, "Use GetNumberOfSampledKeys instead")
@@ -181,14 +194,11 @@ public:
 	virtual void RefreshCacheData();
 
 	//~ Begin UAnimationAsset Interface
-#if WITH_EDITOR
-	virtual void RemapTracksToNewSkeleton(USkeleton* NewSkeleton, bool bConvertSpaces) override;
-#endif
 	virtual void TickAssetPlayer(FAnimTickRecord& Instance, struct FAnimNotifyQueue& NotifyQueue, FAnimAssetTickContext& Context) const override;
-
+	//~ End UAnimationAsset Interface
+	
 	void TickByMarkerAsFollower(FMarkerTickRecord &Instance, FMarkerTickContext &MarkerContext, float& CurrentTime, float& OutPreviousTime, const float MoveDelta, const bool bLooping, const UMirrorDataTable* MirrorTable = nullptr) const;
 	void TickByMarkerAsLeader(FMarkerTickRecord& Instance, FMarkerTickContext& MarkerContext, float& CurrentTime, float& OutPreviousTime, const float MoveDelta, const bool bLooping, const UMirrorDataTable* MirrorTable = nullptr) const;
-	//~ End UAnimationAsset Interface
 
 	/**
 	* Get Bone Transform of the Time given, relative to Parent for all RequiredBones
diff --git a/Engine/Source/Runtime/Engine/Classes/Animation/AnimSequenceHelpers.h b/Engine/Source/Runtime/Engine/Classes/Animation/AnimSequenceHelpers.h
--- a/Engine/Source/Runtime/Engine/Classes/Animation/AnimSequenceHelpers.h
+++ b/Engine/Source/Runtime/Engine/Classes/Animation/AnimSequenceHelpers.h
@@ -19,6 +19,7 @@ struct FBoneAnimationTrack;
 struct FRawAnimSequenceTrack;
 struct FAnimationPoseData;
 class IAnimationDataController;
+struct FVirtualBoneCompactPoseData;
 
 enum class EAnimInterpolationType : uint8;
 
diff --git a/Engine/Source/Runtime/Engine/Classes/Animation/AnimationAsset.h b/Engine/Source/Runtime/Engine/Classes/Animation/AnimationAsset.h
--- a/Engine/Source/Runtime/Engine/Classes/Animation/AnimationAsset.h
+++ b/Engine/Source/Runtime/Engine/Classes/Animation/AnimationAsset.h
@@ -33,6 +33,10 @@ class UPoseAsset;
 class UMirrorDataTable;
 class USkeletalMesh;
 struct FAnimationUpdateContext;
+namespace SmartName
+{
+typedef uint16 UID_Type;
+}
 
 namespace UE { namespace Anim {
 class IAnimNotifyEventContextDataInterface;
@@ -218,15 +222,15 @@ struct FBlendFilter
 /*
  * Pose Curve container for extraction
  * This is used by pose anim node
- * Saves UID/PoseIndex/Value of the curve
+ * Saves Name/PoseIndex/Value of the curve
  */
 struct FPoseCurve
 {
+	UE_DEPRECATED(5.3, "UID is no longer used.")
+	static SmartName::UID_Type	UID;
 
-	// alignment is weird, but 4 here
-	// UID of the curve. Used to get curve value
-	// and also can be used to verify 
-	SmartName::UID_Type	UID;
+	// The name of the curve
+	FName				Name;
 	// PoseIndex of pose asset it's dealing with
 	// used to extract pose value fast
 	int32				PoseIndex;
@@ -234,13 +238,20 @@ struct FPoseCurve
 	float				Value;
 
 	FPoseCurve()
-		: UID(SmartName::MaxUID)
+		: Name(NAME_None)
 		, PoseIndex(INDEX_NONE)
 		, Value(0.f)
 	{}
 
+	FPoseCurve(int32 InPoseIndex, FName InName, float InValue)
+		: Name(InName)
+		, PoseIndex(InPoseIndex)
+		, Value(InValue)
+	{}
+
+	UE_DEPRECATED(5.3, "Please use the constructor that takes an FName.")
 	FPoseCurve(int32 InPoseIndex, SmartName::UID_Type	InUID, float  InValue )
-		: UID(InUID)
+		: Name(NAME_None)
 		, PoseIndex(InPoseIndex)
 		, Value(InValue)
 	{}
diff --git a/Engine/Source/Runtime/Engine/Classes/Animation/MirrorDataTable.h b/Engine/Source/Runtime/Engine/Classes/Animation/MirrorDataTable.h
--- a/Engine/Source/Runtime/Engine/Classes/Animation/MirrorDataTable.h
+++ b/Engine/Source/Runtime/Engine/Classes/Animation/MirrorDataTable.h
@@ -108,12 +108,7 @@ public:
 	ENGINE_API virtual void PostLoad() override;
 
 #if WITH_EDITOR
-
-	ENGINE_API virtual void PreEditChange(FProperty* PropertyThatWillChange) override;
-
 	ENGINE_API virtual void PostEditChangeProperty(FPropertyChangedEvent& PropertyChangedEvent) override;
-
-
 #endif // WITH_EDITOR
 
 	/**
@@ -194,13 +189,16 @@ public:
 
 	// Index of the mirror bone for a given bone index in the reference skeleton, or INDEX_NONE if the bone is not mirrored
 	TCustomBoneIndexArray<FSkeletonPoseBoneIndex, FSkeletonPoseBoneIndex> BoneToMirrorBoneIndex;
-
-	// Array with entries the source UIDs of curves that should be mirrored. 
+	
+	UE_DEPRECATED(5.3, "UID-based mirroring has been remove, please use CurveToMirrorCurveMap.")
 	TArray<SmartName::UID_Type> CurveMirrorSourceUIDArray;
-
-	// Array with the target UIDs of curves that should be mirrored. 
+	
+	UE_DEPRECATED(5.3, "UID-based mirroring has been remove, please use CurveToMirrorCurveMap.")
 	TArray<SmartName::UID_Type> CurveMirrorTargetUIDArray;
 
+	// Map from animation curve to mirrored animation curve
+	TMap<FName, FName> CurveToMirrorCurveMap;
+	
 	// Map from animation notify to mirrored animation notify
 	TMap<FName, FName> AnimNotifyToMirrorAnimNotifyMap;
 	
diff --git a/Engine/Source/Runtime/Engine/Classes/Animation/PoseAsset.h b/Engine/Source/Runtime/Engine/Classes/Animation/PoseAsset.h
--- a/Engine/Source/Runtime/Engine/Classes/Animation/PoseAsset.h
+++ b/Engine/Source/Runtime/Engine/Classes/Animation/PoseAsset.h
@@ -17,6 +17,7 @@
 class UAnimSequence;
 class USkeletalMeshComponent;
 struct FLiveLinkCurveElement;
+struct FReferenceSkeleton;
 
 /** 
  * Pose data 
@@ -85,11 +86,25 @@ struct ENGINE_API FPoseDataContainer
 {
 	GENERATED_USTRUCT_BODY()
 
+public:
+	/** For StructOpsTypeTraits */
+	bool Serialize(FArchive& Ar);
+	void PostSerialize(const FArchive& Ar);
+	
 private:
+#if WITH_EDITORONLY_DATA
+	// pose names - horizontal data
+	UPROPERTY()
+	TArray<FSmartName> PoseNames_DEPRECATED;
+#endif
+
 	// pose names - horizontal data
 	UPROPERTY()
-	TArray<FSmartName>						PoseNames;
+	TArray<FName> PoseFNames;
 
+	// Sorted pose name indices
+	TArray<int32> SortedPoseIndices;
+	
 	// this is list of tracks - vertical data
 	UPROPERTY()
 	TArray<FName>							Tracks;
@@ -112,13 +127,13 @@ private:
 
 	void Reset();
 
-	FPoseData* FindPoseData(FSmartName PoseName);
-	FPoseData* FindOrAddPoseData(FSmartName PoseName);
+	FPoseData* FindPoseData(FName PoseName);
+	FPoseData* FindOrAddPoseData(FName PoseName);
 
 	int32 GetNumPoses() const { return Poses.Num();  }
-	bool Contains(FSmartName PoseName) const { return PoseNames.Contains(PoseName); }
+	bool Contains(FName PoseName) const { return PoseFNames.Contains(PoseName); }
 
-	bool IsValid() const { return PoseNames.Num() == Poses.Num() && Tracks.Num() == TrackBoneIndices.Num(); }
+	bool IsValid() const { return PoseFNames.Num() == Poses.Num() && Tracks.Num() == TrackBoneIndices.Num(); }
 	void GetPoseCurve(const FPoseData* PoseData, FBlendedCurve& OutCurve) const;
 	void BlendPoseCurve(const FPoseData* PoseData, FBlendedCurve& OutCurve, float Weight) const;
 
@@ -131,10 +146,10 @@ private:
 	FTransform GetDefaultTransform(int32 SkeletonIndex, const TArray<FTransform>& RefPose) const;
 
 #if WITH_EDITOR
-	void AddOrUpdatePose(const FSmartName& InPoseName, const TArray<FTransform>& InlocalSpacePose, const TArray<float>& InCurveData);
-	void RenamePose(FSmartName OldPoseName, FSmartName NewPoseName);
-	int32 DeletePose(FSmartName PoseName);
-	bool DeleteCurve(FSmartName CurveName);
+	void AddOrUpdatePose(const FName& InPoseName, const TArray<FTransform>& InlocalSpacePose, const TArray<float>& InCurveData);
+	void RenamePose(FName OldPoseName, FName NewPoseName);
+	int32 DeletePose(FName PoseName);
+	bool DeleteCurve(FName CurveName);
 	bool InsertTrack(const FName& InTrackName, USkeleton* InSkeleton, const TArray<FTransform>& RefPose);
 	
 	bool FillUpSkeletonPose(FPoseData* PoseData, const USkeleton* InSkeleton);
@@ -144,9 +159,23 @@ private:
 	void ConvertToFullPose(USkeleton* InSkeleton, const TArray<FTransform>& RefPose);
 	void ConvertToAdditivePose(const TArray<FTransform>& InBasePose, const TArray<float>& InBaseCurve);
 #endif // WITH_EDITOR
+
+	void RebuildCurveIndexTable();
+	
 	friend class UPoseAsset;
 };
 
+template<>
+struct TStructOpsTypeTraits<FPoseDataContainer> : public TStructOpsTypeTraitsBase2<FPoseDataContainer>
+{
+	enum
+	{
+		WithSerializer = true,
+		WithPostSerialize = true,
+	};
+};
+
+
 /**
  * Pose Asset that can be blended by weight of curves 
  */
@@ -232,10 +261,18 @@ public:
 	ENGINE_API int32 GetNumPoses() const;
 	ENGINE_API int32 GetNumCurves() const;
 	ENGINE_API int32 GetNumTracks() const;
+
+	UE_DEPRECATED(5.3, "Please use GetPoseFNames.")
 	ENGINE_API const TArray<FSmartName> GetPoseNames() const;
-	ENGINE_API const TArray<FName>		GetTrackNames() const;
+	
+	ENGINE_API const TArray<FName>& GetPoseFNames() const;
+	ENGINE_API const TArray<FName>& GetTrackNames() const;
+
+	UE_DEPRECATED(5.3, "Please use GetCurveFNames.")
 	ENGINE_API const TArray<FSmartName> GetCurveNames() const;
-	ENGINE_API const TArray<FAnimCurveBase> GetCurveData() const;
+	
+	ENGINE_API const TArray<FName> GetCurveFNames() const;
+	ENGINE_API const TArray<FAnimCurveBase>& GetCurveData() const;
 	ENGINE_API const TArray<float> GetCurveValues(const int32 PoseIndex) const;
 
 	/** Find index of a track with a given bone name. Returns INDEX_NONE if not found. */
@@ -246,8 +283,9 @@ public:
 	 *	@return	Returns true if OutValue is valid, false if not
 	 */
 	ENGINE_API bool GetCurveValue(const int32 PoseIndex, const int32 CurveIndex, float& OutValue) const;
-	 
-	ENGINE_API bool ContainsPose(const FSmartName& InPoseName) const { return PoseContainer.Contains(InPoseName); }
+
+	UE_DEPRECATED(5.3, "Please use ContainsPose that takes a FName.")
+	ENGINE_API bool ContainsPose(const FSmartName& InPoseName) const { return PoseContainer.Contains(InPoseName.DisplayName); }
 	ENGINE_API bool ContainsPose(const FName& InPoseName) const;
 
 #if WITH_EDITOR
@@ -258,11 +296,26 @@ public:
 	/** Returns the name of all contained poses */
 	UFUNCTION(BlueprintPure, Category=PoseAsset)
 	void GetPoseNames(TArray<FName>& PoseNames) const;
+
+	UE_DEPRECATED(5.3, "Please use AddPoseWithUniqueName.")
+	ENGINE_API bool AddOrUpdatePoseWithUniqueName(const USkeletalMeshComponent* MeshComponent, FSmartName* OutPoseName = nullptr) { return false; }
+	
+	ENGINE_API FName AddPoseWithUniqueName(const USkeletalMeshComponent* MeshComponent);
 	
-	ENGINE_API bool AddOrUpdatePoseWithUniqueName(const USkeletalMeshComponent* MeshComponent, FSmartName* OutPoseName = nullptr);
-	ENGINE_API void AddOrUpdatePose(const FSmartName& PoseName, const USkeletalMeshComponent* MeshComponent, bool bUpdateCurves = true);
+	UE_DEPRECATED(5.3, "Please use AddOrUpdatePose that takes a FName.")
+	void AddOrUpdatePose(const FSmartName& PoseName, const USkeletalMeshComponent* MeshComponent, bool bUpdateCurves = true) { AddOrUpdatePose(PoseName.DisplayName, MeshComponent, bUpdateCurves); }
+
+	ENGINE_API void AddOrUpdatePose(const FName& PoseName, const USkeletalMeshComponent* MeshComponent, bool bUpdateCurves = true);
+
+	UE_DEPRECATED(5.3, "Please use AddReferencePose that takes a FName.")
 	ENGINE_API void AddReferencePose(const FSmartName& PoseName, const FReferenceSkeleton& ReferenceSkeleton);
-	ENGINE_API void CreatePoseFromAnimation(class UAnimSequence* AnimSequence, const TArray<FSmartName>* InPoseNames = nullptr);
+
+	ENGINE_API void AddReferencePose(const FName& PoseName, const FReferenceSkeleton& ReferenceSkeleton);
+
+	UE_DEPRECATED(5.3, "Please use CreatePoseFromAnimation that takes a ptr to an array of FNames.")
+	ENGINE_API void CreatePoseFromAnimation(class UAnimSequence* AnimSequence, const TArray<FSmartName>* InPoseNames) {}
+	
+	ENGINE_API void CreatePoseFromAnimation(class UAnimSequence* AnimSequence, const TArray<FName>* InPoseNames = nullptr);
 
 	/** Contained poses are re-generated from the provided Animation Sequence*/
 	UFUNCTION(BlueprintCallable, Category=PoseAsset)
@@ -273,14 +326,24 @@ public:
 	virtual void ReplaceReferredAnimations(const TMap<UAnimationAsset*, UAnimationAsset*>& ReplacementMap) override;
 	// End AnimationAsset interface
 
-	ENGINE_API bool ModifyPoseName(FName OldPoseName, FName NewPoseName, const SmartName::UID_Type* NewUID);
+	UE_DEPRECATED(5.3, "Please use ModifyPoseName that does not take a UID.")
+	ENGINE_API bool ModifyPoseName(FName OldPoseName, FName NewPoseName, const SmartName::UID_Type* NewUID) { return ModifyPoseName(OldPoseName, NewPoseName); }
+
+	ENGINE_API bool ModifyPoseName(FName OldPoseName, FName NewPoseName);
 
-	// Rename the smart names used by this Pose Asset
+	
+	UE_DEPRECATED(5.3, "Please use RenamePoseOrCurveName.")
 	ENGINE_API void RenameSmartName(const FName& InOriginalName, const FName& InNewName);
 
-	// Remove poses or curves using the smart names supplied
+	// Rename poses or curves using the names supplied
+	ENGINE_API void RenamePoseOrCurveName(const FName& InOriginalName, const FName& InNewName);
+
+	UE_DEPRECATED(5.3, "Please use RemovePoseOrCurveNames.")
 	ENGINE_API void RemoveSmartNames(const TArray<FName>& InNamesToRemove);
 
+	// Remove poses or curves using the names supplied
+	ENGINE_API void RemovePoseOrCurveNames(const TArray<FName>& InNamesToRemove);
+
 	// editor interface
 	virtual void PostEditChangeProperty(FPropertyChangedEvent& PropertyChangedEvent) override;
 	// Return full (local space, non additive) pose. Will do conversion if PoseAsset is Additive. 
@@ -292,7 +355,7 @@ public:
 	ENGINE_API int32 DeletePoses(TArray<FName> PoseNamesToDelete);
 	ENGINE_API int32 DeleteCurves(TArray<FName> CurveNamesToDelete);
 	ENGINE_API bool ConvertSpace(bool bNewAdditivePose, int32 NewBasePoseInde);
-	ENGINE_API const FName GetPoseNameByIndex(int32 InBasePoseIndex) const { return PoseContainer.PoseNames.IsValidIndex(InBasePoseIndex) ? PoseContainer.PoseNames[InBasePoseIndex].DisplayName : NAME_None; }
+	ENGINE_API const FName GetPoseNameByIndex(int32 InBasePoseIndex) const { return PoseContainer.PoseFNames.IsValidIndex(InBasePoseIndex) ? PoseContainer.PoseFNames[InBasePoseIndex] : NAME_None; }
 #endif // WITH_EDITOR
 
 	ENGINE_API int32 GetBasePoseIndex() const { return BasePoseIndex;  }
@@ -318,8 +381,12 @@ public:
 		OnPoseListChanged.Remove(Handle);
 	}
 
+	UE_DEPRECATED(5.3, "Please use GetUniquePoseName scoped to this pose asset.")
 	ENGINE_API static FName GetUniquePoseName(const USkeleton* Skeleton);
+	UE_DEPRECATED(5.3, "Please use GetUniquePoseName.")
 	ENGINE_API static FSmartName GetUniquePoseSmartName(USkeleton* Skeleton);
+	
+	ENGINE_API static FName GetUniquePoseName(UPoseAsset* PoseAsset);
 
 protected:
 	virtual void RemapTracksToNewSkeleton(USkeleton* NewSkeleton, bool bConvertSpaces) override;
@@ -334,7 +401,7 @@ private:
 	void Reinitialize();
 
 	// After any update to SourceLocalPoses, this does update runtime data
-	void AddOrUpdatePose(const FSmartName& PoseName, const TArray<FName>& TrackNames, const TArray<FTransform>& LocalTransform, const TArray<float>& CurveValues);
+	void AddOrUpdatePose(const FName& PoseName, const TArray<FName>& TrackNames, const TArray<FTransform>& LocalTransform, const TArray<float>& CurveValues);
 	void PostProcessData();
 	void BreakAnimationSequenceGUIDComparison();
 #endif // WITH_EDITOR	
diff --git a/Engine/Source/Runtime/Engine/Classes/Animation/Skeleton.h b/Engine/Source/Runtime/Engine/Classes/Animation/Skeleton.h
--- a/Engine/Source/Runtime/Engine/Classes/Animation/Skeleton.h
+++ b/Engine/Source/Runtime/Engine/Classes/Animation/Skeleton.h
@@ -36,6 +36,7 @@ enum class EPackageReloadPhase : uint8;
 class USkeleton;
 typedef SmartName::UID_Type SkeletonAnimCurveUID;
 class USkeleton;
+struct FSkeletonRemapping;
 
 // Delegate used to control global skeleton compatibility
 DECLARE_DELEGATE_RetVal(bool, FAreAllSkeletonsCompatible);
@@ -370,23 +371,75 @@ public:
 	/** Serializable retarget sources for this skeleton **/
 	TMap< FName, FReferencePose > AnimRetargetSources;
 
-	// Typedefs for greater smartname UID readability, add one for each smartname category 
-	typedef SkeletonAnimCurveUID AnimCurveUID;
+	// DEPRECATED - no longer used
+	typedef SmartName::UID_Type AnimCurveUID;
 
-	// Names for smartname mappings, if you're adding a new category of smartnames add a new name here
+	UE_DEPRECATED(5.3, "AnimCurveMappingName is no longer used.")
 	static ENGINE_API const FName AnimCurveMappingName;
 
-	// Names for smartname mappings, if you're adding a new category of smartnames add a new name here
+	UE_DEPRECATED(5.3, "AnimTrackCurveMappingName is no longer used.")
 	static ENGINE_API const FName AnimTrackCurveMappingName;
 
-	// these return container of curve meta data, if you modify this container, 
-	// you'll have to call REfreshCAchedAnimationCurveData to apply
-	ENGINE_API FCurveMetaData* GetCurveMetaData(const FName& CurveName);
-	ENGINE_API const FCurveMetaData* GetCurveMetaData(const FName& CurveName) const;
+	ENGINE_API FCurveMetaData* GetCurveMetaData(FName CurveName);
+	ENGINE_API const FCurveMetaData* GetCurveMetaData(FName CurveName) const;
+
+	UE_DEPRECATED(5.3, "Please use GetCurveMetaData with an FName.")
 	ENGINE_API const FCurveMetaData* GetCurveMetaData(const SmartName::UID_Type CurveUID) const;
+
+	UE_DEPRECATED(5.3, "Please use GetCurveMetaData with an FName.")
 	ENGINE_API FCurveMetaData* GetCurveMetaData(const FSmartName& CurveName);
+	
+	UE_DEPRECATED(5.3, "Please use GetCurveMetaData with an FName.")
 	ENGINE_API const FCurveMetaData* GetCurveMetaData(const FSmartName& CurveName) const;
-	// this is called when you know both flags - called by post serialize
+
+	/**
+	 * Iterate over all curve metadata entries, calling InFunction on each
+	 * @param	InFunction	The function to call
+	 **/
+	ENGINE_API void ForEachCurveMetaData(TFunctionRef<void(FName, const FCurveMetaData&)> InFunction) const;
+	
+	/** @return the number of curve metadata entries **/
+	ENGINE_API int32 GetNumCurveMetaData() const;
+
+	/**
+	 * Adds a curve metadata entry with the specified name
+	 * @param	InCurveName			The name of the curve to find
+	 * @return true if an entry was added, false if an entry already existed
+	 */
+	ENGINE_API bool AddCurveMetaData(FName CurveName);
+
+	/**
+	 * Get an array of all curve metadata names
+	 * @param	OutNames		The array to receive the metadata names 
+	 */
+	ENGINE_API void GetCurveMetaDataNames(TArray<FName>& OutNames) const;
+
+#if WITH_EDITOR
+	/**
+	 * Renames a curve metadata entry. Metadata is preserved, but assigned to a different curve name.
+	 * @param OldName	The name of an existing curve entry
+	 * @param NewName	The name to change the entry to
+	 * @return			true if the rename was successful (the old name was found and the new name didnt collide with an
+	 *					existing entry)
+	 */	
+	ENGINE_API bool RenameCurveMetaData(FName OldName, FName NewName);
+
+	/**
+	 * Removes a curve metadata entry for the specified name.
+	 * @param CurveName	The name of the curve to remove the metadata for
+	 * @return true if the entry was successfully removed (i.e. it existed)
+	 */
+	ENGINE_API bool RemoveCurveMetaData(FName CurveName);
+
+	/**
+	 * Removes a group of curve metadata entries for the specified names.
+	 * @param CurveNames	The names of the curves to remove the metadata for
+	 * @return true if any of the entries were successfully removed (i.e. something changed)
+	 */
+	ENGINE_API bool RemoveCurveMetaData(TArrayView<FName> CurveNames);
+#endif
+
+	// this is called when you know both flags - called by post serialize and import
 	ENGINE_API void AccumulateCurveMetaData(FName CurveName, bool bMaterialSet, bool bMorphtargetSet);
 
 	ENGINE_API bool AddNewVirtualBone(const FName SourceBoneName, const FName TargetBoneName);
@@ -400,8 +453,10 @@ public:
 	void HandleVirtualBoneChanges();
 
 	// return version of AnimCurveUidVersion
-	uint16 GetAnimCurveUidVersion() const { return AnimCurveUidVersion;  }
-	const TArray<uint16>& GetDefaultCurveUIDList() const { return DefaultCurveUIDList; }
+	uint16 GetAnimCurveUidVersion() const;
+
+	UE_DEPRECATED(5.3, "This function is no longer used")
+	const TArray<uint16>& GetDefaultCurveUIDList() const;
 
 	ENGINE_API const TArray<TSoftObjectPtr<USkeleton>>& GetCompatibleSkeletons() const { return CompatibleSkeletons; }
 
@@ -420,26 +475,23 @@ public:
 #endif
 
 protected:
-	// Container for smart name mappings
+	// DEPRECATED - moved to CurveMetaData
 	UPROPERTY()
-	FSmartNameContainer SmartNames;
+	FSmartNameContainer SmartNames_DEPRECATED;
 
 	// Cached ptr to the persistent AnimCurveMapping
-	FSmartNameMapping* AnimCurveMapping;
+	UE_DEPRECATED(5.3, "AnimCurveMapping is no longer used")
+	static FSmartNameMapping* AnimCurveMapping;
 
-	// this is default curve uid list used like ref pose, as default value
-	// don't use this unless you want all curves from the skeleton
-	// FBoneContainer contains only list that is used by current LOD
-	TArray<uint16> DefaultCurveUIDList;
+	UE_DEPRECATED(5.3, "DefaultCurveUIDList is no longer used")
+	static TArray<uint16> DefaultCurveUIDList;
 
 	//Cached marker sync marker names (stripped for non editor)
 	TArray<FName> ExistingMarkerNames;
 
 private:
-	/** Increase the AnimCurveUidVersion so that instances can get the latest information */
-	void IncreaseAnimCurveUidVersion();
-	/** Current  Anim Curve Uid Version. Increase whenever it has to be recalculated */
-	uint16 AnimCurveUidVersion;
+	// Refresh skeleton metadata (updates bone indices for linked bone references)
+	void RefreshSkeletonMetaData();
 
 public:
 	//////////////////////////////////////////////////////////////////////////
@@ -492,41 +544,38 @@ public:
 	ENGINE_API void RemoveSlotGroup(const FName& InSlotName);
 	ENGINE_API void RenameSlotName(const FName& OldName, const FName& NewName);
 
-	////////////////////////////////////////////////////////////////////////////
-	// Smart Name Interfaces
-	////////////////////////////////////////////////////////////////////////////
-	// Adds a new name to the smart name container and modifies the skeleton so it can be saved
-	// return bool - Whether a name was added (false if already present)
 #if WITH_EDITOR
-	ENGINE_API bool AddSmartNameAndModify(FName ContainerName, FName NewDisplayName, FSmartName& NewName);
-
-	// Renames a smartname in the specified container and modifies the skeleton
-	// return bool - Whether the rename was sucessful
-	ENGINE_API bool RenameSmartnameAndModify(FName ContainerName, SmartName::UID_Type Uid, FName NewName);
+	UE_DEPRECATED(5.3, "Please use AddCurveMetaData.")
+	ENGINE_API bool AddSmartNameAndModify(FName ContainerName, FName NewDisplayName, FSmartName& NewName) { return false; }
 
-	// Removes a smartname from the specified container and modifies the skeleton
-	ENGINE_API void RemoveSmartnameAndModify(FName ContainerName, SmartName::UID_Type Uid);
+	UE_DEPRECATED(5.3, "Please use RenameCurveMetaData.")
+	ENGINE_API bool RenameSmartnameAndModify(FName ContainerName, SmartName::UID_Type Uid, FName NewName) { return false; }
+	
+	UE_DEPRECATED(5.3, "Please use RemoveCurveMetaData.")
+	ENGINE_API void RemoveSmartnameAndModify(FName ContainerName, SmartName::UID_Type Uid) {}
 
-	// Removes smartnames from the specified container and modifies the skeleton
-	ENGINE_API void RemoveSmartnamesAndModify(FName ContainerName, const TArray<FName>& Names);
+	UE_DEPRECATED(5.3, "Please use RemoveCurveMetaData.")
+	ENGINE_API void RemoveSmartnamesAndModify(FName ContainerName, const TArray<FName>& Names) {}
 #endif// WITH_EDITOR
 
-	// quick wrapper function for Find UID by name, if not found, it will return SmartName::MaxUID
-	ENGINE_API SmartName::UID_Type GetUIDByName(const FName& ContainerName, const FName& Name) const;
-	ENGINE_API bool GetSmartNameByUID(const FName& ContainerName, SmartName::UID_Type UID, FSmartName& OutSmartName) const;
-	ENGINE_API bool GetSmartNameByName(const FName& ContainerName, const FName& InName, FSmartName& OutSmartName) const;
+	UE_DEPRECATED(5.3, "This function is no longer used")
+	ENGINE_API SmartName::UID_Type GetUIDByName(const FName& ContainerName, const FName& Name) const { return 0; }
 
-	// Get or add a smartname container with the given name
-	ENGINE_API const FSmartNameMapping* GetSmartNameContainer(const FName& ContainerName) const;
+	UE_DEPRECATED(5.3, "This function is no longer used")
+	ENGINE_API bool GetSmartNameByUID(const FName& ContainerName, SmartName::UID_Type UID, FSmartName& OutSmartName) const { return false; }
+	
+	UE_DEPRECATED(5.3, "This function is no longer used")
+	ENGINE_API bool GetSmartNameByName(const FName& ContainerName, const FName& InName, FSmartName& OutSmartName) const { return false; }
+
+	UE_DEPRECATED(5.3, "This function is no longer used")
+	ENGINE_API const FSmartNameMapping* GetSmartNameContainer(const FName& ContainerName) const { return nullptr; }
+
+	UE_DEPRECATED(5.3, "This function is no longer used")
+	ENGINE_API void VerifySmartName(const FName&  ContainerName, FSmartName& InOutSmartName) {}
+	
+	UE_DEPRECATED(5.3, "This function is no longer used")
+	ENGINE_API void VerifySmartNames(const FName&  ContainerName, TArray<FSmartName>& InOutSmartNames) {}
 
-	// make sure the smart name has valid UID and so on
-	ENGINE_API void VerifySmartName(const FName&  ContainerName, FSmartName& InOutSmartName);
-	ENGINE_API void VerifySmartNames(const FName&  ContainerName, TArray<FSmartName>& InOutSmartNames);
-private:
-	// Get or add a smartname container with the given name
-	FSmartNameMapping* GetOrAddSmartNameContainer(const FName& ContainerName);
-	bool VerifySmartNameInternal(const FName&  ContainerName, FSmartName& InOutSmartName);
-	bool FillSmartNameByDisplayName(FSmartNameMapping* Mapping, const FName& DisplayName, FSmartName& OutSmartName);
 #if WITH_EDITORONLY_DATA
 private:
 	/** The default skeletal mesh to use when previewing this skeleton */
@@ -630,6 +679,12 @@ public:
 	// Adds a new anim notify to the cached AnimationNotifies array.
 	ENGINE_API void AddNewAnimationNotify(FName NewAnimNotifyName);
 
+	// Removes an anim notify from the cached AnimationNotifies array.
+	ENGINE_API void RemoveAnimationNotify(FName AnimNotifyName);
+
+	// Renames an anim notify
+	ENGINE_API void RenameAnimationNotify(FName OldAnimNotifyName, FName NewAnimNotifyName);
+
 	ENGINE_API USkeletalMesh* GetAssetPreviewMesh(UObject* InAsset);
 
 	/** Find the first compatible mesh for this skeleton */
@@ -694,8 +749,10 @@ public:
 
 	UE_DEPRECATED(5.2, "Compatibility is now an editor-only concern. Please use IsCompatibleForEditor.")
 	ENGINE_API bool IsCompatibleSkeletonByAssetString(const FString& SkeletonAssetString) const;
-	
+
 	DECLARE_EVENT(USkeleton, FSmartNamesChangedEvent);
+
+	UE_DEPRECATED(5.3, "This member is no longer used. Delegate registration for skeleton metadata can be handled via UAnimCurveMetaData.")
 	FSmartNamesChangedEvent OnSmartNamesChangedEvent;
 
 #if WITH_EDITORONLY_DATA
@@ -705,18 +762,12 @@ public:
 	 */
 	ENGINE_API static FAreAllSkeletonsCompatible AreAllSkeletonsCompatibleDelegate;
 #endif
+
 public:
 	UFUNCTION(BlueprintCallable, Category=Skeleton)
 	ENGINE_API void AddCompatibleSkeleton(const USkeleton* SourceSkeleton);
 	ENGINE_API void RemoveCompatibleSkeleton(const USkeleton* SourceSkeleton);
 
-	/**
-	 * Handle smart name changes.
-	 * This will internally trigger skeleton remappings to be regenerated where needed, for example when editing curves in the animation editor.
-	 **/
-	void HandleSmartNamesChangedEvent();
-
-public:
 	/** 
 	 * Indexing naming convention
 	 * 
diff --git a/Engine/Source/Runtime/Engine/Classes/Animation/SmartName.h b/Engine/Source/Runtime/Engine/Classes/Animation/SmartName.h
--- a/Engine/Source/Runtime/Engine/Classes/Animation/SmartName.h
+++ b/Engine/Source/Runtime/Engine/Classes/Animation/SmartName.h
@@ -6,157 +6,95 @@
 #include "UObject/ObjectMacros.h"
 #include "Misc/Guid.h"
 #include "UObject/Class.h"
-#include "BoneContainer.h"
+#include "Animation/BoneReference.h"
+#include "Animation/AnimTypes.h"
 #include "UObject/AnimPhysObjectVersion.h"
+#include "Animation/AnimCurveMetadata.h"
 #include "SmartName.generated.h"
 
 struct FSmartName;
 
-
-/** Curve Meta Data for each name
- * Unfortunately this should be linked to FName, but no GUID because we don't have GUID in run-time
- * We only add this if anything changed, by default, it is attribute curve
- */
-USTRUCT()
-struct FCurveMetaData
-{
-	GENERATED_USTRUCT_BODY()
-
-	/** connected bones to this meta data */
-	UPROPERTY()
-	TArray<FBoneReference> LinkedBones;
-	
-	/* max LOD (lowest LOD) to evaluate this. -1 means it will evaluate all the time. */
-	UPROPERTY()
-	uint8 MaxLOD;
-	
-	UPROPERTY()
-	FAnimCurveType Type;
-
-	friend FArchive& operator<<(FArchive& Ar, FCurveMetaData& B)
-	{
-		Ar.UsingCustomVersion(FAnimPhysObjectVersion::GUID);
-
-		Ar << B.Type.bMaterial;
-		Ar << B.Type.bMorphtarget;
-		Ar << B.LinkedBones;
-
-		if (Ar.CustomVer(FAnimPhysObjectVersion::GUID) >= FAnimPhysObjectVersion::AddLODToCurveMetaData)
-		{
-			Ar << B.MaxLOD;
-		}
-
-		return Ar;
-	}
-
-	bool Serialize(FArchive& Ar)
-	{
-		Ar << *this;
-		return true;
-	}
-
-	FCurveMetaData()
-		: MaxLOD(0xFF)
-	{
-	}
-};
-
+// DEPRECATED - smart names and their mappings are no longer used
 USTRUCT()
 struct ENGINE_API FSmartNameMapping
 {
 	friend struct FSmartNameMappingIterator;
+	friend class USkeleton;
 	
 	GENERATED_USTRUCT_BODY();
 
 	FSmartNameMapping();
 	
-	// Add a name to the mapping, fails if it already exists
-	// @param InName - The name to add
-	// @return FSmartName - populated smart name
+	UE_DEPRECATED(5.3, "FSmartNameMapping functions are no longer used.")
 	FSmartName AddName(FName InName);
 
-	// Get a name from the mapping
-	// @param Uid - SmartName::UID_Type of the name to retrieve
-	// @param OUT OutName - Retrieved name
-	// @return bool - true if name existed and OutName is valid
+	UE_DEPRECATED(5.3, "FSmartNameMapping functions are no longer used.")
 	bool GetName(const SmartName::UID_Type& Uid, FName& OutName) const;
 	
-	// Fill an array with all used UIDs
-	// @param Array - Array to fill
+	UE_DEPRECATED(5.3, "FSmartNameMapping functions are no longer used.")
 	void FillUidArray(TArray<SmartName::UID_Type>& Array) const;
 
-	// Fill an array with all used names
-	// @param Array - Array to fill
+	UE_DEPRECATED(5.3, "FSmartNameMapping functions are no longer used.")
 	void FillNameArray(TArray<FName>& Array) const;
 
-	// Fill an array with curve names in UID order (None will be placed in invalid UID slots)
+	UE_DEPRECATED(5.3, "FSmartNameMapping functions are no longer used.")
 	void FillUIDToNameArray(TArray<FName>& Array) const;
 
-	// Fill an array with curve types for all used names
-	// @param Array - Array to fill
+	UE_DEPRECATED(5.3, "FSmartNameMapping functions are no longer used.")
 	void FillCurveTypeArray(TArray<FAnimCurveType>& Array) const;
 
-
-	// Fill an array with curve types in UID order
+	UE_DEPRECATED(5.3, "FSmartNameMapping functions are no longer used.")
 	void FillUIDToCurveTypeArray(TArray<FAnimCurveType>& Array) const;
 
 #if WITH_EDITOR
-	// Change a name
-	// @param Uid - SmartName::UID_Type of the name to change
-	// @param NewName - New name to set 
-	// @return bool - true if the name was found and changed, false if the name wasn't present in the mapping
+	UE_DEPRECATED(5.3, "FSmartNameMapping functions are no longer used.")
 	bool Rename(const SmartName::UID_Type& Uid, FName NewName);
 
-	// Remove a name from the mapping
-	// @param Uid - SmartName::UID_Type of the name to remove
-	// @return bool - true if the name was found and removed, false if the name wasn't present in the mapping
+	UE_DEPRECATED(5.3, "FSmartNameMapping functions are no longer used.")
 	bool Remove(const SmartName::UID_Type& Uid);
 
-	// Remove a name from the mapping
-	// @param Uid - SmartName::UID_Type of the name to remove
-	// @return bool - true if the name was found and removed, false if the name wasn't present in the mapping
+	UE_DEPRECATED(5.3, "FSmartNameMapping functions are no longer used.")
 	bool Remove(const FName& Name);
 #endif
 
-	// Return SmartName::UID_Type if it finds it
-	// @param Name - Name of curve to find UID for
-	// @return SmartName::UID_Type - MaxUID if it doesn't find, actual UID if it finds. 
+	UE_DEPRECATED(5.3, "FSmartNameMapping functions are no longer used.")
 	SmartName::UID_Type FindUID(const FName& Name) const;
 
-	// Check whether a name already exists in the mapping
-	// @param Uid - the SmartName::UID_Type to check
-	// @return bool - whether the name was found
+	UE_DEPRECATED(5.3, "FSmartNameMapping functions are no longer used.")
 	bool Exists(const SmartName::UID_Type& Uid) const;
 
-	// Check whether a name already exists in the mapping
-	// @param Name - the name to check
-	// @return bool - whether the name was found
+	UE_DEPRECATED(5.3, "FSmartNameMapping functions are no longer used.")
 	bool Exists(const FName& Name) const;
 
-	// Find Smart Names
+	UE_DEPRECATED(5.3, "FSmartNameMapping functions are no longer used.")
 	bool FindSmartName(FName Name, FSmartName& OutName) const;
+
+	UE_DEPRECATED(5.3, "FSmartNameMapping functions are no longer used.")
 	bool FindSmartNameByUID(SmartName::UID_Type UID, FSmartName& OutName) const;
 
-	// Curve Meta Data Accessors
+	UE_DEPRECATED(5.3, "FSmartNameMapping functions are no longer used.")
 	FCurveMetaData* GetCurveMetaData(FName CurveName);
+	
+	UE_DEPRECATED(5.3, "FSmartNameMapping functions are no longer used.")
 	const FCurveMetaData* GetCurveMetaData(FName CurveName) const;
 	
 #if !WITH_EDITOR
+	UE_DEPRECATED(5.3, "FSmartNameMapping functions are no longer used.")
 	const FCurveMetaData& GetCurveMetaData(SmartName::UID_Type CurveUID) const;
 #endif
-
-	// Serialize this to the provided archive; required for TMap serialization
+	
+	UE_DEPRECATED(5.3, "FSmartNameMapping functions are no longer used.")
 	void Serialize(FArchive& Ar);
-
+	
 	friend FArchive& operator<<(FArchive& Ar, FSmartNameMapping& Elem);
-
-	/* initialize curve meta data for the container */
+	
+	UE_DEPRECATED(5.3, "FSmartNameMapping functions are no longer used.")
 	void InitializeCurveMetaData(class USkeleton* Skeleton);
-
-	/** Get the maximum in use UID */
+	
+	UE_DEPRECATED(5.3, "FSmartNameMapping functions are no longer used.")
 	SmartName::UID_Type GetMaxUID() const { return (SmartName::UID_Type)(CurveNameList.Num() - 1); }
-
-	/** Iterate over all Names in this Mapping */
+	
+	UE_DEPRECATED(5.3, "FSmartNameMapping functions are no longer used.")
 	void Iterate(TFunction<void(const struct FSmartNameMappingIterator& Iterator)> Callback) const;
 
 private:
@@ -191,12 +129,14 @@ struct ENGINE_API FSmartNameMappingIterator
 {
 	public:
 		friend struct FSmartNameMapping;
-	
+
+		UE_DEPRECATED(5.3, "FSmartNameMappingIterator functions are no longer used.")
 		bool GetName(FName& OutCurveName) const
 		{
 			return Mapping->GetName_NoLock(Index, OutCurveName);
 		}
-	
+
+		UE_DEPRECATED(5.3, "FSmartNameMappingIterator functions are no longer used.")
 		const FCurveMetaData* GetCurveMetaData() const
 		{
 			FName Name;
@@ -209,7 +149,8 @@ struct ENGINE_API FSmartNameMappingIterator
 				return nullptr;
 			}
 		}
-	
+
+		UE_DEPRECATED(5.3, "FSmartNameMappingIterator functions are no longer used.")
 		SmartName::UID_Type GetIndex() const { return Index; }
 	
 	private:
@@ -227,16 +168,16 @@ struct ENGINE_API FSmartNameContainer
 {
 	GENERATED_USTRUCT_BODY();
 
-	// Add a new smartname container with the provided name
+	UE_DEPRECATED(5.3, "FSmartNameContainer functions are no longer used.")
 	FSmartNameMapping* AddContainer(FName NewContainerName);
 
-	// Get a container by name	
+	UE_DEPRECATED(5.3, "FSmartNameContainer functions are no longer used.")
 	const FSmartNameMapping* GetContainer(FName ContainerName) const;
 
-	// Serialize this to the provided archive; required for TMap serialization
+	UE_DEPRECATED(5.3, "FSmartNameContainer functions are no longer used.")
 	void Serialize(FArchive& Ar, bool bIsTemplate);
 
-	// Called after load (serialize itself may not be called if the USkeleton we are on is old enough)
+	UE_DEPRECATED(5.3, "FSmartNameContainer functions are no longer used.")
 	void PostLoad();
 
 	friend FArchive& operator<<(FArchive& Ar, FSmartNameContainer& Elem);
@@ -282,26 +223,22 @@ struct ENGINE_API FSmartName
 		, UID(SmartName::MaxUID)
 	{}
 
+	UE_DEPRECATED(5.3, "FSmartName functions are no longer used.")
 	FSmartName(const FName& InDisplayName, SmartName::UID_Type InUID)
 		: DisplayName(InDisplayName)
 		, UID(InUID)
 	{}
-
+	
 	bool operator==(FSmartName const& Other) const
 	{
 		return (DisplayName == Other.DisplayName && UID == Other.UID);
 	}
+
 	bool operator!=(const FSmartName& Other) const
 	{
 		return !(*this == Other);
 	}
-	/**
-	* Serialize the SmartName
-	*
-	* @param Ar	Archive to serialize to
-	*
-	* @return True if the container was serialized
-	*/
+
 	bool Serialize(FArchive& Ar);
 
 	friend FArchive& operator<<(FArchive& Ar, FSmartName& P)
@@ -310,6 +247,7 @@ struct ENGINE_API FSmartName
 		return Ar;
 	}
 
+	UE_DEPRECATED(5.3, "FSmartName functions are no longer used.")
 	bool IsValid() const
 	{
 		return UID != SmartName::MaxUID;
diff --git a/Engine/Source/Runtime/Engine/Classes/Components/SkeletalMeshComponent.h b/Engine/Source/Runtime/Engine/Classes/Components/SkeletalMeshComponent.h
--- a/Engine/Source/Runtime/Engine/Classes/Components/SkeletalMeshComponent.h
+++ b/Engine/Source/Runtime/Engine/Classes/Components/SkeletalMeshComponent.h
@@ -814,10 +814,18 @@ public:
 
 private:
 
+	/** Whether the FilteredAnimCurves list is an allow or deny list */
+	UPROPERTY(transient)
+	bool bFilteredAnimCurvesIsAllowList = false;
+
+	/** Cache curve metadata from mesh and this will be used to identify if it needs to be updated */
+	UPROPERTY(transient)
+	uint16 CachedMeshCurveMetaDataVersion;
+	
 	/** You can choose to disable certain curves if you prefer. 
 	 * This is transient curves that will be ignored by animation system if you choose this */
 	UPROPERTY(transient)
-	TArray<FName> DisallowedAnimCurves;
+	TArray<FName> FilteredAnimCurves;
 
 public:
 
@@ -1284,8 +1292,16 @@ public:
 	UFUNCTION(BlueprintCallable, Category = "Components|SkeletalMesh")
 	void SetAllowedAnimCurvesEvaluation(const TArray<FName>& List, bool bAllow);
 
-	const TArray<FName>& GetDisallowedAnimCurvesEvaluation() const { return DisallowedAnimCurves; }
+	UE_DEPRECATED(5.3, "Please use GetCurveFilterSettings")
+	const TArray<FName>& GetDisallowedAnimCurvesEvaluation() const { return FilteredAnimCurves; }
 
+	/**
+	 * Get the curve settings that will be used for anim evaluation.
+	 * @param InLODOverride		Override the LOD that curves will be calculated for if not INDEX_NONE
+	 * @return the curve settings used for evaluation
+	 */
+	UE::Anim::FCurveFilterSettings GetCurveFilterSettings(int32 InLODOverride = INDEX_NONE) const;
+	
 	/** We detach the Component once we are done playing it.
 	 *
 	 * @param	ParticleSystemComponent that finished
@@ -1525,11 +1541,6 @@ private:
 	 **/
 	TMap<FName, float>	MorphTargetCurves;
 
-	/** 
-	 * Temporary storage for Curve UIDList of evaluating Animation 
-	 */
-	TArray<uint16> CachedCurveUIDList;
-
 public:
 	const TMap<FName, float>& GetMorphTargetCurves() const { return MorphTargetCurves;  }
 	// 
diff --git a/Engine/Source/Runtime/Engine/Classes/Engine/PoseWatchRenderData.h b/Engine/Source/Runtime/Engine/Classes/Engine/PoseWatchRenderData.h
--- a/Engine/Source/Runtime/Engine/Classes/Engine/PoseWatchRenderData.h
+++ b/Engine/Source/Runtime/Engine/Classes/Engine/PoseWatchRenderData.h
@@ -9,6 +9,7 @@
 #include "UObject/ObjectMacros.h"
 #include "UObject/Object.h"
 #include "BoneIndices.h"
+#include "Animation/AnimCurveTypes.h"
 
 struct FReferenceSkeleton;
 
@@ -25,6 +26,8 @@ public:
 
 	bool IsValid() const;
 
+	void SetCurves(const FBlendedCurve& InCurves);
+
 	template<class TAllocator>
 	bool SetPose(const TArray<FBoneIndexType>& InRequiredBones, const TArray<FTransform, TAllocator>& InBoneTransforms)
 	{
@@ -45,6 +48,8 @@ public:
 
 	const TArray<FTransform>& GetBoneTransforms() const;
 
+	const FBlendedHeapCurve& GetCurves() const;
+	
 	const FTransform& GetWorldTransform() const;
 
 	FLinearColor GetBoneColor() const;
@@ -59,7 +64,8 @@ private:
 	FTransform WorldTransform;
 	TArray<FBoneIndexType> RequiredBones;
 	TArray<FTransform> BoneTransforms;
-
+	FBlendedHeapCurve Curves;
+	
 	// Mirrored properties updated on CopyPoseWatchData
 	FLinearColor BoneColor;
 	FVector ViewportOffset;
diff --git a/Engine/Source/Runtime/Engine/Private/Animation/AnimCompressionTypes.cpp b/Engine/Source/Runtime/Engine/Private/Animation/AnimCompressionTypes.cpp
--- a/Engine/Source/Runtime/Engine/Private/Animation/AnimCompressionTypes.cpp
+++ b/Engine/Source/Runtime/Engine/Private/Animation/AnimCompressionTypes.cpp
@@ -129,7 +129,7 @@ public:
 			RequiredBoneIndexArray[BoneIndex] = BoneIndex;
 		}
 
-		RequiredBones.InitializeTo(RequiredBoneIndexArray, FCurveEvaluationOption(true), *InSkeleton);
+		RequiredBones.InitializeTo(RequiredBoneIndexArray, UE::Anim::FCurveFilterSettings(UE::Anim::ECurveFilterMode::None), *InSkeleton);
 	}
 
 };
@@ -142,11 +142,11 @@ static void CopyTransformToRawAnimationData(const FTransform& BoneTransform, FRa
 	Track.ScaleKeys[Frame] = FVector3f(BoneTransform.GetScale3D());
 }
 
-static FFloatCurve* GetFloatCurve(TArray<FFloatCurve>& FloatCurves, const USkeleton::AnimCurveUID& CurveUID)
+static FFloatCurve* GetFloatCurve(TArray<FFloatCurve>& FloatCurves, const FName& CurveName)
 {
 	for (FFloatCurve& Curve : FloatCurves)
 	{
-		if (Curve.Name.UID == CurveUID)
+		if (Curve.GetName() == CurveName)
 		{
 			return &Curve;
 		}
@@ -155,10 +155,9 @@ static FFloatCurve* GetFloatCurve(TArray<FFloatCurve>& FloatCurves, const USkele
 	return nullptr;
 }
 
-static FFloatCurve* FindOrAddCurve(TArray<FFloatCurve>& FloatCurves, const FSmartName& CurveName)
+static FFloatCurve* FindOrAddCurve(TArray<FFloatCurve>& FloatCurves, const FName& CurveName)
 {
-	FFloatCurve* ReturnCurve = GetFloatCurve(FloatCurves, CurveName.UID);
-
+	FFloatCurve* ReturnCurve = GetFloatCurve(FloatCurves, CurveName);
 	if (ReturnCurve == nullptr)
 	{
 		FFloatCurve& NewCurve = FloatCurves.Add_GetRef(FFloatCurve(CurveName, 0));
@@ -166,7 +165,6 @@ static FFloatCurve* FindOrAddCurve(TArray<FFloatCurve>& FloatCurves, const FSmar
 	}
 
 	return ReturnCurve;
-	
 }
 
 static bool IsNewKeyDifferent(const FRichCurveKey& LastKey, float NewValue)
@@ -272,53 +270,45 @@ void FCompressibleAnimData::BakeOutAdditiveIntoRawData(const FFrameRate& SampleR
 			}
 
 			//Write out curve data for this frame
-			for (int32 CurveIndex = 0; CurveIndex < Curve.UIDToArrayIndexLUT->Num(); ++CurveIndex)
+			Curve.ForEachElement([this, &FloatCurves, KeyIndex, PreviousKeyTime, CurrentKeyTime, CurrentFrameTime, &SampleRate](const UE::Anim::FCurveElement& InElement)
 			{
-				USkeleton::AnimCurveUID CurveUID = (USkeleton::AnimCurveUID) CurveIndex;
-				int32 ArrayIndex = Curve.GetArrayIndexByUID(CurveUID);
-				if (ArrayIndex != INDEX_NONE)
+				const float CurveWeight = InElement.Value;
+				FFloatCurve* RawCurve = GetFloatCurve(FloatCurves, InElement.Name);
+				if (!RawCurve && !FMath::IsNearlyZero(CurveWeight)) //Only make a new curve if we are going to give it data
 				{
-					float CurveWeight = Curve.CurveWeights[ArrayIndex];
-					FFloatCurve* RawCurve = GetFloatCurve(FloatCurves, CurveUID);
-					if (!RawCurve && !FMath::IsNearlyZero(CurveWeight)) //Only make a new curve if we are going to give it data
-					{
-						FSmartName NewCurveName;
-						// if we don't have name, there is something wrong here. 
-						ensureAlways(Skeleton->GetSmartNameByUID(USkeleton::AnimCurveMappingName, CurveUID, NewCurveName));
-						// curve flags don't matter much for compressed curves
-						RawCurve = FindOrAddCurve(FloatCurves, NewCurveName);
-					}
+					// curve flags don't matter much for compressed curves
+					RawCurve = FindOrAddCurve(FloatCurves, InElement.Name);
+				}
 
-					if (RawCurve)
+				if (RawCurve)
+				{
+					const bool bHasKeys = RawCurve->FloatCurve.GetNumKeys() > 0;
+					if (!bHasKeys)
 					{
-						const bool bHasKeys = RawCurve->FloatCurve.GetNumKeys() > 0;
-						if (!bHasKeys)
+						//Add pre key of 0
+						if (KeyIndex > 0)
 						{
-							//Add pre key of 0
-							if (KeyIndex > 0)
-							{
-								RawCurve->UpdateOrAddKey(0.f, PreviousKeyTime);
-							}
+							RawCurve->UpdateOrAddKey(0.f, PreviousKeyTime);
 						}
+					}
 
-						if (!bHasKeys || IsNewKeyDifferent(RawCurve->FloatCurve.GetLastKey(), CurveWeight))
+					if (!bHasKeys || IsNewKeyDifferent(RawCurve->FloatCurve.GetLastKey(), CurveWeight))
+					{
+						RawCurve->UpdateOrAddKey(CurveWeight, CurrentKeyTime);
+						TArray<FRichCurveKey>& CurveKeys = RawCurve->FloatCurve.Keys;
+						if (CurveKeys.Num() > 1)
 						{
-							RawCurve->UpdateOrAddKey(CurveWeight, CurrentKeyTime);
-							TArray<FRichCurveKey>& CurveKeys = RawCurve->FloatCurve.Keys;
-							if (CurveKeys.Num() > 1)
+							FRichCurveKey& PrevKey = CurveKeys.Last(1);
+							// Round to frame here as it would have been added at a specific frame boundary (though float->double conversion might mean the value is off)
+							const FFrameTime PrevKeyTime = SampleRate.AsFrameTime(PrevKey.Time).RoundToFrame();
+							if (PrevKeyTime < CurrentFrameTime - 1) // Did we skip a frame, if so need to make previous key const
 							{
-								FRichCurveKey& PrevKey = CurveKeys.Last(1);
-								// Round to frame here as it would have been added at a specific frame boundary (though float->double conversion might mean the value is off)
-								const FFrameTime PrevKeyTime = SampleRate.AsFrameTime(PrevKey.Time).RoundToFrame();
-								if (PrevKeyTime < CurrentFrameTime - 1) // Did we skip a frame, if so need to make previous key const
-								{
-									PrevKey.InterpMode = RCIM_Constant;
-								}
+								PrevKey.InterpMode = RCIM_Constant;
 							}
 						}
 					}
 				}
-			}
+			});
 		}
 	}
 }
@@ -650,7 +640,7 @@ void FCompressibleAnimData::WriteCompressionDataToJSON(TArrayView<FName> Origina
 				{
 					Writer->WriteObjectStart();
 					{
-						Writer->WriteValue(TEXT("curve_name"), FloatCurve.Name.DisplayName.ToString());
+						Writer->WriteValue(TEXT("curve_name"), FloatCurve.GetName().ToString());
 						Writer->WriteValue(TEXT("number_of_keys"), FloatCurve.FloatCurve.GetNumKeys());
 							
 						if(FloatCurve.FloatCurve.GetConstRefOfKeys().Num())
@@ -785,6 +775,13 @@ void FCompressibleAnimData::FetchData(const ITargetPlatform* InPlatform)
 		OriginalTrackNames.Add(AnimTrack.Name);
 	}
 
+	// Pre-sort raw float curves. While FName indices are not stable over serialization, ordering is *close* if names come from the same source.
+	// Sorting increases the chances of linearizing the decompression later
+	RawFloatCurves.Sort([](const FFloatCurve& InLHS, const FFloatCurve& InRHS)
+	{
+		return InLHS.GetName().FastLess(InRHS.GetName());
+	});
+	
 	// Apply any key reduction if possible
 	if (RawAnimationData.Num())
 	{ 
@@ -863,13 +860,6 @@ void FCompressibleAnimData::FetchData(const ITargetPlatform* InPlatform)
         }
 	}
 
-	// Find or add curve names on skeleton
-	const FSmartNameMapping* Mapping = Skeleton->GetSmartNameContainer(USkeleton::AnimCurveMappingName);
-	for (FFloatCurve& Curve : RawFloatCurves)
-	{
-		Skeleton->VerifySmartName(USkeleton::AnimCurveMappingName, Curve.Name);
-	}
-	   
 	if (bShouldPerformStripping)
 	{
 		const int32 NumTracks = RawAnimationData.Num();
@@ -956,13 +946,15 @@ void FCompressibleAnimData::Update(FCompressedAnimSequence& InOutCompressedData)
 	InOutCompressedData.CompressedRawDataSize = GetApproxRawSize();
 
 	const int32 NumCurves = RawFloatCurves.Num();
-	InOutCompressedData.CompressedCurveNames.Reset(NumCurves);
-	InOutCompressedData.CompressedCurveNames.AddUninitialized(NumCurves);
+	InOutCompressedData.IndexedCurveNames.Reset(NumCurves);
+	InOutCompressedData.IndexedCurveNames.AddUninitialized(NumCurves);
 	for (int32 CurveIndex = 0; CurveIndex < NumCurves; ++CurveIndex)
 	{
 		const FFloatCurve& Curve = RawFloatCurves[CurveIndex];
-		InOutCompressedData.CompressedCurveNames[CurveIndex] = Curve.Name;
+		InOutCompressedData.IndexedCurveNames[CurveIndex].CurveName = Curve.GetName();
 	}
+
+	InOutCompressedData.RebuildCurveIndexTable();
 }
 
 #endif // WITH_EDITOR
@@ -1202,7 +1194,7 @@ void FCompressedAnimSequence::SerializeCompressedData(FArchive& Ar, bool bDDCDat
 {
 	Ar << CompressedRawDataSize;
 	Ar << CompressedTrackToSkeletonMapTable;
-	Ar << CompressedCurveNames;
+	Ar << IndexedCurveNames;
 
 	// Serialize the compressed byte stream from the archive to the buffer.
 	int32 NumBytes = CompressedByteStream.Num();
@@ -1304,6 +1296,8 @@ void FCompressedAnimSequence::SerializeCompressedData(FArchive& Ar, bool bDDCDat
 				BoneCompressionCodec->ByteSwapIn(*CompressedDataStructure, CompressedByteStream, MemoryReader);
 			}
 		}
+
+		RebuildCurveIndexTable();
 	}
 	else if (Ar.IsSaving() || Ar.IsCountingMemory())
 	{
@@ -1421,20 +1415,41 @@ void FCompressedAnimSequence::SerializeCompressedData(FArchive& Ar, bool bDDCDat
 		{
 			//Temp DDC debug
 			//DebugData.Display();
-			// Refresh the compressed curve names since the IDs might have changed since
-			for (FSmartName& CurveName : CompressedCurveNames)
-			{
-				Skeleton->VerifySmartName(USkeleton::AnimCurveMappingName, CurveName);
-			}
 		}
 	}
 #endif
 }
 
+void FCompressedAnimSequence::RebuildCurveIndexTable()
+{
+	FMemMark Mark(FMemStack::Get());
+
+	TArray<int32, TMemStackAllocator<>> IndexArray;
+	IndexArray.SetNumUninitialized(IndexedCurveNames.Num());
+	
+	// Create linear indices
+	for(int32 NameIndex = 0; NameIndex < IndexArray.Num(); ++NameIndex)
+	{
+		IndexArray[NameIndex] = NameIndex;
+	}
+
+	// Sort by FName
+	IndexArray.Sort([&IndexedCurveNames = IndexedCurveNames](int32 LHS, int32 RHS)
+	{
+		return IndexedCurveNames[LHS].CurveName.FastLess(IndexedCurveNames[RHS].CurveName);
+	});
+
+	// Index curves
+	for(int32 NameIndex = 0; NameIndex < IndexedCurveNames.Num(); ++NameIndex)
+	{
+		IndexedCurveNames[NameIndex].CurveIndex = IndexArray[NameIndex];
+	}
+}
+
 SIZE_T FCompressedAnimSequence::GetMemorySize() const
 {
 	return	  CompressedTrackToSkeletonMapTable.GetAllocatedSize()
-			+ CompressedCurveNames.GetAllocatedSize()
+			+ IndexedCurveNames.GetAllocatedSize()
 			+ CompressedCurveByteStream.GetAllocatedSize()
 			+ CompressedDataStructure->GetApproxCompressedSize()
 			+ sizeof(FCompressedAnimSequence);
@@ -1443,7 +1458,7 @@ SIZE_T FCompressedAnimSequence::GetMemorySize() const
 void FCompressedAnimSequence::Reset()
 {
 	CompressedTrackToSkeletonMapTable.Empty();
-	CompressedCurveNames.Empty();
+	IndexedCurveNames.Empty();
 	ClearCompressedBoneData();
 	ClearCompressedCurveData();
 	CompressedRawDataSize = 0;
@@ -1556,13 +1571,13 @@ void UpdateSHAWithCurves(FSHA1& Sha, const FRawCurveTracks& InRawCurveData)
 
 	for (const FFloatCurve& Curve : InRawCurveData.FloatCurves)
 			{
-		UpdateWithData(Sha, Curve.Name.UID);
+		UpdateWithData(Sha, Curve.GetName());
 		UpdateWithFloatCurve(Curve.FloatCurve);
 	}
 
 	for (const FTransformCurve& Curve : InRawCurveData.TransformCurves)
 	{
-		UpdateWithData(Sha, Curve.Name.UID);
+		UpdateWithData(Sha, Curve.GetName());
 
 		auto UpdateWithComponent = [&Sha, &UpdateWithFloatCurve](const FVectorCurve& VectorCurve)
 		{
