Animation Curve Runtime & Editor Improvements

diff --git a/Engine/Source/Runtime/Engine/Private/Animation/AnimCurveCompressionCodec_CompressedRichCurve.cpp b/Engine/Source/Runtime/Engine/Private/Animation/AnimCurveCompressionCodec_CompressedRichCurve.cpp
--- a/Engine/Source/Runtime/Engine/Private/Animation/AnimCurveCompressionCodec_CompressedRichCurve.cpp
+++ b/Engine/Source/Runtime/Engine/Private/Animation/AnimCurveCompressionCodec_CompressedRichCurve.cpp
@@ -4,6 +4,7 @@
 #include "Animation/AnimCompressionTypes.h"
 #include "HAL/IConsoleManager.h"
 #include "AnimationCompression.h"
+#include "Animation/AnimCurveUtils.h"
 
 #include UE_INLINE_GENERATED_CPP_BY_NAME(AnimCurveCompressionCodec_CompressedRichCurve)
 
@@ -88,7 +89,7 @@ bool UAnimCurveCompressionCodec_CompressedRichCurve::Compress(const FCompressibl
 				if (!FMath::IsNearlyZero(AbsDelta, MaxCurveError + UE_KINDA_SMALL_NUMBER))
 				{
 					// Delta larger than tolerated error value
-					UE_LOG(LogAnimationCompression, Warning, TEXT("Curve %s: delta too large %f, between %f and %f, at %f"), *Curve.Name.DisplayName.ToString(), AbsDelta, RawValue, CompressedValue, EvalTime);
+					UE_LOG(LogAnimationCompression, Warning, TEXT("Curve %s: delta too large %f, between %f and %f, at %f"), *Curve.GetName().ToString(), AbsDelta, RawValue, CompressedValue, EvalTime);
 				}
 			}
 		}
@@ -126,35 +127,44 @@ void UAnimCurveCompressionCodec_CompressedRichCurve::PopulateDDCKey(FArchive& Ar
 
 void UAnimCurveCompressionCodec_CompressedRichCurve::DecompressCurves(const FCompressedAnimSequence& AnimSeq, FBlendedCurve& Curves, float CurrentTime) const
 {
+	const TArray<FAnimCompressedCurveIndexedName>& IndexedCurveNames = AnimSeq.IndexedCurveNames;
+	const int32 NumCurves = IndexedCurveNames.Num();
+
+	if (NumCurves == 0)
+	{
+		return;
+	}
+
 	const uint8* Buffer = AnimSeq.CompressedCurveByteStream.GetData();
 	const FCurveDesc* CurveDescriptions = (const FCurveDesc*)(Buffer);
 
-	const TArray<FSmartName>& CompressedCurveNames = AnimSeq.CompressedCurveNames;
-	const int32 NumCurves = CompressedCurveNames.Num();
-	for (int32 CurveIndex = 0; CurveIndex < NumCurves; ++CurveIndex)
+	auto GetNameFromIndex = [&IndexedCurveNames](int32 InCurveIndex)
 	{
-		const FSmartName& CurveName = CompressedCurveNames[CurveIndex];
-		if (Curves.IsEnabled(CurveName.UID))
-		{
-			const FCurveDesc& Curve = CurveDescriptions[CurveIndex];
-			const uint8* CompressedKeys = Buffer + Curve.KeyDataOffset;
-			const float Value = FCompressedRichCurve::StaticEval(Curve.CompressionFormat, Curve.KeyTimeCompressionFormat, Curve.PreInfinityExtrap, Curve.PostInfinityExtrap, Curve.ConstantValueNumKeys, CompressedKeys, CurrentTime);
-			Curves.Set(CurveName.UID, Value);
-		}
-	}
+		return IndexedCurveNames[IndexedCurveNames[InCurveIndex].CurveIndex].CurveName;
+	};
+
+	auto GetValueFromIndex = [&CurveDescriptions, &IndexedCurveNames, Buffer, CurrentTime](int32 InCurveIndex)
+	{
+		const FCurveDesc& Curve = CurveDescriptions[IndexedCurveNames[InCurveIndex].CurveIndex];
+		const uint8* CompressedKeys = Buffer + Curve.KeyDataOffset;
+		const float Value = FCompressedRichCurve::StaticEval(Curve.CompressionFormat, Curve.KeyTimeCompressionFormat, Curve.PreInfinityExtrap, Curve.PostInfinityExtrap, Curve.ConstantValueNumKeys, CompressedKeys, CurrentTime);
+		return Value;
+	};
+	
+	UE::Anim::FCurveUtils::BuildSorted(Curves, NumCurves, GetNameFromIndex, GetValueFromIndex, Curves.GetFilter());
 }
 
-float UAnimCurveCompressionCodec_CompressedRichCurve::DecompressCurve(const FCompressedAnimSequence& AnimSeq, SmartName::UID_Type CurveUID, float CurrentTime) const
+float UAnimCurveCompressionCodec_CompressedRichCurve::DecompressCurve(const FCompressedAnimSequence& AnimSeq, FName InCurveName, float CurrentTime) const
 {
 	const uint8* Buffer = AnimSeq.CompressedCurveByteStream.GetData();
 	const FCurveDesc* CurveDescriptions = (const FCurveDesc*)(Buffer);
-
-	const TArray<FSmartName>& CompressedCurveNames = AnimSeq.CompressedCurveNames;
-	const int32 NumCurves = CompressedCurveNames.Num();
+	
+	const TArray<FAnimCompressedCurveIndexedName>& IndexedCurveNames = AnimSeq.IndexedCurveNames;
+	const int32 NumCurves = IndexedCurveNames.Num();
 	for (int32 CurveIndex = 0; CurveIndex < NumCurves; ++CurveIndex)
 	{
-		const FSmartName& CurveName = CompressedCurveNames[CurveIndex];
-		if (CurveName.UID == CurveUID)
+		const FName& CurveName = IndexedCurveNames[CurveIndex].CurveName;
+		if (CurveName == InCurveName)
 		{
 			const FCurveDesc& Curve = CurveDescriptions[CurveIndex];
 			const uint8* CompressedKeys = Buffer + Curve.KeyDataOffset;
diff --git a/Engine/Source/Runtime/Engine/Private/Animation/AnimCurveCompressionCodec_UniformIndexable.cpp b/Engine/Source/Runtime/Engine/Private/Animation/AnimCurveCompressionCodec_UniformIndexable.cpp
--- a/Engine/Source/Runtime/Engine/Private/Animation/AnimCurveCompressionCodec_UniformIndexable.cpp
+++ b/Engine/Source/Runtime/Engine/Private/Animation/AnimCurveCompressionCodec_UniformIndexable.cpp
@@ -3,17 +3,29 @@
 #include "Animation/AnimCurveCompressionCodec_UniformIndexable.h"
 #include "Animation/AnimSequence.h"
 #include "EngineLogs.h"
+#include "Animation/AttributesRuntime.h"
+#include "Animation/BuiltInAttributeTypes.h"
+#include "Animation/AnimCurveUtils.h"
 
 #include UE_INLINE_GENERATED_CPP_BY_NAME(AnimCurveCompressionCodec_UniformIndexable)
 
-FAnimCurveBufferAccess::FAnimCurveBufferAccess(const UAnimSequenceBase* InSequenceBase, USkeleton::AnimCurveUID InUID)
+FAnimCurveBufferAccess::FAnimCurveBufferAccess(const UAnimSequenceBase* InSequenceBase, SmartName::UID_Type InUID)
 	: NumSamples(INDEX_NONE)
 	, SampleRate(0.f)
 	, CompressedBuffer(nullptr)
 	, RawCurve(nullptr)
 	, bUseCompressedData(false)
 {
-	if (InSequenceBase->HasCurveData(InUID, false))
+}
+
+FAnimCurveBufferAccess::FAnimCurveBufferAccess(const UAnimSequenceBase* InSequenceBase, FName InName)
+	: NumSamples(INDEX_NONE)
+	, SampleRate(0.f)
+	, CompressedBuffer(nullptr)
+	, RawCurve(nullptr)
+	, bUseCompressedData(false)
+{
+	if (InSequenceBase->HasCurveData(InName, false))
 	{
 		if (const UAnimSequence* InSequence = Cast<UAnimSequence>(InSequenceBase))
 		{
@@ -22,7 +34,7 @@ FAnimCurveBufferAccess::FAnimCurveBufferAccess(const UAnimSequenceBase* InSequen
 				const FCompressedAnimSequence& CompressedAnimSequence = InSequence->CompressedData;
 				if (UAnimCurveCompressionCodec_UniformIndexable* IndexableCompressedCurve = Cast<UAnimCurveCompressionCodec_UniformIndexable>(CompressedAnimSequence.CurveCompressionCodec))
 				{
-					bUseCompressedData = IndexableCompressedCurve->GetCurveBufferAndSamples(CompressedAnimSequence, InUID, CompressedBuffer, NumSamples, SampleRate);
+					bUseCompressedData = IndexableCompressedCurve->GetCurveBufferAndSamples(CompressedAnimSequence, InName, CompressedBuffer, NumSamples, SampleRate);
 				}
 				else
 				{
@@ -33,7 +45,7 @@ FAnimCurveBufferAccess::FAnimCurveBufferAccess(const UAnimSequenceBase* InSequen
 
 		if (!CompressedBuffer)
 		{
-			RawCurve = (const FFloatCurve*)InSequenceBase->GetCurveData().GetCurveData(InUID);
+			RawCurve = (const FFloatCurve*)InSequenceBase->GetCurveData().GetCurveData(InName);
 
 			if (RawCurve)
 			{
@@ -156,8 +168,8 @@ void UAnimCurveCompressionCodec_UniformIndexable::PopulateDDCKey(FArchive& Ar)
 
 void UAnimCurveCompressionCodec_UniformIndexable::DecompressCurves(const FCompressedAnimSequence& AnimSeq, FBlendedCurve& Curves, float CurrentTime) const
 {
-	const TArray<FSmartName>& CompressedCurveNames = AnimSeq.CompressedCurveNames;
-	const int32 NumCurves = CompressedCurveNames.Num();
+	const TArray<FAnimCompressedCurveIndexedName>& IndexedCurveNames = AnimSeq.IndexedCurveNames;
+	const int32 NumCurves = IndexedCurveNames.Num();
 
 	if (NumCurves == 0)
 	{
@@ -187,31 +199,33 @@ void UAnimCurveCompressionCodec_UniformIndexable::DecompressCurves(const FCompre
 	const int32 SampleIndex1 = FMath::Min(SampleIndex0 + 1, NumSamples - 1);
 	const float InterpolationAlpha = SamplePoint - float(SampleIndex0);
 
-	for (int32 CurveIndex = 0; CurveIndex < NumCurves; ++CurveIndex)
+	auto GetNameFromIndex = [&IndexedCurveNames](int32 InCurveIndex)
 	{
-		const FSmartName& CurveName = CompressedCurveNames[CurveIndex];
-		
-		if (Curves.IsEnabled(CurveName.UID))
-		{
-			const float* CurveSamplesPtr = AnimatedSamplesPtr + (CurveIndex * NumSamples);
+		return IndexedCurveNames[IndexedCurveNames[InCurveIndex].CurveIndex].CurveName;
+	};
 
-			const float Sample0 = CurveSamplesPtr[SampleIndex0];
-			const float Sample1 = CurveSamplesPtr[SampleIndex1];
-			const float Sample = FMath::Lerp(Sample0, Sample1, InterpolationAlpha);
+	auto GetValueFromIndex = [&IndexedCurveNames, AnimatedSamplesPtr, NumSamples, SampleIndex0, SampleIndex1, InterpolationAlpha](int32 InCurveIndex)
+	{
+		const int32 CurveIndex = IndexedCurveNames[InCurveIndex].CurveIndex;
+		const float* CurveSamplesPtr = AnimatedSamplesPtr + (CurveIndex * NumSamples);
 
-			Curves.Set(CurveName.UID, Sample);
-		}
-	}
+		const float Sample0 = CurveSamplesPtr[SampleIndex0];
+		const float Sample1 = CurveSamplesPtr[SampleIndex1];
+		const float Sample = FMath::Lerp(Sample0, Sample1, InterpolationAlpha);
+		return Sample;
+	};
+
+	UE::Anim::FCurveUtils::BuildSorted(Curves, NumCurves, GetNameFromIndex, GetValueFromIndex, Curves.GetFilter());
 }
 
-float UAnimCurveCompressionCodec_UniformIndexable::DecompressCurve(const FCompressedAnimSequence& AnimSeq, SmartName::UID_Type CurveUID, float CurrentTime) const
+float UAnimCurveCompressionCodec_UniformIndexable::DecompressCurve(const FCompressedAnimSequence& AnimSeq, FName InCurveName, float CurrentTime) const
 {
-	const TArray<FSmartName>& CompressedCurveNames = AnimSeq.CompressedCurveNames;
-	const int32 NumCurves = CompressedCurveNames.Num();
+	const TArray<FAnimCompressedCurveIndexedName>& IndexedCurveNames = AnimSeq.IndexedCurveNames;
+	const int32 NumCurves = IndexedCurveNames.Num();
 
 	if (NumCurves == 0)
 	{
-		return 0.f;
+		return 0.0f;
 	}
 
 	const uint8* Buffer = AnimSeq.CompressedCurveByteStream.GetData();
@@ -223,7 +237,7 @@ float UAnimCurveCompressionCodec_UniformIndexable::DecompressCurve(const FCompre
 
 	if (NumSamples == 0)
 	{
-		return 0.f;
+		return 0.0f;
 	}
 
 	const float SampleRate = *(const float*)&Buffer[BufferOffset];
@@ -239,24 +253,26 @@ float UAnimCurveCompressionCodec_UniformIndexable::DecompressCurve(const FCompre
 
 	for (int32 CurveIndex = 0; CurveIndex < NumCurves; ++CurveIndex)
 	{
-		const FSmartName& CurveName = CompressedCurveNames[CurveIndex];
-
-		if (CurveName.UID == CurveUID)
+		const FName& CurveName = IndexedCurveNames[CurveIndex].CurveName;
+		if (CurveName == InCurveName)
 		{
 			const float* CurveSamplesPtr = AnimatedSamplesPtr + (CurveIndex * NumSamples);
 
 			const float Sample0 = CurveSamplesPtr[SampleIndex0];
 			const float Sample1 = CurveSamplesPtr[SampleIndex1];
-			return FMath::Lerp(Sample0, Sample1, InterpolationAlpha);
+			const float Sample = FMath::Lerp(Sample0, Sample1, InterpolationAlpha);
+
+			return Sample;
 		}
 	}
-	return 0.f;
+
+	return 0.0f;
 }
 
-bool UAnimCurveCompressionCodec_UniformIndexable::GetCurveBufferAndSamples(const FCompressedAnimSequence& AnimSeq, SmartName::UID_Type CurveUID, const float*& OutCurveBuffer, int32& OutSamples, float& OutSampleRate)
+bool UAnimCurveCompressionCodec_UniformIndexable::GetCurveBufferAndSamples(const FCompressedAnimSequence& AnimSeq, FName CurveName, const float*& OutCurveBuffer, int32& OutSamples, float& OutSampleRate)
 {
-	const TArray<FSmartName>& CompressedCurveNames = AnimSeq.CompressedCurveNames;
-	const int32 NumCurves = CompressedCurveNames.Num();
+	const TArray<FAnimCompressedCurveIndexedName>& IndexedCurveNames = AnimSeq.IndexedCurveNames;
+	const int32 NumCurves = IndexedCurveNames.Num();
 
 	if (NumCurves == 0)
 	{
@@ -283,8 +299,7 @@ bool UAnimCurveCompressionCodec_UniformIndexable::GetCurveBufferAndSamples(const
 
 	for (int32 CurveIndex = 0; CurveIndex < NumCurves; ++CurveIndex)
 	{
-		const FSmartName& CurveName = CompressedCurveNames[CurveIndex];
-		if (CurveName.UID == CurveUID)
+		if (IndexedCurveNames[CurveIndex].CurveName == CurveName)
 		{
 			OutCurveBuffer = AnimatedSamplesPtr + (CurveIndex * OutSamples);
 			return true;
diff --git a/Engine/Source/Runtime/Engine/Private/Animation/AnimCurveCompressionCodec_UniformlySampled.cpp b/Engine/Source/Runtime/Engine/Private/Animation/AnimCurveCompressionCodec_UniformlySampled.cpp
--- a/Engine/Source/Runtime/Engine/Private/Animation/AnimCurveCompressionCodec_UniformlySampled.cpp
+++ b/Engine/Source/Runtime/Engine/Private/Animation/AnimCurveCompressionCodec_UniformlySampled.cpp
@@ -2,6 +2,7 @@
 
 #include "Animation/AnimCurveCompressionCodec_UniformlySampled.h"
 #include "Animation/AnimCompressionTypes.h"
+#include "Animation/AnimCurveUtils.h"
 
 #include UE_INLINE_GENERATED_CPP_BY_NAME(AnimCurveCompressionCodec_UniformlySampled)
 
@@ -162,8 +163,8 @@ void UAnimCurveCompressionCodec_UniformlySampled::PopulateDDCKey(FArchive& Ar)
 
 void UAnimCurveCompressionCodec_UniformlySampled::DecompressCurves(const FCompressedAnimSequence& AnimSeq, FBlendedCurve& Curves, float CurrentTime) const
 {
-	const TArray<FSmartName>& CompressedCurveNames = AnimSeq.CompressedCurveNames;
-	const int32 NumCurves = CompressedCurveNames.Num();
+	const TArray<FAnimCompressedCurveIndexedName>& IndexedCurveNames = AnimSeq.IndexedCurveNames;
+	const int32 NumCurves = IndexedCurveNames.Num();
 
 	if (NumCurves == 0)
 	{
@@ -207,35 +208,42 @@ void UAnimCurveCompressionCodec_UniformlySampled::DecompressCurves(const FCompre
 	const float* AnimatedSamples0 = AnimatedSamplesPtr + (SampleIndex0 * NumAnimatedCurves);
 	const float* AnimatedSamples1 = AnimatedSamplesPtr + (SampleIndex1 * NumAnimatedCurves);
 
-	for (int32 CurveIndex = 0, ConstantCurveIndex = 0, AnimatedCurveIndex = 0; CurveIndex < NumCurves; ++CurveIndex)
+	int32 ConstantCurveIndex = 0, AnimatedCurveIndex = 0;
+
+	auto GetNameFromIndex = [&IndexedCurveNames](int32 InCurveIndex)
 	{
-		const FSmartName& CurveName = CompressedCurveNames[CurveIndex];
+		return IndexedCurveNames[IndexedCurveNames[InCurveIndex].CurveIndex].CurveName;
+	};
+
+	auto GetValueFromIndex = [&IndexedCurveNames, ConstantCurvesBitsetPtr, &ConstantCurveIndex, &AnimatedCurveIndex, ConstantSamplesPtr, AnimatedSamples0, AnimatedSamples1, InterpolationAlpha](int32 InCurveIndex) 
+	{
+		const int32 CurveIndex = IndexedCurveNames[InCurveIndex].CurveIndex;
 		const bool bIsConstant = (ConstantCurvesBitsetPtr[CurveIndex / 32] & (1 << (CurveIndex % 32))) != 0;
-		if (Curves.IsEnabled(CurveName.UID))
-		{
-			float Sample;
-			if (bIsConstant)
-			{
-				Sample = ConstantSamplesPtr[ConstantCurveIndex];
-			}
-			else
-			{
-				const float Sample0 = AnimatedSamples0[AnimatedCurveIndex];
-				const float Sample1 = AnimatedSamples1[AnimatedCurveIndex];
-				Sample = FMath::Lerp(Sample0, Sample1, InterpolationAlpha);
-			}
 
-			Curves.Set(CurveName.UID, Sample);
+		float Sample;
+		if (bIsConstant)
+		{
+			Sample = ConstantSamplesPtr[ConstantCurveIndex];
+		}
+		else
+		{
+			const float Sample0 = AnimatedSamples0[AnimatedCurveIndex];
+			const float Sample1 = AnimatedSamples1[AnimatedCurveIndex];
+			Sample = FMath::Lerp(Sample0, Sample1, InterpolationAlpha);
 		}
 
 		(bIsConstant ? ConstantCurveIndex : AnimatedCurveIndex)++;
-	}
+
+		return Sample;
+	};
+
+	UE::Anim::FCurveUtils::BuildSorted(Curves, NumCurves, GetNameFromIndex, GetValueFromIndex, Curves.GetFilter());
 }
 
-float UAnimCurveCompressionCodec_UniformlySampled::DecompressCurve(const FCompressedAnimSequence& AnimSeq, SmartName::UID_Type CurveUID, float CurrentTime) const
+float UAnimCurveCompressionCodec_UniformlySampled::DecompressCurve(const FCompressedAnimSequence& AnimSeq, FName InCurveName, float CurrentTime) const
 {
-	const TArray<FSmartName>& CompressedCurveNames = AnimSeq.CompressedCurveNames;
-	const int32 NumCurves = CompressedCurveNames.Num();
+	const TArray<FAnimCompressedCurveIndexedName>& IndexedCurveNames = AnimSeq.IndexedCurveNames;
+	const int32 NumCurves = IndexedCurveNames.Num();
 
 	if (NumCurves == 0)
 	{
@@ -267,9 +275,9 @@ float UAnimCurveCompressionCodec_UniformlySampled::DecompressCurve(const FCompre
 
 	for (int32 CurveIndex = 0, ConstantCurveIndex = 0, AnimatedCurveIndex = 0; CurveIndex < NumCurves; ++CurveIndex)
 	{
-		const FSmartName& CurveName = CompressedCurveNames[CurveIndex];
+		const FName& CurveName = IndexedCurveNames[CurveIndex].CurveName;
 		const bool bIsConstant = (ConstantCurvesBitsetPtr[CurveIndex / 32] & (1 << (CurveIndex % 32))) != 0;
-		if (CurveName.UID == CurveUID)
+		if (CurveName == InCurveName)
 		{
 			float Sample;
 			if (bIsConstant)
diff --git a/Engine/Source/Runtime/Engine/Private/Animation/AnimCurveTypes.cpp b/Engine/Source/Runtime/Engine/Private/Animation/AnimCurveTypes.cpp
--- a/Engine/Source/Runtime/Engine/Private/Animation/AnimCurveTypes.cpp
+++ b/Engine/Source/Runtime/Engine/Private/Animation/AnimCurveTypes.cpp
@@ -8,47 +8,75 @@
 #include "UObject/AnimObjectVersion.h"
 #include "Math/RandomStream.h"
 #include "Animation/AnimSequenceBase.h"
+#include "Animation/Skeleton.h"
+#include "BoneContainer.h"
+#include "UObject/UE5MainStreamObjectVersion.h"
+#include "Animation/AnimCurveUtils.h"
 
 DECLARE_CYCLE_STAT(TEXT("EvalRawCurveData"), STAT_EvalRawCurveData, STATGROUP_Anim);
 
+namespace UE::Anim
+{
+PRAGMA_DISABLE_DEPRECATION_WARNINGS
+	TArray<float> FBaseBlendedCurve_DEPRECATED::CurveWeights;
+	TBitArray<> FBaseBlendedCurve_DEPRECATED::ValidCurveWeights;
+	TArray<uint16> const* FBaseBlendedCurve_DEPRECATED::UIDToArrayIndexLUT;
+	uint16 FBaseBlendedCurve_DEPRECATED::NumValidCurveCount;
+	bool FBaseBlendedCurve_DEPRECATED::bInitialized;
+PRAGMA_ENABLE_DEPRECATION_WARNINGS
+}
+
 /////////////////////////////////////////////////////
 // FFloatCurve
 
-void FAnimCurveBase::PostSerialize(FArchive& Ar)
+void FAnimCurveBase::PostSerializeFixup(FArchive& Ar)
+{
+#if WITH_EDITORONLY_DATA
+	if (Ar.CustomVer(FFrameworkObjectVersion::GUID) < FFrameworkObjectVersion::SmartNameRefactor)
+	{
+		if (Ar.UEVer() >= VER_UE4_SKELETON_ADD_SMARTNAMES)
+		{
+			SmartName::UID_Type CurveUid = SmartName::MaxUID;
+			Ar << CurveUid;
+			Name_DEPRECATED.UID = CurveUid;
+		}
+	}
+#endif
+}
+
+bool FAnimCurveBase::Serialize(FArchive& Ar)
 {
-	SmartName::UID_Type CurveUid = SmartName::MaxUID;
 	Ar.UsingCustomVersion(FFrameworkObjectVersion::GUID);
 	Ar.UsingCustomVersion(FAnimObjectVersion::GUID);
+	Ar.UsingCustomVersion(FUE5MainStreamObjectVersion::GUID);
+
+	// Return false to defer to regular serialization 
+	return false;
+}
 
+void FAnimCurveBase::PostSerialize(const FArchive& Ar)
+{
+#if WITH_EDITORONLY_DATA
 	if (Ar.IsLoading())
 	{
 		if (Ar.CustomVer(FFrameworkObjectVersion::GUID) < FFrameworkObjectVersion::SmartNameRefactor)
 		{
-			if (Ar.UEVer() >= VER_UE4_SKELETON_ADD_SMARTNAMES)
-			{
-				Ar << CurveUid;
-
-				Name.UID = CurveUid;
-
-#if WITH_EDITORONLY_DATA
-				Name.DisplayName = LastObservedName_DEPRECATED;
-#endif
-			}
-#if WITH_EDITORONLY_DATA
-			else
-			{
-				Name.DisplayName = LastObservedName_DEPRECATED;
-			}
-#endif
+			Name_DEPRECATED.DisplayName = LastObservedName_DEPRECATED;
 		}
-
-#if WITH_EDITORONLY_DATA
+		
 		if(Ar.CustomVer(FAnimObjectVersion::GUID) < FAnimObjectVersion::AnimSequenceCurveColors)
 		{
+			// Need to set the curve name before we generate a new color
+			CurveName = Name_DEPRECATED.DisplayName;
 			Color = MakeColor();
 		}
-#endif
+
+		if(Ar.CustomVer(FUE5MainStreamObjectVersion::GUID) < FUE5MainStreamObjectVersion::AnimationRemoveSmartNames)
+		{
+			CurveName = Name_DEPRECATED.DisplayName;
+		}
 	}
+#endif
 }
 
 void FAnimCurveBase::SetCurveTypeFlag(EAnimAssetCurveFlags InFlag, bool bValue)
@@ -89,7 +117,7 @@ int32 FAnimCurveBase::GetCurveTypeFlags() const
 FLinearColor FAnimCurveBase::MakeColor()
 {
 	// Create a color based on the hash of the name
-	FRandomStream Stream(GetTypeHash(Name.DisplayName));
+	FRandomStream Stream(GetTypeHash(CurveName));
 	const uint8 Hue = (uint8)(Stream.FRand() * 255.0f);
 	return FLinearColor::MakeFromHSV8(Hue, 196, 196);
 }
@@ -364,42 +392,19 @@ FVectorCurve* FTransformCurve::GetVectorCurveByIndex(int32 Index)
 
 bool FCachedFloatCurve::IsValid(const UAnimSequenceBase* InAnimSequence) const
 {
-	return ((CurveName != NAME_None) && InAnimSequence->HasCurveData(GetAnimCurveUID(InAnimSequence)));
+	return ((CurveName != NAME_None) && InAnimSequence->HasCurveData(CurveName));
 }
 
 float FCachedFloatCurve::GetValueAtPosition(const UAnimSequenceBase* InAnimSequence, const float& InPosition) const
 {
-	return InAnimSequence->EvaluateCurveData(GetAnimCurveUID(InAnimSequence), InPosition);
-}
-
-SkeletonAnimCurveUID FCachedFloatCurve::GetAnimCurveUID(const UAnimSequenceBase* InAnimSequence) const
-{
-	if (CurveName != CachedCurveName && InAnimSequence)
-	{
-		if (const USkeleton* Skeleton = InAnimSequence->GetSkeleton())
-		{
-			const FSmartNameMapping* CurveNameMapping = Skeleton->GetSmartNameContainer(USkeleton::AnimCurveMappingName);
-			if (CurveNameMapping)
-			{
-				CachedUID = CurveNameMapping->FindUID(CurveName);
-				CachedCurveName = CurveName;
-			}
-		}
-	}
-
-	return CachedUID;
+	return InAnimSequence->EvaluateCurveData(CurveName, InPosition);
 }
 
 const FFloatCurve* FCachedFloatCurve::GetFloatCurve(const UAnimSequenceBase* InAnimSequence) const
 {
 	if (InAnimSequence)
 	{
-		SkeletonAnimCurveUID DistanceCurveUID = GetAnimCurveUID(InAnimSequence);
-		if (DistanceCurveUID != SmartName::MaxUID)
-		{
-			const FAnimationCurveIdentifier CurveId(DistanceCurveUID, ERawCurveTrackTypes::RCT_Float);
-			return (const FFloatCurve*)(InAnimSequence->GetCurveData().GetCurveData(DistanceCurveUID));
-		}
+		return static_cast<const FFloatCurve*>(InAnimSequence->GetCurveData().GetCurveData(CurveName));
 	}
 
 	return nullptr;
@@ -411,19 +416,19 @@ const FFloatCurve* FCachedFloatCurve::GetFloatCurve(const UAnimSequenceBase* InA
 void FRawCurveTracks::EvaluateCurveData( FBlendedCurve& Curves, float CurrentTime ) const
 {
 	SCOPE_CYCLE_COUNTER(STAT_EvalRawCurveData);
-	if (Curves.NumValidCurveCount > 0)
+	
+	auto GetNameFromIndex = [this](int32 InCurveIndex)
 	{
-		// evaluate the curve data at the CurrentTime and add to Instance
-		for (auto CurveIter = FloatCurves.CreateConstIterator(); CurveIter; ++CurveIter)
-		{
-			const FFloatCurve& Curve = *CurveIter;
-			if (Curves.IsEnabled(Curve.Name.UID))
-			{
-				float Value = Curve.Evaluate(CurrentTime);
-				Curves.Set(Curve.Name.UID, Value);
-			}
-		}
-	}
+		return FloatCurves[InCurveIndex].GetName();
+	};
+
+	auto GetValueFromIndex = [this, CurrentTime](int32 InCurveIndex)
+	{
+		return FloatCurves[InCurveIndex].Evaluate(CurrentTime);
+	};
+	
+	// evaluate the curve data at the CurrentTime and add to Instance
+	UE::Anim::FCurveUtils::BuildUnsorted(Curves, FloatCurves.Num(), GetNameFromIndex, GetValueFromIndex, Curves.GetFilter());
 }
 
 #if WITH_EDITOR
@@ -446,7 +451,7 @@ void FRawCurveTracks::EvaluateTransformCurveData(USkeleton * Skeleton, TMap<FNam
 		}
 
 		// Add or retrieve curve
-		FName CurveName = Curve.Name.DisplayName;
+		FName CurveName = Curve.GetName();
 		
 		// note we're not checking Curve.GetCurveTypeFlags() yet
 		FTransform & Value = OutCurves.FindOrAdd(CurveName);
@@ -454,39 +459,40 @@ void FRawCurveTracks::EvaluateTransformCurveData(USkeleton * Skeleton, TMap<FNam
 	}
 }
 #endif
-FAnimCurveBase * FRawCurveTracks::GetCurveData(SkeletonAnimCurveUID Uid, ERawCurveTrackTypes SupportedCurveType /*= FloatType*/)
+
+const FAnimCurveBase * FRawCurveTracks::GetCurveData(FName Name, ERawCurveTrackTypes SupportedCurveType /*= FloatType*/) const
 {
 	switch (SupportedCurveType)
 	{
 #if WITH_EDITOR
 	case ERawCurveTrackTypes::RCT_Vector:
-		return GetCurveDataImpl<FVectorCurve>(VectorCurves, Uid);
+		return GetCurveDataImpl<FVectorCurve>(VectorCurves, Name);
 	case ERawCurveTrackTypes::RCT_Transform:
-		return GetCurveDataImpl<FTransformCurve>(TransformCurves, Uid);
+		return GetCurveDataImpl<FTransformCurve>(TransformCurves, Name);
 #endif // WITH_EDITOR
 	case ERawCurveTrackTypes::RCT_Float:
 	default:
-		return GetCurveDataImpl<FFloatCurve>(FloatCurves, Uid);
+		return GetCurveDataImpl<FFloatCurve>(FloatCurves, Name);
 	}
 }
 
-const FAnimCurveBase * FRawCurveTracks::GetCurveData(SkeletonAnimCurveUID Uid, ERawCurveTrackTypes SupportedCurveType /*= FloatType*/) const
+FAnimCurveBase * FRawCurveTracks::GetCurveData(FName Name, ERawCurveTrackTypes SupportedCurveType /*= FloatType*/)
 {
 	switch (SupportedCurveType)
 	{
 #if WITH_EDITOR
 	case ERawCurveTrackTypes::RCT_Vector:
-		return GetCurveDataImpl<FVectorCurve>(VectorCurves, Uid);
+		return GetCurveDataImpl<FVectorCurve>(VectorCurves, Name);
 	case ERawCurveTrackTypes::RCT_Transform:
-		return GetCurveDataImpl<FTransformCurve>(TransformCurves, Uid);
+		return GetCurveDataImpl<FTransformCurve>(TransformCurves, Name);
 #endif // WITH_EDITOR
 	case ERawCurveTrackTypes::RCT_Float:
 	default:
-		return GetCurveDataImpl<FFloatCurve>(FloatCurves, Uid);
+		return GetCurveDataImpl<FFloatCurve>(FloatCurves, Name);
 	}
 }
 
-bool FRawCurveTracks::DeleteCurveData(const FSmartName& CurveToDelete, ERawCurveTrackTypes SupportedCurveType /*= FloatType*/)
+bool FRawCurveTracks::DeleteCurveData(const FName& CurveToDelete, ERawCurveTrackTypes SupportedCurveType /*= FloatType*/)
 {
 	switch(SupportedCurveType)
 	{
@@ -522,13 +528,13 @@ void FRawCurveTracks::DeleteAllCurveData(ERawCurveTrackTypes SupportedCurveType
 }
 
 #if WITH_EDITOR
-void FRawCurveTracks::AddFloatCurveKey(const FSmartName& NewCurve, int32 CurveFlags, float Time, float Value)
+void FRawCurveTracks::AddFloatCurveKey(const FName& NewCurve, int32 CurveFlags, float Time, float Value)
 {
-	FFloatCurve* FloatCurve = GetCurveDataImpl<FFloatCurve>(FloatCurves, NewCurve.UID);
+	FFloatCurve* FloatCurve = GetCurveDataImpl<FFloatCurve>(FloatCurves, NewCurve);
 	if (FloatCurve == nullptr)
 	{
 		AddCurveData(NewCurve, CurveFlags, ERawCurveTrackTypes::RCT_Float);
-		FloatCurve = GetCurveDataImpl<FFloatCurve>(FloatCurves, NewCurve.UID);
+		FloatCurve = GetCurveDataImpl<FFloatCurve>(FloatCurves, NewCurve);
 	}
 
 	if (FloatCurve->GetCurveTypeFlags() != CurveFlags)
@@ -549,7 +555,7 @@ void FRawCurveTracks::RemoveRedundantKeys(float Tolerance /*= UE_SMALL_NUMBER*/,
 }
 #endif
 
-bool FRawCurveTracks::AddCurveData(const FSmartName& NewCurve, int32 CurveFlags /*= ACF_DefaultCurve*/, ERawCurveTrackTypes SupportedCurveType /*= FloatType*/)
+bool FRawCurveTracks::AddCurveData(const FName& NewCurve, int32 CurveFlags /*= ACF_DefaultCurve*/, ERawCurveTrackTypes SupportedCurveType /*= FloatType*/)
 {
 	switch(SupportedCurveType)
 	{
@@ -585,12 +591,12 @@ void FRawCurveTracks::Resize(float TotalLength, bool bInsert/* whether insert or
 #endif
 }
 
-void FRawCurveTracks::PostSerialize(FArchive& Ar)
+void FRawCurveTracks::PostSerializeFixup(FArchive& Ar)
 {
 	// @TODO: If we're about to serialize vector curve, add here
 	for(FFloatCurve& Curve : FloatCurves)
 	{
-		Curve.PostSerialize(Ar);
+		Curve.PostSerializeFixup(Ar);
 	}
 #if WITH_EDITORONLY_DATA
 	if( !Ar.IsCooking() )
@@ -599,33 +605,14 @@ void FRawCurveTracks::PostSerialize(FArchive& Ar)
 		{
 			for( FTransformCurve& Curve : TransformCurves )
 			{
-				Curve.PostSerialize( Ar );
+				Curve.PostSerializeFixup( Ar );
 			}
-
 		}
 	}
 #endif // WITH_EDITORONLY_DATA
 }
 
-void FRawCurveTracks::RefreshName(const FSmartNameMapping* NameMapping, ERawCurveTrackTypes SupportedCurveType /*= FloatType*/)
-{
-	switch(SupportedCurveType)
-	{
-#if WITH_EDITOR
-	case ERawCurveTrackTypes::RCT_Vector:
-		UpdateLastObservedNamesImpl<FVectorCurve>(VectorCurves, NameMapping);
-		break;
-	case ERawCurveTrackTypes::RCT_Transform:
-		UpdateLastObservedNamesImpl<FTransformCurve>(TransformCurves, NameMapping);
-		break;
-#endif // WITH_EDITOR
-	case ERawCurveTrackTypes::RCT_Float:
-	default:
-		UpdateLastObservedNamesImpl<FFloatCurve>(FloatCurves, NameMapping);
-	}
-}
-
-bool FRawCurveTracks::DuplicateCurveData(const FSmartName& CurveToCopy, const FSmartName& NewCurve, ERawCurveTrackTypes SupportedCurveType /*= FloatType*/)
+bool FRawCurveTracks::DuplicateCurveData(const FName& CurveToCopy, const FName& NewCurve, ERawCurveTrackTypes SupportedCurveType /*= FloatType*/)
 {
 	switch(SupportedCurveType)
 	{
@@ -652,39 +639,39 @@ bool FRawCurveTracks::DuplicateCurveData(const FSmartName& CurveToCopy, const FS
 // but this has to be refactored once we'd like to move onto serialize
 ///////////////////////////////////
 template <typename DataType>
-DataType * FRawCurveTracks::GetCurveDataImpl(TArray<DataType> & Curves, SkeletonAnimCurveUID Uid)
+DataType * FRawCurveTracks::GetCurveDataImpl(TArray<DataType> & Curves, FName Name)
 {
-	for(DataType& Curve : Curves)
+	for (DataType& Curve : Curves)
 	{
-		if(Curve.Name.UID == Uid)
+		if (Curve.GetName() == Name)
 		{
 			return &Curve;
 		}
 	}
 
-	return NULL;
+	return nullptr;
 }
 
 template <typename DataType>
-const DataType * FRawCurveTracks::GetCurveDataImpl(const TArray<DataType> & Curves, SkeletonAnimCurveUID Uid) const
+const DataType * FRawCurveTracks::GetCurveDataImpl(const TArray<DataType> & Curves, FName Name) const
 {
 	for (const DataType& Curve : Curves)
 	{
-		if (Curve.Name.UID == Uid)
+		if (Curve.GetName() == Name)
 		{
 			return &Curve;
 		}
 	}
 
-	return NULL;
+	return nullptr;
 }
 
 template <typename DataType>
-bool FRawCurveTracks::DeleteCurveDataImpl(TArray<DataType> & Curves, const FSmartName& CurveToDelete)
+bool FRawCurveTracks::DeleteCurveDataImpl(TArray<DataType> & Curves, const FName& CurveToDelete)
 {
 	for(int32 Idx = 0; Idx < Curves.Num(); ++Idx)
 	{
-		if(Curves[Idx].Name.UID == CurveToDelete.UID)
+		if(Curves[Idx].GetName() == CurveToDelete)
 		{
 			Curves.RemoveAt(Idx);
 			return true;
@@ -695,9 +682,9 @@ bool FRawCurveTracks::DeleteCurveDataImpl(TArray<DataType> & Curves, const FSmar
 }
 
 template <typename DataType>
-bool FRawCurveTracks::AddCurveDataImpl(TArray<DataType> & Curves, const FSmartName& NewCurve, int32 CurveFlags)
+bool FRawCurveTracks::AddCurveDataImpl(TArray<DataType> & Curves, const FName& NewCurve, int32 CurveFlags)
 {
-	if(GetCurveDataImpl<DataType>(Curves, NewCurve.UID) == NULL)
+	if(GetCurveDataImpl<DataType>(Curves, NewCurve) == NULL)
 	{
 		Curves.Add(DataType(NewCurve, CurveFlags));
 		return true;
@@ -706,22 +693,10 @@ bool FRawCurveTracks::AddCurveDataImpl(TArray<DataType> & Curves, const FSmartNa
 }
 
 template <typename DataType>
-void FRawCurveTracks::UpdateLastObservedNamesImpl(TArray<DataType> & Curves, const FSmartNameMapping* NameMapping)
-{
-	if(NameMapping)
-	{
-		for(DataType& Curve : Curves)
-		{
-			NameMapping->GetName(Curve.Name.UID, Curve.Name.DisplayName);
-		}
-	}
-}
-
-template <typename DataType>
-bool FRawCurveTracks::DuplicateCurveDataImpl(TArray<DataType> & Curves, const FSmartName& CurveToCopy, const FSmartName& NewCurve)
+bool FRawCurveTracks::DuplicateCurveDataImpl(TArray<DataType> & Curves, const FName& CurveToCopy, const FName& NewCurve)
 {
-	DataType* ExistingCurve = GetCurveDataImpl<DataType>(Curves, CurveToCopy.UID);
-	if(ExistingCurve && GetCurveDataImpl<DataType>(Curves, NewCurve.UID) == NULL)
+	DataType* ExistingCurve = GetCurveDataImpl<DataType>(Curves, CurveToCopy);
+	if(ExistingCurve && GetCurveDataImpl<DataType>(Curves, NewCurve) == NULL)
 	{
 		// Add the curve to the track and set its data to the existing curve
 		Curves.Add(DataType(NewCurve, ExistingCurve->GetCurveTypeFlags()));
@@ -740,19 +715,12 @@ FArchive& operator<<(FArchive& Ar, FRawCurveTracks& D)
 	return Ar;
 }
 
-///////////////////////////////////////////////////////////////////////
-// FAnimCurveParam
-
-void FAnimCurveParam::Initialize(USkeleton* Skeleton)
+void FBlendedCurve::InitFrom(const FBoneContainer& InBoneContainer)
 {
-	// Initialize for curve UID
-	if (Name != NAME_None)
-	{
-		UID = Skeleton->GetUIDByName(USkeleton::AnimCurveMappingName, Name);
-	}
-	else
-	{
-		// invalidate current UID
-		UID = SmartName::MaxUID;
-	}
+	TBaseBlendedCurve<FAnimStackAllocator>::SetFilter(&InBoneContainer.GetCurveFilter());
 }
+
+void FBlendedHeapCurve::InitFrom(const FBoneContainer& InBoneContainer)
+{
+	TBaseBlendedCurve<FDefaultAllocator>::SetFilter(&InBoneContainer.GetCurveFilter()); 
+}
\ No newline at end of file
diff --git a/Engine/Source/Runtime/Engine/Private/Animation/AnimDataModel.cpp b/Engine/Source/Runtime/Engine/Private/Animation/AnimDataModel.cpp
--- a/Engine/Source/Runtime/Engine/Private/Animation/AnimDataModel.cpp
+++ b/Engine/Source/Runtime/Engine/Private/Animation/AnimDataModel.cpp
@@ -36,7 +36,7 @@ void UAnimDataModel::PostLoad()
 		{
 			FloatCurve.FloatCurve.AutoSetTangents();
 			FCurvePayload Payload;
-			Payload.Identifier = FAnimationCurveIdentifier(FloatCurve.Name.UID, ERawCurveTrackTypes::RCT_Float);
+			Payload.Identifier = FAnimationCurveIdentifier(FloatCurve.GetName(), ERawCurveTrackTypes::RCT_Float);
 			GetNotifier().Notify(EAnimDataModelNotifyType::CurveChanged, Payload);
 		}
 		GetNotifier().Notify(EAnimDataModelNotifyType::BracketClosed);
@@ -356,7 +356,7 @@ const FFloatCurve* UAnimDataModel::FindFloatCurve(const FAnimationCurveIdentifie
 	ensure(CurveIdentifier.CurveType == ERawCurveTrackTypes::RCT_Float);
 	for (const FFloatCurve& FloatCurve : CurveData.FloatCurves)
 	{
-		if (FloatCurve.Name == CurveIdentifier.InternalName || FloatCurve.Name.UID == CurveIdentifier.InternalName.UID)
+		if (FloatCurve.GetName() == CurveIdentifier.CurveName)
 		{
 			return &FloatCurve;
 		}
@@ -370,7 +370,7 @@ const FTransformCurve* UAnimDataModel::FindTransformCurve(const FAnimationCurveI
 	ensure(CurveIdentifier.CurveType == ERawCurveTrackTypes::RCT_Transform);
 	for (const FTransformCurve& TransformCurve : CurveData.TransformCurves)
 	{
-		if (TransformCurve.Name == CurveIdentifier.InternalName || TransformCurve.Name.UID == CurveIdentifier.InternalName.UID)
+		if (TransformCurve.GetName() == CurveIdentifier.CurveName)
 		{
 			return &TransformCurve;
 		}
@@ -519,6 +519,11 @@ FGuid UAnimDataModel::GenerateGuid() const
 		Sha.Update((uint8*)(&Data), sizeof(Data));
 	};
 
+	auto UpdateWithString = [&](const FString& Data)
+	{
+		Sha.UpdateWithString(*Data, Data.Len());
+	};
+	
 	for (const FBoneAnimationTrack& Track : BoneAnimationTracks)
 	{
 		UpdateSHAWithArray(Track.InternalTrackData.PosKeys);
@@ -536,14 +541,14 @@ FGuid UAnimDataModel::GenerateGuid() const
 
 	for (const FFloatCurve& Curve : CurveData.FloatCurves)
 	{
-		const FString CurveName = Curve.Name.DisplayName.ToString();
+		const FString CurveName = Curve.GetName().ToString();
 		UpdateSHAWithArray(CurveName.GetCharArray()),
 		UpdateWithFloatCurve(Curve.FloatCurve);
 	}
 
     for (const FTransformCurve& Curve : CurveData.TransformCurves)
     {
-    	UpdateWithData(Curve.Name.UID);
+    	UpdateWithString(Curve.GetName().ToString());
 
     	auto UpdateWithComponent = [&UpdateWithFloatCurve](const FVectorCurve& VectorCurve)
     	{
@@ -763,19 +768,7 @@ void UAnimDataModel::Evaluate(FAnimationPoseData& InOutPoseData, const UE::Anim:
 	}
 
 	// Evaluate animation float curve data
-	FBlendedCurve& Curves = InOutPoseData.GetCurve();
-	if (Curves.NumValidCurveCount > 0)
-	{
-		// evaluate the curve data at the CurrentTime and add to Instance
-		for (const FFloatCurve& Curve : CurveData.FloatCurves)
-		{
-			if (Curves.IsEnabled(Curve.Name.UID))
-			{
-				float Value = Curve.Evaluate(Time);
-				Curves.Set(Curve.Name.UID, Value);
-			}
-		}
-	}
+	UE::Anim::EvaluateFloatCurvesFromModel(this, InOutPoseData.GetCurve(), Time);
 
 	const double TimePerFrame = FrameRate.AsInterval();
 	TMap<FName, const FTransformCurve*> ActiveCurves;
@@ -790,11 +783,8 @@ void UAnimDataModel::Evaluate(FAnimationPoseData& InOutPoseData, const UE::Anim:
 				continue;
 			}
 
-			// Add or retrieve curve
-			const FName& CurveName = Curve.Name.DisplayName;
-
 			// note we're not checking Curve.GetCurveTypeFlags() yet
-			ActiveCurves.FindOrAdd(CurveName, &Curve);
+			ActiveCurves.FindOrAdd(Curve.GetName(), &Curve);
 		}
 	}
 
@@ -906,7 +896,7 @@ FTransformCurve* UAnimDataModel::FindMutableTransformCurveById(const FAnimationC
 {
 	for (FTransformCurve& TransformCurve : CurveData.TransformCurves)
 	{
-		if (TransformCurve.Name.UID == CurveIdentifier.InternalName.UID)
+		if (TransformCurve.GetName() == CurveIdentifier.CurveName)
 		{
 			return &TransformCurve;
 		}
@@ -919,7 +909,7 @@ FFloatCurve* UAnimDataModel::FindMutableFloatCurveById(const FAnimationCurveIden
 {
 	for (FFloatCurve& FloatCurve : CurveData.FloatCurves)
 	{
-		if (FloatCurve.Name.UID == CurveIdentifier.InternalName.UID)
+		if (FloatCurve.GetName() == CurveIdentifier.CurveName)
 		{
 			return &FloatCurve;
 		}
diff --git a/Engine/Source/Runtime/Engine/Private/Animation/AnimInstance.cpp b/Engine/Source/Runtime/Engine/Private/Animation/AnimInstance.cpp
--- a/Engine/Source/Runtime/Engine/Private/Animation/AnimInstance.cpp
+++ b/Engine/Source/Runtime/Engine/Private/Animation/AnimInstance.cpp
@@ -822,6 +822,7 @@ void UAnimInstance::ParallelEvaluateAnimation(bool bForceRefPose, const USkeleta
 			
 		// Run the anim blueprint
 		Proxy.EvaluateAnimation(EvaluationContext);
+
 		// Move the curves
 		OutEvaluationData.OutCurve.CopyFrom(EvaluationContext.Curve);
 		OutEvaluationData.OutPose.CopyBonesFrom(EvaluationContext.Pose);
@@ -1304,9 +1305,16 @@ void UAnimInstance::RecalcRequiredBones()
 	}
 }
 
+void UAnimInstance::RecalcRequiredCurves(const UE::Anim::FCurveFilterSettings& InCurveFilterSettings)
+{
+	GetProxyOnGameThread<FAnimInstanceProxy>().RecalcRequiredCurves(InCurveFilterSettings);
+}
+
 void UAnimInstance::RecalcRequiredCurves(const FCurveEvaluationOption& CurveEvalOption)
 {
+	PRAGMA_DISABLE_DEPRECATION_WARNINGS
 	GetProxyOnGameThread<FAnimInstanceProxy>().RecalcRequiredCurves(CurveEvalOption);
+	PRAGMA_ENABLE_DEPRECATION_WARNINGS
 }
 
 USkeletalMeshComponent* UAnimInstance::GetSkelMeshComponent() const
@@ -1364,30 +1372,9 @@ void UAnimInstance::PostInitProperties()
 	}
 }
 
-void UAnimInstance::AddCurveValue(const FName& CurveName, float Value)
-{
-	const FSmartNameMapping* Mapping = CurrentSkeleton->GetSmartNameContainer(USkeleton::AnimCurveMappingName);
-	if(Mapping)
-	{
-		AddCurveValue(*Mapping, CurveName, Value);
-	}
-}
-
-void UAnimInstance::AddCurveValue(const FSmartNameMapping& Mapping, const FName& CurveName, float Value)
-{
-	GetProxyOnAnyThread<FAnimInstanceProxy>().AddCurveValue(Mapping, CurveName, Value);
-}
-
-void UAnimInstance::AddCurveValue(const SmartName::UID_Type Uid, float Value)
+void UAnimInstance::AddCurveValue(const FName& CurveName, float Value, bool bMorphtarget, bool bMaterial)
 {
-	FName CurrentCurveName;
-	// Grab the smartname mapping from our current skeleton and resolve the curve name. We cannot cache
-	// the smart name mapping as the skeleton can change at any time.
-	if(const FSmartNameMapping* NameMapping = CurrentSkeleton->GetSmartNameContainer(USkeleton::AnimCurveMappingName))
-	{
-		NameMapping->GetName(Uid, CurrentCurveName);
-	}
-	AddCurveValue(CurrentCurveName, Value);
+	GetProxyOnAnyThread<FAnimInstanceProxy>().AddCurveValue(CurveName, Value, bMorphtarget, bMaterial);
 }
 
 void UAnimInstance::UpdateCurvesToComponents(USkeletalMeshComponent* Component /*= nullptr*/)
@@ -1752,17 +1739,7 @@ void UAnimInstance::GetActiveCurveNames(EAnimCurveType CurveType, TArray<FName>&
 
 void UAnimInstance::GetAllCurveNames(TArray<FName>& OutNames) const
 {
-	USkeletalMeshComponent* SkelMeshComp = GetOwningComponent();
-	if (SkelMeshComp && SkelMeshComp->GetSkeletalMeshAsset() && SkelMeshComp->GetSkeletalMeshAsset()->GetSkeleton())
-	{
-		const USkeleton* CurSkeleton = SkelMeshComp->GetSkeletalMeshAsset()->GetSkeleton();
-
-		const FSmartNameMapping* Mapping = CurSkeleton->GetSmartNameContainer(USkeleton::AnimCurveMappingName);
-		if (Mapping)
-		{
-			Mapping->FillNameArray(OutNames);
-		}
-	}
+	GetActiveCurveNames(EAnimCurveType::AttributeCurve, OutNames);
 }
 
 void UAnimInstance::SetRootMotionMode(TEnumAsByte<ERootMotionMode::Type> Value)
diff --git a/Engine/Source/Runtime/Engine/Private/Animation/AnimInstanceProxy.cpp b/Engine/Source/Runtime/Engine/Private/Animation/AnimInstanceProxy.cpp
--- a/Engine/Source/Runtime/Engine/Private/Animation/AnimInstanceProxy.cpp
+++ b/Engine/Source/Runtime/Engine/Private/Animation/AnimInstanceProxy.cpp
@@ -593,7 +593,7 @@ void FAnimInstanceProxy::PostUpdate(UAnimInstance* InAnimInstance) const
 			{
 				if (!PoseWatch.PoseWatchPoseElement->GetIsEnabled())
 				{
-					TRACE_ANIM_POSE_WATCH(*this, PoseWatch.NodeID, TArray<FTransform>(), TArray<uint16>(), FTransform::Identity, false);
+					TRACE_ANIM_POSE_WATCH(*this, PoseWatch.PoseWatchPoseElement, PoseWatch.NodeID, TArray<FTransform>(), FBlendedHeapCurve(), TArray<uint16>(), FTransform::Identity, false);
 				}
 			}
 		}
@@ -1116,7 +1116,7 @@ void FAnimInstanceProxy::RecalcRequiredBones(USkeletalMeshComponent* Component,
 {
 	DECLARE_SCOPE_HIERARCHICAL_COUNTER_FUNC()
 
-	RequiredBones.InitializeTo(Component->RequiredBones, FCurveEvaluationOption(Component->GetAllowedAnimCurveEvaluate(), &Component->GetDisallowedAnimCurvesEvaluation(), Component->GetPredictedLODLevel()), *Asset);
+	RequiredBones.InitializeTo(Component->RequiredBones, Component->GetCurveFilterSettings(), *Asset);
 
 	// If there is a ref pose override, we want to replace ref pose in RequiredBones
 	// Update ref pose in required bones structure (either set it, or clear it, depending on if one is set on the Component)
@@ -1132,11 +1132,21 @@ void FAnimInstanceProxy::RecalcRequiredBones(USkeletalMeshComponent* Component,
 	bBoneCachesInvalidated = true;
 }
 
+void FAnimInstanceProxy::RecalcRequiredCurves(const UE::Anim::FCurveFilterSettings& CurveFilterSettings)
+{
+	DECLARE_SCOPE_HIERARCHICAL_COUNTER_FUNC()
+
+	RequiredBones.CacheRequiredAnimCurves(CurveFilterSettings);
+	bBoneCachesInvalidated = true;
+}
+
 void FAnimInstanceProxy::RecalcRequiredCurves(const FCurveEvaluationOption& CurveEvalOption)
 {
 	DECLARE_SCOPE_HIERARCHICAL_COUNTER_FUNC()
 
+	PRAGMA_DISABLE_DEPRECATION_WARNINGS
 	RequiredBones.CacheRequiredAnimCurveUids(CurveEvalOption);
+	PRAGMA_ENABLE_DEPRECATION_WARNINGS
 	bBoneCachesInvalidated = true;
 }
 
@@ -3260,6 +3270,11 @@ void FAnimInstanceProxy::RecordNodeAttribute(const FAnimInstanceProxy& InSourceP
 }
 
 void FAnimInstanceProxy::RegisterWatchedPose(const FCompactPose& Pose, int32 LinkID)
+{
+	RegisterWatchedPose(Pose, FBlendedCurve(), LinkID);
+}
+
+void FAnimInstanceProxy::RegisterWatchedPose(const FCompactPose& Pose, const FBlendedCurve& InCurve, int32 LinkID)
 {
 	if (bIsBeingDebugged)
 	{
@@ -3309,9 +3324,10 @@ void FAnimInstanceProxy::RegisterWatchedPose(const FCompactPose& Pose, int32 Lin
 						}
 
 						PoseWatch.SetPose(TmpRequiredBones, BoneTransforms);
+						PoseWatch.SetCurves(InCurve);
 						PoseWatch.SetWorldTransform(SkelMeshComponent->GetComponentTransform());
 
-						TRACE_ANIM_POSE_WATCH(*this, PoseWatch.NodeID, PoseWatch.GetBoneTransforms(), PoseWatch.GetRequiredBones(), PoseWatch.GetWorldTransform(), true);
+						TRACE_ANIM_POSE_WATCH(*this, PoseWatch.PoseWatchPoseElement, PoseWatch.NodeID, PoseWatch.GetBoneTransforms(), PoseWatch.GetCurves(), PoseWatch.GetRequiredBones(), PoseWatch.GetWorldTransform(), PoseWatch.PoseWatchPoseElement->GetIsVisible());
 						break;
 					}
 				}
@@ -3321,6 +3337,11 @@ void FAnimInstanceProxy::RegisterWatchedPose(const FCompactPose& Pose, int32 Lin
 }
 
 void FAnimInstanceProxy::RegisterWatchedPose(const FCSPose<FCompactPose>& Pose, int32 LinkID)
+{
+	RegisterWatchedPose(Pose, FBlendedCurve(), LinkID);
+}
+
+void FAnimInstanceProxy::RegisterWatchedPose(const FCSPose<FCompactPose>& Pose, const FBlendedCurve& InCurve, int32 LinkID)
 {
 	if (bIsBeingDebugged)
 	{
@@ -3341,9 +3362,10 @@ void FAnimInstanceProxy::RegisterWatchedPose(const FCSPose<FCompactPose>& Pose,
 						const TArray<FTransform, FAnimStackAllocator>& BoneTransforms = TempPose.GetBones();
 						const TArray<FBoneIndexType>& TmpRequiredBones = TempPose.GetBoneContainer().GetBoneIndicesArray();
 						PoseWatch.SetPose(TmpRequiredBones, BoneTransforms);
+						PoseWatch.SetCurves(InCurve);
 						PoseWatch.SetWorldTransform(SkelMeshComponent->GetComponentTransform());
 
-						TRACE_ANIM_POSE_WATCH(*this, PoseWatch.NodeID, PoseWatch.GetBoneTransforms(), PoseWatch.GetRequiredBones(), PoseWatch.GetWorldTransform(), true);
+						TRACE_ANIM_POSE_WATCH(*this, PoseWatch.PoseWatchPoseElement, PoseWatch.NodeID, PoseWatch.GetBoneTransforms(), PoseWatch.GetCurves(), PoseWatch.GetRequiredBones(), PoseWatch.GetWorldTransform(), PoseWatch.PoseWatchPoseElement->GetIsVisible());
 						break;
 					}
 				}
@@ -3407,50 +3429,25 @@ void FAnimInstanceProxy::UpdateCurvesToEvaluationContext(const FAnimationEvaluat
 
 	ResetAnimationCurves();
 
-	if(InContext.Curve.UIDToArrayIndexLUT != nullptr && InContext.Curve.UIDToArrayIndexLUT->Num() > 0)
+	InContext.Curve.ForEachElement([this](const UE::Anim::FCurveElement& InCurveElement)
 	{
-		const TArray<uint16>& UIDToArrayIndexLookupTable = RequiredBones.GetUIDToArrayLookupTable();
-		const FSmartNameMapping* Mapping = RequiredBones.GetSkeletonAsset()->GetSmartNameContainer(USkeleton::AnimCurveMappingName);
-		
-		check(UIDToArrayIndexLookupTable.Num() == InContext.Curve.UIDToArrayIndexLUT->Num());
+		AnimationCurves[(uint8)EAnimCurveType::AttributeCurve].Add(InCurveElement.Name, InCurveElement.Value);
+	});
 
-		for (int32 CurveUID = 0; CurveUID < UIDToArrayIndexLookupTable.Num(); ++CurveUID)
+	UE::Anim::FNamedValueArrayUtils::Intersection(InContext.Curve, RequiredBones.GetCurveFlags(), 
+		[this](const UE::Anim::FCurveElement& InCurveElement, const UE::Anim::FCurveElementFlags& InCurveFlagsElement)
 		{
-			uint16 ArrayIndex = UIDToArrayIndexLookupTable[CurveUID];
-		
-			if (ArrayIndex != MAX_uint16 && 
-				ensureAlwaysMsgf(InContext.Curve.CurveWeights.IsValidIndex(ArrayIndex), TEXT("%s Animation Instance contains out of bound UIDList."), *AnimInstanceObject->GetClass()->GetName()) &&
-				InContext.Curve.ValidCurveWeights[ArrayIndex])
+			if(EnumHasAnyFlags(InCurveFlagsElement.Flags | InCurveElement.Flags, UE::Anim::ECurveElementFlags::MorphTarget))
 			{
-				FName CurveName;
-				if (Mapping->GetName(CurveUID, CurveName))
-				{
-#if !WITH_EDITOR
-					const FCurveMetaData* CurveMetaData = &Mapping->GetCurveMetaData(CurveUID);
-#else
-					if (const FCurveMetaData* CurveMetaData = Mapping->GetCurveMetaData(CurveName))
-#endif
-					{
-						const FAnimCurveType& CurveType = CurveMetaData->Type;
-						const float CurveWeight = InContext.Curve.CurveWeights[ArrayIndex];
-					
-						AnimationCurves[(uint8)EAnimCurveType::AttributeCurve].Add(CurveName, CurveWeight);
-
-						if(CurveType.bMorphtarget)
-						{
-							AnimationCurves[(uint8)EAnimCurveType::MorphTargetCurve].Add(CurveName, CurveWeight);
-						}
+				AnimationCurves[(uint8)EAnimCurveType::MorphTargetCurve].Add(InCurveElement.Name, InCurveElement.Value);
+			}
 
-						if(CurveType.bMaterial)
-						{
-							MaterialParametersToClear.Remove(CurveName);
-							AnimationCurves[(uint8)EAnimCurveType::MaterialCurve].Add(CurveName, CurveWeight);
-						}
-					}	
-				}
+			if(EnumHasAnyFlags(InCurveFlagsElement.Flags | InCurveElement.Flags, UE::Anim::ECurveElementFlags::Material))
+			{
+				MaterialParametersToClear.Remove(InCurveElement.Name);
+				AnimationCurves[(uint8)EAnimCurveType::MaterialCurve].Add(InCurveElement.Name, InCurveElement.Value);
 			}
-		}
-	}
+		});
 }
 
 void FAnimInstanceProxy::UpdateCurvesPostEvaluation(USkeletalMeshComponent* SkelMeshComp)
@@ -3488,7 +3485,7 @@ bool FAnimInstanceProxy::HasActiveCurves() const
 	return false;
 }
 
-void FAnimInstanceProxy::AddCurveValue(const FSmartNameMapping& Mapping, const FName& CurveName, float Value)
+void FAnimInstanceProxy::AddCurveValue(const FName& CurveName, float Value, bool bMorphtarget, bool bMaterial)
 {
 	DECLARE_SCOPE_HIERARCHICAL_COUNTER_FUNC()
 
@@ -3506,35 +3503,31 @@ void FAnimInstanceProxy::AddCurveValue(const FSmartNameMapping& Mapping, const F
 		AnimationCurves[(uint8)EAnimCurveType::AttributeCurve].Add(CurveName, Value);
 	}
 
-	const FCurveMetaData* CurveMetaData = Mapping.GetCurveMetaData(CurveName);
-	if (CurveMetaData)
+	if (bMorphtarget)
 	{
-		if (CurveMetaData->Type.bMorphtarget)
+		CurveValPtr = AnimationCurves[(uint8)EAnimCurveType::MorphTargetCurve].Find(CurveName);
+		if (CurveValPtr)
 		{
-			CurveValPtr = AnimationCurves[(uint8)EAnimCurveType::MorphTargetCurve].Find(CurveName);
-			if (CurveValPtr)
-			{
-				// sum up, in the future we might normalize, but for now this just sums up
-				// this won't work well if all of them have full weight - i.e. additive 
-				*CurveValPtr = Value;
-			}
-			else
-			{
-				AnimationCurves[(uint8)EAnimCurveType::MorphTargetCurve].Add(CurveName, Value);
-			}
+			// sum up, in the future we might normalize, but for now this just sums up
+			// this won't work well if all of them have full weight - i.e. additive 
+			*CurveValPtr = Value;
 		}
-		if (CurveMetaData->Type.bMaterial)
+		else
 		{
-			MaterialParametersToClear.RemoveSwap(CurveName);
-			CurveValPtr = AnimationCurves[(uint8)EAnimCurveType::MaterialCurve].Find(CurveName);
-			if (CurveValPtr)
-			{
-				*CurveValPtr = Value;
-			}
-			else
-			{
-				AnimationCurves[(uint8)EAnimCurveType::MaterialCurve].Add(CurveName, Value);
-			}
+			AnimationCurves[(uint8)EAnimCurveType::MorphTargetCurve].Add(CurveName, Value);
+		}
+	}
+	if (bMaterial)
+	{
+		MaterialParametersToClear.RemoveSwap(CurveName);
+		CurveValPtr = AnimationCurves[(uint8)EAnimCurveType::MaterialCurve].Find(CurveName);
+		if (CurveValPtr)
+		{
+			*CurveValPtr = Value;
+		}
+		else
+		{
+			AnimationCurves[(uint8)EAnimCurveType::MaterialCurve].Add(CurveName, Value);
 		}
 	}
 }
diff --git a/Engine/Source/Runtime/Engine/Private/Animation/AnimMontage.cpp b/Engine/Source/Runtime/Engine/Private/Animation/AnimMontage.cpp
--- a/Engine/Source/Runtime/Engine/Private/Animation/AnimMontage.cpp
+++ b/Engine/Source/Runtime/Engine/Private/Animation/AnimMontage.cpp
@@ -2944,17 +2944,7 @@ void UAnimMontage::BakeTimeStretchCurve()
 	const FFloatCurve* TimeStretchFloatCurve = nullptr;
 	if (ShouldDataModelBeValid())
 	{		
-		if (const USkeleton* MySkeleton = GetSkeleton())
-		{
-			if (const FSmartNameMapping* CurveNameMapping = MySkeleton->GetSmartNameContainer(USkeleton::AnimCurveMappingName))
-			{
-				const USkeleton::AnimCurveUID CurveUID = CurveNameMapping->FindUID(TimeStretchCurveName);
-				if (CurveUID != SmartName::MaxUID)
-				{
-					TimeStretchFloatCurve = GetDataModel()->FindFloatCurve(FAnimationCurveIdentifier(CurveUID, ERawCurveTrackTypes::RCT_Float));
-				}
-			}
-		}
+		TimeStretchFloatCurve = GetDataModel()->FindFloatCurve(FAnimationCurveIdentifier(TimeStretchCurveName, ERawCurveTrackTypes::RCT_Float));
 	}
 
 	if (TimeStretchFloatCurve == nullptr)
diff --git a/Engine/Source/Runtime/Engine/Private/Animation/AnimNodeBase.cpp b/Engine/Source/Runtime/Engine/Private/Animation/AnimNodeBase.cpp
--- a/Engine/Source/Runtime/Engine/Private/Animation/AnimNodeBase.cpp
+++ b/Engine/Source/Runtime/Engine/Private/Animation/AnimNodeBase.cpp
@@ -412,7 +412,7 @@ void FPoseLink::Evaluate(FPoseContext& Output)
 					}
 				}
 			}
-			if(Output.Curve.NumValid() > 0)
+			if(Output.Curve.Num() > 0)
 			{
 				Output.AnimInstanceProxy->RecordNodeAttribute(*Output.AnimInstanceProxy, SourceID, LinkID, UE::Anim::FAttributes::Curves);
 			}
@@ -424,7 +424,7 @@ void FPoseLink::Evaluate(FPoseContext& Output)
 #endif
 
 #if WITH_EDITOR
-		Output.AnimInstanceProxy->RegisterWatchedPose(Output.Pose, LinkID);
+		Output.AnimInstanceProxy->RegisterWatchedPose(Output.Pose, Output.Curve, LinkID);
 #endif
 	}
 	else
@@ -484,7 +484,7 @@ void FComponentSpacePoseLink::EvaluateComponentSpace(FComponentSpacePoseContext&
 					}
 				}
 			}
-			if(Output.Curve.NumValid() > 0)
+			if(Output.Curve.Num() > 0)
 			{
 				Output.AnimInstanceProxy->RecordNodeAttribute(*Output.AnimInstanceProxy, SourceID, LinkID, UE::Anim::FAttributes::Curves);
 			}
@@ -492,7 +492,7 @@ void FComponentSpacePoseLink::EvaluateComponentSpace(FComponentSpacePoseContext&
 #endif
 
 #if WITH_EDITOR
-		Output.AnimInstanceProxy->RegisterWatchedPose(Output.Pose, LinkID);
+		Output.AnimInstanceProxy->RegisterWatchedPose(Output.Pose, Output.Curve, LinkID);
 #endif
 	}
 	else
diff --git a/Engine/Source/Runtime/Engine/Private/Animation/AnimNode_Inertialization.cpp b/Engine/Source/Runtime/Engine/Private/Animation/AnimNode_Inertialization.cpp
--- a/Engine/Source/Runtime/Engine/Private/Animation/AnimNode_Inertialization.cpp
+++ b/Engine/Source/Runtime/Engine/Private/Animation/AnimNode_Inertialization.cpp
@@ -8,6 +8,7 @@
 #include "Components/SkeletalMeshComponent.h"
 #include "Animation/AnimBlueprintGeneratedClass.h"
 #include "Logging/TokenizedMessage.h"
+#include "Animation/AnimCurveUtils.h"
 
 #include UE_INLINE_GENERATED_CPP_BY_NAME(AnimNode_Inertialization)
 
@@ -27,6 +28,9 @@ TAutoConsoleVariable<int32> CVarAnimInertializationIgnoreDeficit(TEXT("a.AnimNod
 static constexpr int32 INERTIALIZATION_MAX_POSE_SNAPSHOTS = 2;
 static constexpr float INERTIALIZATION_TIME_EPSILON = 1.0e-7f;
 
+// DEPRECATED: Used for backwards compatibility
+float FInertializationCurveDiffElement::Delta = 0.0f;
+
 namespace UE { namespace Anim {
 
 // Inertialization request event bound to a node
@@ -133,19 +137,6 @@ void FAnimNode_Inertialization::Initialize_AnyThread(const FAnimationInitializeC
 	Deactivate();
 
 	InertializationPoseDiff.Reset();
-	CachedFilteredCurvesUIDs.Reset();
-
-	const USkeleton* Skeleton = Context.AnimInstanceProxy->GetSkeleton();
-	check(Skeleton);
-	for (const FName& CurveName : FilteredCurves)
-	{
-		SmartName::UID_Type NameUID = Skeleton->GetUIDByName(USkeleton::AnimCurveMappingName, CurveName);
-		if (NameUID != SmartName::MaxUID)
-		{
-			// Grab UIDs of filtered curves to avoid lookup later
-			CachedFilteredCurvesUIDs.Add(NameUID);
-		}
-	}
 }
 
 
@@ -409,10 +400,6 @@ void FAnimNode_Inertialization::Evaluate_AnyThread(FPoseContext& Output)
 		for (int32 i = 0; i < INERTIALIZATION_MAX_POSE_SNAPSHOTS - 1; ++i)
 		{
 			Swap(PoseSnapshots[i], PoseSnapshots[i+1]);
-
-			// Verify the swap operation used move assignment to properly fix up curve LUT pointers
-			checkSlow(PoseSnapshots[i].Curves.BlendedCurve.UIDToArrayIndexLUT == &PoseSnapshots[i].Curves.CurveUIDToArrayIndexLUT);
-			checkSlow(PoseSnapshots[i+1].Curves.BlendedCurve.UIDToArrayIndexLUT == &PoseSnapshots[i+1].Curves.CurveUIDToArrayIndexLUT);
 		}
 
 		// Overwrite the (now irrelevant) pose in the last slot with the new post snapshot
@@ -489,7 +476,15 @@ void FAnimNode_Inertialization::StartInertialization(FPoseContext& Context, FIne
 		}
 	}
 
-	OutPoseDiff.InitFrom(Context.Pose, Context.Curve, Context.AnimInstanceProxy->GetComponentTransform(), AttachParentName, PreviousPose1, PreviousPose2, CachedFilteredCurvesUIDs);
+	// Initialize curve filter if required 
+	if(FilteredCurves.Num() != CurveFilter.Num())
+	{
+		CurveFilter.Empty();
+		CurveFilter.SetFilterMode(UE::Anim::ECurveFilterMode::DisallowFiltered);
+		CurveFilter.AppendNames(FilteredCurves);
+	}
+	
+	OutPoseDiff.InitFrom(Context.Pose, Context.Curve, Context.AnimInstanceProxy->GetComponentTransform(), AttachParentName, PreviousPose1, PreviousPose2, CurveFilter);
 }
 
 
@@ -544,7 +539,7 @@ void FInertializationPose::InitFrom(const FCompactPose& Pose, const FBlendedCurv
 // Prev2		the pose and curves from two frames before
 // DeltaTime	the time elapsed between Prev1 and Pose
 //
-void FInertializationPoseDiff::InitFrom(const FCompactPose& Pose, const FBlendedCurve& Curves, const FTransform& ComponentTransform, const FName& AttachParentName, const FInertializationPose& Prev1, const FInertializationPose& Prev2, const TSet<SmartName::UID_Type>& FilteredCurvesUIDs)
+void FInertializationPoseDiff::InitFrom(const FCompactPose& Pose, const FBlendedCurve& Curves, const FTransform& ComponentTransform, const FName& AttachParentName, const FInertializationPose& Prev1, const FInertializationPose& Prev2, const UE::Anim::FCurveFilter& CurveFilter)
 {
 	const FBoneContainer& BoneContainer = Pose.GetBoneContainer();
 
@@ -713,40 +708,31 @@ void FInertializationPoseDiff::InitFrom(const FCompactPose& Pose, const FBlended
 	}
 
 	// Compute the curve differences
-	const int32 CurveNum = Curves.IsValid() ? Curves.UIDToArrayIndexLUT->Num() : 0;
-	CurveDiffs.Empty(CurveNum);
-	CurveDiffs.AddZeroed(CurveNum);
-	for (int32 CurveUID = 0; CurveUID != CurveNum; ++CurveUID)
+	// First copy in current values
+	CurveDiffs.CopyFrom(Curves);
+
+	// Apply filtering to diffs to remove anything we dont want to inertialize
+	if(CurveFilter.Num() > 0)
 	{
-		const int32 CurrIdx = Curves.GetArrayIndexByUID(CurveUID);
-		const int32 Prev1Idx = Prev1.Curves.BlendedCurve.GetArrayIndexByUID(CurveUID);
-		if (CurrIdx == INDEX_NONE || Prev1Idx == INDEX_NONE)
-		{
-			// CurveDiff is zeroed
-			continue;
-		}
+		UE::Anim::FCurveUtils::Filter(CurveDiffs, CurveFilter);
+	}
 
-		// Check if the curve is in our filter set. Leave CurveDiff zeroed if it is.
-		if (FilteredCurvesUIDs.Contains((SmartName::UID_Type)CurveUID))
+	// Compute differences & derivatives
+	UE::Anim::FNamedValueArrayUtils::Union(CurveDiffs, Prev1.Curves.BlendedCurve, Prev2.Curves.BlendedCurve,
+		[DeltaTime = Prev1.DeltaTime](FInertializationCurveDiffElement& OutResultElement, const UE::Anim::FCurveElement& InElement0, const UE::Anim::FCurveElement& InElement1, UE::Anim::ENamedValueUnionFlags InFlags)
 		{
-			continue;
-		}
+			const float CurrWeight = OutResultElement.Value;
+			const float Prev1Weight = InElement0.Value;
 
-		const float CurrWeight = Curves.CurveWeights[CurrIdx];
-		const float Prev1Weight = Prev1.Curves.BlendedCurve.CurveWeights[Prev1Idx];
-		FInertializationCurveDiff& CurveDiff = CurveDiffs[CurveUID];
+			// Store delta in curve value
+			OutResultElement.Value = Prev1Weight - CurrWeight;
 
-		// Note we intentionally ignore FBlendedCurve::ValidCurveWeights. We want to ease in/out when only one
-		// curve is valid, and we'll compute a zero delta and derivative when both are invalid.
-		CurveDiff.Delta = Prev1Weight - CurrWeight;
-
-		const int32 Prev2Idx = Prev2.Curves.BlendedCurve.GetArrayIndexByUID(CurveUID);
-		if (Prev2Idx != INDEX_NONE && Prev1.DeltaTime > UE_KINDA_SMALL_NUMBER)
-		{
-			const float Prev2Weight = Prev2.Curves.BlendedCurve.CurveWeights[Prev2Idx];
-			CurveDiff.Derivative = (Prev1Weight - Prev2Weight) / Prev1.DeltaTime;
-		}
-	}
+			if(DeltaTime > UE_KINDA_SMALL_NUMBER)
+			{
+				const float Prev2Weight = InElement1.Value;
+				OutResultElement.Derivative = (Prev1Weight - Prev2Weight) / DeltaTime;
+			}
+		});
 }
 
 
@@ -786,27 +772,12 @@ void FInertializationPoseDiff::ApplyTo(FCompactPose& Pose, FBlendedCurve& Curves
 	Pose.NormalizeRotations();
 
 	// Apply curve differences
-	if (Curves.IsValid())
-	{
-		// Note Curves.Elements is indexed indirectly via lookup table while CurveDiffs is indexed directly by curve ID
-		const int32 CurveNum = FMath::Min(Curves.UIDToArrayIndexLUT->Num(), CurveDiffs.Num());
-		for (int32 CurveUID = 0; CurveUID != CurveNum; ++CurveUID)
+	UE::Anim::FNamedValueArrayUtils::Union(Curves, CurveDiffs,
+		[&InertializationElapsedTime, &InertializationDuration](UE::Anim::FCurveElement& OutResultElement, const FInertializationCurveDiffElement& InParamElement, UE::Anim::ENamedValueUnionFlags InFlags)
 		{
-			const int32 CurrIdx = Curves.GetArrayIndexByUID(CurveUID);
-			if (CurrIdx == INDEX_NONE)
-			{
-				continue;
-			}
-
-			const FInertializationCurveDiff& CurveDiff = CurveDiffs[CurveUID];
-			const float C = CalcInertialFloat(CurveDiff.Delta, CurveDiff.Derivative, InertializationElapsedTime, InertializationDuration);
-			if (C != 0.0f)
-			{
-				Curves.CurveWeights[CurrIdx] += C;
-				Curves.ValidCurveWeights[CurrIdx] = true;
-			}
-		}
-	}
+			OutResultElement.Value = CalcInertialFloat(InParamElement.Value, InParamElement.Derivative, InertializationElapsedTime, InertializationDuration);
+			OutResultElement.Flags |= InParamElement.Flags;
+		});
 }
 
 
diff --git a/Engine/Source/Runtime/Engine/Private/Animation/AnimNode_LinkedInputPose.cpp b/Engine/Source/Runtime/Engine/Private/Animation/AnimNode_LinkedInputPose.cpp
--- a/Engine/Source/Runtime/Engine/Private/Animation/AnimNode_LinkedInputPose.cpp
+++ b/Engine/Source/Runtime/Engine/Private/Animation/AnimNode_LinkedInputPose.cpp
@@ -50,7 +50,7 @@ void FAnimNode_LinkedInputPose::Evaluate_AnyThread(FPoseContext& Output)
 		Output.Curve.MoveFrom(InputContext.Curve);
 		Output.CustomAttributes.MoveFrom(InputContext.CustomAttributes);
 	}
-	else if(CachedInputPose.IsValid() && CachedInputCurve.IsValid() && ensure(Output.Pose.GetNumBones() == CachedInputPose.GetNumBones()))
+	else if(CachedInputPose.IsValid() && ensure(Output.Pose.GetNumBones() == CachedInputPose.GetNumBones()))
 	{
 		Output.Pose.CopyBonesFrom(CachedInputPose);
 		Output.Curve.CopyFrom(CachedInputCurve);
diff --git a/Engine/Source/Runtime/Engine/Private/Animation/AnimSequence.cpp b/Engine/Source/Runtime/Engine/Private/Animation/AnimSequence.cpp
--- a/Engine/Source/Runtime/Engine/Private/Animation/AnimSequence.cpp
+++ b/Engine/Source/Runtime/Engine/Private/Animation/AnimSequence.cpp
@@ -911,19 +911,12 @@ void UAnimSequence::PostLoad()
 
 	if (USkeleton* CurrentSkeleton = GetSkeleton())
 	{
-	#if WITH_EDITOR
-    		for (const FAnimSyncMarker& SyncMarker : AuthoredSyncMarkers)
-    		{
-    			CurrentSkeleton->RegisterMarkerName(SyncMarker.MarkerName);
-    		}
-    #endif
-    
-    #if !WITH_EDITOR
-		for (FSmartName& CurveName : CompressedData.CompressedCurveNames)
+#if WITH_EDITOR
+		for (const FAnimSyncMarker& SyncMarker : AuthoredSyncMarkers)
 		{
-			CurrentSkeleton->VerifySmartName(USkeleton::AnimCurveMappingName, CurveName);
+			CurrentSkeleton->RegisterMarkerName(SyncMarker.MarkerName);
 		}
-    #endif
+#endif
 	}
 }
 
@@ -1151,20 +1144,16 @@ void UAnimSequence::GetBoneTransform(FTransform& OutAtom, FSkeletonPoseBoneIndex
 		ValidateModel();
 		const FName BoneName = GetSkeleton()->GetReferenceSkeleton().GetBoneName(BoneIndex.GetInt());
 		OutAtom = DataModelInterface->EvaluateBoneTrackTransform(BoneName, DataModelInterface->GetFrameRate().AsFrameTime(Time), Interpolation);
-		if (const USkeleton* CurrentSkeleton = GetSkeleton())
+
+		const FAnimationCurveIdentifier TransformCurveId(BoneName, ERawCurveTrackTypes::RCT_Transform);
+		if (const FTransformCurve* TransformCurvePtr = DataModelInterface->FindTransformCurve(TransformCurveId))
 		{
-			FSmartName CurveSmartName;
-			CurrentSkeleton->GetSmartNameByName(USkeleton::AnimTrackCurveMappingName, BoneName, CurveSmartName);
-			const FAnimationCurveIdentifier TransformCurveId(CurveSmartName, ERawCurveTrackTypes::RCT_Transform);
-			if (const FTransformCurve* TransformCurvePtr = DataModelInterface->FindTransformCurve(TransformCurveId))
-			{
-				const FTransform AdditiveTransform = TransformCurvePtr->Evaluate(Time, 1.f);
-				const FTransform LocalTransform = OutAtom;
-				OutAtom.SetRotation(LocalTransform.GetRotation() * AdditiveTransform.GetRotation());
-				OutAtom.SetTranslation(LocalTransform.TransformPosition(AdditiveTransform.GetTranslation()));
-				OutAtom.SetScale3D(LocalTransform.GetScale3D() * AdditiveTransform.GetScale3D());	
-			}
-		}	
+			const FTransform AdditiveTransform = TransformCurvePtr->Evaluate(Time, 1.f);
+			const FTransform LocalTransform = OutAtom;
+			OutAtom.SetRotation(LocalTransform.GetRotation() * AdditiveTransform.GetRotation());
+			OutAtom.SetTranslation(LocalTransform.TransformPosition(AdditiveTransform.GetTranslation()));
+			OutAtom.SetScale3D(LocalTransform.GetScale3D() * AdditiveTransform.GetScale3D());	
+		}
 #endif
 	}
 }
@@ -1188,20 +1177,15 @@ void UAnimSequence::GetBoneTransform(FTransform& OutAtom, FSkeletonPoseBoneIndex
 		const FName BoneName = GetSkeleton()->GetReferenceSkeleton().GetBoneName(BoneIndex.GetInt());
 		OutAtom = DataModelInterface->EvaluateBoneTrackTransform(BoneName, DataModelInterface->GetFrameRate().AsFrameTime(DecompContext.GetEvaluationTime()), Interpolation);
 
-		if (const USkeleton* CurrentSkeleton = GetSkeleton())
+		const FAnimationCurveIdentifier TransformCurveId(BoneName, ERawCurveTrackTypes::RCT_Transform);
+		if (const FTransformCurve* TransformCurvePtr = DataModelInterface->FindTransformCurve(TransformCurveId))
 		{
-			FSmartName CurveSmartName;
-			CurrentSkeleton->GetSmartNameByName(USkeleton::AnimTrackCurveMappingName, BoneName, CurveSmartName);		
-			const FAnimationCurveIdentifier TransformCurveId(CurveSmartName, ERawCurveTrackTypes::RCT_Transform);
-			if (const FTransformCurve* TransformCurvePtr = DataModelInterface->FindTransformCurve(TransformCurveId))
-			{
-				const FTransform AdditiveTransform = TransformCurvePtr->Evaluate(DecompContext.GetEvaluationTime(), 1.f);
-				const FTransform LocalTransform = OutAtom;
-				OutAtom.SetRotation(LocalTransform.GetRotation() * AdditiveTransform.GetRotation());
-				OutAtom.SetTranslation(LocalTransform.TransformPosition(AdditiveTransform.GetTranslation()));
-				OutAtom.SetScale3D(LocalTransform.GetScale3D() * AdditiveTransform.GetScale3D());	
-			}
-		}	
+			const FTransform AdditiveTransform = TransformCurvePtr->Evaluate(DecompContext.GetEvaluationTime(), 1.f);
+			const FTransform LocalTransform = OutAtom;
+			OutAtom.SetRotation(LocalTransform.GetRotation() * AdditiveTransform.GetRotation());
+			OutAtom.SetTranslation(LocalTransform.TransformPosition(AdditiveTransform.GetTranslation()));
+			OutAtom.SetScale3D(LocalTransform.GetScale3D() * AdditiveTransform.GetScale3D());	
+		}
 #endif
 	}
 }
@@ -1546,11 +1530,8 @@ void UAnimSequence::GetBonePose(FAnimationPoseData& OutAnimationPoseData, const
 	if (bUseRawDataForPoseExtraction)
 	{
 		{
-			FSkeletonRemappingCurve RemappedCurve(OutAnimationPoseData.GetCurve(), RequiredBones, GetSkeleton());
-			FAnimationPoseData RemappedPoseData(OutAnimationPoseData.GetPose(), RemappedCurve.GetCurve(), OutAnimationPoseData.GetAttributes());
-
 			const UE::Anim::DataModel::FEvaluationContext EvaluationContext(ExtractionContext.CurrentTime, DataModelInterface->GetFrameRate(), GetRetargetTransformsSourceName(), GetRetargetTransforms(), Interpolation);
-			DataModelInterface->Evaluate(RemappedPoseData, EvaluationContext);
+			DataModelInterface->Evaluate(OutAnimationPoseData, EvaluationContext);
 		}
 
 		if ((ExtractionContext.bExtractRootMotion && RootMotionReset.bEnableRootMotion) || RootMotionReset.bForceRootLock)
@@ -1572,14 +1553,12 @@ void UAnimSequence::GetBonePose(FAnimationPoseData& OutAnimationPoseData, const
 
 	// (Always) evaluate compressed curve data
 	{
-		// Scoped so that the RemappedCurve destructs and isn't kept alive longer than needed.
-		FSkeletonRemappingCurve RemappedCurve(OutAnimationPoseData.GetCurve(), RequiredBones, GetSkeleton());
 #if WITH_EDITOR
 		// When evaluating from raw animation data, UE::Anim::BuildPoseFromModel will populate the curve data
 		if (!bUseRawDataForPoseExtraction)
 #endif // WITH_EDITOR
 		{
-			EvaluateCurveData(RemappedCurve.GetCurve(), ExtractionContext.CurrentTime, bUseRawDataForPoseExtraction);
+			EvaluateCurveData(OutAnimationPoseData.GetCurve(), ExtractionContext.CurrentTime, bUseRawDataForPoseExtraction);
 		}
 	}
 
@@ -1587,6 +1566,12 @@ void UAnimSequence::GetBonePose(FAnimationPoseData& OutAnimationPoseData, const
 	EvaluateAttributes(OutAnimationPoseData, ExtractionContext, false);
 }
 
+const TArray<FSmartName>& UAnimSequence::GetCompressedCurveNames() const
+{
+	static TArray<FSmartName> Dummy;
+	return Dummy;
+}
+
 void UAnimSequence::GetBonePose_Additive(FCompactPose& OutPose, FBlendedCurve& OutCurve, const FAnimExtractContext& ExtractionContext) const
 {
 	UE::Anim::FStackAttributeContainer TempAttributes;
@@ -1740,17 +1725,19 @@ bool UAnimSequence::CanEvaluateRawAnimationData() const
 	return false;
 }
 
-#if WITH_EDITOR
-void UAnimSequence::UpdateCompressedCurveName(SmartName::UID_Type CurveUID, const struct FSmartName& NewCurveName)
+#if WITH_EDITORONLY_DATA
+void UAnimSequence::UpdateCompressedCurveName(const FName& OldCurveName, const FName& NewCurveName)
 {
-	for (FSmartName& CurveName : CompressedData.CompressedCurveNames)
+	for (FAnimCompressedCurveIndexedName& IndexedCurveName : CompressedData.IndexedCurveNames)
 	{
-		if (CurveName.UID == CurveUID)
+		if (IndexedCurveName.CurveName == OldCurveName)
 		{
-			CurveName = NewCurveName;
+			IndexedCurveName.CurveName = NewCurveName;
 			break;
 		}
 	}
+
+	CompressedData.RebuildCurveIndexTable();
 }
 #endif // WITH_EDITOR
 
@@ -3158,14 +3145,11 @@ void UAnimSequence::AddKeyToSequence(float Time, const FName& BoneName, const FT
 	FName CurveName = BoneName;
 	USkeleton * CurrentSkeleton = GetSkeleton();
 	check (CurrentSkeleton);
-
-	FSmartName NewCurveName;
-	CurrentSkeleton->AddSmartNameAndModify(USkeleton::AnimTrackCurveMappingName, CurveName, NewCurveName);
-
+	
 	ValidateModel();
 
 	IAnimationDataController::FScopedBracket ScopedBracket(Controller, LOCTEXT("AddKeyToSequence_Bracket", "Adding key to sequence"));
-	FAnimationCurveIdentifier TransformCurveId(NewCurveName, ERawCurveTrackTypes::RCT_Transform);
+	FAnimationCurveIdentifier TransformCurveId(CurveName, ERawCurveTrackTypes::RCT_Transform);
 	Controller->AddCurve(TransformCurveId, AACF_DriveTrack | AACF_Editable);
 
 	const FTransformCurve* TransformCurve = DataModelInterface->FindTransformCurve(TransformCurveId);
@@ -3325,12 +3309,7 @@ int32 UAnimSequence::GetNumberOfSampledKeys() const
 void UAnimSequence::EvaluateCurveData(FBlendedCurve& OutCurve, float CurrentTime, bool bForceUseRawData) const
 {
 	SCOPE_CYCLE_COUNTER(STAT_AnimSeq_EvalCurveData);
-
-	if (OutCurve.NumValidCurveCount == 0)
-	{
-		return;
-	}
-
+	
 	const bool bEvaluateRawData =
 	    PRAGMA_DISABLE_DEPRECATION_WARNINGS
 	    bUseRawDataOnly
@@ -3340,19 +3319,7 @@ void UAnimSequence::EvaluateCurveData(FBlendedCurve& OutCurve, float CurrentTime
 	if (CanEvaluateRawAnimationData() && bEvaluateRawData)
 	{
 #if WITH_EDITOR
-		// Evaluate float curves from the AnimationData Model
-		if (OutCurve.NumValidCurveCount > 0)
-		{
-			for (auto CurveIter = DataModelInterface->GetFloatCurves().CreateConstIterator(); CurveIter; ++CurveIter)
-			{
-				const FFloatCurve& Curve = *CurveIter;
-				if (OutCurve.IsEnabled(Curve.Name.UID))
-				{
-					const float Value = Curve.Evaluate(CurrentTime);
-					OutCurve.Set(Curve.Name.UID, Value);
-				}
-			}
-		}
+		UE::Anim::EvaluateFloatCurvesFromModel(DataModelInterface.GetInterface(), OutCurve, CurrentTime);
 #else
 		Super::EvaluateCurveData(OutCurve, CurrentTime, bForceUseRawData);
 #endif
@@ -3364,9 +3331,9 @@ void UAnimSequence::EvaluateCurveData(FBlendedCurve& OutCurve, float CurrentTime
 	}
 }
 
-float UAnimSequence::EvaluateCurveData(SmartName::UID_Type CurveUID, float CurrentTime, bool bForceUseRawData) const
+float UAnimSequence::EvaluateCurveData(FName CurveName, float CurrentTime, bool bForceUseRawData) const
 {
-	SCOPE_CYCLE_COUNTER(STAT_AnimSeq_EvalCurveData);
+	QUICK_SCOPE_CYCLE_COUNTER(EvaluateCurveDataByName);
 
 	const bool bEvaluateRawData =
 	 PRAGMA_DISABLE_DEPRECATION_WARNINGS
@@ -3376,17 +3343,17 @@ float UAnimSequence::EvaluateCurveData(SmartName::UID_Type CurveUID, float Curre
 	check(!bForceUseRawData || CanEvaluateRawAnimationData());
 	if (CanEvaluateRawAnimationData() && bEvaluateRawData)
 	{
-		return Super::EvaluateCurveData(CurveUID, CurrentTime, bForceUseRawData);
+		return Super::EvaluateCurveData(CurveName, CurrentTime, bForceUseRawData);
 	}
 	else if(IsCurveCompressedDataValid() && CompressedData.CurveCompressionCodec)
 	{
-		return CompressedData.CurveCompressionCodec->DecompressCurve(CompressedData, CurveUID, CurrentTime);
+		return CompressedData.CurveCompressionCodec->DecompressCurve(CompressedData, CurveName, CurrentTime);
 	}
 
 	return 0.f;
 }
 
-bool UAnimSequence::HasCurveData(SmartName::UID_Type CurveUID, bool bForceUseRawData) const
+bool UAnimSequence::HasCurveData(FName CurveName, bool bForceUseRawData) const
 {
 	const bool bEvaluateRawData =
 	PRAGMA_DISABLE_DEPRECATION_WARNINGS
@@ -3396,18 +3363,20 @@ bool UAnimSequence::HasCurveData(SmartName::UID_Type CurveUID, bool bForceUseRaw
 	check(!bForceUseRawData || CanEvaluateRawAnimationData());
 	if (CanEvaluateRawAnimationData() && bEvaluateRawData)
 	{
-		return Super::HasCurveData(CurveUID, bForceUseRawData);
+		PRAGMA_DISABLE_DEPRECATION_WARNINGS
+		return Super::HasCurveData(CurveName, bForceUseRawData);
+		PRAGMA_ENABLE_DEPRECATION_WARNINGS
 	}
 	else if(IsCurveCompressedDataValid() && CompressedData.CurveCompressionCodec)
 	{
-	for (const FSmartName& CurveName : CompressedData.CompressedCurveNames)
-	{
-		if (CurveName.UID == CurveUID)
+		for (const FAnimCompressedCurveIndexedName& IndexedCurveName : CompressedData.IndexedCurveNames)
 		{
-			return true;
+			if (IndexedCurveName.CurveName == CurveName)
+			{
+				return true;
+			}
 		}
 	}
-	}
 
 	return false;
 }
@@ -3843,7 +3812,7 @@ void UAnimSequence::SynchronousAnimatedBoneAttributesCompression()
 				RequiredBoneIndexArray[BoneIndex] = BoneIndex;
 			}
 
-			RequiredBones.InitializeTo(RequiredBoneIndexArray, FCurveEvaluationOption(true), *InSkeleton);
+			RequiredBones.InitializeTo(RequiredBoneIndexArray, UE::Anim::FCurveFilterSettings(UE::Anim::ECurveFilterMode::None), *InSkeleton);
 		}
 	};
 	
@@ -4897,14 +4866,7 @@ void UAnimSequence::OnModelModified(const EAnimDataModelNotifyType& NotifyType,
 		case EAnimDataModelNotifyType::CurveRenamed:
 		{
 			const FCurveRenamedPayload& TypedPayload = Payload.GetPayload<FCurveRenamedPayload>();
-			UpdateCompressedCurveName(TypedPayload.Identifier.InternalName.UID, TypedPayload.NewIdentifier.InternalName);
-			if (TypedPayload.Identifier.InternalName.DisplayName == TypedPayload.NewIdentifier.InternalName.DisplayName
-				&& TypedPayload.Identifier.InternalName.UID == SmartName::MaxUID && TypedPayload.NewIdentifier.InternalName.UID != SmartName::MaxUID)
-			{
-				// We are renaming a Curve with a FSmartName::Invalid UID, this means it was just loaded and we are retrieving
-				// its actual UID from its DisplayName. This should *not* mark the package dirty.
-				bShouldMarkPackageDirty = false;
-			}
+			UpdateCompressedCurveName(TypedPayload.Identifier.CurveName, TypedPayload.NewIdentifier.CurveName);
 			break;
 		}
 		default:
diff --git a/Engine/Source/Runtime/Engine/Private/Animation/AnimSequenceBase.cpp b/Engine/Source/Runtime/Engine/Private/Animation/AnimSequenceBase.cpp
--- a/Engine/Source/Runtime/Engine/Private/Animation/AnimSequenceBase.cpp
+++ b/Engine/Source/Runtime/Engine/Private/Animation/AnimSequenceBase.cpp
@@ -59,7 +59,7 @@ UAnimSequenceBase::UAnimSequenceBase(const FObjectInitializer& ObjectInitializer
 
 bool UAnimSequenceBase::IsPostLoadThreadSafe() const
 {
-	return false;	// PostLoad is not thread safe because of the call to VerifyCurveNames() (calling USkeleton::VerifySmartName) that can mutate a shared map in the skeleton.
+	return WITH_EDITORONLY_DATA == 0;	// Not thread safe in editor because new objects can be constructed on upgrade and the skeleton can be modified
 }
 
 #if WITH_EDITORONLY_DATA
@@ -85,23 +85,14 @@ void UAnimSequenceBase::DeclareConstructClasses(TArray<FTopLevelAssetPath>& OutC
 }
 #endif
 
-template <typename DataType>
-void VerifyCurveNames(USkeleton& Skeleton, const FName& NameContainer, TArray<DataType>& CurveList)
-{
-	for (DataType& Curve : CurveList)
-	{
-		Skeleton.VerifySmartName(NameContainer, Curve.Name);
-	}
-}
-
 void UAnimSequenceBase::PostLoad()
 {
 #if WITH_EDITORONLY_DATA
 	if (!HasAnyFlags(EObjectFlags::RF_ClassDefaultObject))
 	{
 		LLM_SCOPE(ELLMTag::Animation);
-		
-		auto PreloadSkeletonAndVerifyCurves = [this]()
+
+		auto PreloadSkeleton = [this]()
 		{
 			if (USkeleton* MySkeleton = GetSkeleton())
 			{
@@ -110,11 +101,6 @@ void UAnimSequenceBase::PostLoad()
 					SkeletonLinker->Preload(MySkeleton);
 				}
 				MySkeleton->ConditionalPostLoad();
-
-				PRAGMA_DISABLE_DEPRECATION_WARNINGS
-				VerifyCurveNames<FFloatCurve>(*MySkeleton, USkeleton::AnimCurveMappingName, RawCurveData.FloatCurves);
-				VerifyCurveNames<FTransformCurve>(*MySkeleton, USkeleton::AnimTrackCurveMappingName, RawCurveData.TransformCurves);
-				PRAGMA_ENABLE_DEPRECATION_WARNINGS
 			}
 		};
 
@@ -140,7 +126,7 @@ void UAnimSequenceBase::PostLoad()
 		    GetController();
 		    BindToModelModificationEvent();
 
-		    PreloadSkeletonAndVerifyCurves();
+		    PreloadSkeleton();
 
 		    if (bRequiresModelPopulation || bRequiresModelCreation)
 		    {
@@ -198,13 +184,14 @@ void UAnimSequenceBase::PostLoad()
 		}
 		else
 		{
-			PreloadSkeletonAndVerifyCurves();
+			PreloadSkeleton();
 		}
 	}
 #endif // WITH_EDITORONLY_DATA
 
 	Super::PostLoad();
 
+#if WITH_EDITORONLY_DATA
 	// Convert Notifies to new data
 	if( GIsEditor && Notifies.Num() > 0 )
 	{
@@ -221,12 +208,11 @@ void UAnimSequenceBase::PostLoad()
 		}
 	}
 
-#if WITH_EDITOR
 	InitializeNotifyTrack();
-#endif
+#endif	// WITH_EDITORONLY_DATA
 	RefreshCacheData();
 
-#if WITH_EDITOR
+#if WITH_EDITORONLY_DATA
 	if (!GetPackage()->GetHasBeenEndLoaded())
 	{
 		FCoreUObjectDelegates::OnEndLoadPackage.AddUObject(this, &UAnimSequenceBase::OnEndLoadPackage);
@@ -234,12 +220,10 @@ void UAnimSequenceBase::PostLoad()
 	else
 	{
 		OnAnimModelLoaded();
-	}	
-#endif
+	}
 
 	if(USkeleton* MySkeleton = GetSkeleton())
 	{
-#if WITH_EDITOR
 		if (IsDataModelValid())
 		{
 			const bool bDoNotTransactAction = false;
@@ -251,37 +235,28 @@ void UAnimSequenceBase::PostLoad()
 					for (int32 Index = 0; Index < FloatCurves.Num(); ++Index)
 					{
 						const FFloatCurve& Curve = FloatCurves[Index];
-						if (Curve.Name.DisplayName == NAME_None)
+						if (Curve.GetName() == NAME_None)
 						{
 							// give unique name
 							const FName UniqueName = FName(*FString(GetName() + TEXT("_CurveNameFix_") + FString::FromInt(Index)));
-							UE_LOG(LogAnimation, Warning, TEXT("[AnimSequence %s] contains invalid curve name \'None\'. Renaming this to %s. Please fix this curve in the editor. "), *GetFullName(), *Curve.Name.DisplayName.ToString());
+							UE_LOG(LogAnimation, Warning, TEXT("[AnimSequence %s] contains invalid curve name \'None\'. Renaming this to %s. Please fix this curve in the editor. "), *GetFullName(), *Curve.GetName().ToString());
 
-							FSmartName NewSmartName = Curve.Name;
-							NewSmartName.DisplayName = UniqueName;
-
-							Controller->RenameCurve(FAnimationCurveIdentifier(Curve.Name, ERawCurveTrackTypes::RCT_Float), FAnimationCurveIdentifier(NewSmartName, ERawCurveTrackTypes::RCT_Float), bDoNotTransactAction);
+							Controller->RenameCurve(FAnimationCurveIdentifier(Curve.GetName(), ERawCurveTrackTypes::RCT_Float), FAnimationCurveIdentifier(UniqueName, ERawCurveTrackTypes::RCT_Float), bDoNotTransactAction);
 						}
 					}
 				}
 				Controller->CloseBracket(bDoNotTransactAction);
 			}
 		}
-#endif
 
 		// this should continue to add if skeleton hasn't been saved either 
 		// we don't wipe out data, so make sure you add back in if required
 		if (GetLinkerCustomVersion(FFrameworkObjectVersion::GUID) < FFrameworkObjectVersion::MoveCurveTypesToSkeleton
 			|| MySkeleton->GetLinkerCustomVersion(FFrameworkObjectVersion::GUID) < FFrameworkObjectVersion::MoveCurveTypesToSkeleton)
 		{
-#if WITH_EDITOR
 			// This is safe as the data model will have been created during this PostLoad call
 			const TArray<FFloatCurve>& FloatCurves = DataModelInterface->GetFloatCurves();
-#else
-			PRAGMA_DISABLE_DEPRECATION_WARNINGS
-			const TArray<FFloatCurve>& FloatCurves = RawCurveData.FloatCurves;
-			PRAGMA_ENABLE_DEPRECATION_WARNINGS
-#endif
+
 			// fix up curve flags to skeleton
 			for (const FFloatCurve& Curve : FloatCurves)
 			{
@@ -291,11 +266,12 @@ void UAnimSequenceBase::PostLoad()
 				// only add this if that has to 
 				if (bMorphtargetSet || bMaterialSet)
 				{
-					MySkeleton->AccumulateCurveMetaData(Curve.Name.DisplayName, bMaterialSet, bMorphtargetSet);
+					MySkeleton->AccumulateCurveMetaData(Curve.GetName(), bMaterialSet, bMorphtargetSet);
 				}
 			}
 		}
 	}
+#endif // WITH_EDITORONLY_DATA
 }
 
 void UAnimSequenceBase::PostDuplicate(EDuplicateMode::Type DuplicateMode)
@@ -358,6 +334,28 @@ void UAnimSequenceBase::RemoveNotifies()
 	RefreshCacheData();
 }
 
+#if WITH_EDITOR
+void UAnimSequenceBase::RenameNotifies(FName InOldName, FName InNewName)
+{
+	Modify();
+
+	for(FAnimNotifyEvent& Notify : Notifies)
+	{
+		// Only handle named notifies
+		if(!Notify.IsBlueprintNotify())
+		{
+			if(Notify.NotifyName == InOldName)
+			{
+				Notify.NotifyName = InNewName;
+			}
+		}
+	}
+
+	// notification broadcast
+	OnNotifyChanged.Broadcast();
+}
+#endif
+
 bool UAnimSequenceBase::IsNotifyAvailable() const
 {
 	return (Notifies.Num() != 0) && (GetPlayLength() > 0.f);
@@ -524,16 +522,6 @@ void UAnimSequenceBase::GetAnimNotifiesFromDeltaPositions(const float& PreviousP
 	}
 }
 
-#if WITH_EDITOR
-void UAnimSequenceBase::RemapTracksToNewSkeleton(USkeleton* NewSkeleton, bool bConvertSpaces)
-{
-	Super::RemapTracksToNewSkeleton(NewSkeleton, bConvertSpaces);
-
-	Controller->FindOrAddCurveNamesOnSkeleton(NewSkeleton, ERawCurveTrackTypes::RCT_Float);
-	Controller->FindOrAddCurveNamesOnSkeleton(NewSkeleton, ERawCurveTrackTypes::RCT_Transform);
-}
-#endif
-
 void UAnimSequenceBase::TickAssetPlayer(FAnimTickRecord& Instance, struct FAnimNotifyQueue& NotifyQueue, FAnimAssetTickContext& Context) const
 {
 	float CurrentTime = *(Instance.TimeAccumulator);
@@ -849,37 +837,31 @@ EAnimEventTriggerOffsets::Type UAnimSequenceBase::CalculateOffsetForNotify(float
 void UAnimSequenceBase::GetAssetRegistryTags(TArray<FAssetRegistryTag>& OutTags) const
 {
 	Super::GetAssetRegistryTags(OutTags);
-	if(Notifies.Num() > 0)
-	{
-		FString NotifyList;
 
-		// add notifies to 
-		for(auto Iter=Notifies.CreateConstIterator(); Iter; ++Iter)
-		{
-			// only add if not BP anim notify since they're handled separate
-			if(Iter->IsBlueprintNotify() == false)
-			{
-				NotifyList += FString::Printf(TEXT("%s%s"), *Iter->NotifyName.ToString(), *USkeleton::AnimNotifyTagDelimiter);
-			}
-		}
-		
-		if(NotifyList.Len() > 0)
+	// Add notify IDs to a tag list, or a delimiter if we have no notifies.
+	// The delimiter is necessary so we can distinguish between data with no curves and old data, as the asset registry
+	// strips tags that have empty values 
+	FString NotifyList = USkeleton::AnimNotifyTagDelimiter;
+	for(auto Iter=Notifies.CreateConstIterator(); Iter; ++Iter)
+	{
+		// only add if not BP anim notify since they're handled separately
+		if(Iter->IsBlueprintNotify() == false)
 		{
-			OutTags.Add(FAssetRegistryTag(USkeleton::AnimNotifyTag, NotifyList, FAssetRegistryTag::TT_Hidden));
+			NotifyList += FString::Printf(TEXT("%s%s"), *Iter->NotifyName.ToString(), *USkeleton::AnimNotifyTagDelimiter);
 		}
 	}
+	
+	OutTags.Add(FAssetRegistryTag(USkeleton::AnimNotifyTag, NotifyList, FAssetRegistryTag::TT_Hidden));
 
-	// Add curve IDs to a tag list, or a blank tag if we have no curves.
-	// The blank list is necessary when we attempt to delete a curve so
-	// an old asset can be detected from its asset data so we load as few
-	// as possible.
-	FString CurveNameList;
-
+	// Add curve IDs to a tag list, or a delimiter if we have no curves.
+	// The delimiter is necessary so we can distinguish between data with no curves and old data, as the asset registry
+	// strips tags that have empty values 
+	FString CurveNameList = USkeleton::CurveTagDelimiter;
 	if (DataModelInterface.GetObject())
 	{
 		for(const FFloatCurve& Curve : DataModelInterface->GetFloatCurves())
 		{
-			CurveNameList += FString::Printf(TEXT("%s%s"), *Curve.Name.DisplayName.ToString(), *USkeleton::CurveTagDelimiter);
+			CurveNameList += FString::Printf(TEXT("%s%s"), *Curve.GetName().ToString(), *USkeleton::CurveTagDelimiter);
 		}
 	}
 	OutTags.Add(FAssetRegistryTag(USkeleton::CurveNameTag, CurveNameList, FAssetRegistryTag::TT_Hidden));
@@ -944,7 +926,7 @@ void UAnimSequenceBase::RefreshParentAssetData()
 		Controller->RemoveAllCurvesOfType(ERawCurveTrackTypes::RCT_Float);
 		for (const FFloatCurve& FloatCurve : ParentDataModel->GetFloatCurves())
 		{
-			const FAnimationCurveIdentifier CurveId = UAnimationCurveIdentifierExtensions::FindCurveIdentifier(GetSkeleton(), FloatCurve.Name.DisplayName, ERawCurveTrackTypes::RCT_Float);
+			const FAnimationCurveIdentifier CurveId(FloatCurve.GetName(), ERawCurveTrackTypes::RCT_Float);
 			Controller->AddCurve(CurveId, FloatCurve.GetCurveTypeFlags());
 			Controller->SetCurveKeys(CurveId, FloatCurve.FloatCurve.GetConstRefOfKeys());
 		}
@@ -952,7 +934,7 @@ void UAnimSequenceBase::RefreshParentAssetData()
 		Controller->RemoveAllCurvesOfType(ERawCurveTrackTypes::RCT_Transform);
 		for (const FTransformCurve& TransformCurve : ParentDataModel->GetTransformCurves())
 		{
-			const FAnimationCurveIdentifier CurveId(TransformCurve.Name, ERawCurveTrackTypes::RCT_Transform);
+			const FAnimationCurveIdentifier CurveId(TransformCurve.GetName(), ERawCurveTrackTypes::RCT_Transform);
 			Controller->AddCurve(CurveId, TransformCurve.GetCurveTypeFlags());
 
 			// Set each individual channel rich curve keys, to account for any custom tangents etc.
@@ -1027,10 +1009,10 @@ void UAnimSequenceBase::EvaluateCurveData(FBlendedCurve& OutCurve, float Current
 	PRAGMA_ENABLE_DEPRECATION_WARNINGS
 }
 
-float UAnimSequenceBase::EvaluateCurveData(SmartName::UID_Type CurveUID, float CurrentTime, bool bForceUseRawData) const
+float UAnimSequenceBase::EvaluateCurveData(FName CurveName, float CurrentTime, bool bForceUseRawData) const
 {
 	PRAGMA_DISABLE_DEPRECATION_WARNINGS
-	const FFloatCurve* Curve = (const FFloatCurve*)RawCurveData.GetCurveData(CurveUID, ERawCurveTrackTypes::RCT_Float);
+	const FFloatCurve* Curve = (const FFloatCurve*)RawCurveData.GetCurveData(CurveName, ERawCurveTrackTypes::RCT_Float);
 	PRAGMA_ENABLE_DEPRECATION_WARNINGS
 
 	return Curve != nullptr ? Curve->Evaluate(CurrentTime) : 0.0f;
@@ -1043,17 +1025,17 @@ const FRawCurveTracks& UAnimSequenceBase::GetCurveData() const
 	PRAGMA_ENABLE_DEPRECATION_WARNINGS
 }
 
-bool UAnimSequenceBase::HasCurveData(SmartName::UID_Type CurveUID, bool bForceUseRawData) const
+bool UAnimSequenceBase::HasCurveData(FName CurveName, bool bForceUseRawData) const
 {
 #if WITH_EDITOR
 	if (IsDataModelValid())
 	{
-		return DataModelInterface->FindFloatCurve(FAnimationCurveIdentifier(CurveUID, ERawCurveTrackTypes::RCT_Float)) != nullptr;
-	}
+		return DataModelInterface->FindFloatCurve(FAnimationCurveIdentifier(CurveName, ERawCurveTrackTypes::RCT_Float)) != nullptr;
+	}	
 #endif
 	
 	PRAGMA_DISABLE_DEPRECATION_WARNINGS
-	return RawCurveData.GetCurveData(CurveUID) != nullptr;
+	return RawCurveData.GetCurveData(CurveName) != nullptr;
 	PRAGMA_ENABLE_DEPRECATION_WARNINGS
 }
 
@@ -1073,7 +1055,7 @@ void UAnimSequenceBase::Serialize(FArchive& Ar)
 
 	// fix up version issue and so on
 	PRAGMA_DISABLE_DEPRECATION_WARNINGS
-	RawCurveData.PostSerialize(Ar);
+	RawCurveData.PostSerializeFixup(Ar);
 	PRAGMA_ENABLE_DEPRECATION_WARNINGS
 }
 
@@ -1286,18 +1268,18 @@ void UAnimSequenceBase::OnModelModified(const EAnimDataModelNotifyType& NotifyTy
 				if (Identifier.CurveType == ERawCurveTrackTypes::RCT_Float)
 				{
 					PRAGMA_DISABLE_DEPRECATION_WARNINGS
-					return RawCurveData.FloatCurves.FindByPredicate([SmartName=Identifier.InternalName](FFloatCurve& Curve)
+					return RawCurveData.FloatCurves.FindByPredicate([CurveName=Identifier.CurveName](FFloatCurve& Curve)
 	                {
-	                    return Curve.Name == SmartName;
+	                    return Curve.GetName() == CurveName;
 	                });
 					PRAGMA_ENABLE_DEPRECATION_WARNINGS
 				}
 				else if (Identifier.CurveType == ERawCurveTrackTypes::RCT_Transform)
 				{
 					PRAGMA_DISABLE_DEPRECATION_WARNINGS
-					return RawCurveData.TransformCurves.FindByPredicate([SmartName=Identifier.InternalName](FTransformCurve& Curve)
+					return RawCurveData.TransformCurves.FindByPredicate([CurveName=Identifier.CurveName](FTransformCurve& Curve)
                     {
-                        return Curve.Name == SmartName;
+                        return Curve.GetName() == CurveName;
                     });
 					PRAGMA_ENABLE_DEPRECATION_WARNINGS
 				}
@@ -1307,7 +1289,7 @@ void UAnimSequenceBase::OnModelModified(const EAnimDataModelNotifyType& NotifyTy
 				
 			if (CurvePtr)
 			{
-				CurvePtr->Name = TypedPayload.NewIdentifier.InternalName;
+				CurvePtr->SetName(TypedPayload.NewIdentifier.CurveName);
 			}
 			break;
 		}
@@ -1485,12 +1467,8 @@ void UAnimSequenceBase::OnAnimModelLoaded()
 			for (int32 Index = 0; Index < FloatCurves.Num(); ++Index)
 			{
 				const FFloatCurve& Curve = FloatCurves[Index];
-				ensureMsgf(Curve.Name.DisplayName != NAME_None, TEXT("[AnimSequencer %s] has invalid curve name."), *GetFullName());
+				ensureMsgf(Curve.GetName() != NAME_None, TEXT("[AnimSequencer %s] has invalid curve name."), *GetFullName());
 			}
-
-			constexpr bool bDoNotTransactAction = false;
-			Controller->FindOrAddCurveNamesOnSkeleton(MySkeleton, ERawCurveTrackTypes::RCT_Float, bDoNotTransactAction);
-			Controller->FindOrAddCurveNamesOnSkeleton(MySkeleton, ERawCurveTrackTypes::RCT_Transform, bDoNotTransactAction);
 		}
 	}
 }
diff --git a/Engine/Source/Runtime/Engine/Private/Animation/AnimSequenceHelpers.cpp b/Engine/Source/Runtime/Engine/Private/Animation/AnimSequenceHelpers.cpp
--- a/Engine/Source/Runtime/Engine/Private/Animation/AnimSequenceHelpers.cpp
+++ b/Engine/Source/Runtime/Engine/Private/Animation/AnimSequenceHelpers.cpp
@@ -18,6 +18,7 @@
 #include "Animation/SkeletonRemapping.h"
 #include "Animation/SkeletonRemappingRegistry.h"
 #include "Animation/AttributesRuntime.h"
+#include "Animation/AnimCurveUtils.h"
 
 #define LOCTEXT_NAMESPACE "AnimSequenceHelpers"
 
@@ -53,18 +54,21 @@ void EvaluateFloatCurvesFromModel(const IAnimationDataModel* Model, FBlendedCurv
 {
 	check(Model);
 
-	if (OutCurves.NumValidCurveCount > 0)
+	// Evaluate into a temporary curve, then filter by enabled curves
+	const TArray<FFloatCurve>& ModelCurves = Model->GetFloatCurves();
+	const int32 NumCurves = ModelCurves.Num();
+
+	auto GetNameFromIndex = [&ModelCurves](int32 InCurveIndex)
 	{
-		// evaluate the curve data at the Time and add to Instance
-		for (const FFloatCurve& Curve : Model->GetFloatCurves())
-		{
-			if (OutCurves.IsEnabled(Curve.Name.UID))
-			{
-				float Value = Curve.Evaluate(Time);
-				OutCurves.Set(Curve.Name.UID, Value);
-			}
-		}
-	}
+		return ModelCurves[InCurveIndex].GetName();
+	};
+
+	auto GetValueFromIndex = [&ModelCurves, Time](int32 InCurveIndex)
+	{
+		return ModelCurves[InCurveIndex].Evaluate(Time);
+	};
+
+	UE::Anim::FCurveUtils::BuildUnsorted(OutCurves, NumCurves, GetNameFromIndex, GetValueFromIndex, OutCurves.GetFilter());
 }
 
 void EvaluateTransformCurvesFromModel(const IAnimationDataModel* Model, TMap<FName, FTransform>& OutCurves, double Time, float BlendWeight)
@@ -80,10 +84,8 @@ void EvaluateTransformCurvesFromModel(const IAnimationDataModel* Model, TMap<FNa
 			}
 
 			// Add or retrieve curve
-			const FName& CurveName = Curve.Name.DisplayName;
-
 			// note we're not checking Curve.GetCurveTypeFlags() yet
-			FTransform& Value = OutCurves.FindOrAdd(CurveName);
+			FTransform& Value = OutCurves.FindOrAdd(Curve.GetName());
 			Value = Curve.Evaluate(Time, BlendWeight);
 		}
 	}
@@ -104,7 +106,7 @@ void GetBoneTransformFromModel(const IAnimationDataModel* Model, FTransform& Out
 
 	for (const FTransformCurve& AdditiveTransformCurve : Model->GetTransformCurves())
 	{
-		if (AdditiveTransformCurve.Name.DisplayName == TrackNames[TrackIndex])
+		if (AdditiveTransformCurve.GetName() == TrackNames[TrackIndex])
 		{
 			const float TimeInterval = Model->GetFrameRate().AsSeconds(KeyIndex);
 			const FTransform AdditiveTransform = AdditiveTransformCurve.Evaluate(TimeInterval, 1.f);
@@ -122,13 +124,14 @@ void CopyCurveDataToModel(const FRawCurveTracks& CurveData, const USkeleton* Ske
 	// Populate float curve data
 	for (const FFloatCurve& FloatCurve : CurveData.FloatCurves)
 	{
-		const FAnimationCurveIdentifier CurveId = UAnimationCurveIdentifierExtensions::FindCurveIdentifier(Skeleton, FloatCurve.Name.DisplayName, ERawCurveTrackTypes::RCT_Float);
+		const FAnimationCurveIdentifier CurveId = UAnimationCurveIdentifierExtensions::FindCurveIdentifier(Skeleton, FloatCurve.GetName(), ERawCurveTrackTypes::RCT_Float);
 		if (CurveId.IsValid())
 		{
 			Controller.AddCurve(CurveId, FloatCurve.GetCurveTypeFlags());
 			FCurveAttributes Attributes;
 			Attributes.SetPreExtrapolation(FloatCurve.FloatCurve.PreInfinityExtrap);
 			Attributes.SetPostExtrapolation(FloatCurve.FloatCurve.PostInfinityExtrap);					
+			Controller.SetCurveColor(CurveId, FloatCurve.GetColor());
 			Controller.SetCurveAttributes(CurveId, Attributes);
 			Controller.SetCurveColor(CurveId, FloatCurve.GetColor());
 			Controller.SetCurveKeys(CurveId, FloatCurve.FloatCurve.GetConstRefOfKeys());
@@ -138,7 +141,7 @@ void CopyCurveDataToModel(const FRawCurveTracks& CurveData, const USkeleton* Ske
 	// Populate transform curve data
 	for (const FTransformCurve& TransformCurve : CurveData.TransformCurves)
 	{
-		const FAnimationCurveIdentifier CurveId = UAnimationCurveIdentifierExtensions::FindCurveIdentifier(Skeleton, TransformCurve.Name.DisplayName, ERawCurveTrackTypes::RCT_Transform);
+		const FAnimationCurveIdentifier CurveId = UAnimationCurveIdentifierExtensions::FindCurveIdentifier(Skeleton, TransformCurve.GetName(), ERawCurveTrackTypes::RCT_Transform);
 		if (CurveId.IsValid())
 		{
 			Controller.AddCurve(CurveId, TransformCurve.GetCurveTypeFlags());
diff --git a/Engine/Source/Runtime/Engine/Private/Animation/AnimSingleNodeInstanceProxy.cpp b/Engine/Source/Runtime/Engine/Private/Animation/AnimSingleNodeInstanceProxy.cpp
--- a/Engine/Source/Runtime/Engine/Private/Animation/AnimSingleNodeInstanceProxy.cpp
+++ b/Engine/Source/Runtime/Engine/Private/Animation/AnimSingleNodeInstanceProxy.cpp
@@ -61,17 +61,10 @@ const UMirrorDataTable* FAnimSingleNodeInstanceProxy::GetMirrorDataTable()
 void FAnimSingleNodeInstanceProxy::PropagatePreviewCurve(FPoseContext& Output) 
 {
 	USkeleton* MySkeleton = GetSkeleton();
+
 	for (auto Iter = PreviewCurveOverride.CreateConstIterator(); Iter; ++Iter)
 	{
-		const FName& Name = Iter.Key();
-		const float Value = Iter.Value();
-
-		FSmartName PreviewCurveName;
-
-		if (MySkeleton->GetSmartNameByName(USkeleton::AnimCurveMappingName, Name, PreviewCurveName))
-		{
-			Output.Curve.Set(PreviewCurveName.UID, Value);
-		}
+		Output.Curve.Set(Iter.Key(), Iter.Value());
 	}
 }
 #endif // WITH_EDITORONLY_DATA
@@ -476,7 +469,7 @@ void FAnimNode_SingleNode::Evaluate_AnyThread(FPoseContext& Output)
 		{
 			if (PoseAsset->GetSkeleton() != nullptr)
 			{
-				const TArray<FSmartName>& PoseNames = PoseAsset->GetPoseNames();
+				const TArray<FName>& PoseNames = PoseAsset->GetPoseFNames();
 
 				int32 TotalPoses = PoseNames.Num();
 				FAnimExtractContext ExtractContext;
@@ -484,12 +477,9 @@ void FAnimNode_SingleNode::Evaluate_AnyThread(FPoseContext& Output)
 
 				for (int32 PoseIndex = 0; PoseIndex <PoseNames.Num(); ++PoseIndex)
 				{
-					const FSmartName& PoseName = PoseNames[PoseIndex];
-					if (PoseName.UID != SmartName::MaxUID)
-					{
-						ExtractContext.PoseCurves[PoseIndex].PoseIndex = PoseIndex;
-						ExtractContext.PoseCurves[PoseIndex].Value = Output.Curve.Get(PoseName.UID);
-					}
+					const FName& PoseName = PoseNames[PoseIndex];
+					ExtractContext.PoseCurves[PoseIndex].PoseIndex = PoseIndex;
+					ExtractContext.PoseCurves[PoseIndex].Value = Output.Curve.Get(PoseName);
 				}
 
 				if (PoseAsset->IsValidAdditive())
diff --git a/Engine/Source/Runtime/Engine/Private/Animation/AnimStreamable.cpp b/Engine/Source/Runtime/Engine/Private/Animation/AnimStreamable.cpp
--- a/Engine/Source/Runtime/Engine/Private/Animation/AnimStreamable.cpp
+++ b/Engine/Source/Runtime/Engine/Private/Animation/AnimStreamable.cpp
@@ -430,14 +430,6 @@ void UAnimStreamable::PostLoad()
 	}
 #else
 	IStreamingManager::Get().GetAnimationStreamingManager().AddStreamingAnim(this); // This will be handled by RequestCompressedData in editor builds
-
-	if (USkeleton* CurrentSkeleton = GetSkeleton())
-	{
-		for (FSmartName& CurveName : GetRunningPlatformData().Chunks[0].CompressedAnimSequence->CompressedCurveNames)
-		{
-			CurrentSkeleton->VerifySmartName(USkeleton::AnimCurveMappingName, CurveName);
-		}
-	}
 #endif
 }
 
