Animation Curve Runtime & Editor Improvements

diff --git a/Engine/Source/Runtime/Engine/Private/Animation/AnimTrace.cpp b/Engine/Source/Runtime/Engine/Private/Animation/AnimTrace.cpp
--- a/Engine/Source/Runtime/Engine/Private/Animation/AnimTrace.cpp
+++ b/Engine/Source/Runtime/Engine/Private/Animation/AnimTrace.cpp
@@ -14,6 +14,7 @@
 #include "Animation/AnimNotifies/AnimNotifyState.h"
 #include "Trace/Trace.inl"
 #include "UObject/UObjectAnnotation.h"
+#include "Engine/PoseWatch.h"
 
 UE_TRACE_CHANNEL_DEFINE(AnimationChannel);
 
@@ -232,13 +233,18 @@ UE_TRACE_EVENT_BEGIN(Animation, Sync)
 	UE_TRACE_EVENT_FIELD(uint32, GroupNameId)
 UE_TRACE_EVENT_END()
 
-UE_TRACE_EVENT_BEGIN(Animation, PoseWatch)
+UE_TRACE_EVENT_BEGIN(Animation, PoseWatch2)
 	UE_TRACE_EVENT_FIELD(uint64, Cycle)
 	UE_TRACE_EVENT_FIELD(double, RecordingTime)
+	UE_TRACE_EVENT_FIELD(uint64, ComponentId)
 	UE_TRACE_EVENT_FIELD(uint64, AnimInstanceId)
 	UE_TRACE_EVENT_FIELD(uint64, PoseWatchId)
+	UE_TRACE_EVENT_FIELD(uint32, NameId)
+	UE_TRACE_EVENT_FIELD(uint32, Color)
 	UE_TRACE_EVENT_FIELD(float[], WorldTransform)
 	UE_TRACE_EVENT_FIELD(float[], BoneTransforms)
+	UE_TRACE_EVENT_FIELD(float[], CurveValues)
+	UE_TRACE_EVENT_FIELD(uint32[], CurveIds)
 	UE_TRACE_EVENT_FIELD(uint16[], RequiredBones)
 	UE_TRACE_EVENT_FIELD(bool, bIsEnabled)
 UE_TRACE_EVENT_END()
@@ -1177,7 +1183,7 @@ void FAnimTrace::OutputSync(const FAnimInstanceProxy& InSourceProxy, int32 InSou
 		<< Sync.GroupNameId(GroupNameId);
 }
 
-void FAnimTrace::OutputPoseWatch(const FAnimInstanceProxy& InSourceProxy, int32 InPoseWatchId, const TArray<FTransform>& BoneTransforms, const TArray<FBoneIndexType>& RequiredBones, const FTransform& WorldTransform, const bool bIsEnabled)
+void FAnimTrace::OutputPoseWatch(const FAnimInstanceProxy& InSourceProxy, UPoseWatchPoseElement* InPoseWatchElement, int32 InPoseWatchId, const TArray<FTransform>& BoneTransforms, const FBlendedHeapCurve& InCurves, const TArray<FBoneIndexType>& RequiredBones, const FTransform& WorldTransform, const bool bIsEnabled)
 {
 	bool bChannelEnabled = UE_TRACE_CHANNELEXPR_IS_ENABLED(AnimationChannel);
 	if (!bChannelEnabled)
@@ -1191,18 +1197,45 @@ void FAnimTrace::OutputPoseWatch(const FAnimInstanceProxy& InSourceProxy, int32
 	}
 
 	const UAnimInstance* AnimInstance = CastChecked<UAnimInstance>(InSourceProxy.GetAnimInstanceObject());
+	const USkeletalMeshComponent* Component = CastChecked<USkeletalMeshComponent>(AnimInstance->GetSkelMeshComponent());
 
-	TRACE_OBJECT(InSourceProxy.GetAnimInstanceObject());
+	TRACE_OBJECT(Component);
+	TRACE_OBJECT(AnimInstance);
+
+	TArray<float>& CurveValues = FAnimTraceScratchBuffers::Get().CurveValues;
+	CurveValues.Reset();
+	CurveValues.SetNumUninitialized(InCurves.Num());
+	TArray<uint32>& CurveIds = FAnimTraceScratchBuffers::Get().CurveIds;
+	CurveIds.Reset();
+	CurveIds.SetNumUninitialized(InCurves.Num());
 
-	UE_TRACE_LOG(Animation, PoseWatch, AnimationChannel)
-		<< PoseWatch.Cycle(FPlatformTime::Cycles64())
-		<< PoseWatch.RecordingTime(FObjectTrace::GetWorldElapsedTime(AnimInstance->GetWorld()))
-		<< PoseWatch.AnimInstanceId(FObjectTrace::GetObjectId(InSourceProxy.GetAnimInstanceObject()))
-		<< PoseWatch.PoseWatchId(InPoseWatchId)
-		<< PoseWatch.WorldTransform(reinterpret_cast<const float*>(&WorldTransform), sizeof(FTransform) / sizeof(float))
-		<< PoseWatch.BoneTransforms(reinterpret_cast<const float*>(BoneTransforms.GetData()), BoneTransforms.Num() * (sizeof(FTransform) / sizeof(float)))
-		<< PoseWatch.RequiredBones(RequiredBones.GetData(), RequiredBones.Num())
-		<< PoseWatch.bIsEnabled(bIsEnabled);
+	if(InCurves.Num() > 0)
+	{
+		int32 CurveIndex = 0;
+		InCurves.ForEachElement([&CurveValues, &CurveIds, &CurveIndex](const UE::Anim::FCurveElement& InCurveElement)
+		{
+			CurveIds[CurveIndex] = OutputName(InCurveElement.Name);
+			CurveValues[CurveIndex] = InCurveElement.Value;
+			CurveIndex++;
+		});
+	}
+	
+	UE_TRACE_LOG(Animation, PoseWatch2, AnimationChannel)
+		<< PoseWatch2.Cycle(FPlatformTime::Cycles64())
+		<< PoseWatch2.RecordingTime(FObjectTrace::GetWorldElapsedTime(AnimInstance->GetWorld()))
+		<< PoseWatch2.ComponentId(FObjectTrace::GetObjectId(Component))
+		<< PoseWatch2.AnimInstanceId(FObjectTrace::GetObjectId(InSourceProxy.GetAnimInstanceObject()))
+		<< PoseWatch2.PoseWatchId(InPoseWatchId)
+#if WITH_EDITOR
+		<< PoseWatch2.NameId(OutputName(*InPoseWatchElement->GetParent()->GetLabel().ToString()))
+		<< PoseWatch2.Color(InPoseWatchElement->GetColor().DWColor())
+#endif
+		<< PoseWatch2.WorldTransform(reinterpret_cast<const float*>(&WorldTransform), sizeof(FTransform) / sizeof(float))
+		<< PoseWatch2.BoneTransforms(reinterpret_cast<const float*>(BoneTransforms.GetData()), BoneTransforms.Num() * (sizeof(FTransform) / sizeof(float)))
+		<< PoseWatch2.CurveValues(CurveValues.GetData(), CurveValues.Num())
+		<< PoseWatch2.CurveIds(CurveIds.GetData(), CurveIds.Num())
+		<< PoseWatch2.RequiredBones(RequiredBones.GetData(), RequiredBones.Num())
+		<< PoseWatch2.bIsEnabled(bIsEnabled);
 }
 
 #endif
diff --git a/Engine/Source/Runtime/Engine/Private/Animation/AnimationRuntime.cpp b/Engine/Source/Runtime/Engine/Private/Animation/AnimationRuntime.cpp
--- a/Engine/Source/Runtime/Engine/Private/Animation/AnimationRuntime.cpp
+++ b/Engine/Source/Runtime/Engine/Private/Animation/AnimationRuntime.cpp
@@ -223,7 +223,7 @@ FORCEINLINE void BlendCurves(const TArrayView<const FBlendedCurve> SourceCurves,
 {
 	if (SourceCurves.Num() > 0)
 	{
-#if INTEL_ISPC
+/*#if INTEL_ISPC
 		if (bAnim_BlendCurves_ISPC_Enabled)
 		{
 			OutCurve.InitFrom(SourceCurves[0]);
@@ -241,7 +241,7 @@ FORCEINLINE void BlendCurves(const TArrayView<const FBlendedCurve> SourceCurves,
 			}
 		}
 		else
-#endif
+#endif*/
 		{
 			OutCurve.Override(SourceCurves[0], SourceWeights[0]);
 			for (int32 CurveIndex = 1; CurveIndex < SourceCurves.Num(); ++CurveIndex)
@@ -1176,55 +1176,11 @@ void FAnimationRuntime::AccumulateMeshSpaceRotationAdditiveToLocalPoseInternal(F
 	}
 }
 
-void  FAnimationRuntime::MirrorCurves(FBlendedCurve& Curves, const UMirrorDataTable& MirrorDataTable)
+void FAnimationRuntime::MirrorCurves(FBlendedCurve& Curves, const UMirrorDataTable& MirrorDataTable)
 {
-	if (Curves.UIDToArrayIndexLUT == nullptr)
+	for(const TPair<FName, FName>& MirrorPair : MirrorDataTable.CurveToMirrorCurveMap)
 	{
-		return;
-	}
-
-	int32 NumMirrorUIDs = MirrorDataTable.CurveMirrorSourceUIDArray.Num();
-
-	for (int32 MirrorIndex = 0; MirrorIndex < NumMirrorUIDs; ++MirrorIndex)
-	{
-		SmartName::UID_Type SourceMirrorUID = MirrorDataTable.CurveMirrorSourceUIDArray[MirrorIndex];
-		SmartName::UID_Type TargetMirrorUID = MirrorDataTable.CurveMirrorTargetUIDArray[MirrorIndex];
-		int32 SourceMirrorIndex = Curves.GetArrayIndexByUID(SourceMirrorUID);
-		int32 TargetMirrorIndex = Curves.GetArrayIndexByUID(TargetMirrorUID);
-		if (SourceMirrorIndex != INDEX_NONE && TargetMirrorIndex != INDEX_NONE)
-		{
-			if (Curves.ValidCurveWeights[SourceMirrorIndex])
-			{
-				if (Curves.ValidCurveWeights[TargetMirrorIndex])
-				{
-					// Determine if we should swap or overwrite values. If the map has a paired entries (left->right and right->left) 
-					// these entries will appear beside each other in the arrays
-					SmartName::UID_Type NextSourceMirrorUID = INDEX_NONE; 
-					if (MirrorIndex + 1 < NumMirrorUIDs)
-					{
-						NextSourceMirrorUID = MirrorDataTable.CurveMirrorSourceUIDArray[MirrorIndex + 1];
-					}
-					if (NextSourceMirrorUID == TargetMirrorUID)
-					{
-						float SwapWeight = Curves.CurveWeights[TargetMirrorIndex];
-						Curves.CurveWeights[TargetMirrorIndex] = Curves.CurveWeights[SourceMirrorIndex];
-						Curves.CurveWeights[SourceMirrorIndex] = SwapWeight;
-						// skip over the next entry since it is swapped 
-						MirrorIndex++;
-					}
-					else
-					{
-						Curves.CurveWeights[TargetMirrorIndex] = Curves.CurveWeights[SourceMirrorIndex];
-					}
-				}
-				else
-				{
-					Curves.CurveWeights[TargetMirrorIndex] = Curves.CurveWeights[SourceMirrorIndex];
-					Curves.ValidCurveWeights[TargetMirrorIndex] = true;
-					Curves.ValidCurveWeights[SourceMirrorIndex] = false;
-				}
-			}
-		}
+		Curves.Mirror(MirrorPair.Key, MirrorPair.Value);
 	}
 }
 
diff --git a/Engine/Source/Runtime/Engine/Private/Animation/CurveIdentifier.cpp b/Engine/Source/Runtime/Engine/Private/Animation/CurveIdentifier.cpp
--- a/Engine/Source/Runtime/Engine/Private/Animation/CurveIdentifier.cpp
+++ b/Engine/Source/Runtime/Engine/Private/Animation/CurveIdentifier.cpp
@@ -3,29 +3,36 @@
 #include "Animation/AnimData/CurveIdentifier.h"
 #include "Animation/Skeleton.h"
 #include "EngineLogs.h"
+#include "Animation/Skeleton.h"
+#include "UObject/UE5MainStreamObjectVersion.h"
 
 #include UE_INLINE_GENERATED_CPP_BY_NAME(CurveIdentifier)
 
-#if WITH_EDITOR
-FAnimationCurveIdentifier UAnimationCurveIdentifierExtensions::GetCurveIdentifier(USkeleton* InSkeleton, FName Name, ERawCurveTrackTypes CurveType)
+void FAnimationCurveIdentifier::PostSerialize(const FArchive& Ar)
 {
-	FAnimationCurveIdentifier Identifier;
-
-	if (InSkeleton)
+#if WITH_EDITORONLY_DATA
+	if (Ar.IsLoading())
 	{
-		const FName MappingName = CurveType == ERawCurveTrackTypes::RCT_Float ? USkeleton::AnimCurveMappingName : USkeleton::AnimTrackCurveMappingName;
-		const bool bExists = InSkeleton->GetSmartNameByName(MappingName, Name, Identifier.InternalName);
-		if (!bExists)
+		if(Ar.CustomVer(FUE5MainStreamObjectVersion::GUID) < FUE5MainStreamObjectVersion::AnimationRemoveSmartNames)
 		{
-			InSkeleton->AddSmartNameAndModify(MappingName, Name, Identifier.InternalName);
+			CurveName = InternalName_DEPRECATED.DisplayName;
 		}
-
-		Identifier.CurveType = CurveType;
-	}
-	else
-	{
-		UE_LOG(LogAnimation, Warning, TEXT("Invalid Skeleton provided for GetCurveIdentifier"));
 	}
+#endif
+}
+
+#if WITH_EDITOR
+void UAnimationCurveIdentifierExtensions::SetCurveIdentifier(FAnimationCurveIdentifier& InOutIdentifier, FName Name, ERawCurveTrackTypes CurveType)
+{
+	InOutIdentifier.CurveName = Name;
+	InOutIdentifier.CurveType = CurveType;
+}
+
+FAnimationCurveIdentifier UAnimationCurveIdentifierExtensions::GetCurveIdentifier(USkeleton* InSkeleton, FName Name, ERawCurveTrackTypes CurveType)
+{
+	FAnimationCurveIdentifier Identifier;
+	Identifier.CurveName = Name;
+	Identifier.CurveType = CurveType;
 
 	return Identifier;
 }
@@ -33,17 +40,8 @@ FAnimationCurveIdentifier UAnimationCurveIdentifierExtensions::GetCurveIdentifie
 FAnimationCurveIdentifier UAnimationCurveIdentifierExtensions::FindCurveIdentifier(const USkeleton* InSkeleton, FName Name, ERawCurveTrackTypes CurveType)
 {
 	FAnimationCurveIdentifier Identifier;
-
-	if (InSkeleton)
-	{
-		const FName MappingName = CurveType == ERawCurveTrackTypes::RCT_Float ? USkeleton::AnimCurveMappingName : USkeleton::AnimTrackCurveMappingName;
-		const bool bExists = InSkeleton->GetSmartNameByName(MappingName, Name, Identifier.InternalName);
-		Identifier.CurveType = CurveType;
-	}
-	else
-	{
-		UE_LOG(LogAnimation, Warning, TEXT("Invalid Skeleton provided for FindCurveIdentifier"));
-	}
+	Identifier.CurveName = Name;
+	Identifier.CurveType = CurveType;
 
 	return Identifier;
 }
@@ -51,35 +49,6 @@ FAnimationCurveIdentifier UAnimationCurveIdentifierExtensions::FindCurveIdentifi
 TArray<FAnimationCurveIdentifier> UAnimationCurveIdentifierExtensions::GetCurveIdentifiers(USkeleton* InSkeleton, ERawCurveTrackTypes CurveType)
 {
 	TArray<FAnimationCurveIdentifier> Identifiers;
-	if (InSkeleton)
-	{
-		const FName MappingName = CurveType == ERawCurveTrackTypes::RCT_Float ? USkeleton::AnimCurveMappingName : USkeleton::AnimTrackCurveMappingName;
-		const FSmartNameMapping* SmartNameContainer = InSkeleton->GetSmartNameContainer(MappingName);
-		if (SmartNameContainer)
-		{
-			TArray<SmartName::UID_Type> UIDs;
-			SmartNameContainer->FillUidArray(UIDs);
-
-			for (const SmartName::UID_Type& UID : UIDs)
-			{
-				FSmartName SmartName;
-				if (SmartNameContainer->FindSmartNameByUID(UID, SmartName))
-				{
-
-					Identifiers.Add(FAnimationCurveIdentifier(SmartName, CurveType));
-				}
-			}
-		}
-		else
-		{
-			UE_LOG(LogAnimation, Warning, TEXT("Unable to find smartname mapping %s in Skeleton (%s)"), *MappingName.ToString(), *InSkeleton->GetName());
-		}
-	}
-	else
-	{
-		UE_LOG(LogAnimation, Warning, TEXT("Invalid Skeleton provided for GetCurveIdentifiers"));
-	}
-
 	return Identifiers;
 }
 
diff --git a/Engine/Source/Runtime/Engine/Private/Animation/MirrorDataTable.cpp b/Engine/Source/Runtime/Engine/Private/Animation/MirrorDataTable.cpp
--- a/Engine/Source/Runtime/Engine/Private/Animation/MirrorDataTable.cpp
+++ b/Engine/Source/Runtime/Engine/Private/Animation/MirrorDataTable.cpp
@@ -60,46 +60,21 @@ UMirrorDataTable::UMirrorDataTable(const FObjectInitializer& ObjectInitializer)
 #if WITH_EDITORONLY_DATA
 	OnDataTableChanged().AddUObject(this, &UMirrorDataTable::FillMirrorArrays);
 #endif 
-
 }
 
 void UMirrorDataTable::PostLoad()
 {
 	Super::PostLoad();
-	if (Skeleton)
-	{
-		Skeleton->ConditionalPostLoad();
-		Skeleton->OnSmartNamesChangedEvent.AddUObject(this, &UMirrorDataTable::FillMirrorArrays);
-	}
+
 	FillMirrorArrays(); 
 }
 
 #if WITH_EDITOR
 
-void UMirrorDataTable::PreEditChange(FProperty* PropertyThatWillChange)
-{
-	Super::PreEditChange(PropertyThatWillChange);
-
-	if (PropertyThatWillChange && PropertyThatWillChange->GetFName() == GET_MEMBER_NAME_STRING_CHECKED(UMirrorDataTable, Skeleton))
-	{
-		if (Skeleton)
-		{
-			Skeleton->OnSmartNamesChangedEvent.RemoveAll(this);
-		}
-	}
-}
-
 void UMirrorDataTable::PostEditChangeProperty(FPropertyChangedEvent& PropertyChangedEvent)
 {
 	Super::PostEditChangeProperty(PropertyChangedEvent);
 	FillMirrorArrays(); 
-	if (PropertyChangedEvent.Property && PropertyChangedEvent.Property->GetFName() == GET_MEMBER_NAME_CHECKED(UMirrorDataTable, Skeleton))
-	{
-		if (Skeleton)
-		{
-			Skeleton->OnSmartNamesChangedEvent.AddUObject(this, &UMirrorDataTable::FillMirrorArrays);
-		}
-	}
 }
 
 #endif // WITH_EDITOR
@@ -253,7 +228,7 @@ void UMirrorDataTable::FindReplaceMirroredNames()
 			AddMirrorRow(Notify, MirroredName, EMirrorRowType::AnimationNotify);
 		}
 	}
-
+	
 	for (const FName& SyncMarker : Skeleton->GetExistingMarkerNames())
 	{
 		FName MirroredName = FindReplace(SyncMarker);
@@ -263,26 +238,16 @@ void UMirrorDataTable::FindReplaceMirroredNames()
 		}
 	}
 
-	const FSmartNameMapping* CurveSmartNames = Skeleton->GetSmartNameContainer(USkeleton::AnimCurveMappingName);
-	if (CurveSmartNames)
+	TArray<FName> CurveNames;
+	Skeleton->GetCurveMetaDataNames(CurveNames);
+
+	// For every source curve, try to find the curve with the same name in the target.
+	for (const FName& CurveName : CurveNames)
 	{
-		const SmartName::UID_Type NumItems = CurveSmartNames->GetMaxUID() + 1;
-		TSet<FName> CurveNames; 
-		// For every source curve, try to find the curve with the same name in the target.
-		for (SmartName::UID_Type Index = 0; Index < NumItems; ++Index)
+		FName MirroredName = FindReplace(CurveName);
+		if (!MirroredName.IsNone() && CurveNames.Contains(MirroredName))
 		{
-			FName CurveName;
-			CurveSmartNames->GetName(Index, CurveName);
-			CurveNames.Add(CurveName);
-		}
-
-		for (const FName& CurveName : CurveNames)
-		{
-			FName MirroredName = FindReplace(CurveName);
-			if (!MirroredName.IsNone() && CurveNames.Contains(MirroredName))
-			{
-				AddMirrorRow(CurveName, MirroredName, EMirrorRowType::Curve);
-			}
+			AddMirrorRow(CurveName, MirroredName, EMirrorRowType::Curve);
 		}
 	}
 
@@ -393,19 +358,16 @@ void UMirrorDataTable::FillMirrorArrays()
 {
 	SyncToMirrorSyncMap.Empty();
 	AnimNotifyToMirrorAnimNotifyMap.Empty();
+	CurveToMirrorCurveMap.Empty();
 	if (!Skeleton)
 	{
 		BoneToMirrorBoneIndex.Empty();
-		CurveMirrorSourceUIDArray.Empty(); 
-		CurveMirrorTargetUIDArray.Empty(); 
 		return; 
 	}
 
 	FillMirrorBoneIndexes(Skeleton, BoneToMirrorBoneIndex);
-
-	TMap<FName, FName> CurveToMirrorCurveMap;
 	
-	ForeachRow<FMirrorTableRow>(TEXT("UMirrorDataTable::FillMirrorArrays"), [this, &CurveToMirrorCurveMap](const FName& Key, const FMirrorTableRow& Value) mutable
+	ForeachRow<FMirrorTableRow>(TEXT("UMirrorDataTable::FillMirrorArrays"), [this](const FName& Key, const FMirrorTableRow& Value) mutable
 		{
 			switch (Value.MirrorEntryType)
 			{
@@ -429,38 +391,6 @@ void UMirrorDataTable::FillMirrorArrays()
 			}
 		}
 	);
-
-	// Ensure post load prepares the smart names
-	Skeleton->ConditionalPostLoad();
-
-	//ensure that pairs always appear beside each other in the arrays
-	TSet<SmartName::UID_Type> AddedSourceUIDs; 
-	const FSmartNameMapping* CurveSmartNames = Skeleton->GetSmartNameContainer(USkeleton::AnimCurveMappingName);
-
-	if (CurveSmartNames)
-	{
-		CurveMirrorSourceUIDArray.Reset(CurveToMirrorCurveMap.Num());
-		CurveMirrorTargetUIDArray.Reset(CurveToMirrorCurveMap.Num());
-		for (auto& Elem : CurveToMirrorCurveMap)
-		{
-			SmartName::UID_Type SourceCurveUID = CurveSmartNames->FindUID(Elem.Key);
-			SmartName::UID_Type TargetCurveUID = CurveSmartNames->FindUID(Elem.Value);
-			if (SourceCurveUID != INDEX_NONE && TargetCurveUID != INDEX_NONE && !AddedSourceUIDs.Contains(SourceCurveUID))
-			{
-				CurveMirrorSourceUIDArray.Add(SourceCurveUID);
-				AddedSourceUIDs.Add(SourceCurveUID); 
-				CurveMirrorTargetUIDArray.Add(TargetCurveUID);
-				if (CurveToMirrorCurveMap.Contains(Elem.Value) && CurveSmartNames->FindUID(CurveToMirrorCurveMap[Elem.Value]) == SourceCurveUID)
-				{
-					CurveMirrorSourceUIDArray.Add(TargetCurveUID);
-					AddedSourceUIDs.Add(TargetCurveUID);
-					CurveMirrorTargetUIDArray.Add(SourceCurveUID);
-				}
-			}
-		}
-		CurveMirrorSourceUIDArray.Shrink(); 
-		CurveMirrorTargetUIDArray.Shrink();
-	}
 }
 
 #undef LOCTEXT_NAMESPACE
diff --git a/Engine/Source/Runtime/Engine/Private/Animation/PoseAsset.cpp b/Engine/Source/Runtime/Engine/Private/Animation/PoseAsset.cpp
--- a/Engine/Source/Runtime/Engine/Private/Animation/PoseAsset.cpp
+++ b/Engine/Source/Runtime/Engine/Private/Animation/PoseAsset.cpp
@@ -20,6 +20,11 @@
 #include "UObject/UnrealType.h"
 #include "SkeletonRemappingRegistry.h"
 #include "Animation/SkeletonRemapping.h"
+#include "UObject/UE5MainStreamObjectVersion.h"
+#include "UObject/UE5ReleaseStreamObjectVersion.h"
+#include "Widgets/Notifications/SNotificationList.h"
+#include "Animation/AnimCurveUtils.h"
+#include "UObject/UE5MainStreamObjectVersion.h"
 
 #include UE_INLINE_GENERATED_CPP_BY_NAME(PoseAsset)
 
@@ -31,10 +36,35 @@
 //
 //////////////////////////////////////////////////////////////////////////////////////////////////////////////////
 
+bool FPoseDataContainer::Serialize(FArchive& Ar)
+{
+#if WITH_EDITORONLY_DATA
+	Ar.UsingCustomVersion(FUE5MainStreamObjectVersion::GUID);
+#endif
+
+	return false;
+}
+
+void FPoseDataContainer::PostSerialize(const FArchive& Ar)
+{
+#if WITH_EDITORONLY_DATA
+	if(Ar.CustomVer(FUE5MainStreamObjectVersion::GUID) < FUE5MainStreamObjectVersion::AnimationRemoveSmartNames)
+	{
+		for(const FSmartName& PoseName : PoseNames_DEPRECATED)
+		{
+			PoseFNames.Add(PoseName.DisplayName);
+		}
+	}
+#endif
+
+	RebuildCurveIndexTable();
+}
+
 void FPoseDataContainer::Reset()
 {
 	// clear everything
-	PoseNames.Reset();
+	PoseFNames.Reset();
+	SortedPoseIndices.Reset();
 	Poses.Reset();
 	Tracks.Reset();
 	TrackBoneIndices.Reset();
@@ -48,12 +78,17 @@ void FPoseDataContainer::GetPoseCurve(const FPoseData* PoseData, FBlendedCurve&
 		const TArray<float>& CurveValues = PoseData->CurveData;
 		checkSlow(CurveValues.Num() == Curves.Num());
 
-		// extract curve - not optimized, can use optimization
-		for (int32 CurveIndex = 0; CurveIndex < Curves.Num(); ++CurveIndex)
+		auto GetNameFromIndex = [this](int32 InCurveIndex)
 		{
-			const FAnimCurveBase& Curve = Curves[CurveIndex];
-			OutCurve.Set(Curve.Name.UID, CurveValues[CurveIndex]);
-		}
+			return PoseFNames[SortedPoseIndices[InCurveIndex]];
+		};
+
+		auto GetValueFromIndex = [this, &CurveValues](int32 InCurveIndex)
+		{
+			return CurveValues[SortedPoseIndices[InCurveIndex]];
+		};
+
+		UE::Anim::FCurveUtils::BuildSorted(OutCurve, CurveValues.Num(), GetNameFromIndex, GetValueFromIndex, OutCurve.GetFilter());
 	}
 }
 
@@ -64,20 +99,23 @@ void FPoseDataContainer::BlendPoseCurve(const FPoseData* PoseData, FBlendedCurve
 		const TArray<float>& CurveValues = PoseData->CurveData;
 		checkSlow(CurveValues.Num() == Curves.Num());
 
-		for (int32 CurveIndex = 0; CurveIndex < Curves.Num(); ++CurveIndex)
-		{
-			const FAnimCurveBase& Curve = Curves[CurveIndex];
+		FBlendedCurve Curve;
+		Curve.SetFilter(InOutCurve.GetFilter());
 
-			InOutCurve.Set(Curve.Name.UID, CurveValues[CurveIndex] * Weight + InOutCurve.Get(Curve.Name.UID));
-		}
+		GetPoseCurve(PoseData, InOutCurve);
+
+		UE::Anim::FNamedValueArrayUtils::Union(InOutCurve, Curve,
+			[Weight](UE::Anim::FCurveElement& InOutElement, const UE::Anim::FCurveElement& InElement, UE::Anim::ENamedValueUnionFlags InFlags)
+			{
+				InOutElement.Value = InElement.Value * Weight + InOutElement.Value;
+				InOutElement.Flags |= InElement.Flags;
+			});
 	}
 }
 
-
-
-FPoseData* FPoseDataContainer::FindPoseData(FSmartName PoseName)
+FPoseData* FPoseDataContainer::FindPoseData(FName PoseName)
 {
-	int32 PoseIndex = PoseNames.Find(PoseName);
+	int32 PoseIndex = PoseFNames.Find(PoseName);
 	if (PoseIndex != INDEX_NONE)
 	{
 		return &Poses[PoseIndex];
@@ -86,12 +124,12 @@ FPoseData* FPoseDataContainer::FindPoseData(FSmartName PoseName)
 	return nullptr;
 }
 
-FPoseData* FPoseDataContainer::FindOrAddPoseData(FSmartName PoseName)
+FPoseData* FPoseDataContainer::FindOrAddPoseData(FName PoseName)
 {
-	int32 PoseIndex = PoseNames.Find(PoseName);
+	int32 PoseIndex = PoseFNames.Find(PoseName);
 	if (PoseIndex == INDEX_NONE)
 	{
-		PoseIndex = PoseNames.Add(PoseName);
+		PoseIndex = PoseFNames.Add(PoseName);
 		check(PoseIndex == Poses.AddZeroed(1));
 	}
 
@@ -121,10 +159,10 @@ FTransform FPoseDataContainer::GetDefaultTransform(int32 SkeletonIndex, const TA
 
 
 #if WITH_EDITOR
-void FPoseDataContainer::AddOrUpdatePose(const FSmartName& InPoseName, const TArray<FTransform>& InLocalSpacePose, const TArray<float>& InCurveData)
+void FPoseDataContainer::AddOrUpdatePose(const FName& InPoseName, const TArray<FTransform>& InLocalSpacePose, const TArray<float>& InCurveData)
 {
-	// make sure the transform is correct size
-	if (ensureAlways(InLocalSpacePose.Num() == Tracks.Num()))
+	// make sure the transforms and curves are the correct size
+	if (ensureAlways(InLocalSpacePose.Num() == Tracks.Num()) && ensureAlways(InCurveData.Num() == Curves.Num()))
 	{
 		// find or add pose data
 		FPoseData* PoseDataPtr = FindOrAddPoseData(InPoseName);
@@ -190,21 +228,21 @@ bool FPoseDataContainer::FillUpSkeletonPose(FPoseData* PoseData, const USkeleton
 	return false;
 }
 
-void FPoseDataContainer::RenamePose(FSmartName OldPoseName, FSmartName NewPoseName)
+void FPoseDataContainer::RenamePose(FName OldPoseName, FName NewPoseName)
 {
-	int32 PoseIndex = PoseNames.Find(OldPoseName);
+	int32 PoseIndex = PoseFNames.Find(OldPoseName);
 	if (PoseIndex != INDEX_NONE)
 	{
-		PoseNames[PoseIndex] = NewPoseName;
+		PoseFNames[PoseIndex] = NewPoseName;
 	}
 }
 
-int32 FPoseDataContainer::DeletePose(FSmartName PoseName)
+int32 FPoseDataContainer::DeletePose(FName PoseName)
 {
-	int32 PoseIndex = PoseNames.Find(PoseName);
+	int32 PoseIndex = PoseFNames.Find(PoseName);
 	if (PoseIndex != INDEX_NONE)
 	{
-		PoseNames.RemoveAt(PoseIndex);
+		PoseFNames.RemoveAt(PoseIndex);
 		Poses.RemoveAt(PoseIndex);
 		return PoseIndex;
 	}
@@ -212,11 +250,11 @@ int32 FPoseDataContainer::DeletePose(FSmartName PoseName)
 	return INDEX_NONE;
 }
 
-bool FPoseDataContainer::DeleteCurve(FSmartName CurveName)
+bool FPoseDataContainer::DeleteCurve(FName CurveName)
 {
 	for (int32 CurveIndex = 0; CurveIndex < Curves.Num(); ++CurveIndex)
 	{
-		if (Curves[CurveIndex].Name == CurveName)
+		if (Curves[CurveIndex].GetName() == CurveName)
 		{
 			Curves.RemoveAt(CurveIndex);
 
@@ -345,6 +383,21 @@ void FPoseDataContainer::ConvertToAdditivePose(const TArray<FTransform>& InBaseP
 	}
 }
 #endif // WITH_EDITOR
+
+void FPoseDataContainer::RebuildCurveIndexTable()
+{
+	// Recreate sorted curve index table
+	SortedPoseIndices.SetNumUninitialized(PoseFNames.Num());
+	for(int32 NameIndex = 0; NameIndex < SortedPoseIndices.Num(); ++NameIndex)
+	{
+		SortedPoseIndices[NameIndex] = NameIndex;
+	}
+
+	SortedPoseIndices.Sort([&PoseFNames = PoseFNames](int32 LHS, int32 RHS)
+	{
+		return PoseFNames[LHS].FastLess(PoseFNames[RHS]);
+	});
+}
 /////////////////////////////////////////////////////
 // UPoseAsset
 /////////////////////////////////////////////////////
@@ -481,23 +534,19 @@ void UPoseAsset::GetAnimationCurveOnly(TArray<FName>& InCurveNames, TArray<float
 		// we iterate through to see if we have that corresponding pose
 		for (int32 CurveIndex = 0; CurveIndex < InCurveNames.Num(); ++CurveIndex)
 		{
-			FSmartName PoseSmartName;
-			if (MySkeleton->GetSmartNameByName(USkeleton::AnimCurveMappingName, InCurveNames[CurveIndex], PoseSmartName))
+			int32 PoseIndex = PoseContainer.PoseFNames.Find(InCurveNames[CurveIndex]);
+			if (ensure(PoseIndex != INDEX_NONE))
 			{
-				int32 PoseIndex = PoseContainer.PoseNames.Find(PoseSmartName);
-				if (ensure(PoseIndex != INDEX_NONE))
-				{
-					const FPoseData& PoseData = PoseContainer.Poses[PoseIndex];
-					const float Value = InCurveValues[CurveIndex];
+				const FPoseData& PoseData = PoseContainer.Poses[PoseIndex];
+				const float Value = InCurveValues[CurveIndex];
 
-					// we only add to the list if it's not additive Or if it's additive, we don't want to add base pose index
-					// and has weight
-					if ((!bAdditivePose || PoseIndex != BasePoseIndex) && FAnimationRuntime::HasWeight(Value))
-					{
-						TotalWeight += Value;												
-						PoseWeights[PoseIndex] = Value;
-						WeightedPoseIndices.Add(PoseIndex);
-					}
+				// we only add to the list if it's not additive Or if it's additive, we don't want to add base pose index
+				// and has weight
+				if ((!bAdditivePose || PoseIndex != BasePoseIndex) && FAnimationRuntime::HasWeight(Value))
+				{
+					TotalWeight += Value;												
+					PoseWeights[PoseIndex] = Value;
+					WeightedPoseIndices.Add(PoseIndex);
 				}
 			}
 		}
@@ -505,17 +554,8 @@ void UPoseAsset::GetAnimationCurveOnly(TArray<FName>& InCurveNames, TArray<float
 		const int32 TotalNumberOfValidPoses = WeightedPoseIndices.Num();
 		if (TotalNumberOfValidPoses > 0)
 		{
-			// collect curves
-			TArray<uint16> CurveUIDList;
-			CurveUIDList.AddUninitialized(PoseContainer.Curves.Num());
-			for (int32 CurveIndex = 0; CurveIndex < PoseContainer.Curves.Num(); ++CurveIndex)
-			{
-				CurveUIDList[CurveIndex] = PoseContainer.Curves[CurveIndex].Name.UID;
-			}
-
 			// blend curves
 			FBlendedCurve BlendedCurve;
-			BlendedCurve.InitFrom(&CurveUIDList);
 
 			//if full pose, we'll have to normalize by weight
 			if (bNormalizeWeight && TotalWeight > 1.f)
@@ -542,19 +582,12 @@ void UPoseAsset::GetAnimationCurveOnly(TArray<FName>& InCurveNames, TArray<float
 
 			OutCurveNames.Reset();
 			OutCurveValues.Reset();
-			for (int32 Idx = 0; Idx < CurveUIDList.Num(); ++Idx)
+
+			BlendedCurve.ForEachElement([&OutCurveNames, &OutCurveValues](const UE::Anim::FCurveElement& InElement)
 			{
-				USkeleton::AnimCurveUID UID = Idx;
-				if (BlendedCurve.IsEnabled(UID))
-				{
-					FSmartName CurveName;
-					if (MySkeleton->GetSmartNameByUID(USkeleton::AnimCurveMappingName, UID, CurveName))
-					{
-						OutCurveNames.Add(CurveName.DisplayName);
-						OutCurveValues.Add(BlendedCurve.Get(UID));
-					}
-				}
-			}
+				OutCurveNames.Add(InElement.Name);
+				OutCurveValues.Add(InElement.Value);
+			});
 		}
 	}
 }
@@ -849,16 +882,16 @@ bool UPoseAsset::GetAnimationPose(struct FAnimationPoseData& OutAnimationPoseDat
 
 bool UPoseAsset::IsPostLoadThreadSafe() const
 {
-	return false;	// PostLoad is not thread safe because of the call to VerifySmartName() that can mutate a shared map in the skeleton.
+	return WITH_EDITORONLY_DATA == 0;	// Not thread safe in editor as the skeleton can be modified on version upgrade
 }
 
 void UPoseAsset::PostLoad()
 {
 	Super::PostLoad();
 
-// moved to PostLoad because Skeleton is not completely loaded when we do this in Serialize
-// and we need Skeleton
-#if WITH_EDITOR
+#if WITH_EDITORONLY_DATA
+	// moved to PostLoad because Skeleton is not completely loaded when we do this in Serialize
+	// and we need Skeleton
 	if (GetLinkerCustomVersion(FFrameworkObjectVersion::GUID) < FFrameworkObjectVersion::PoseAssetSupportPerBoneMask && GetLinkerCustomVersion(FAnimPhysObjectVersion::GUID) >= FAnimPhysObjectVersion::SaveEditorOnlyFullPoseForPoseAsset)
 	{
 		// fix curve names
@@ -932,19 +965,11 @@ void UPoseAsset::PostLoad()
 			}
 		}	
 	}	
-#endif // WITH_EDITOR
 
 	// fix curve names
 	USkeleton* MySkeleton = GetSkeleton();
 	if (MySkeleton)
 	{
-		MySkeleton->VerifySmartNames(USkeleton::AnimCurveMappingName, PoseContainer.PoseNames);
-
-		for (FAnimCurveBase& Curve : PoseContainer.Curves)
-		{
-			MySkeleton->VerifySmartName(USkeleton::AnimCurveMappingName, Curve.Name);
-		}
-
 		// double loop but this check only should happen once per asset
 		// this should continue to add if skeleton hasn't been saved either 
 		if (GetLinkerCustomVersion(FFrameworkObjectVersion::GUID) < FFrameworkObjectVersion::MoveCurveTypesToSkeleton 
@@ -959,20 +984,19 @@ void UPoseAsset::PostLoad()
 				// only add this if that has to 
 				if (bMorphtargetSet || bMaterialSet)
 				{
-					MySkeleton->AccumulateCurveMetaData(Curve.Name.DisplayName, bMaterialSet, bMorphtargetSet);
+					MySkeleton->AccumulateCurveMetaData(Curve.GetName(), bMaterialSet, bMorphtargetSet);
 				}
 			}
 		}
 	}
 
 	// I have to fix pose names
-#if WITH_EDITOR
 	if(RemoveInvalidTracks())
 	{
 		PostProcessData();
 	}
 	else
-#endif // WITH_EDITOR
+#endif // WITH_EDITORONLY_DATA
 	{
 		UpdateTrackBoneIndices();
 	}
@@ -1014,19 +1038,22 @@ void UPoseAsset::GetAssetRegistryTags(TArray<FAssetRegistryTag>& OutTags) const
 	OutTags.Add(FAssetRegistryTag("Poses", FString::FromInt(GetNumPoses()), FAssetRegistryTag::TT_Numerical));
 #if WITH_EDITOR
 	TArray<FName> Names;
-	Names.Reserve(PoseContainer.PoseNames.Num() + PoseContainer.Curves.Num());
+	Names.Reserve(PoseContainer.PoseFNames.Num() + PoseContainer.Curves.Num());
 
-	for (const FSmartName& SmartName : PoseContainer.PoseNames)
+	for (const FName& PoseName : PoseContainer.PoseFNames)
 	{
-		Names.Add(SmartName.DisplayName);
+		Names.Add(PoseName);
 	}
 
 	for (const FAnimCurveBase& Curve : PoseContainer.Curves)
 	{
-		Names.AddUnique(Curve.Name.DisplayName);
+		Names.AddUnique(Curve.GetName());
 	}
-
-	FString PoseNameList;
+	
+	// Add curve IDs to a tag list, or a delimiter if we have no curves.
+	// The delimiter is necessary so we can distinguish between data with no curves and old data, as the asset registry
+	// strips tags that have empty values 
+	FString PoseNameList = USkeleton::CurveTagDelimiter;
 	for(const FName& Name : Names)
 	{
 		PoseNameList += FString::Printf(TEXT("%s%s"), *Name.ToString(), *USkeleton::CurveTagDelimiter);
@@ -1050,29 +1077,38 @@ int32 UPoseAsset::GetNumTracks() const
 	return PoseContainer.Tracks.Num();
 }
 
+const TArray<FName>& UPoseAsset::GetPoseFNames() const
+{
+	return PoseContainer.PoseFNames;
+}
 
 const TArray<FSmartName> UPoseAsset::GetPoseNames() const
 {
-	return PoseContainer.PoseNames;
+	return TArray<FSmartName>();
 }
 
-const TArray<FName>	UPoseAsset::GetTrackNames() const
+const TArray<FName>& UPoseAsset::GetTrackNames() const
 {
 	return PoseContainer.Tracks;
 }
 
 const TArray<FSmartName> UPoseAsset::GetCurveNames() const
 {
-	TArray<FSmartName> CurveNames;
+	return TArray<FSmartName>();
+}
+
+const TArray<FName> UPoseAsset::GetCurveFNames() const
+{
+	TArray<FName> CurveNames;
 	for (int32 CurveIndex = 0; CurveIndex < PoseContainer.Curves.Num(); ++CurveIndex)
 	{
-		CurveNames.Add(PoseContainer.Curves[CurveIndex].Name);
+		CurveNames.Add(PoseContainer.Curves[CurveIndex].GetName());
 	}
 
 	return CurveNames;
 }
 
-const TArray<FAnimCurveBase> UPoseAsset::GetCurveData() const
+const TArray<FAnimCurveBase>& UPoseAsset::GetCurveData() const
 {
 	return PoseContainer.Curves;
 }
@@ -1122,9 +1158,9 @@ const int32 UPoseAsset::GetTrackIndexByName(const FName& InTrackName) const
 
 bool UPoseAsset::ContainsPose(const FName& InPoseName) const
 {
-	for (const FSmartName& PoseName : PoseContainer.PoseNames)
+	for (const FName& PoseName : PoseContainer.PoseFNames)
 	{
-		if (PoseName.DisplayName == InPoseName)
+		if (PoseName == InPoseName)
 		{
 			return true;
 		}
@@ -1136,7 +1172,7 @@ bool UPoseAsset::ContainsPose(const FName& InPoseName) const
 #if WITH_EDITOR
 void UPoseAsset::RenamePose(const FName& OriginalPoseName, const FName& NewPoseName)
 {
-	ModifyPoseName(OriginalPoseName, NewPoseName, nullptr);
+	ModifyPoseName(OriginalPoseName, NewPoseName);
 }
 
 void UPoseAsset::GetPoseNames(TArray<FName>& PoseNames) const
@@ -1170,26 +1206,19 @@ void UPoseAsset::BreakAnimationSequenceGUIDComparison()
 	SourceAnimationRawDataGUID.Invalidate();
 }
 
-bool UPoseAsset::AddOrUpdatePoseWithUniqueName(const USkeletalMeshComponent* MeshComponent, FSmartName* OutPoseName /*= nullptr*/)
+FName UPoseAsset::AddPoseWithUniqueName(const USkeletalMeshComponent* MeshComponent)
 {
-	bool bSavedAdditivePose = bAdditivePose;
-
-	FSmartName NewPoseName = GetUniquePoseSmartName(GetSkeleton());
+	const FName NewPoseName = GetUniquePoseName(this);
 	AddOrUpdatePose(NewPoseName, MeshComponent);
 
-	if (OutPoseName)
-	{
-		*OutPoseName = NewPoseName;
-	}
-
 	PostProcessData();
 
 	OnPoseListChanged.Broadcast();
 
-	return true;
+	return NewPoseName;
 }
 
-void UPoseAsset::AddReferencePose(const FSmartName& PoseName, const FReferenceSkeleton& RefSkeleton)
+void UPoseAsset::AddReferencePose(const FName& PoseName, const FReferenceSkeleton& RefSkeleton)
 {
 	TArray<FTransform> BoneTransforms;
 	TArray<FName> TrackNames;
@@ -1208,7 +1237,7 @@ void UPoseAsset::AddReferencePose(const FSmartName& PoseName, const FReferenceSk
 	PostProcessData();
 }
 
-void UPoseAsset::AddOrUpdatePose(const FSmartName& PoseName, const USkeletalMeshComponent* MeshComponent, bool bUpdateCurves)
+void UPoseAsset::AddOrUpdatePose(const FName& PoseName, const USkeletalMeshComponent* MeshComponent, bool bUpdateCurves)
 {
 	const USkeleton* MySkeleton = GetSkeleton();
 	if (MySkeleton && MeshComponent && MeshComponent->GetSkeletalMeshAsset())
@@ -1231,27 +1260,17 @@ void UPoseAsset::AddOrUpdatePose(const FSmartName& PoseName, const USkeletalMesh
 				BoneTransform[BoneIndex] = BoneTransform[BoneIndex].GetRelativeTransform(BoneTransform[ParentIndex]);
 			}
 		}
-
-		const USkeleton* MeshSkeleton = MeshComponent->GetSkeletalMeshAsset()->GetSkeleton();
-		const FSmartNameMapping* Mapping = MeshSkeleton->GetSmartNameContainer(USkeleton::AnimCurveMappingName);
 		
 		TArray<float> NewCurveValues;
 		NewCurveValues.AddZeroed(PoseContainer.Curves.Num());
 
-		if(Mapping)
+		const FBlendedHeapCurve& MeshCurves = MeshComponent->GetAnimationCurves();
+
+		for (int32 NewCurveIndex = 0; NewCurveIndex < NewCurveValues.Num(); ++NewCurveIndex)
 		{
-			const FBlendedHeapCurve& MeshCurves = MeshComponent->GetAnimationCurves();
- 
-			for (int32 NewCurveIndex = 0; NewCurveIndex < NewCurveValues.Num(); ++NewCurveIndex)
-			{
-				const FAnimCurveBase& Curve = PoseContainer.Curves[NewCurveIndex];
-				const SmartName::UID_Type CurveUID = Mapping->FindUID(Curve.Name.DisplayName);
-				if (CurveUID != SmartName::MaxUID)
-				{
-					const float MeshCurveValue = MeshCurves.Get(CurveUID);
-					NewCurveValues[NewCurveIndex] = MeshCurveValue;
-				}
-			}
+			const FAnimCurveBase& Curve = PoseContainer.Curves[NewCurveIndex];
+			const float MeshCurveValue = MeshCurves.Get(Curve.GetName());
+			NewCurveValues[NewCurveIndex] = MeshCurveValue;
 		}
 		
 		BreakAnimationSequenceGUIDComparison();
@@ -1263,7 +1282,7 @@ void UPoseAsset::AddOrUpdatePose(const FSmartName& PoseName, const USkeletalMesh
 	}
 }
 
-void UPoseAsset::AddOrUpdatePose(const FSmartName& PoseName, const TArray<FName>& TrackNames, const TArray<FTransform>& LocalTransform, const TArray<float>& CurveValues)
+void UPoseAsset::AddOrUpdatePose(const FName& PoseName, const TArray<FName>& TrackNames, const TArray<FTransform>& LocalTransform, const TArray<float>& CurveValues)
 {
 	const USkeleton* MySkeleton = GetSkeleton();
 	if (MySkeleton)
@@ -1376,32 +1395,50 @@ void UPoseAsset::Reinitialize()
 
 void UPoseAsset::RenameSmartName(const FName& InOriginalName, const FName& InNewName)
 {
-	for (FSmartName SmartName : PoseContainer.PoseNames)
+	RenamePoseOrCurveName(InOriginalName, InNewName);
+}
+
+void UPoseAsset::RenamePoseOrCurveName(const FName& InOriginalName, const FName& InNewName)
+{
+	if(PoseContainer.PoseFNames.Contains(InNewName) || PoseContainer.Curves.ContainsByPredicate([InNewName](const FAnimCurveBase& InCurve){ return InCurve.GetName() == InNewName; }))
+	{
+		// Cant rename on top of something we already have - this will create duplicates
+		return;
+	}
+
+	for (FName Name : PoseContainer.PoseFNames)
 	{
-		if (SmartName.DisplayName == InOriginalName)
+		if (Name == InOriginalName)
 		{
-			SmartName.DisplayName = InNewName;
+			Name = InNewName;
 			break;
 		}
 	}
 
 	for (FAnimCurveBase& Curve : PoseContainer.Curves)
 	{
-		if (Curve.Name.DisplayName == InOriginalName)
+		if (Curve.GetName() == InOriginalName)
 		{
-			Curve.Name.DisplayName = InNewName;
+			Curve.SetName(InNewName);
 			break;
 		}
 	}
+
+	PoseContainer.RebuildCurveIndexTable();
 }
 
 void UPoseAsset::RemoveSmartNames(const TArray<FName>& InNamesToRemove)
+{
+	RemovePoseOrCurveNames(InNamesToRemove);
+}
+
+void UPoseAsset::RemovePoseOrCurveNames(const TArray<FName>& InNamesToRemove)
 {
 	DeletePoses(InNamesToRemove);
 	DeleteCurves(InNamesToRemove);
 }
 
-void UPoseAsset::CreatePoseFromAnimation(class UAnimSequence* AnimSequence, const TArray<FSmartName>* InPoseNames/*== nullptr*/)
+void UPoseAsset::CreatePoseFromAnimation(class UAnimSequence* AnimSequence, const TArray<FName>* InPoseNames/*== nullptr*/)
 {
 	if (AnimSequence)
 	{
@@ -1442,34 +1479,28 @@ void UPoseAsset::CreatePoseFromAnimation(class UAnimSequence* AnimSequence, cons
 				const int32 NumTracks = TrackNames.Num();
 				NewPose.Reset(NumTracks);
 				NewPose.AddUninitialized(NumTracks);
-
-				// @Todo fill up curve data
-				TArray<float> CurveData;
+				
 				float IntervalBetweenKeys = (NumPoses > 1)? AnimSequence->GetPlayLength() / (NumPoses -1 ) : 0.f;
 
 				// add curves - only float curves
 				const FAnimationCurveData& AnimationCurveData = DataModel->GetCurveData();
-
 				const int32 TotalFloatCurveCount = AnimationCurveData.FloatCurves.Num();
-				
-				// have to construct own UIDList;
-				// copy default UIDLIst
-				TArray<SmartName::UID_Type> UIDList;
 
 				if (TotalFloatCurveCount > 0)
 				{
 					for (const FFloatCurve& Curve : AnimationCurveData.FloatCurves)
 					{
-						PoseContainer.Curves.Add(FAnimCurveBase(Curve.Name, Curve.GetCurveTypeFlags()));
-						UIDList.Add(Curve.Name.UID);
+						PoseContainer.Curves.Add(FAnimCurveBase(Curve.GetName(), Curve.GetCurveTypeFlags()));
 					}
 				}
 
-				CurveData.AddZeroed(UIDList.Num());
 				// add to skeleton UID, so that it knows the curve data
 				for (int32 PoseIndex = 0; PoseIndex < NumPoses; ++PoseIndex)
 				{
-					FSmartName NewPoseName = (InPoseNames && InPoseNames->IsValidIndex(PoseIndex))? (*InPoseNames)[PoseIndex] : GetUniquePoseSmartName(TargetSkeleton);
+					TArray<float> CurveData;
+					CurveData.Reserve(TotalFloatCurveCount);
+					
+					FName NewPoseName = (InPoseNames && InPoseNames->IsValidIndex(PoseIndex))? (*InPoseNames)[PoseIndex] : GetUniquePoseName(this);
 					// now get rawanimationdata, and each key is converted to new pose
 					for (int32 TrackIndex = 0; TrackIndex < NumTracks; ++TrackIndex)
 					{
@@ -1482,15 +1513,12 @@ void UPoseAsset::CreatePoseFromAnimation(class UAnimSequence* AnimSequence, cons
 						// have to do iterate over time
 						// support curve
 						FBlendedCurve SourceCurve;
-						SourceCurve.InitFrom(&TargetSkeleton->GetDefaultCurveUIDList());
 						AnimSequence->EvaluateCurveData(SourceCurve, PoseIndex*IntervalBetweenKeys, true);
-
-						// copy back to CurveData
-						for (int32 CurveIndex = 0; CurveIndex < CurveData.Num(); ++CurveIndex)
+						
+						SourceCurve.ForEachElement([&CurveData](const UE::Anim::FCurveElement& InElement)
 						{
-							CurveData[CurveIndex] = SourceCurve.Get(UIDList[CurveIndex]);
-						}
-
+							CurveData.Add(InElement.Value);
+						});
 						check(CurveData.Num() == PoseContainer.Curves.Num());
 					}
 				
@@ -1513,7 +1541,7 @@ void UPoseAsset::UpdatePoseFromAnimation(class UAnimSequence* AnimSequence)
 		// when you update pose, right now, it just only keeps pose names
 		// in the future we might want to make it more flexible
 		// back up old pose names
-		const TArray<FSmartName> OldPoseNames = PoseContainer.PoseNames;
+		const TArray<FName> OldPoseNames = PoseContainer.PoseFNames;
 		const bool bOldAdditive = bAdditivePose;
 		int32 OldBasePoseIndex = BasePoseIndex;
 		CreatePoseFromAnimation(AnimSequence, &OldPoseNames);
@@ -1531,36 +1559,23 @@ void UPoseAsset::UpdatePoseFromAnimation(class UAnimSequence* AnimSequence)
 			ConvertToAdditivePose(OldBasePoseIndex);
 		}
 
+		PoseContainer.RebuildCurveIndexTable();
+
 		OnPoseListChanged.Broadcast();
 	}
 }
 
-bool UPoseAsset::ModifyPoseName(FName OldPoseName, FName NewPoseName, const SmartName::UID_Type* NewUID)
+bool UPoseAsset::ModifyPoseName(FName OldPoseName, FName NewPoseName)
 {
-	USkeleton* MySkeleton = GetSkeleton();
-
 	if (ContainsPose(NewPoseName))
 	{
 		// already exists, return 
 		return false;
 	}
 
-	FSmartName OldPoseSmartName;
-	ensureAlways(MySkeleton->GetSmartNameByName(USkeleton::AnimCurveMappingName, OldPoseName, OldPoseSmartName));
-
-	if (FPoseData* PoseData = PoseContainer.FindPoseData(OldPoseSmartName))
+	if (FPoseData* PoseData = PoseContainer.FindPoseData(OldPoseName))
 	{
-		FSmartName NewPoseSmartName;
-		if (NewUID)
-		{
-			MySkeleton->GetSmartNameByUID(USkeleton::AnimCurveMappingName, *NewUID, NewPoseSmartName);
-		}
-		else
-		{
-			MySkeleton->AddSmartNameAndModify(USkeleton::AnimCurveMappingName, NewPoseName, NewPoseSmartName);
-		}
-
-		PoseContainer.RenamePose(OldPoseSmartName, NewPoseSmartName);
+		PoseContainer.RenamePose(OldPoseName, NewPoseName);
 		OnPoseListChanged.Broadcast();
 
 		return true;
@@ -1573,27 +1588,21 @@ int32 UPoseAsset::DeletePoses(TArray<FName> PoseNamesToDelete)
 {
 	int32 ItemsDeleted = 0;
 
-	USkeleton* MySkeleton = GetSkeleton();
-
 	for (const FName& PoseName : PoseNamesToDelete)
 	{
-		FSmartName PoseSmartName;
-		if (MySkeleton->GetSmartNameByName(USkeleton::AnimCurveMappingName, PoseName, PoseSmartName))
+		int32 PoseIndexDeleted = PoseContainer.DeletePose(PoseName);
+		if (PoseIndexDeleted != INDEX_NONE)
 		{
-			int32 PoseIndexDeleted = PoseContainer.DeletePose(PoseSmartName);
-			if (PoseIndexDeleted != INDEX_NONE)
+			++ItemsDeleted;
+			// if base pose index is same as pose index deleted
+			if (BasePoseIndex == PoseIndexDeleted)
 			{
-				++ItemsDeleted;
-				// if base pose index is same as pose index deleted
-				if (BasePoseIndex == PoseIndexDeleted)
-				{
-					BasePoseIndex = INDEX_NONE;
-				}
-				// if base pose index is after this, we reduce the number
-				else if (BasePoseIndex > PoseIndexDeleted)
-				{
-					--BasePoseIndex;
-				}
+				BasePoseIndex = INDEX_NONE;
+			}
+			// if base pose index is after this, we reduce the number
+			else if (BasePoseIndex > PoseIndexDeleted)
+			{
+				--BasePoseIndex;
 			}
 		}
 	}
@@ -1613,14 +1622,10 @@ int32 UPoseAsset::DeleteCurves(TArray<FName> CurveNamesToDelete)
 {
 	int32 ItemsDeleted = 0;
 
-	USkeleton* MySkeleton = GetSkeleton();
-
 	for (const FName& CurveName : CurveNamesToDelete)
 	{
-		FSmartName CurveSmartName;
-		if (MySkeleton->GetSmartNameByName(USkeleton::AnimCurveMappingName, CurveName, CurveSmartName))
+		if(PoseContainer.DeleteCurve(CurveName))
 		{
-			PoseContainer.DeleteCurve(CurveSmartName);
 			++ItemsDeleted;
 		}
 	}
@@ -1705,9 +1710,9 @@ bool UPoseAsset::ConvertSpace(bool bNewAdditivePose, int32 NewBasePoseIndex)
 
 const int32 UPoseAsset::GetPoseIndexByName(const FName& InBasePoseName) const
 {
-	for (int32 PoseIndex = 0; PoseIndex < PoseContainer.PoseNames.Num(); ++PoseIndex)
+	for (int32 PoseIndex = 0; PoseIndex < PoseContainer.PoseFNames.Num(); ++PoseIndex)
 	{
-		if (PoseContainer.PoseNames[PoseIndex].DisplayName == InBasePoseName)
+		if (PoseContainer.PoseFNames[PoseIndex] == InBasePoseName)
 		{
 			return PoseIndex;
 		}
@@ -1721,7 +1726,7 @@ const int32 UPoseAsset::GetCurveIndexByName(const FName& InCurveName) const
 	for (int32 TestIdx = 0; TestIdx < PoseContainer.Curves.Num(); TestIdx++)
 	{
 		const FAnimCurveBase& Curve = PoseContainer.Curves[TestIdx];
-		if (Curve.Name.DisplayName == InCurveName)
+		if (Curve.GetName() == InCurveName)
 		{
 			return TestIdx;
 		}
@@ -1837,48 +1842,33 @@ void FPoseDataContainer::DeleteTrack(int32 TrackIndex)
 #if WITH_EDITOR
 FName UPoseAsset::GetUniquePoseName(const USkeleton* Skeleton)
 {
-	FName NewName = NAME_None;
-	if (Skeleton)
-	{
-		int32 NameIndex = 0;
-		SmartName::UID_Type NewUID;
-		do
-		{
-			NewName = FName(*FString::Printf(TEXT("Pose_%d"), NameIndex++));
-			NewUID = Skeleton->GetUIDByName(USkeleton::AnimCurveMappingName, NewName);
-		} while (NewUID != SmartName::MaxUID);
-	}
-
-	return NewName;
+	return NAME_None;
 }
 
 FSmartName UPoseAsset::GetUniquePoseSmartName(USkeleton* Skeleton)
 {
-	const FName NewName = GetUniquePoseName(Skeleton);
-	FSmartName NewPoseName;
-	if (NewName != NAME_None)
+	return FSmartName();
+}
+
+FName UPoseAsset::GetUniquePoseName(UPoseAsset* PoseAsset)
+{
+	check(PoseAsset);
+	int32 NameIndex = 0;
+	FName NewName;
+
+	do
 	{
-		Skeleton->AddSmartNameAndModify(USkeleton::AnimCurveMappingName, NewName, NewPoseName);
+		NewName = FName(*FString::Printf(TEXT("Pose_%d"), NameIndex++));
 	}
-
-	return NewPoseName;
+	while(PoseAsset->ContainsPose(NewName));
+	
+	return NewName;
 }
 
 void UPoseAsset::RemapTracksToNewSkeleton(USkeleton* NewSkeleton, bool bConvertSpaces)
 {
 	Super::RemapTracksToNewSkeleton(NewSkeleton, bConvertSpaces);
 
-	// after remap, should verify if the names are still valid in this skeleton
-	if (NewSkeleton)
-	{
-		NewSkeleton->VerifySmartNames(USkeleton::AnimCurveMappingName, PoseContainer.PoseNames);
-
-		for (FAnimCurveBase& Curve : PoseContainer.Curves)
-		{
-			NewSkeleton->VerifySmartName(USkeleton::AnimCurveMappingName, Curve.Name);
-		}
-	}
-
 	PostProcessData();
 }
 
diff --git a/Engine/Source/Runtime/Engine/Private/Animation/Skeleton.cpp b/Engine/Source/Runtime/Engine/Private/Animation/Skeleton.cpp
--- a/Engine/Source/Runtime/Engine/Private/Animation/Skeleton.cpp
+++ b/Engine/Source/Runtime/Engine/Private/Animation/Skeleton.cpp
@@ -28,12 +28,19 @@
 #include "UObject/AnimObjectVersion.h"
 #include "EngineUtils.h"
 #include "SkeletonRemappingRegistry.h"
+#include "UObject/UE5MainStreamObjectVersion.h"
 
 #include UE_INLINE_GENERATED_CPP_BY_NAME(Skeleton)
 
 #define LOCTEXT_NAMESPACE "Skeleton"
 #define ROOT_BONE_PARENT	INDEX_NONE
 
+// DEPRECATED - legacy support
+PRAGMA_DISABLE_DEPRECATION_WARNINGS
+TArray<uint16> USkeleton::DefaultCurveUIDList;
+FSmartNameMapping* USkeleton::AnimCurveMapping = nullptr;
+PRAGMA_ENABLE_DEPRECATION_WARNINGS
+
 #if WITH_EDITOR
 const FName USkeleton::AnimNotifyTag = FName(TEXT("AnimNotifyList"));
 const FString USkeleton::AnimNotifyTagDelimiter = TEXT(";");
@@ -302,9 +309,6 @@ bool USkeleton::IsCompatibleSkeletonByAssetData(const FAssetData& AssetData, con
 USkeleton::USkeleton(const FObjectInitializer& ObjectInitializer)
 	: Super(ObjectInitializer)
 {
-	// Make sure we have somewhere for curve names.
-	AnimCurveMapping = SmartNames.AddContainer(AnimCurveMappingName);
-	AnimCurveUidVersion = 0;
 	CachedSoftObjectPtr = TSoftObjectPtr<USkeleton>(this);
 
 	if (HasAnyFlags(RF_ClassDefaultObject))
@@ -324,8 +328,6 @@ void USkeleton::BeginDestroy()
 	{
 		FCoreUObjectDelegates::OnPackageReloaded.RemoveAll(this);
 	}
-
-	OnSmartNamesChangedEvent.Clear();
 }
 
 void USkeleton::PostInitProperties()
@@ -342,7 +344,7 @@ void USkeleton::PostInitProperties()
 
 bool USkeleton::IsPostLoadThreadSafe() const
 {
-	return true;
+	return WITH_EDITORONLY_DATA == 0;
 }
 
 void USkeleton::PostLoad()
@@ -350,44 +352,41 @@ void USkeleton::PostLoad()
 	Super::PostLoad();
 	LLM_SCOPE(ELLMTag::Animation);
 
+#if WITH_EDITORONLY_DATA
 	if( GetLinker() && (GetLinker()->UEVer() < VER_UE4_REFERENCE_SKELETON_REFACTOR) )
 	{
 		// Convert RefLocalPoses & BoneTree to FReferenceSkeleton
 		ConvertToFReferenceSkeleton();
 	}
+#endif
 
 	// catch any case if guid isn't valid
 	check(Guid.IsValid());
 
-	SmartNames.PostLoad();
-	
-	// Cache smart name uids for animation curve names
-	IncreaseAnimCurveUidVersion();
-
 	// refresh linked bone indices
-	FSmartNameMapping* CurveMappingTable = SmartNames.GetContainerInternal(USkeleton::AnimCurveMappingName);
-	if (CurveMappingTable)
-	{
-		CurveMappingTable->InitializeCurveMetaData(this);
-	}
-	
-	// Listen to smart name changes so we can regenerate our skeleton mappings.
-	if (!OnSmartNamesChangedEvent.IsBoundToObject(this))
-	{
-		OnSmartNamesChangedEvent.AddUObject(this, &USkeleton::HandleSmartNamesChangedEvent);
-	}
-	
+	RefreshSkeletonMetaData();
+
 	// Cleanup CompatibleSkeletons for convenience. This basically removes any soft object pointers that has an invalid soft object name.
 	CompatibleSkeletons = CompatibleSkeletons.FilterByPredicate([](const TSoftObjectPtr<USkeleton>& Skeleton)
 	{
 		return Skeleton.ToSoftObjectPath().IsValid();
 	});
-}
 
-// Regenerate all required skeleton remappings whenever curves change.
-void USkeleton::HandleSmartNamesChangedEvent()
-{
-	UE::Anim::FSkeletonRemappingRegistry::Get().RefreshCurveMappings(this);
+#if WITH_EDITORONLY_DATA
+	if(GetLinkerCustomVersion(FUE5MainStreamObjectVersion::GUID) < FUE5MainStreamObjectVersion::AnimationRemoveSmartNames)
+	{
+		// Move curve metadata over to asset user data
+		UAnimCurveMetaData* AnimCurveMetaData = NewObject<UAnimCurveMetaData>(this, NAME_None, RF_Transactional);
+		for(const TPair<FName, FSmartNameMapping>& NameMappingPair : SmartNames_DEPRECATED.NameMappings)
+		{
+			for(const TPair<FName, FCurveMetaData>& NameMetaDataPair : NameMappingPair.Value.CurveMetaDataMap)
+			{
+				AnimCurveMetaData->CurveMetaData.Add(NameMetaDataPair.Key, NameMetaDataPair.Value);
+			}
+		}
+		AddAssetUserData(AnimCurveMetaData);
+	}
+#endif
 }
 
 void USkeleton::PostDuplicate(bool bDuplicateForPIE)
@@ -404,6 +403,7 @@ void USkeleton::PostDuplicate(bool bDuplicateForPIE)
 void USkeleton::Serialize( FArchive& Ar )
 {
 	Ar.UsingCustomVersion(FAnimObjectVersion::GUID);
+	Ar.UsingCustomVersion(FUE5MainStreamObjectVersion::GUID);
 
 	DECLARE_SCOPE_CYCLE_COUNTER( TEXT("USkeleton::Serialize"), STAT_Skeleton_Serialize, STATGROUP_LoadTime );
 
@@ -468,9 +468,9 @@ void USkeleton::Serialize( FArchive& Ar )
 	// If we should be using smartnames, serialize the mappings
 	if(Ar.UEVer() >= VER_UE4_SKELETON_ADD_SMARTNAMES)
 	{
-		SmartNames.Serialize(Ar, IsTemplate());
-
-		AnimCurveMapping = SmartNames.GetContainerInternal(USkeleton::AnimCurveMappingName);
+		PRAGMA_DISABLE_DEPRECATION_WARNINGS
+		SmartNames_DEPRECATED.Serialize(Ar, IsTemplate());
+		PRAGMA_ENABLE_DEPRECATION_WARNINGS
 	}
 
 	// Build look up table between Slot nodes and their Group.
@@ -1239,6 +1239,29 @@ void USkeleton::AddNewAnimationNotify(FName NewAnimNotifyName)
 	}
 }
 
+void USkeleton::RemoveAnimationNotify(FName AnimNotifyName)
+{
+	if (AnimNotifyName != NAME_None)
+	{
+		AnimationNotifies.Remove(AnimNotifyName);
+	}
+}
+
+void USkeleton::RenameAnimationNotify(FName OldAnimNotifyName, FName NewAnimNotifyName)
+{
+	if(!AnimationNotifies.Contains(NewAnimNotifyName))
+	{
+		for(FName& NotifyName : AnimationNotifies)
+		{
+			if(NotifyName == OldAnimNotifyName)
+			{
+				NotifyName = NewAnimNotifyName;
+				break;
+			}
+		}
+	}
+}
+
 USkeletalMesh* USkeleton::FindCompatibleMesh() const
 {
 	FARFilter Filter;
@@ -1381,12 +1404,7 @@ void USkeleton::HandleSkeletonHierarchyChange()
 	RefreshAllRetargetSources();
 #endif
 
-	// refresh curve meta data that contains joint info
-	FSmartNameMapping* CurveMappingTable = SmartNames.GetContainerInternal(USkeleton::AnimCurveMappingName);
-	if (CurveMappingTable)
-	{
-		CurveMappingTable->InitializeCurveMetaData(this);
-	}
+	RefreshSkeletonMetaData();
 
 	// Remove entries from Blend Profiles for bones that no longer exists
 	for (UBlendProfile* Profile : BlendProfiles)
@@ -1578,206 +1596,61 @@ void USkeleton::RenameSlotName(const FName& OldName, const FName& NewName)
 	SetSlotGroupName(NewName, GroupName);
 }
 
-#if WITH_EDITOR
-
-bool USkeleton::AddSmartNameAndModify(FName ContainerName, FName NewDisplayName, FSmartName& OutNewName)
+bool USkeleton::AddCurveMetaData(FName CurveName)
 {
-	if (NewDisplayName != NAME_None)
+	if(UAnimCurveMetaData* AnimCurveMetaData = GetAssetUserData<UAnimCurveMetaData>())
 	{
-		OutNewName.DisplayName = NewDisplayName;
-		const bool bAdded = VerifySmartNameInternal(ContainerName, OutNewName);
-
-		if (bAdded)
-		{
-			IncreaseAnimCurveUidVersion();
-		}
-
-		return bAdded;
+		return AnimCurveMetaData->AddCurveMetaData(CurveName);
 	}
-
 	return false;
 }
 
-bool USkeleton::RenameSmartnameAndModify(FName ContainerName, SmartName::UID_Type Uid, FName NewName)
-{
-	bool Successful = false;
-	if (NewName != NAME_None)
-	{
-		FSmartNameMapping* RequestedMapping = SmartNames.GetContainerInternal(ContainerName);
-		if (RequestedMapping)
-		{
-			FName CurrentName;
-			RequestedMapping->GetName(Uid, CurrentName);
-			if (CurrentName != NewName)
-			{
-				Modify();
-				Successful = RequestedMapping->Rename(Uid, NewName);
-				IncreaseAnimCurveUidVersion();
-			}
-		}
-	}
-	return Successful;
-}
-
-void USkeleton::RemoveSmartnameAndModify(FName ContainerName, SmartName::UID_Type Uid)
-{
-	FSmartNameMapping* RequestedMapping = SmartNames.GetContainerInternal(ContainerName);
-	if (RequestedMapping)
-	{
-		Modify();
-		if (RequestedMapping->Remove(Uid))
-		{
-			IncreaseAnimCurveUidVersion();
-		}
-	}
-}
-
-void USkeleton::RemoveSmartnamesAndModify(FName ContainerName, const TArray<FName>& Names)
-{
-	FSmartNameMapping* RequestedMapping = SmartNames.GetContainerInternal(ContainerName);
-	if (RequestedMapping)
-	{
-		bool bModified = false;
-		for (const FName& CurveName : Names)
-		{
-			if (RequestedMapping->Exists(CurveName))
-			{
-				if (!bModified)
-				{
-					Modify();
-				}
-				RequestedMapping->Remove(CurveName);
-				bModified = true;
-			}
-		}
-
-		if (bModified)
-		{
-			IncreaseAnimCurveUidVersion();
-		}
-	}
-}
-#endif // WITH_EDITOR
+#if WITH_EDITOR
 
-bool USkeleton::GetSmartNameByUID(const FName& ContainerName, SmartName::UID_Type UID, FSmartName& OutSmartName) const
+bool USkeleton::RenameCurveMetaData(FName OldName, FName NewName)
 {
-	const FSmartNameMapping* RequestedMapping = SmartNames.GetContainerInternal(ContainerName);
-	if (RequestedMapping)
+	if(UAnimCurveMetaData* AnimCurveMetaData = GetAssetUserData<UAnimCurveMetaData>())
 	{
-		return (RequestedMapping->FindSmartNameByUID(UID, OutSmartName));
+		return AnimCurveMetaData->RenameCurveMetaData(OldName, NewName);
 	}
-
 	return false;
 }
 
-bool USkeleton::GetSmartNameByName(const FName& ContainerName, const FName& InName, FSmartName& OutSmartName) const
+bool USkeleton::RemoveCurveMetaData(FName CurveName)
 {
-	const FSmartNameMapping* RequestedMapping = SmartNames.GetContainerInternal(ContainerName);
-	if (RequestedMapping)
+	if(UAnimCurveMetaData* AnimCurveMetaData = GetAssetUserData<UAnimCurveMetaData>())
 	{
-		return (RequestedMapping->FindSmartName(InName, OutSmartName));
+		return AnimCurveMetaData->RemoveCurveMetaData(CurveName);
 	}
-
 	return false;
 }
 
-SmartName::UID_Type USkeleton::GetUIDByName(const FName& ContainerName, const FName& Name) const
-{
-	const FSmartNameMapping* RequestedMapping = SmartNames.GetContainerInternal(ContainerName);
-	if (RequestedMapping)
-	{
-		return RequestedMapping->FindUID(Name);
-	}
-
-	return SmartName::MaxUID;
-}
-
-// this does very simple thing.
-// @todo: this for now prioritize FName because that is main issue right now
-// @todo: @fixme: this has to be fixed when we have GUID
-void USkeleton::VerifySmartName(const FName& ContainerName, FSmartName& InOutSmartName)
-{
-	if (VerifySmartNameInternal(ContainerName, InOutSmartName) && ContainerName == USkeleton::AnimCurveMappingName)
-	{
-		IncreaseAnimCurveUidVersion();
-	}
-}
-
-
-bool USkeleton::FillSmartNameByDisplayName(FSmartNameMapping* Mapping, const FName& DisplayName, FSmartName& OutSmartName)
+bool USkeleton::RemoveCurveMetaData(TArrayView<FName> CurveNames)
 {
-	FSmartName SkeletonName;
-	if (Mapping->FindSmartName(DisplayName, SkeletonName))
+	if(UAnimCurveMetaData* AnimCurveMetaData = GetAssetUserData<UAnimCurveMetaData>())
 	{
-		OutSmartName.DisplayName = DisplayName;
-
-		// if not editor, we assume name is always correct
-		OutSmartName.UID = SkeletonName.UID;
-		return true;
+		return AnimCurveMetaData->RemoveCurveMetaData(CurveNames);
 	}
-
 	return false;
 }
 
-bool USkeleton::VerifySmartNameInternal(const FName&  ContainerName, FSmartName& InOutSmartName)
-{
-	FSmartNameMapping* Mapping = GetOrAddSmartNameContainer(ContainerName);
-	if (Mapping != nullptr)
-	{
-		if (!Mapping->FindSmartName(InOutSmartName.DisplayName, InOutSmartName))
-		{
-#if WITH_EDITOR
-			if (IsInGameThread())
-			{
-				Modify();
-			}
-#endif
-			InOutSmartName = Mapping->AddName(InOutSmartName.DisplayName);
-			return true;
-		}
-	}
-
-	return false;
-}
+#endif // WITH_EDITOR
 
-void USkeleton::VerifySmartNames(const FName&  ContainerName, TArray<FSmartName>& InOutSmartNames)
+uint16 USkeleton::GetAnimCurveUidVersion() const
 {
-	bool bRefreshCache = false;
-
-	for(auto& SmartName: InOutSmartNames)
-	{
-		VerifySmartNameInternal(ContainerName, SmartName);
-	}
-
-	if (ContainerName == USkeleton::AnimCurveMappingName)
+	if(UAnimCurveMetaData* AnimCurveMetaData = const_cast<USkeleton*>(this)->GetAssetUserData<UAnimCurveMetaData>())
 	{
-		IncreaseAnimCurveUidVersion();
+		return AnimCurveMetaData->GetVersionNumber();
 	}
+	return 0;
 }
 
-FSmartNameMapping* USkeleton::GetOrAddSmartNameContainer(const FName& ContainerName)
+void USkeleton::GetCurveMetaDataNames(TArray<FName>& OutNames) const
 {
-	FSmartNameMapping* Mapping = SmartNames.GetContainerInternal(ContainerName);
-	if (Mapping == nullptr)
+	if(const UAnimCurveMetaData* AnimCurveMetaData = const_cast<USkeleton*>(this)->GetAssetUserData<UAnimCurveMetaData>())
 	{
-#if WITH_EDITOR
-		if (IsInGameThread())
-		{
-			Modify();
-		}
-#endif
-		IncreaseAnimCurveUidVersion();
-		SmartNames.AddContainer(ContainerName);
-		AnimCurveMapping = SmartNames.GetContainerInternal(USkeleton::AnimCurveMappingName);
-		Mapping = SmartNames.GetContainerInternal(ContainerName);		
+		return AnimCurveMetaData->GetCurveMetaDataNames(OutNames);
 	}
-
-	return Mapping;
-}
-
-const FSmartNameMapping* USkeleton::GetSmartNameContainer(const FName& ContainerName) const
-{
-	return SmartNames.GetContainer(ContainerName);
 }
 
 void USkeleton::RegenerateGuid()
@@ -1792,113 +1665,83 @@ void USkeleton::RegenerateVirtualBoneGuid()
 	check(VirtualBoneGuid.IsValid());
 }
 
-void USkeleton::IncreaseAnimCurveUidVersion()
+FCurveMetaData* USkeleton::GetCurveMetaData(FName CurveName)
 {
-	// this will be checked by skeletalmeshcomponent and if it's same, it won't care. If it's different, it will rebuild UID list
-	++AnimCurveUidVersion;
-
-	// update default uid list
-	const FSmartNameMapping* Mapping = AnimCurveMapping;
-	if (ensureAlways(Mapping))
+	if(UAnimCurveMetaData* AnimCurveMetaData = GetAssetUserData<UAnimCurveMetaData>())
 	{
-		DefaultCurveUIDList.Reset();
-		// this should exactly work with what current index is
-		Mapping->FillUidArray(DefaultCurveUIDList);
+		return AnimCurveMetaData->GetCurveMetaData(CurveName);
 	}
+	return nullptr;
 }
 
-FCurveMetaData* USkeleton::GetCurveMetaData(const FName& CurveName)
+const FCurveMetaData* USkeleton::GetCurveMetaData(FName CurveName) const
 {
-	FSmartNameMapping* Mapping = AnimCurveMapping;
-	if (ensureAlways(Mapping))
+	if(const UAnimCurveMetaData* AnimCurveMetaData = const_cast<USkeleton*>(this)->GetAssetUserData<UAnimCurveMetaData>())
 	{
-		return Mapping->GetCurveMetaData(CurveName);
+		return AnimCurveMetaData->GetCurveMetaData(CurveName);
 	}
-
 	return nullptr;
 }
 
-const FCurveMetaData* USkeleton::GetCurveMetaData(const FName& CurveName) const
+void USkeleton::ForEachCurveMetaData(TFunctionRef<void(FName, const FCurveMetaData&)> InFunction) const
 {
-	const FSmartNameMapping* Mapping = AnimCurveMapping;
-	if (ensureAlways(Mapping))
+	if(const UAnimCurveMetaData* AnimCurveMetaData = const_cast<USkeleton*>(this)->GetAssetUserData<UAnimCurveMetaData>())
 	{
-		return Mapping->GetCurveMetaData(CurveName);
+		AnimCurveMetaData->ForEachCurveMetaData(InFunction);
 	}
-
-	return nullptr;
 }
 
-const FCurveMetaData* USkeleton::GetCurveMetaData(const SmartName::UID_Type CurveUID) const
+int32 USkeleton::GetNumCurveMetaData() const
 {
-	const FSmartNameMapping* Mapping = AnimCurveMapping;
-	if (ensureAlways(Mapping))
+	if(const UAnimCurveMetaData* AnimCurveMetaData = const_cast<USkeleton*>(this)->GetAssetUserData<UAnimCurveMetaData>())
 	{
-#if WITH_EDITOR
-		FSmartName SmartName;
-		if (Mapping->FindSmartNameByUID(CurveUID, SmartName))
-		{
-			return Mapping->GetCurveMetaData(SmartName.DisplayName);
-		}
-#else
-		return &Mapping->GetCurveMetaData(CurveUID);
-#endif
+		return AnimCurveMetaData->GetNumCurveMetaData();
 	}
+	return 0;
+}
 
+const FCurveMetaData* USkeleton::GetCurveMetaData(const SmartName::UID_Type CurveUID) const
+{
 	return nullptr;
 }
 
 FCurveMetaData* USkeleton::GetCurveMetaData(const FSmartName& CurveName)
 {
-	FSmartNameMapping* Mapping = AnimCurveMapping;
-	if (ensureAlways(Mapping))
-	{
-		// the name might have changed, make sure it's up-to-date
-		FName DisplayName;
-		Mapping->GetName(CurveName.UID, DisplayName);
-		return Mapping->GetCurveMetaData(DisplayName);
-	}
-
-	return nullptr;
+	return GetCurveMetaData(CurveName.DisplayName);
 }
 
 const FCurveMetaData* USkeleton::GetCurveMetaData(const FSmartName& CurveName) const
 {
-	const FSmartNameMapping* Mapping = AnimCurveMapping;
-	if (ensureAlways(Mapping))
-	{
-		// the name might have changed, make sure it's up-to-date
-		FName DisplayName;
-		Mapping->GetName(CurveName.UID, DisplayName);
-		return Mapping->GetCurveMetaData(DisplayName);
-	}
-
-	return nullptr;
+	return GetCurveMetaData(CurveName.DisplayName);
 }
 
-// this is called when you know both flags - called by post serialize
 void USkeleton::AccumulateCurveMetaData(FName CurveName, bool bMaterialSet, bool bMorphtargetSet)
 {
 	if (bMaterialSet || bMorphtargetSet)
 	{
-		const FSmartNameMapping* Mapping = AnimCurveMapping;
-		if (ensureAlways(Mapping))
+		// Add user data if not already added
+		UAnimCurveMetaData* AnimCurveMetaData = GetAssetUserData<UAnimCurveMetaData>();
+		if(AnimCurveMetaData == nullptr)
 		{
-			// if we don't have name, add one
-			if (Mapping->Exists(CurveName))
+			AnimCurveMetaData = NewObject<UAnimCurveMetaData>(this, NAME_None, RF_Transactional);
+			AddAssetUserData(AnimCurveMetaData);
+		}
+
+		// Add curve if not already present
+		AddCurveMetaData(CurveName);
+
+		if(FCurveMetaData* FoundCurveMetaData = GetCurveMetaData(CurveName))
+		{
+			bool bOldMaterial = FoundCurveMetaData->Type.bMaterial;
+			bool bOldMorphtarget = FoundCurveMetaData->Type.bMorphtarget;
+			// we don't want to undo previous flags, if it was true, we just allow more to it. 
+			FoundCurveMetaData->Type.bMaterial |= bMaterialSet;
+			FoundCurveMetaData->Type.bMorphtarget |= bMorphtargetSet;
+
+			if (bOldMaterial != FoundCurveMetaData->Type.bMaterial 
+				|| bOldMorphtarget != FoundCurveMetaData->Type.bMorphtarget)
 			{
-				FCurveMetaData* CurveMetaData = GetCurveMetaData(CurveName);
-				bool bOldMaterial = CurveMetaData->Type.bMaterial;
-				bool bOldMorphtarget = CurveMetaData->Type.bMorphtarget;
-				// we don't want to undo previous flags, if it was true, we just alolw more to it. 
-				CurveMetaData->Type.bMaterial |= bMaterialSet;
-				CurveMetaData->Type.bMorphtarget |= bMorphtargetSet;
-
-				if (IsInGameThread() && (bOldMaterial != CurveMetaData->Type.bMaterial
-					|| bOldMorphtarget != CurveMetaData->Type.bMorphtarget))
-				{
-					MarkPackageDirty();
-				}
+				MarkPackageDirty();
 			}
 		}
 	}
@@ -2051,11 +1894,7 @@ void USkeleton::HandleVirtualBoneChanges()
 	}
 
 	// refresh curve meta data that contains joint info
-	FSmartNameMapping* CurveMappingTable = SmartNames.GetContainerInternal(USkeleton::AnimCurveMappingName);
-	if (CurveMappingTable)
-	{
-		CurveMappingTable->InitializeCurveMetaData(this);
-	}
+	RefreshSkeletonMetaData();
 
 	for (TObjectIterator<USkinnedMeshComponent> It; It; ++It)
 	{
@@ -2250,6 +2089,12 @@ void USkeleton::GetAssetRegistryTags(TArray<FAssetRegistryTag>& OutTags) const
 	}
 
 	OutTags.Add(FAssetRegistryTag(USkeleton::CompatibleSkeletonsNameTag, CompatibleSkeletonsBuilder.ToString(), FAssetRegistryTag::TT_Hidden));
+
+	// Allow asset user data to output tags
+	for(UAssetUserData* AssetUserDataItem : AssetUserData)
+	{
+		AssetUserDataItem->GetAssetRegistryTags(OutTags);
+	}
 }
 
 #endif //WITH_EDITOR
@@ -2363,5 +2208,13 @@ void USkeleton::HandlePackageReloaded(const EPackageReloadPhase InPackageReloadP
 	}
 }
 
+void USkeleton::RefreshSkeletonMetaData()
+{
+	if(UAnimCurveMetaData* AnimCurveMetaData = GetAssetUserData<UAnimCurveMetaData>())
+	{
+		AnimCurveMetaData->RefreshBoneIndices(this);
+	}
+}
+
 #undef LOCTEXT_NAMESPACE 
 
diff --git a/Engine/Source/Runtime/Engine/Private/Animation/SkeletonRemapping.cpp b/Engine/Source/Runtime/Engine/Private/Animation/SkeletonRemapping.cpp
--- a/Engine/Source/Runtime/Engine/Private/Animation/SkeletonRemapping.cpp
+++ b/Engine/Source/Runtime/Engine/Private/Animation/SkeletonRemapping.cpp
@@ -186,8 +186,6 @@ void FSkeletonRemapping::GenerateMapping()
 			RetargetingTable[TargetBoneIndex] = MakeTuple(FQuat::Identity, FQuat::Identity);
 		}
 	}
-
-	GenerateCurveMapping();
 }
 
 void FSkeletonRemapping::ComposeWith(const FSkeletonRemapping& OtherSkeletonRemapping)
@@ -243,55 +241,8 @@ void FSkeletonRemapping::ComposeWith(const FSkeletonRemapping& OtherSkeletonRema
 	}
 }
 
-void FSkeletonRemapping::GenerateCurveMapping()
+const TArray<SmartName::UID_Type>& FSkeletonRemapping::GetSourceToTargetCurveMapping() const
 {
-	const USkeleton* Source = GetSourceSkeleton().Get();
-	const USkeleton* Target = GetTargetSkeleton().Get();
-	check(Source && Target);
-
-	// Get the curve mappings.
-	const FSmartNameMapping* SourceMapping = Source->GetSmartNameContainer(USkeleton::AnimCurveMappingName);
-	const FSmartNameMapping* TargetMapping = Target->GetSmartNameContainer(USkeleton::AnimCurveMappingName);
-	if (!SourceMapping) // Source has no curves, so there is nothing to add.
-	{
-		return;
-	}
-
-	// Check whether we have curves or not in both source and target.
-	bool bHasSourceCurves = false;
-	bool bHasTargetCurves = false;
-	if (TargetMapping)
-	{
-		bHasTargetCurves = (TargetMapping->GetMaxUID() != SmartName::MaxUID);
-		bHasSourceCurves = (SourceMapping->GetMaxUID() != SmartName::MaxUID);
-	}
-
-	// Init the mapping table.
-	const SmartName::UID_Type NumItems = SourceMapping->GetMaxUID() + 1;
-	if (bHasSourceCurves)
-	{
-		SourceToTargetCurveMapping.Init(MAX_uint16, NumItems);
-	}
-
-	// If we have no source or target curves (while having valid smart name mappings), there is nothing to do.
-	if (!bHasTargetCurves && !bHasSourceCurves)
-	{
-		return;
-	}
-
-	// For every source curve, try to find the curve with the same name in the target.
-	for (SmartName::UID_Type Index = 0; Index < NumItems; ++Index)
-	{
-		FName CurveName;
-		SourceMapping->GetName(Index, CurveName);
-
-		// Make sure the curve name is valid and we can actually find a curve with the same name in the target.
-		FSmartName TargetSmartName;
-		if (!CurveName.IsValid() || !TargetMapping->FindSmartName(CurveName, TargetSmartName))
-		{
-			continue;
-		}
-
-		SourceToTargetCurveMapping[Index] = TargetSmartName.UID;
-	}
-}
\ No newline at end of file
+	static TArray<SmartName::UID_Type> Dummy;
+	return Dummy;
+}
diff --git a/Engine/Source/Runtime/Engine/Private/Animation/SkeletonRemapping.h b/Engine/Source/Runtime/Engine/Private/Animation/SkeletonRemapping.h
--- a/Engine/Source/Runtime/Engine/Private/Animation/SkeletonRemapping.h
+++ b/Engine/Source/Runtime/Engine/Private/Animation/SkeletonRemapping.h
@@ -142,25 +142,14 @@ struct FSkeletonRemapping
 		return QQ.Get<0>() * SourceRotation * QQ.Get<0>().Inverse();
 	}
 
-	/**
-	 * Get the curve mapping array, which maps from the source curve UID to the target curve UID.
-	 * You can access this array like this:
-	 * 
-	 * \code{.cpp}
-	 * const SmartName::UID_Type TargetCurveUID = GetSourceToTargetCurveMapping()[SourceCurveUID];
-	 * \endcode
-	 */
-	const TArray<SmartName::UID_Type>& GetSourceToTargetCurveMapping() const { return SourceToTargetCurveMapping; }
+	UE_DEPRECATED(5.3, "This function is no longer used")
+	const TArray<SmartName::UID_Type>& GetSourceToTargetCurveMapping() const;
 
 	/** Refreshes the mapping. Empties mapping arrays and regenerates them against the stored skeletons. */
 	void RegenerateMapping();
 	
-	/**
-	 * Generates the mapping table for the curves, based on curve names.
-	 * Basically this will look at the curve names in the smart name table on the source, and tries to find matching ones in the table of 
-	 * the target skeleton. It then maps the UID's.
-	 */
-	void GenerateCurveMapping();
+	UE_DEPRECATED(5.3, "This function is no longer used")
+	void GenerateCurveMapping() {}
 
 	/** Check to see if a reference pose retarget is required between the source & target hierarchies */
 	bool RequiresReferencePoseRetarget() const
@@ -187,9 +176,6 @@ private:
 
 	// Table of source skeleton bone indexes (indexed by target skeleton bone index)
 	TArray<int32> TargetToSourceBoneIndexes;
-	
-	// Maps curve UIDs between source and target (indexed by source curve UID).
-	TArray<SmartName::UID_Type> SourceToTargetCurveMapping;
 
 	// Table of precalculated constants for retargeting from source to target (indexed by target skeleton bone index)
 	TArray<TTuple<FQuat, FQuat>> RetargetingTable;
diff --git a/Engine/Source/Runtime/Engine/Private/Animation/SkeletonRemappingRegistry.cpp b/Engine/Source/Runtime/Engine/Private/Animation/SkeletonRemappingRegistry.cpp
--- a/Engine/Source/Runtime/Engine/Private/Animation/SkeletonRemappingRegistry.cpp
+++ b/Engine/Source/Runtime/Engine/Private/Animation/SkeletonRemappingRegistry.cpp
@@ -99,22 +99,6 @@ const FSkeletonRemapping& FSkeletonRemappingRegistry::GetRemapping(const USkelet
 	}
 }
 
-void FSkeletonRemappingRegistry::RefreshCurveMappings(USkeleton* InSkeleton)
-{
-	TArray<TSharedPtr<FSkeletonRemapping>> ExistingMappings;
-	{
-		FRWScopeLock ReadLock(MappingsLock, SLT_ReadOnly);
-		PerSkeletonMappings.MultiFind(InSkeleton, ExistingMappings);
-	}
-	{
-		FRWScopeLock WriteLock(MappingsLock, SLT_Write);
-		for(const TSharedPtr<FSkeletonRemapping>& ExistingMapping : ExistingMappings)
-		{
-			ExistingMapping->GenerateCurveMapping();
-		}
-	}
-}
-
 void FSkeletonRemappingRegistry::RefreshMappings(USkeleton* InSkeleton)
 {
 	TArray<TSharedPtr<FSkeletonRemapping>> ExistingMappings;
diff --git a/Engine/Source/Runtime/Engine/Private/Animation/SkeletonRemappingRegistry.h b/Engine/Source/Runtime/Engine/Private/Animation/SkeletonRemappingRegistry.h
--- a/Engine/Source/Runtime/Engine/Private/Animation/SkeletonRemappingRegistry.h
+++ b/Engine/Source/Runtime/Engine/Private/Animation/SkeletonRemappingRegistry.h
@@ -30,9 +30,6 @@ public:
 
 	// Refresh any existing mappings that the supplied skeleton is used with
 	void RefreshMappings(USkeleton* InSkeleton);
-	
-	// Refresh any existing curve mappings that the supplied skeleton is used with
-	void RefreshCurveMappings(USkeleton* InSkeleton);
 
 private:
 	friend FSkeletonRemappingRegistryPrivate;
diff --git a/Engine/Source/Runtime/Engine/Private/BoneContainer.cpp b/Engine/Source/Runtime/Engine/Private/BoneContainer.cpp
--- a/Engine/Source/Runtime/Engine/Private/BoneContainer.cpp
+++ b/Engine/Source/Runtime/Engine/Private/BoneContainer.cpp
@@ -9,52 +9,10 @@
 #include "Animation/SkeletonRemappingRegistry.h"
 #include "Animation/SkeletonRemapping.h"
 
-#include UE_INLINE_GENERATED_CPP_BY_NAME(BoneContainer)
-
 LLM_DEFINE_TAG(BoneContainer);
 
 DEFINE_LOG_CATEGORY(LogSkeletalControl);
 
-//////////////////////////////////////////////////////////////////////////
-// FSkeletonRemappingCurve
-
-FSkeletonRemappingCurve::FSkeletonRemappingCurve(FBlendedCurve& InCurve, FBoneContainer& InBoneContainer, const FSkeletonRemapping* InSkeletonRemapping)
-	: Curve(InCurve)
-	, BoneContainer(InBoneContainer)
-	, bIsRemapping(true)
-{
-	LLM_SCOPE_BYNAME(TEXT("Animation/BoneContainer"));
-	const FCachedSkeletonCurveMapping& CurveMapping = InBoneContainer.GetOrCreateCachedCurveMapping(InSkeletonRemapping);
-	Curve.UIDToArrayIndexLUT = &CurveMapping.UIDToArrayIndices;
-}
-
-FSkeletonRemappingCurve::FSkeletonRemappingCurve(FBlendedCurve& InCurve, FBoneContainer& InBoneContainer, const USkeleton* SourceSkeleton)
-	: Curve(InCurve)
-	, BoneContainer(InBoneContainer)
-{
-	LLM_SCOPE_BYNAME(TEXT("Animation/BoneContainer"));
-	const FSkeletonRemapping& SkeletonRemapping = UE::Anim::FSkeletonRemappingRegistry::Get().GetRemapping(SourceSkeleton, BoneContainer.GetSkeletonAsset());
-	if (SkeletonRemapping.IsValid()) // No remapping is required, just continue as we normally would.
-	{
-		const FCachedSkeletonCurveMapping& CurveMapping = BoneContainer.GetOrCreateCachedCurveMapping(&SkeletonRemapping);
-		Curve.UIDToArrayIndexLUT = &CurveMapping.UIDToArrayIndices;
-		bIsRemapping = true;
-	}
-	else
-	{
-		bIsRemapping = false;
-	}
-}
-
-FSkeletonRemappingCurve::~FSkeletonRemappingCurve()
-{
-	if (bIsRemapping)
-	{
-		Curve.UIDToArrayIndexLUT = &BoneContainer.GetUIDToArrayLookupTableBackup();
-	}
-}
-
-
 //////////////////////////////////////////////////////////////////////////
 // FBoneContainer
 
@@ -63,7 +21,6 @@ FBoneContainer::FBoneContainer()
 , AssetSkeletalMesh(nullptr)
 , AssetSkeleton(nullptr)
 , RefSkeleton(nullptr)
-, UIDToArrayIndexLUTValidCount(0)
 , SerialNumber(0)
 #if DO_CHECK
 , CalculatedForLOD(INDEX_NONE)
@@ -85,7 +42,6 @@ FBoneContainer::FBoneContainer(const TArray<FBoneIndexType>& InRequiredBoneIndex
 , AssetSkeletalMesh(nullptr)
 , AssetSkeleton(nullptr)
 , RefSkeleton(nullptr)
-, UIDToArrayIndexLUTValidCount(0)
 , SerialNumber(0)
 #if DO_CHECK
 , CalculatedForLOD(INDEX_NONE)
@@ -95,16 +51,42 @@ FBoneContainer::FBoneContainer(const TArray<FBoneIndexType>& InRequiredBoneIndex
 , bUseSourceData(false)
 {
 	LLM_SCOPE_BYNAME(TEXT("Animation/BoneContainer"));
-	Initialize(CurveEvalOption);
+
+	const UE::Anim::FCurveFilterSettings CurveFilterSettings(CurveEvalOption.bAllowCurveEvaluation ? UE::Anim::ECurveFilterMode::None : UE::Anim::ECurveFilterMode::DisallowAll, CurveEvalOption.DisallowedList, CurveEvalOption.LODIndex);
+	Initialize(CurveFilterSettings);
+}
+
+FBoneContainer::FBoneContainer(const TArray<FBoneIndexType>& InRequiredBoneIndexArray, const UE::Anim::FCurveFilterSettings& InCurveFilterSettings, UObject& InAsset)
+	: BoneIndicesArray(InRequiredBoneIndexArray)
+	, Asset(&InAsset)
+	, AssetSkeletalMesh(nullptr)
+	, AssetSkeleton(nullptr)
+	, RefSkeleton(nullptr)
+	, SerialNumber(0)
+#if DO_CHECK
+	, CalculatedForLOD(INDEX_NONE)
+#endif
+	, bDisableRetargeting(false)
+	, bUseRAWData(false)
+	, bUseSourceData(false)
+{
+	LLM_SCOPE_BYNAME(TEXT("Animation/BoneContainer"));
+	Initialize(InCurveFilterSettings);
 }
 
 void FBoneContainer::InitializeTo(const TArray<FBoneIndexType>& InRequiredBoneIndexArray, const FCurveEvaluationOption& CurveEvalOption, UObject& InAsset)
+{
+	const UE::Anim::FCurveFilterSettings CurveFilterSettings(CurveEvalOption.bAllowCurveEvaluation ? UE::Anim::ECurveFilterMode::None : UE::Anim::ECurveFilterMode::DisallowAll, CurveEvalOption.DisallowedList, CurveEvalOption.LODIndex);
+	InitializeTo(InRequiredBoneIndexArray, CurveFilterSettings, InAsset);
+}
+
+void FBoneContainer::InitializeTo(const TArray<FBoneIndexType>& InRequiredBoneIndexArray, const UE::Anim::FCurveFilterSettings& CurveFilterSettings, UObject& InAsset)
 {
 	LLM_SCOPE_BYNAME(TEXT("Animation/BoneContainer"));
 	BoneIndicesArray = InRequiredBoneIndexArray;
 	Asset = &InAsset;
 
-	Initialize(CurveEvalOption);
+	Initialize(CurveFilterSettings);
 }
 
 struct FBoneContainerScratchArea : public TThreadSingleton<FBoneContainerScratchArea>
@@ -112,7 +94,7 @@ struct FBoneContainerScratchArea : public TThreadSingleton<FBoneContainerScratch
 	TArray<int32> MeshIndexToCompactPoseIndex;
 };
 
-void FBoneContainer::Initialize(const FCurveEvaluationOption& CurveEvalOption)
+void FBoneContainer::Initialize(const UE::Anim::FCurveFilterSettings& InCurveFilterSettings)
 {
 	LLM_SCOPE_BYNAME(TEXT("Animation/BoneContainer"));
 	RefSkeleton = nullptr;
@@ -128,7 +110,7 @@ void FBoneContainer::Initialize(const FCurveEvaluationOption& CurveEvalOption)
 	USkeleton* AssetSkeletonObj = nullptr;
 
 #if DO_CHECK
-	CalculatedForLOD = CurveEvalOption.LODIndex;
+	CalculatedForLOD = InCurveFilterSettings.LODIndex;
 #endif
 	
 	if (AssetSkeletalMeshObj)
@@ -242,8 +224,8 @@ void FBoneContainer::Initialize(const FCurveEvaluationOption& CurveEvalOption)
 			VirtualBoneCompactPoseData.Add(FVirtualBoneCompactPoseData(FCompactPoseBoneIndex(VBInd), FCompactPoseBoneIndex(SourceInd), FCompactPoseBoneIndex(TargetInd)));
 		}
 	}
-	// cache required curve UID list according to new bone sets
-	CacheRequiredAnimCurveUids(CurveEvalOption);
+	// cache required curves list according to new bone sets
+	CacheRequiredAnimCurves(InCurveFilterSettings);
 
 	// Reset retargeting cached data look up table.
 	RetargetSourceCachedDataLUT.Reset();
@@ -251,114 +233,126 @@ void FBoneContainer::Initialize(const FCurveEvaluationOption& CurveEvalOption)
 	RegenerateSerialNumber();
 }
 
-void FBoneContainer::CacheRequiredAnimCurveUids(const FCurveEvaluationOption& CurveEvalOption)
+void FBoneContainer::CacheRequiredAnimCurves(const UE::Anim::FCurveFilterSettings& InCurveFilterSettings)
 {
 	LLM_SCOPE_BYNAME(TEXT("Animation/BoneContainer"));
 
-	// Clear the skeleton remapping backup LUT.
-	// We will initialize this only when needed, using a lazy approach.
-	UIDToArrayIndexLUTBackup.Reset();
+	CurveFilter.Empty();
+	CurveFilter.SetFilterMode(InCurveFilterSettings.FilterMode);
+
+	CurveFlags.Empty();
 
-	if (AssetSkeleton.IsValid())
+	if (USkeleton* Skeleton = AssetSkeleton.Get())
 	{
-		// this is placeholder. In the future, this will change to work with linked joint of curve meta data
-		// anim curve name Uids; For now it adds all of them
-		const FSmartNameMapping* Mapping = AssetSkeleton->GetSmartNameContainer(USkeleton::AnimCurveMappingName);
-		if (Mapping != nullptr)
+		// Copy filter curves.
+		TArray<FName, TInlineAllocator<32>> FilterCurves;
+		if(InCurveFilterSettings.FilterCurves)
 		{
-			UIDToArrayIndexLUT.Reset();
-			UIDToArrayIndexLUTValidCount = 0;
-			
-			const SmartName::UID_Type MaxUID = Mapping->GetMaxUID();
-
-			if (MaxUID == SmartName::MaxUID)
-			{
-				// No smart names, nothing to do
-				return;
-			}
-			
-			//Init UID LUT to everything unused
-			UIDToArrayIndexLUT.AddUninitialized(MaxUID+1);
-			for (SmartName::UID_Type& Item : UIDToArrayIndexLUT)
+			FilterCurves = *InCurveFilterSettings.FilterCurves;
+		}
+		
+		if (InCurveFilterSettings.FilterMode != UE::Anim::ECurveFilterMode::DisallowAll)
+		{
+			// Apply curve metadata, LOD and linked bones filtering
+			Skeleton->ForEachCurveMetaData([this, &InCurveFilterSettings, &FilterCurves](const FName& InCurveName, const FCurveMetaData& InMetaData)
 			{
-				Item = SmartName::MaxUID;
-			}
+				bool bBeingUsed = true;
+				UE::Anim::ECurveElementFlags Flags = UE::Anim::ECurveElementFlags::None;
+				UE::Anim::ECurveFilterFlags FilterFlags = UE::Anim::ECurveFilterFlags::None;
+				if(InMetaData.Type.bMaterial)
+				{
+					Flags |= UE::Anim::ECurveElementFlags::Material;
+				}
 
-			// if the linked joints don't exists in RequiredBones, remove itself
-			int32 NumAvailableUIDs = 0;
-			Mapping->Iterate([this, &CurveEvalOption, &NumAvailableUIDs](const FSmartNameMappingIterator& Iterator)
-			{
-				FName CurveName;
-				if (Iterator.GetName(CurveName))
+				if(InMetaData.Type.bMorphtarget)
 				{
-					bool bBeingUsed = true;
-					if (!CurveEvalOption.bAllowCurveEvaluation)
-					{
-						bBeingUsed = false;
-					}
-					else
+					Flags |= UE::Anim::ECurveElementFlags::MorphTarget;
+				}
+
+				const int32 Index = FilterCurves.IndexOfByKey(InCurveName);
+				if(Index != INDEX_NONE)
+				{
+					FilterFlags |= UE::Anim::ECurveFilterFlags::Filtered;
+					FilterCurves.RemoveAtSwap(Index, 1, false);
+				}
+				
+				if (InMetaData.MaxLOD < InCurveFilterSettings.LODIndex)
+				{
+					bBeingUsed = false;
+				}
+				else if (InMetaData.LinkedBones.Num() > 0)
+				{
+					bBeingUsed = false;
+					for (const FBoneReference& LinkedBoneReference : InMetaData.LinkedBones)
 					{
-						// CurveNameIndex shouyld match to UID
-						if (CurveName == NAME_None)
-						{
-							bBeingUsed = false;
-						}
-						else if (CurveEvalOption.DisallowedList && CurveEvalOption.DisallowedList->Contains(CurveName))
-						{
-							//remove the UID
-							bBeingUsed = false;
-						}
-						else
+						// when you enter first time, sometimes it does not have all info yet
+						if (LinkedBoneReference.BoneIndex != INDEX_NONE && LinkedBoneReference.BoneName != NAME_None)
 						{
-							const FCurveMetaData* CurveMetaData = Iterator.GetCurveMetaData();
-							if (CurveMetaData)
+							// this linked bone always use skeleton index
+							ensure(LinkedBoneReference.bUseSkeletonIndex);
+							// we want to make sure all the joints are removed from RequiredBones before removing this curve
+							if (GetCompactPoseIndexFromSkeletonIndex(LinkedBoneReference.BoneIndex) != INDEX_NONE)
 							{
-								if (CurveMetaData->MaxLOD < CurveEvalOption.LODIndex)
-								{
-									bBeingUsed = false;
-								}
-								else if (CurveMetaData->LinkedBones.Num() > 0)
-								{
-									bBeingUsed = false;
-									for (int32 LinkedBoneIndex = 0; LinkedBoneIndex < CurveMetaData->LinkedBones.Num(); ++LinkedBoneIndex)
-									{
-										const FBoneReference& BoneReference = CurveMetaData->LinkedBones[LinkedBoneIndex];
-										// when you enter first time, sometimes it does not have all info yet
-										if (BoneReference.BoneIndex != INDEX_NONE && BoneReference.BoneName != NAME_None)
-										{
-											// this linked bone alkways use skeleton index
-											ensure(BoneReference.bUseSkeletonIndex);
-											// we want to make sure all the joints are removed from RequiredBones before removing this UID
-											if (GetCompactPoseIndexFromSkeletonIndex(BoneReference.BoneIndex) != INDEX_NONE)
-											{
-												// still has some joint that matters, do not remove
-												bBeingUsed = true;
-												break;
-											}
-										}
-									}
-								}
+								// still has some joint that matters, do not remove
+								bBeingUsed = true;
+								break;
 							}
 						}
 					}
+				}
 
-					if (bBeingUsed)
-					{
-						UIDToArrayIndexLUT[Iterator.GetIndex()] = NumAvailableUIDs++;
-					}
+				if (!bBeingUsed)
+				{
+					FilterFlags |= UE::Anim::ECurveFilterFlags::Disallowed;
+				}
+
+				if (FilterFlags != UE::Anim::ECurveFilterFlags::None)
+				{
+					// Add curve with any relevant flags to filter
+					CurveFilter.Add(InCurveName, FilterFlags);
+				}
+
+				if(Flags != UE::Anim::ECurveElementFlags::None)
+				{
+					// Add curve with any relevant flags to bulk flags
+					CurveFlags.Add(InCurveName, Flags);
 				}
-				UIDToArrayIndexLUTValidCount = NumAvailableUIDs;
 			});
+
+			// Now add any filtered curves we didnt find in metadata
+			CurveFilter.AppendNames(FilterCurves);
 		}
 	}
-	else
+
+	if (USkeletalMesh* SkeletalMesh = AssetSkeletalMesh.Get())
 	{
-		UIDToArrayIndexLUT.Reset();
-		UIDToArrayIndexLUTValidCount = 0;
-	}
+		// Override any metadata with the skeletal mesh
+		if(const UAnimCurveMetaData* MetaDataUserData = SkeletalMesh->GetAssetUserData<UAnimCurveMetaData>())
+		{
+			UE::Anim::FBulkCurveFlags MeshCurveFlags;
 
-	// Make sure we regenerate our cached curve mappings next time they are requested.
-	MarkAllCachedCurveMappingsDirty();
+			// Apply morph target flags to any morph curves
+			MetaDataUserData->ForEachCurveMetaData([&MeshCurveFlags](FName InCurveName, const FCurveMetaData& InCurveMetaData)
+			{
+				UE::Anim::ECurveElementFlags Flags = UE::Anim::ECurveElementFlags::None;
+				if(InCurveMetaData.Type.bMaterial)
+				{
+					Flags |= UE::Anim::ECurveElementFlags::Material;
+				}
+				if(InCurveMetaData.Type.bMorphtarget)
+				{
+					Flags |= UE::Anim::ECurveElementFlags::MorphTarget;
+				}
+
+				if(Flags != UE::Anim::ECurveElementFlags::None)
+				{
+					MeshCurveFlags.Add(InCurveName, Flags);
+				}
+			});
+
+			UE::Anim::FNamedValueArrayUtils::Union(CurveFlags, MeshCurveFlags);
+		}
+	}
 
 	RegenerateSerialNumber();
 }
@@ -375,66 +369,6 @@ void FBoneContainer::RegenerateSerialNumber()
 	}
 }
 
-void FBoneContainer::MarkAllCachedCurveMappingsDirty()
-{
-	for (auto& CurveMapping : CachedCurveMappingTable)
-	{
-		CurveMapping.Value.bIsDirty = true;
-	}
-}
-
-const FCachedSkeletonCurveMapping& FBoneContainer::GetOrCreateCachedCurveMapping(const FSkeletonRemapping* SkeletonRemapping)
-{
-	LLM_SCOPE_BYNAME(TEXT("Animation/BoneContainer"));
-	
-	check(SkeletonRemapping);
-
-	const USkeleton* SourceSkeleton = SkeletonRemapping->GetSourceSkeleton().Get();
-	check(SourceSkeleton);
-
-	// Make a backup, used for skeleton remapping of curves.
-	if (UIDToArrayIndexLUTBackup.IsEmpty())
-	{
-		UIDToArrayIndexLUTBackup = UIDToArrayIndexLUT;
-	}
-
-	// Check if we already have some cached data for this skeleton.
-	FCachedSkeletonCurveMapping* CachedData = CachedCurveMappingTable.Find(SourceSkeleton);
-	if (CachedData)
-	{
-		// Only return the object when we don't need to update its mapping.
-		if (!CachedData->bIsDirty)
-		{
-			return *CachedData;
-		}
-	}
-	else
-	{
-		CachedData = &CachedCurveMappingTable.Add(const_cast<USkeleton*>(SourceSkeleton));
-	}
-
-	UE_LOG(LogAnimation, Verbose, TEXT("Generating mapping for %s to %s"), 
-		*FAssetData(SkeletonRemapping->GetSourceSkeleton().Get()).AssetName.ToString(),
-		*FAssetData(SkeletonRemapping->GetTargetSkeleton().Get()).AssetName.ToString());
-
-	CachedData->UIDToArrayIndices = SkeletonRemapping->GetSourceToTargetCurveMapping();
-
-	// Now remap the UIDs to our local curve indexes.
-	const int32 NumMappings = CachedData->UIDToArrayIndices.Num();
-	for (int32 MappingIndex = 0; MappingIndex < NumMappings; ++MappingIndex)
-	{
-		const SmartName::UID_Type CurrentUID = CachedData->UIDToArrayIndices[MappingIndex];
-		if (CurrentUID != MAX_uint16)
-		{
-			const int32 CurveIndex = UIDToArrayIndexLUT.Find(CurrentUID);
-			CachedData->UIDToArrayIndices[MappingIndex] = (CurveIndex != INDEX_NONE) ? static_cast<SmartName::UID_Type>(CurveIndex) : MAX_uint16;
-		}
-	}
-
-	CachedData->bIsDirty = false;
-	return *CachedData;
-}
-
 const FRetargetSourceCachedData& FBoneContainer::GetRetargetSourceCachedData(const FName& InRetargetSourceName) const
 {
 	LLM_SCOPE_BYNAME(TEXT("Animation/BoneContainer"));
@@ -598,59 +532,25 @@ void FBoneContainer::RemapFromSkeleton(USkeleton const & SourceSkeleton)
 	PoseToSkeletonBoneIndexArray = SkeletonToPoseBoneIndexArray;
 }
 
-
-/////////////////////////////////////////////////////
-// FBoneReference
-
-bool FBoneReference::Initialize(const FBoneContainer& RequiredBones)
+TArray<uint16> const& FBoneContainer::GetUIDToArrayLookupTable() const
 {
-#if WITH_EDITOR
-	BoneName = *BoneName.ToString().TrimStartAndEnd();
-#endif
-	BoneIndex = RequiredBones.GetPoseBoneIndexForBoneName(BoneName);
-
-	bUseSkeletonIndex = false;
-	// If bone name is not found, look into the leader skeleton to see if it's found there.
-	// SkeletalMeshes can exclude bones from the leader skeleton, and that's OK.
-	// If it's not found in the leader skeleton, the bone does not exist at all! so we should report it as a warning.
-	if (BoneIndex == INDEX_NONE && BoneName != NAME_None)
-	{
-		if (USkeleton* SkeletonAsset = RequiredBones.GetSkeletonAsset())
-		{
-			if (SkeletonAsset->GetReferenceSkeleton().FindBoneIndex(BoneName) == INDEX_NONE)
-			{
-				UE_LOG(LogAnimation, Warning, TEXT("FBoneReference::Initialize BoneIndex for Bone '%s' does not exist in Skeleton '%s'"),
-					*BoneName.ToString(), *GetNameSafe(SkeletonAsset));
-			}
-		}
-	}
-
-	CachedCompactPoseIndex = RequiredBones.MakeCompactPoseIndex(GetMeshPoseIndex(RequiredBones));
-
-	return (BoneIndex != INDEX_NONE);
+	static TArray<uint16> Dummy;
+	return Dummy;
 }
 
-bool FBoneReference::Initialize(const USkeleton* Skeleton)
+int32 FBoneContainer::GetUIDToArrayIndexLookupTableValidCount() const
 {
-	if (Skeleton && (BoneName != NAME_None))
-	{
-#if WITH_EDITOR
-		BoneName = *BoneName.ToString().TrimStartAndEnd();
-#endif
-		BoneIndex = Skeleton->GetReferenceSkeleton().FindBoneIndex(BoneName);
-		bUseSkeletonIndex = true;
-	}
-	else
-	{
-		BoneIndex = INDEX_NONE;
-	}
-
-	CachedCompactPoseIndex = FCompactPoseBoneIndex(INDEX_NONE);
+	return 0;
+}
 
-	return (BoneIndex != INDEX_NONE);
+TArray<FAnimCurveType> const& FBoneContainer::GetUIDToCurveTypeLookupTable() const
+{
+	static TArray<FAnimCurveType> Dummy;
+	return Dummy;
 }
 
-bool FBoneReference::IsValidToEvaluate(const FBoneContainer& RequiredBones) const
+TArray<SmartName::UID_Type> const& FBoneContainer::GetUIDToArrayLookupTableBackup() const
 {
-	return (BoneIndex != INDEX_NONE && RequiredBones.Contains(BoneIndex));
+	static TArray<SmartName::UID_Type> Dummy;
+	return Dummy;
 }
diff --git a/Engine/Source/Runtime/Engine/Private/Components/PoseableMeshComponent.cpp b/Engine/Source/Runtime/Engine/Private/Components/PoseableMeshComponent.cpp
--- a/Engine/Source/Runtime/Engine/Private/Components/PoseableMeshComponent.cpp
+++ b/Engine/Source/Runtime/Engine/Private/Components/PoseableMeshComponent.cpp
@@ -36,7 +36,7 @@ bool UPoseableMeshComponent::AllocateTransformData()
 				RequiredBoneIndexArray[BoneIndex] = BoneIndex;
 			}
 
-			RequiredBones.InitializeTo(RequiredBoneIndexArray, FCurveEvaluationOption(false), *GetSkinnedAsset());
+			RequiredBones.InitializeTo(RequiredBoneIndexArray, UE::Anim::FCurveFilterSettings(UE::Anim::ECurveFilterMode::DisallowAll), *GetSkinnedAsset());
 		}
 
 		FillComponentSpaceTransforms();
