Animation Curve Runtime & Editor Improvements

diff --git a/Engine/Source/Runtime/Engine/Private/Components/SkeletalMeshComponent.cpp b/Engine/Source/Runtime/Engine/Private/Components/SkeletalMeshComponent.cpp
--- a/Engine/Source/Runtime/Engine/Private/Components/SkeletalMeshComponent.cpp
+++ b/Engine/Source/Runtime/Engine/Private/Components/SkeletalMeshComponent.cpp
@@ -1727,27 +1727,22 @@ void USkeletalMeshComponent::RecalcRequiredCurves()
 		return;
 	}
 
-	if (GetSkeletalMeshAsset()->GetSkeleton())
-	{
-		CachedCurveUIDList = GetSkeletalMeshAsset()->GetSkeleton()->GetDefaultCurveUIDList();
-	}
-
-	const FCurveEvaluationOption CurveEvalOption(bAllowAnimCurveEvaluation, &DisallowedAnimCurves, GetPredictedLODLevel());
+	UE::Anim::FCurveFilterSettings CurveFilterSettings = GetCurveFilterSettings();
 
 	// make sure animation requiredcurve to mark as dirty
 	if (AnimScriptInstance)
 	{
-		AnimScriptInstance->RecalcRequiredCurves(CurveEvalOption);
+		AnimScriptInstance->RecalcRequiredCurves(CurveFilterSettings);
 	}
 
 	for(UAnimInstance* LinkedInstance : LinkedInstances)
 	{
-		LinkedInstance->RecalcRequiredCurves(CurveEvalOption);
+		LinkedInstance->RecalcRequiredCurves(CurveFilterSettings);
 	}
 
 	if(PostProcessAnimInstance)
 	{
-		PostProcessAnimInstance->RecalcRequiredCurves(CurveEvalOption);
+		PostProcessAnimInstance->RecalcRequiredCurves(CurveFilterSettings);
 	}
 
 	MarkRequiredCurveUpToDate();
@@ -1966,15 +1961,44 @@ void USkeletalMeshComponent::RecalcRequiredBones(int32 LODIndex)
 
 void USkeletalMeshComponent::MarkRequiredCurveUpToDate()
 {
-	if (GetSkeletalMeshAsset() && GetSkeletalMeshAsset()->GetSkeleton())
+	if(USkeletalMesh* Mesh = GetSkeletalMeshAsset())
 	{
-		CachedAnimCurveUidVersion = GetSkeletalMeshAsset()->GetSkeleton()->GetAnimCurveUidVersion();
+		if(USkeleton* Skeleton = Mesh->GetSkeleton())
+		{
+			CachedAnimCurveUidVersion = GetSkeletalMeshAsset()->GetSkeleton()->GetAnimCurveUidVersion();
+		}
+
+		if(UAnimCurveMetaData* AnimCurveMetaData = Mesh->GetAssetUserData<UAnimCurveMetaData>())
+		{
+			CachedMeshCurveMetaDataVersion = AnimCurveMetaData->GetVersionNumber();
+		}
+		else
+		{
+			CachedMeshCurveMetaDataVersion = 0;
+		}
 	}
 }
 
 bool USkeletalMeshComponent::AreRequiredCurvesUpToDate() const
 {
-	return (!GetSkeletalMeshAsset() || !GetSkeletalMeshAsset()->GetSkeleton() || CachedAnimCurveUidVersion == GetSkeletalMeshAsset()->GetSkeleton()->GetAnimCurveUidVersion());
+	if(USkeletalMesh* Mesh = GetSkeletalMeshAsset())
+	{
+		uint16 MeshVersion = 0;
+		if(UAnimCurveMetaData* AnimCurveMetaData = Mesh->GetAssetUserData<UAnimCurveMetaData>())
+		{
+			MeshVersion = AnimCurveMetaData->GetVersionNumber();
+		}
+
+		uint16 SkeletonVersion = 0;
+		if(USkeleton* Skeleton = Mesh->GetSkeleton())
+		{
+			SkeletonVersion = Skeleton->GetAnimCurveUidVersion();
+		}
+
+		return CachedAnimCurveUidVersion == SkeletonVersion && CachedMeshCurveMetaDataVersion == MeshVersion;
+	}
+
+	return true;
 }
 
 void USkeletalMeshComponent::EvaluateAnimation(const USkeletalMesh* InSkeletalMesh, UAnimInstance* InAnimInstance, FVector& OutRootBoneTranslation, FBlendedHeapCurve& OutCurve, FCompactPose& OutPose, UE::Anim::FHeapAttributeContainer& OutAttributes) const
@@ -1994,10 +2018,6 @@ void USkeletalMeshComponent::EvaluateAnimation(const USkeletalMesh* InSkeletalMe
 		FParallelEvaluationData EvaluationData = { OutCurve, OutPose, OutAttributes };
 		InAnimInstance->ParallelEvaluateAnimation(bForceRefpose, InSkeletalMesh, EvaluationData);
 	}
-	else
-	{
-		OutCurve.InitFrom(&CachedCurveUIDList);
-	}
 }
 
 void USkeletalMeshComponent::UpdateFollowerComponent()
@@ -2349,16 +2369,7 @@ void USkeletalMeshComponent::RefreshBoneTransforms(FActorComponentTickFunction*
 
 	PRAGMA_ENABLE_DEPRECATION_WARNINGS
 
-	TArray<uint16> const* CurrentAnimCurveUIDFinder = (AnimScriptInstance) ? &AnimScriptInstance->GetRequiredBones().GetUIDToArrayLookupTable() : 
-		((ShouldEvaluatePostProcessInstance() && PostProcessAnimInstance) ? &PostProcessAnimInstance->GetRequiredBones().GetUIDToArrayLookupTable() : nullptr);
-	const bool bAnimInstanceHasCurveUIDList = CurrentAnimCurveUIDFinder != nullptr;
-
-	const int32 CurrentCurveCount = (CurrentAnimCurveUIDFinder) ? FBlendedCurve::GetValidElementCount(CurrentAnimCurveUIDFinder) : 0;
-
-	const bool bInvalidCachedCurve = bDoEvaluationRateOptimization && 
-									 bAnimInstanceHasCurveUIDList &&
-									(CachedCurve.UIDToArrayIndexLUT != CurrentAnimCurveUIDFinder || CachedCurve.Num() != CurrentCurveCount);
-
+	const bool bInvalidCachedCurve = bDoEvaluationRateOptimization && CachedCurve.Num() != AnimCurves.Num();
 
 	const bool bInvalidCachedAttributes = bDoEvaluationRateOptimization && CachedAttributes != CustomAttributes;
 
@@ -2366,7 +2377,7 @@ void USkeletalMeshComponent::RefreshBoneTransforms(FActorComponentTickFunction*
 
 	const bool bShouldInterpolateSkippedFrames = (bExternalTickRateControlled && bExternalInterpolate) || (bCachedShouldUseUpdateRateOptimizations && AnimUpdateRateParams->ShouldInterpolateSkippedFrames());
 
-	const bool bShouldDoInterpolation = TickFunction != nullptr && bDoEvaluationRateOptimization && !bInvalidCachedBones && bShouldInterpolateSkippedFrames && bAnimInstanceHasCurveUIDList;
+	const bool bShouldDoInterpolation = TickFunction != nullptr && bDoEvaluationRateOptimization && !bInvalidCachedBones && bShouldInterpolateSkippedFrames;
 
 	const bool bShouldDoParallelInterpolation = bShouldDoInterpolation && CVarUseParallelAnimationInterpolation.GetValueOnGameThread() == 1;
 
@@ -2389,22 +2400,12 @@ void USkeletalMeshComponent::RefreshBoneTransforms(FActorComponentTickFunction*
 	AnimEvaluationContext.AnimInstance = AnimScriptInstance;
 	AnimEvaluationContext.PostProcessAnimInstance = (ShouldEvaluatePostProcessInstance())? ToRawPtr(PostProcessAnimInstance): nullptr;
 
-	if (CurrentAnimCurveUIDFinder)
-	{
-		if (AnimCurves.UIDToArrayIndexLUT != CurrentAnimCurveUIDFinder || AnimCurves.Num() != CurrentCurveCount)
-		{
-			AnimCurves.InitFrom(CurrentAnimCurveUIDFinder);
-		}
-	}
-	else
-	{
-		AnimCurves.Empty();
-	}
+	AnimCurves.Empty();
 
 	AnimEvaluationContext.bDoEvaluation = bShouldDoEvaluation;
 	AnimEvaluationContext.bDoInterpolation = bShouldDoInterpolation;
 	AnimEvaluationContext.bDuplicateToCacheBones = bInvalidCachedBones || (bDoEvaluationRateOptimization && AnimEvaluationContext.bDoEvaluation && !AnimEvaluationContext.bDoInterpolation);
-	AnimEvaluationContext.bDuplicateToCacheCurve = bInvalidCachedCurve || (bDoEvaluationRateOptimization && AnimEvaluationContext.bDoEvaluation && !AnimEvaluationContext.bDoInterpolation && CurrentAnimCurveUIDFinder != nullptr);
+	AnimEvaluationContext.bDuplicateToCacheCurve = bInvalidCachedCurve || (bDoEvaluationRateOptimization && AnimEvaluationContext.bDoEvaluation && !AnimEvaluationContext.bDoInterpolation);
 
 	AnimEvaluationContext.bDuplicateToCachedAttributes = bInvalidCachedAttributes || (bDoEvaluationRateOptimization && AnimEvaluationContext.bDoEvaluation && !AnimEvaluationContext.bDoInterpolation);
 
@@ -2476,10 +2477,8 @@ void USkeletalMeshComponent::RefreshBoneTransforms(FActorComponentTickFunction*
 					LocalEditableSpaceBases.Reset();
 					LocalEditableSpaceBases.Append(CachedComponentSpaceTransforms);
 				}
-				if (CachedCurve.IsValid())
-				{
-					AnimCurves.CopyFrom(CachedCurve);
-				}
+
+				AnimCurves.CopyFrom(CachedCurve);
 
 				if (CachedAttributes.ContainsData())
 				{
@@ -2648,12 +2647,9 @@ void USkeletalMeshComponent::PostAnimEvaluation(FAnimationEvaluationContext& Eva
 	{
 		if (EvaluationContext.bDuplicateToCacheCurve)
 		{
-			ensureAlwaysMsgf(AnimCurves.IsValid(), TEXT("Animation Curve is invalid (%s). TotalCount(%d) "),
-				*GetPathNameSafe(GetSkeletalMeshAsset()), AnimCurves.NumValidCurveCount);
 			CachedCurve.CopyFrom(AnimCurves);
 		}
-
-	
+		
 		if (EvaluationContext.bDuplicateToCachedAttributes)
 		{
 			CachedAttributes.CopyFrom(CustomAttributes);
@@ -4049,8 +4045,6 @@ void USkeletalMeshComponent::ParallelDuplicateAndInterpolate(FAnimationEvaluatio
 	{
 		if (InAnimEvaluationContext.bDuplicateToCacheCurve)
 		{
-			ensureAlwaysMsgf(InAnimEvaluationContext.Curve.IsValid(), TEXT("Animation Curve is invalid (%s). TotalCount(%d) "),
-				*GetPathNameSafe(GetSkeletalMeshAsset()), InAnimEvaluationContext.Curve.NumValidCurveCount);
 			InAnimEvaluationContext.CachedCurve.CopyFrom(InAnimEvaluationContext.Curve);
 		}
 
@@ -4502,26 +4496,16 @@ void USkeletalMeshComponent::SetAllowAnimCurveEvaluation(bool bInAllow)
 
 void USkeletalMeshComponent::AllowAnimCurveEvaluation(FName NameOfCurve, bool bAllow)
 {
-	// if allow is same as disallowed curve, which means it mismatches
-	if (bAllow == DisallowedAnimCurves.Contains(NameOfCurve))
+	if (bAllow == bFilteredAnimCurvesIsAllowList)
 	{
-		if (bAllow)
-		{
-			DisallowedAnimCurves.Remove(NameOfCurve);
-			CachedAnimCurveUidVersion = 0;
-		}
-		else
-		{
-			DisallowedAnimCurves.Add(NameOfCurve);
-			CachedAnimCurveUidVersion = 0;
-
-		}
+		FilteredAnimCurves.Add(NameOfCurve);
+		CachedAnimCurveUidVersion = 0;
 	}
 }
 
 void USkeletalMeshComponent::ResetAllowedAnimCurveEvaluation()
 {
-	DisallowedAnimCurves.Reset();
+	FilteredAnimCurves.Reset();
 	CachedAnimCurveUidVersion = 0;
 }
 
@@ -4529,41 +4513,40 @@ void USkeletalMeshComponent::SetAllowedAnimCurvesEvaluation(const TArray<FName>&
 {
 	// Reset already clears the version - CachedAnimCurveUidVersion = 0;
 	ResetAllowedAnimCurveEvaluation();
-	if (bAllow)
+	FilteredAnimCurves.Append(List);
+	bFilteredAnimCurvesIsAllowList = bAllow;
+}
+
+UE::Anim::FCurveFilterSettings USkeletalMeshComponent::GetCurveFilterSettings(int32 InLODOverride) const
+{
+	UE::Anim::FCurveFilterSettings CurveFilterSettings;
+	if(bAllowAnimCurveEvaluation)
 	{
-		struct FFilterDisallowedList
+		if(FilteredAnimCurves.Num() > 0)
 		{
-			FFilterDisallowedList(const TArray<FName>& InAllowedList) : AllowedList(InAllowedList) {}
-
-			FORCEINLINE bool operator()(const FName& Name) const
-			{
-				return AllowedList.Contains(Name);
-			}
-			const TArray<FName>& AllowedList;
-		};
-
-		if (GetSkeletalMeshAsset())
+			CurveFilterSettings.FilterMode = bFilteredAnimCurvesIsAllowList ? UE::Anim::ECurveFilterMode::AllowOnlyFiltered : UE::Anim::ECurveFilterMode::DisallowFiltered;
+			CurveFilterSettings.FilterCurves = &FilteredAnimCurves;
+		}
+		else
 		{
-			USkeleton* Skeleton = GetSkeletalMeshAsset()->GetSkeleton();
-			if (Skeleton)
-			{
-				const FSmartNameMapping* Mapping = Skeleton->GetSmartNameContainer(USkeleton::AnimCurveMappingName);
-				if (Mapping != nullptr)
-				{
-					TArray<FName> CurveNames;
-					Mapping->FillNameArray(CurveNames);
-
-					DisallowedAnimCurves = CurveNames;
-					DisallowedAnimCurves.RemoveAllSwap(FFilterDisallowedList(List));
-				}
-			}
-
+			CurveFilterSettings.FilterMode = UE::Anim::ECurveFilterMode::None;
 		}
 	}
 	else
 	{
-		DisallowedAnimCurves = List;
+		CurveFilterSettings.FilterMode = UE::Anim::ECurveFilterMode::DisallowAll;
+	}
+
+	if(InLODOverride != INDEX_NONE)
+	{
+		CurveFilterSettings.LODIndex = InLODOverride;
+	}
+	else
+	{
+		CurveFilterSettings.LODIndex = GetPredictedLODLevel();
 	}
+
+	return CurveFilterSettings;
 }
 
 const TArray<FTransform>& USkeletalMeshComponent::GetCachedComponentSpaceTransforms() const
diff --git a/Engine/Source/Runtime/Engine/Private/Engine/PoseWatchRenderData.cpp b/Engine/Source/Runtime/Engine/Private/Engine/PoseWatchRenderData.cpp
--- a/Engine/Source/Runtime/Engine/Private/Engine/PoseWatchRenderData.cpp
+++ b/Engine/Source/Runtime/Engine/Private/Engine/PoseWatchRenderData.cpp
@@ -18,6 +18,11 @@ FAnimNodePoseWatch::FAnimNodePoseWatch()
 	ViewportOffset = FVector::ZeroVector;
 }
 
+void FAnimNodePoseWatch::SetCurves(const FBlendedCurve& InCurves)
+{
+	Curves.CopyFrom(InCurves);
+}
+
 void FAnimNodePoseWatch::SetWorldTransform(const FTransform& InWorldTransform)
 {
 	WorldTransform = InWorldTransform;
@@ -64,6 +69,11 @@ const TArray<FTransform>& FAnimNodePoseWatch::GetBoneTransforms() const
 	return BoneTransforms;
 }
 
+const FBlendedHeapCurve& FAnimNodePoseWatch::GetCurves() const
+{
+	return Curves;
+}
+
 const FTransform& FAnimNodePoseWatch::GetWorldTransform() const
 {
 	return WorldTransform;
diff --git a/Engine/Source/Runtime/Engine/Private/SkeletalMesh.cpp b/Engine/Source/Runtime/Engine/Private/SkeletalMesh.cpp
--- a/Engine/Source/Runtime/Engine/Private/SkeletalMesh.cpp
+++ b/Engine/Source/Runtime/Engine/Private/SkeletalMesh.cpp
@@ -2953,7 +2953,7 @@ void USkeletalMesh::PostLoadVerifyAndFixBadTangent()
 
 bool USkeletalMesh::IsPostLoadThreadSafe() const
 {
-	return false;	// PostLoad is not thread safe because of the call to InitMorphTargets, which can call VerifySmartName() that can mutate a shared map in the skeleton.
+	return false;	// PostLoad is not thread safe
 }
 
 void USkeletalMesh::BeginPostLoadInternal(FSkinnedAssetPostLoadContext& Context)
@@ -3505,6 +3505,15 @@ void USkeletalMesh::GetAssetRegistryTags(TArray<FAssetRegistryTag>& OutTags) con
 	// The tag must be added unconditionally, because some code calls this function on the CDO to find out what
 	// tags are available.
 	OutTags.Add(FAssetRegistryTag("MaxBoneInfluences", MaxBoneInfluencesString, FAssetRegistryTag::TT_Numerical));
+
+	// Allow asset user data to output tags
+	for(UAssetUserData* AssetUserDataItem : AssetUserData)
+	{
+		if(AssetUserDataItem)
+		{
+			AssetUserDataItem->GetAssetRegistryTags(OutTags);
+		}
+	}
 #endif // WITH_EDITORONLY_DATA
 	
 	Super::GetAssetRegistryTags(OutTags);
@@ -3656,17 +3665,8 @@ void USkeletalMesh::InitMorphTargets()
 		{
 			GetMorphTargetIndexMap().Add(ShapeName, Index);
 
-			// register as morphtarget curves
-			if (GetSkeleton())
-			{
-				FSmartName CurveName;
-				CurveName.DisplayName = ShapeName;
-				
-				// verify will make sure it adds to the curve if not found
-				// the reason of using this is to make sure it works in editor/non-editor
-				GetSkeleton()->VerifySmartName(USkeleton::AnimCurveMappingName, CurveName);
-				GetSkeleton()->AccumulateCurveMetaData(ShapeName, false, true);
-			}
+			// Note: we dont register as morph target curves here as curves metadata can now be
+			// specified on this mesh, which can now opt out of the morph flag being set
 		}
 	}
 }
diff --git a/Engine/Source/Runtime/Engine/Public/Animation/AnimCompressionTypes.h b/Engine/Source/Runtime/Engine/Public/Animation/AnimCompressionTypes.h
--- a/Engine/Source/Runtime/Engine/Public/Animation/AnimCompressionTypes.h
+++ b/Engine/Source/Runtime/Engine/Public/Animation/AnimCompressionTypes.h
@@ -24,6 +24,11 @@
 
 #include "ProfilingDebugging/CsvProfiler.h"
 #include "Animation/AnimationDecompression.h"
+#include "BoneContainer.h"
+
+#if WITH_EDITOR
+#include "ReferenceSkeleton.h"
+#endif
 
 #include "AnimCompressionTypes.generated.h"
 
@@ -731,6 +736,36 @@ struct TIsContiguousContainer<TMaybeMappedArray<T, Alignment>>
 	static constexpr bool Value = TIsContiguousContainer<TArray<T, TMaybeMappedAllocator<Alignment>>>::Value;
 };
 
+USTRUCT()
+struct FAnimCompressedCurveIndexedName
+{
+	GENERATED_BODY()
+
+	FAnimCompressedCurveIndexedName() = default;
+	
+	bool operator<(const FAnimCompressedCurveIndexedName& InOther) const
+	{
+		return InOther.CurveName.FastLess(CurveName);
+	}
+
+	friend FArchive& operator<<(FArchive& Ar, FAnimCompressedCurveIndexedName& Item)
+	{
+		Ar << Item.CurveName;
+		if(Ar.IsCountingMemory())
+		{
+			Ar << Item.CurveIndex;
+		}
+		return Ar;
+	}
+	
+	// Name of the curve, used for sorting
+	UPROPERTY()
+	FName CurveName = NAME_None;
+
+	// Index into the compressed buffer
+	int32 CurveIndex = INDEX_NONE;
+};
+
 struct ENGINE_API FCompressedAnimSequence
 {
 public:
@@ -742,12 +777,12 @@ public:
 	 */
 	TArray<struct FTrackToSkeletonMap> CompressedTrackToSkeletonMapTable;
 
-	/**
-	 * Much like track indices above, we need to be able to remap curve names. The FName inside FSmartName
-	 * is serialized and the UID is generated at runtime. Stored in the DDC.
-	 */
+	UE_DEPRECATED(5.3, "Please use IndexedCurveNames.")
 	TArray<struct FSmartName> CompressedCurveNames;
 
+	/** Curve names for each compressed curve track, indices sorted by FName */
+	TArray<FAnimCompressedCurveIndexedName> IndexedCurveNames;
+
 	/**
 	 * ByteStream for compressed animation data.
 	 * The memory layout is dependent on the algorithm used to compress the anim sequence.
@@ -788,6 +823,8 @@ public:
 
 	void SerializeCompressedData(FArchive& Ar, bool bDDCData, UObject* DataOwner, USkeleton* Skeleton, UAnimBoneCompressionSettings* BoneCompressionSettings, UAnimCurveCompressionSettings* CurveCompressionSettings, bool bCanUseBulkData=true);
 
+	void RebuildCurveIndexTable();
+	
 	int32 GetSkeletonIndexFromTrackIndex(const int32 TrackIndex) const
 	{
 		return CompressedTrackToSkeletonMapTable[TrackIndex].BoneTreeIndex;
diff --git a/Engine/Source/Runtime/Engine/Public/Animation/AnimCurveTypes.h b/Engine/Source/Runtime/Engine/Public/Animation/AnimCurveTypes.h
--- a/Engine/Source/Runtime/Engine/Public/Animation/AnimCurveTypes.h
+++ b/Engine/Source/Runtime/Engine/Public/Animation/AnimCurveTypes.h
@@ -8,6 +8,8 @@
 #include "Animation/SmartName.h"
 #include "Curves/RichCurve.h"
 #include "Misc/EnumRange.h"
+#include "AnimCurveElementFlags.h"
+#include "Animation/NamedValueArray.h"
 
 #if UE_ENABLE_INCLUDE_ORDER_DEPRECATED_IN_5_2
 #include "Animation/Skeleton.h"
@@ -17,6 +19,15 @@
 
 typedef SmartName::UID_Type SkeletonAnimCurveUID;
 
+class USkeleton;
+struct FBoneContainer;
+
+namespace UE::Anim
+{
+struct FCurveUtils;
+struct FCurveFilter;
+}
+
 UENUM()
 enum class EAnimCurveType : uint8 
 {
@@ -53,10 +64,7 @@ static const EAnimAssetCurveFlags AACF_DefaultCurve = AACF_Editable;
 
 ENUM_RANGE_BY_FIRST_AND_LAST(EAnimAssetCurveFlags, AACF_DriveMorphTarget_DEPRECATED, AACF_Disabled);
 
-/** UI Curve Parameter type
- * This gets name, and cached UID and use it when needed
- * Also it contains curve types 
- */
+/** DEPRECATED - no longer used */
 USTRUCT()
 struct ENGINE_API FAnimCurveParam
 {
@@ -72,20 +80,14 @@ struct ENGINE_API FAnimCurveParam
 		: UID(SmartName::MaxUID)
 	{}
 
-	// initialize
-	void Initialize(USkeleton* Skeleton);
-
-	// this doesn't check CurveType flag
-	// because it's possible you don't care about your curve types
-	bool IsValid() const
-	{
-		return UID != SmartName::MaxUID;
-	}
+	UE_DEPRECATED(5.3, "FAnimCurveParam is no longer used.")
+	void Initialize(USkeleton* Skeleton) {}
 
-	bool IsValidToEvaluate() const
-	{
-		return IsValid();
-	}
+	UE_DEPRECATED(5.3, "FAnimCurveParam is no longer used.")
+	bool IsValid() const { return false; }
+	
+	UE_DEPRECATED(5.3, "FAnimCurveParam is no longer used.")
+	bool IsValidToEvaluate() const { return false; }
 };
 /**
  * Float curve data for one track
@@ -96,22 +98,20 @@ struct ENGINE_API FAnimCurveBase
 	GENERATED_USTRUCT_BODY()
 
 #if WITH_EDITORONLY_DATA
-	// Last observed name of the curve. We store this so we can recover from situations that
-	// mean the skeleton doesn't have a mapped name for our UID (such as a user saving the an
-	// animation but not the skeleton).
 	UPROPERTY()
 	FName		LastObservedName_DEPRECATED;
-#endif
 
 	UPROPERTY()
-	FSmartName	Name;
+	FSmartName	Name_DEPRECATED;
 
-#if WITH_EDITORONLY_DATA
 	UPROPERTY()
 	FLinearColor Color;
 #endif
 
 private:
+	UPROPERTY()
+	FName CurveName;
+
 	// this flag is mostly used by editor only now
 	// however I can't remove this to editor only because 
 	// we need DEPRECATED Flag to be loaded in game
@@ -129,8 +129,18 @@ public:
 #endif
 	}
 
+	FAnimCurveBase(FName InName, int32 InCurveTypeFlags)
+		: CurveName(InName)
+		, CurveTypeFlags(InCurveTypeFlags)
+	{
+#if WITH_EDITORONLY_DATA
+		Color = MakeColor();
+#endif
+	}
+
+	UE_DEPRECATED(5.3, "Please use the cosntructor that takes an FName.")
 	FAnimCurveBase(FSmartName InName, int32 InCurveTypeFlags)
-		: Name(InName)
+		: CurveName(InName.DisplayName)
 		, CurveTypeFlags(InCurveTypeFlags)
 	{
 #if WITH_EDITORONLY_DATA
@@ -138,8 +148,11 @@ public:
 #endif
 	}
 
-	// To be able to use typedef'd types we need to serialize manually
-	void PostSerialize(FArchive& Ar);
+	// This allows loading data that was saved between VER_UE4_SKELETON_ADD_SMARTNAMES and FFrameworkObjectVersion::SmartNameRefactor
+	void PostSerializeFixup(FArchive& Ar);
+
+	bool Serialize(FArchive& Ar);
+	void PostSerialize(const FArchive& Ar);
 
 	/**
 	 * Set InFlag to bValue
@@ -166,6 +179,18 @@ public:
 	 */
 	int32 GetCurveTypeFlags() const;
 
+	/** Get the name of this curve */
+	FName GetName() const
+	{
+		return CurveName;
+	}
+
+	/** Set the name of this curve */
+	void SetName(FName InName)
+	{
+		CurveName = InName;
+	}
+
 #if WITH_EDITORONLY_DATA
 	/** Get the color used to display this curve in the editor */
 	FLinearColor GetColor() const { return Color; }
@@ -176,6 +201,16 @@ private:
 #endif
 };
 
+template<>
+struct TStructOpsTypeTraits<FAnimCurveBase> : public TStructOpsTypeTraitsBase2<FAnimCurveBase>
+{
+	enum 
+	{
+		WithSerializer = true,
+		WithPostSerialize = true,
+	};
+};
+
 USTRUCT(BlueprintType)
 struct ENGINE_API FFloatCurve : public FAnimCurveBase
 {
@@ -187,11 +222,19 @@ struct ENGINE_API FFloatCurve : public FAnimCurveBase
 
 	FFloatCurve(){}
 
+	UE_DEPRECATED(5.3, "Please use the constructor that takes an FName.")
 	FFloatCurve(FSmartName InName, int32 InCurveTypeFlags)
+	PRAGMA_DISABLE_DEPRECATION_WARNINGS
 		: FAnimCurveBase(InName, InCurveTypeFlags)
+	PRAGMA_ENABLE_DEPRECATION_WARNINGS
 	{
 	}
 
+	FFloatCurve(FName InName, int32 InCurveTypeFlags)
+		: FAnimCurveBase(InName, InCurveTypeFlags)
+	{
+	}
+	
 	// we don't want to have = operator. This only copies curves, but leaving naming and everything else intact. 
 	void CopyCurve(const FFloatCurve& SourceCurve);
 	float Evaluate(float CurrentTime) const;
@@ -200,6 +243,16 @@ struct ENGINE_API FFloatCurve : public FAnimCurveBase
 	void Resize(float NewLength, bool bInsert/* whether insert or remove*/, float OldStartTime, float OldEndTime);
 };
 
+template<>
+struct TStructOpsTypeTraits<FFloatCurve> : public TStructOpsTypeTraitsBase2<FFloatCurve>
+{
+	enum 
+	{
+		WithSerializer = true,
+		WithPostSerialize = true,
+	};
+};
+
 USTRUCT()
 struct ENGINE_API FVectorCurve : public FAnimCurveBase
 {
@@ -219,7 +272,15 @@ struct ENGINE_API FVectorCurve : public FAnimCurveBase
 
 	FVectorCurve(){}
 
+	UE_DEPRECATED(5.3, "Please use the constructor that takes an FName.")
 	FVectorCurve(FSmartName InName, int32 InCurveTypeFlags)
+	PRAGMA_DISABLE_DEPRECATION_WARNINGS
+		: FAnimCurveBase(InName, InCurveTypeFlags)
+	PRAGMA_ENABLE_DEPRECATION_WARNINGS
+	{
+	}
+	
+	FVectorCurve(FName InName, int32 InCurveTypeFlags)
 		: FAnimCurveBase(InName, InCurveTypeFlags)
 	{
 	}
@@ -234,6 +295,16 @@ struct ENGINE_API FVectorCurve : public FAnimCurveBase
 	int32 GetNumKeys() const;
 };
 
+template<>
+struct TStructOpsTypeTraits<FVectorCurve> : public TStructOpsTypeTraitsBase2<FVectorCurve>
+{
+	enum 
+	{
+		WithSerializer = true,
+		WithPostSerialize = true,
+	};
+};
+
 USTRUCT(BlueprintType)
 struct ENGINE_API FTransformCurve: public FAnimCurveBase
 {
@@ -255,7 +326,15 @@ struct ENGINE_API FTransformCurve: public FAnimCurveBase
 
 	FTransformCurve(){}
 
+	UE_DEPRECATED(5.3, "Please use the constructor that takes an FName.")
 	FTransformCurve(FSmartName InName, int32 InCurveTypeFlags)
+	PRAGMA_DISABLE_DEPRECATION_WARNINGS
+		: FAnimCurveBase(InName, InCurveTypeFlags)
+	PRAGMA_ENABLE_DEPRECATION_WARNINGS
+	{
+	}
+	
+	FTransformCurve(FName InName, int32 InCurveTypeFlags)
 		: FAnimCurveBase(InName, InCurveTypeFlags)
 	{
 	}
@@ -271,6 +350,16 @@ struct ENGINE_API FTransformCurve: public FAnimCurveBase
 	FVectorCurve* GetVectorCurveByIndex(int32 Index);
 };
 
+template<>
+struct TStructOpsTypeTraits<FTransformCurve> : public TStructOpsTypeTraitsBase2<FTransformCurve>
+{
+	enum 
+	{
+		WithSerializer = true,
+		WithPostSerialize = true,
+	};
+};
+
 USTRUCT(BlueprintType)
 struct ENGINE_API FCachedFloatCurve
 {
@@ -280,27 +369,21 @@ public:
 	UPROPERTY(EditAnywhere, BlueprintReadWrite, Category = "Curve Settings")
 	FName CurveName;
 
-private:
-	mutable SkeletonAnimCurveUID CachedUID;
-	mutable FName CachedCurveName;
-
 public:
-	FCachedFloatCurve()
-		: CachedUID(SmartName::MaxUID)
-	{}
 
 	bool IsValid(const UAnimSequenceBase* InAnimSequence) const;
 	float GetValueAtPosition(const UAnimSequenceBase* InAnimSequence, const float& InPosition) const;
 	const FFloatCurve* GetFloatCurve(const UAnimSequenceBase* InAnimSequence) const;
 
 protected:
-	SkeletonAnimCurveUID GetAnimCurveUID(const UAnimSequenceBase* InAnimSequence) const;
+	UE_DEPRECATED(5.3, "Please just use CurveName.")
+	SmartName::UID_Type GetAnimCurveUID(const UAnimSequenceBase* InAnimSequence) const { return 0; }
 };
 
 /**
 * This is array of curves that run when collecting curves natively 
 */
-struct FCurveElement
+struct UE_DEPRECATED(5.3, "FCurveElement in the global namespace is no longer used.") FCurveElement
 {
 	/** Curve Value */
 	float					Value;
@@ -329,198 +412,313 @@ struct FCurveElement
 	}
 };
 
+namespace UE::Anim
+{
+
+// An element for a named curve value
+struct FCurveElement
+{
+	FCurveElement() = default;
+
+	FCurveElement(FName InName)
+		: Name(InName)
+	{}
+
+	FCurveElement(FName InName, float InValue)
+		: Name(InName)
+		, Value(InValue)
+	{}
+
+	FCurveElement(FName InName, ECurveElementFlags InFlags)
+		: Name(InName)
+		, Flags(InFlags) 
+	{}
+	
+	FCurveElement(FName InName, float InValue, ECurveElementFlags InFlags)
+		: Name(InName)
+		, Value(InValue)
+		, Flags(InFlags) 
+	{}
+	
+	FName Name = NAME_None;
+	float Value = 0.0f;
+	ECurveElementFlags Flags = ECurveElementFlags::None;
+};
+
+// An element for a named curve value that is also indexed.
+// Useful in bulk curve operations that need to pay heed to disabled curves
+struct FCurveElementIndexed : public UE::Anim::FCurveElement
+{
+	FCurveElementIndexed() = default;
+
+	FCurveElementIndexed(FName InName, int32 InIndex)
+		: UE::Anim::FCurveElement(InName)
+		, Index(InIndex)
+	{}
+
+	int32 Index = INDEX_NONE;
+};
+
+// Deprecated base for TBaseBlendedCurve for legacy support
+// We need to use this non-templated base to avoid the deprecated (and now static) members taking up per-instance memory
+struct ENGINE_API FBaseBlendedCurve_DEPRECATED
+{
+	UE_DEPRECATED(5.3, "Direct access to CurveWeights is no longer allowed as it is no longer used.")
+	static TArray<float> CurveWeights;
+	
+	UE_DEPRECATED(5.3, "Direct access to ValidCurveWeights is no longer allowed as it is no longer used.")
+	static TBitArray<> ValidCurveWeights;
+	
+	UE_DEPRECATED(5.3, "Direct access to UIDToArrayIndexLUT is no longer allowed as it is no longer used.")
+	static TArray<uint16> const * UIDToArrayIndexLUT;
+	
+	UE_DEPRECATED(5.3, "Direct access to NumValidCurveCount is no longer allowed as it is no longer used.")
+	static uint16 NumValidCurveCount;
+	
+	UE_DEPRECATED(5.3, "Direct access to bInitialized is no longer allowed.")
+	static bool bInitialized;
+};
+
+}
+
 /**
  * This struct is used to create curve snap shot of current time when extracted
  */
-template <typename InAllocator>
-struct FBaseBlendedCurve
+template <typename InAllocatorType = FAnimStackAllocator, typename InElementType = UE::Anim::FCurveElement>
+struct TBaseBlendedCurve : public UE::Anim::TNamedValueArray<InAllocatorType, InElementType>, public UE::Anim::FBaseBlendedCurve_DEPRECATED
 {
-	typedef InAllocator   Allocator;
-	/**
-	* List of curve weights for this pose
-	*/
-	TArray<float, Allocator> CurveWeights;
+public:
+	typedef InAllocatorType AllocatorType;
+	typedef InElementType ElementType;
 
-	/**
-	 * A bitmask to indicate which weights are valid.
-	 */
-	TBitArray<Allocator> ValidCurveWeights;
+	template<typename OtherAllocator, typename OtherElementType>
+	friend struct TBaseBlendedCurve;
 
-	/**
-	* UID to array index look up table for Elements
-	* This eliminates the look up cost
-	*/
-	TArray<uint16> const * UIDToArrayIndexLUT;
+	friend struct UE::Anim::FCurveUtils;
 
-	/** 
-	 * Valid Curve count 
-	 * This should match Elements.Num() 
-	 */
-	uint16 NumValidCurveCount;
+	typedef UE::Anim::TNamedValueArray<AllocatorType, ElementType> Super;
 
-	/**
-	 * constructor
-	 */
-	FBaseBlendedCurve()
-		: UIDToArrayIndexLUT(nullptr)
-		, NumValidCurveCount(0)
-		, bInitialized(false)
+private:
+	// Filter to use when building curves
+	const UE::Anim::FCurveFilter* Filter = nullptr;
+
+public:
+	UE_DEPRECATED(5.3, "InitFrom can no longer be called with a LUT, please use another initialization method or just dont initialize the curve.")
+	void InitFrom(TArray<uint16> const * InUIDToArrayIndexLUT)
 	{
+		Super::Elements.Reset();
+		Filter = nullptr;
+		Super::bSorted = false;
 	}
 
-	/** Initialize Curve Data from following data */
-	void InitFrom(const FBoneContainer& RequiredBones)
+	/** Initialize from another curve */
+	template <typename OtherAllocator, typename OtherElementType>
+	void InitFrom(const TBaseBlendedCurve<OtherAllocator, OtherElementType>& InCurveToInitFrom)
 	{
-		UIDToArrayIndexLUT = &RequiredBones.GetUIDToArrayLookupTable();
-		NumValidCurveCount = (uint16)RequiredBones.GetUIDToArrayIndexLookupTableValidCount();
-		CurveWeights.Reset();
-		CurveWeights.AddZeroed(NumValidCurveCount);
-		ValidCurveWeights.Init(false, NumValidCurveCount);
-		// no name, means no curve
-		bInitialized = true;
+		CURVE_PROFILE_CYCLE_COUNTER(TBaseBlendedCurve_InitFrom);
+		
+		if constexpr(std::is_same<AllocatorType, OtherAllocator>::value && std::is_same<ElementType, OtherElementType>::value)
+		{
+			if (ensure(&InCurveToInitFrom != this))
+			{
+				Super::Elements = InCurveToInitFrom.Elements;
+				Super::bSorted = InCurveToInitFrom.bSorted;
+				Filter = InCurveToInitFrom.Filter;
+			}
+		}
+		else if constexpr(std::is_same<ElementType, OtherElementType>::value)
+		{
+			Super::Elements = InCurveToInitFrom.Elements;
+			Super::bSorted = InCurveToInitFrom.bSorted;
+			Filter = InCurveToInitFrom.Filter;
+		}
+		else
+		{
+			Super::Elements.Reset();
+
+			UE::Anim::FNamedValueArrayUtils::Union(*this, InCurveToInitFrom,
+				[](ElementType& InOutResultElement, OtherElementType& InParamElement, UE::Anim::ENamedValueUnionFlags InFlags)
+				{
+					InOutResultElement.Value = InParamElement.Value;
+					InOutResultElement.Flags = InParamElement.Flags;
+				});
+
+			Super::bSorted = InCurveToInitFrom.bSorted;
+			Filter = InCurveToInitFrom.Filter;
+		}
 	}
 
-	void InitFrom(TArray<uint16> const * InUIDToArrayIndexLUT)
+	UE_DEPRECATED(5.3, "Please use InvalidateCurveWeight with a curve name")
+	void InvalidateCurveWeight(SmartName::UID_Type InUid)
 	{
-		check(InUIDToArrayIndexLUT != nullptr);
-		UIDToArrayIndexLUT = InUIDToArrayIndexLUT;
-		NumValidCurveCount = (uint16)GetValidElementCount(UIDToArrayIndexLUT);
-		CurveWeights.Reset();
-		CurveWeights.AddZeroed(NumValidCurveCount);
-		ValidCurveWeights.Init(false, NumValidCurveCount);
-		// no name, means no curve
-		bInitialized = true;
 	}
 
-	template <typename OtherAllocator>
-	void InitFrom(const FBaseBlendedCurve<OtherAllocator>& InCurveToInitFrom)
+	/** Invalidate value of the named curve */
+	void InvalidateCurveWeight(FName InName)
+	{
+		const int32 ElementIndex = Super::IndexOf(InName);
+		if(ElementIndex != INDEX_NONE)
+		{
+			Super::Elements.RemoveAt(ElementIndex);
+		}
+	}
+	
+	UE_DEPRECATED(5.3, "Please use Set with a curve name or consider using one of the bulk APIs (e.g. UE::Anim::FCurveUtils::BulkSet).")
+	void Set(SmartName::UID_Type InUid, float InValue)
 	{
-		// make sure this doesn't happen
-		check(InCurveToInitFrom.UIDToArrayIndexLUT != nullptr);
-		UIDToArrayIndexLUT = InCurveToInitFrom.UIDToArrayIndexLUT;
-		NumValidCurveCount = InCurveToInitFrom.NumValidCurveCount;
-
-		CurveWeights.Reset();
-		CurveWeights.AddZeroed(NumValidCurveCount);
-		ValidCurveWeights.Init(false, NumValidCurveCount);
-		bInitialized = true;
 	}
 
-	void InitFrom(const FBaseBlendedCurve<Allocator>& InCurveToInitFrom)
+	/**
+	 * Set value of curve named InName to InValue. Has no effect if the curve does not contain an element for InName.
+	 * Note that this performs a binary search per-call. Consider using a combiner operation to set multiple element's values.
+	 * @param	InName	the name of the curve to set
+	 * @param	InValue	the value of the curve to set
+	 */
+	void Set(FName InName, float InValue)
 	{
-		// make sure this doesn't happen
-		if (ensure(&InCurveToInitFrom != this))
+		if(ElementType* CurveElement = Super::Find(InName))
 		{
-			check(InCurveToInitFrom.UIDToArrayIndexLUT != nullptr);
-			UIDToArrayIndexLUT = InCurveToInitFrom.UIDToArrayIndexLUT;
-			NumValidCurveCount = InCurveToInitFrom.NumValidCurveCount;
-			CurveWeights.Reset();
-			CurveWeights.AddZeroed(NumValidCurveCount);
-			ValidCurveWeights.Init(false, NumValidCurveCount);
-			bInitialized = true;
+			CurveElement->Value = InValue;
 		}
 	}
 
-	/** Invalidate value of InUID */
-	void InvalidateCurveWeight(SkeletonAnimCurveUID InUid)
+	/**
+	 * Set flags of curve named InName to InValue. Has no effect if the curve does not contain an element for InName.
+	 * Note that this performs a binary search per-call. Consider using a combiner operation to set multiple element's flags.
+	 * @param	InName	the name of the curve to set
+	 * @param	InFlags	the flags of the curve to set
+	 */
+	void SetFlags(FName InName, UE::Anim::ECurveElementFlags InFlags)
 	{
-		check(bInitialized);
-
-		int32 ArrayIndex = GetArrayIndexByUID(InUid);
-		if (ArrayIndex != INDEX_NONE)
+		if(ElementType* CurveElement = Super::Find(InName))
 		{
-			CurveWeights[ArrayIndex] = 0.f;
-			ValidCurveWeights[ArrayIndex] = false;
+			CurveElement->Flags = InFlags;
 		}
 	}
 
-	/** Set value of InUID to InValue */
-	void Set(SkeletonAnimCurveUID InUid, float InValue)
+	UE_DEPRECATED(5.3, "Please use Get with a curve name or consider using one of the bulk APIs (e.g. UE::Anim::FCurveUtils::BulkGet)")
+	float Get(SmartName::UID_Type InUid) const
 	{
-		check(bInitialized);
+		return 0.f;
+	}
 
-		int32 ArrayIndex = GetArrayIndexByUID(InUid);
-		if (ArrayIndex != INDEX_NONE)
+	/**
+	 * Get value of curve element named InName.
+	 * Note that this performs a binary search per-call. Consider using a ForEach* or a bulk API call
+	 * (e.g. UE::Anim::FCurveUtils::BulkGet) to get multiple element's values.
+	 * @param	InName	the name of the curve element to get
+	 * @return  the value of the curve element. If this curve does not contain an element with the supplied name, returns 0.0 
+	 */
+	float Get(FName InName) const
+	{
+		if(const ElementType* CurveElement = Super::Find(InName))
 		{
-			CurveWeights[ArrayIndex] = InValue;
-			ValidCurveWeights[ArrayIndex] = true;
+			return CurveElement->Value;
 		}
+		return 0.0f;
 	}
 
-	/** Get Value of InUID - @todo : add validation check here and make sure caller also knows it's not valid*/
-	float Get(SkeletonAnimCurveUID InUid) const
+	/**
+	 * Get flags for curve element named InName
+	 * Note that this performs a binary search per-call. Consider using a ForEach* call to get multiple element's flags.
+	 * @param	InName	the name of the curve element to get
+	 * @return  the flags of the curve element. If this curve does not contain an element with the supplied name, returns UE::Anim::ECurveElementFlags::None 
+	 */
+	UE::Anim::ECurveElementFlags GetFlags(FName InName) const
 	{
-		check(bInitialized);
-
-		int32 ArrayIndex = GetArrayIndexByUID(InUid);
-		if (ArrayIndex != INDEX_NONE)
+		if(const ElementType* CurveElement = Super::Find(InName))
 		{
-			return CurveWeights[ArrayIndex];
+			return CurveElement->Flags;
 		}
-
-		return 0.f;
+		return UE::Anim::ECurveElementFlags::None;
 	}
 
-	/** Get Value of InUID with validation and default value */
-	float Get(SkeletonAnimCurveUID InUid, bool& OutIsValid, float InDefaultValue=0.f) const
+	UE_DEPRECATED(5.3, "Please use Get with a curve name or consider using one of the bulk APIs (e.g. UE::Anim::FCurveUtils::BulkGet).")
+	float Get(SmartName::UID_Type InUid, bool& OutIsValid, float InDefaultValue=0.f) const
 	{
-		check(bInitialized);
+		return 0.0f;
+	}
 
-		const int32 ArrayIndex = GetArrayIndexByUID(InUid);
-		if ((ArrayIndex != INDEX_NONE) && ValidCurveWeights[ArrayIndex])
+	/**
+	 * Get Value of curve named InName with validation and default value
+	 * Note that this performs a binary search per-call. Consider using a ForEach* or one of the bulk APIs
+	 * (e.g. UE::Anim::FCurveUtils::BulkGet) call to get multiple element's flags.
+	 * @param	InName	the name of the curve element to get
+	 * @param	HasElement	whether this curve contains the supplied named element 
+	 * @param	InDefaultValue	the default value to use in case the curve does not contain the supplied element
+	 * @return  the value of the curve element. If this curve does not contain an element with the supplied name, returns InDefaultValue
+	 */
+	float Get(FName InName, bool& OutHasElement, float InDefaultValue = 0.0f) const
+	{
+		if(const ElementType* CurveElement = Super::Find(InName))
 		{
-			OutIsValid = true;
-			return CurveWeights[ArrayIndex];
+			OutHasElement = true;
+			return CurveElement->Value;
 		}
-
-		OutIsValid = false;
+		
+		OutHasElement = false;
 		return InDefaultValue;
 	}
 
-	/** Get Array Index by UID */
-	int32 GetArrayIndexByUID(SkeletonAnimCurveUID InUid) const
+	/**
+	 * Mirror the values & flags of the two named curves if both exist.
+	 * If only InName0 exists, InName0 will be renamed to InName1.
+	 * If only InName1 exists, InName1 will be renamed to InName0.
+	 * @param	InName0	the name of first curve to mirror
+	 * @param	InName1	the name of second curve to mirror
+	 */
+	void Mirror(FName InName0, FName InName1)
 	{
-		int32 ArrayIndex = (*UIDToArrayIndexLUT).IsValidIndex(InUid) ? (*UIDToArrayIndexLUT)[InUid] : MAX_uint16;
-		if (ArrayIndex != MAX_uint16)
+		ElementType* CurveElement0 = Super::Find(InName0);
+		ElementType* CurveElement1 = Super::Find(InName1);
+
+		// Both exist, so swap values & flags
+		if(CurveElement0 && CurveElement1)
 		{
-			return ArrayIndex;
+			Swap(CurveElement0->Value, CurveElement1->Value);
+			Swap(CurveElement0->Flags, CurveElement1->Flags);
 		}
+		// First exists, but not second. Rename to second
+		else if(CurveElement0)
+		{
+			CurveElement0->Name = InName1;
+			Super::bSorted = false;
+		}
+		// Second exists, but not first. Rename to first
+		else if(CurveElement1)
+		{
+			CurveElement1->Name = InName0;
+			Super::bSorted = false;
+		}
+	}
+	
+	UE_DEPRECATED(5.3, "This function has been removed.")
+	int32 GetArrayIndexByUID(SmartName::UID_Type InUid) const
+	{
 		return INDEX_NONE;
 	}
 
-	/** return true if enabled. return false otherwise. */
-	bool IsEnabled(SkeletonAnimCurveUID InUid) const
+	UE_DEPRECATED(5.3, "Element validity/enabled state is now handled via UE::Anim::FCurveFilter.")
+	bool IsEnabled(SmartName::UID_Type InUid) const
 	{
-		check(bInitialized);
-
-		return (GetArrayIndexByUID(InUid) != INDEX_NONE);
+		return false;
 	}
 	
-	/** Get Valid Element Count from given UIDToArrayIndexLUT */
+	UE_DEPRECATED(5.3, "Element validity/enabled state is now handled via UE::Anim::FCurveFilter.")
 	static int32 GetValidElementCount(TArray<uint16> const* InUIDToArrayIndexLUT) 
 	{
-		int32 Count = 0;
-		if (InUIDToArrayIndexLUT)
-		{
-			const int32 ArraySize = InUIDToArrayIndexLUT->Num();
-			for (int32 Index = 0; Index < ArraySize; ++Index)
-			{
-				if ((*InUIDToArrayIndexLUT)[Index] != MAX_uint16)
-				{
-					++Count;
-				}
-			}
-		}
+		return 0;
+	}
 
-		return Count;
-	}	
+public:
 	/**
-	 * Blend (A, B) using Alpha, same as Lerp
+	 * Blend (A, B) using Alpha
 	 */
-	//@Todo curve flags won't transfer over - it only overwrites
-	void Lerp(const FBaseBlendedCurve& A, const FBaseBlendedCurve& B, float Alpha)
+	template<typename AllocatorA, typename ElementTypeA, typename AllocatorB, typename ElementTypeB>
+	void Lerp(const TBaseBlendedCurve<AllocatorA, ElementTypeA>& A, const TBaseBlendedCurve<AllocatorB, ElementTypeB>& B, float Alpha)
 	{
-		check(A.Num() == B.Num());
+		CURVE_PROFILE_CYCLE_COUNTER(TBaseBlendedCurve_Lerp);
+		
 		if (!FAnimWeight::IsRelevant(FMath::Abs(Alpha)))
 		{
 			// if blend is all the way for child1, then just copy its bone atoms
@@ -533,24 +731,23 @@ struct FBaseBlendedCurve
 		}
 		else
 		{
-			InitFrom(A);
-
-			// Only consider curve elements where either or both elements are valid.
-			for (TConstDualEitherSetBitIterator<InAllocator, InAllocator> It(A.ValidCurveWeights, B.ValidCurveWeights); It; ++It)
+			// Combine using Lerp. Result is a merged set of curves in 'this'. 
+			UE::Anim::FNamedValueArrayUtils::Union(*this, A, B, [&Alpha](ElementType& OutResult, const ElementTypeA& InElement0, const ElementTypeB& InElement1, UE::Anim::ENamedValueUnionFlags InFlags)
 			{
-				int32 Idx = It.GetIndex();
-				ValidCurveWeights[Idx] = true;
-				CurveWeights[Idx] = FMath::Lerp(A.CurveWeights[Idx], B.CurveWeights[Idx], Alpha);
-			}
+				OutResult.Value = FMath::Lerp(InElement0.Value, InElement1.Value, Alpha);
+				OutResult.Flags = InElement0.Flags | InElement1.Flags;
+			});
 		}
 	}
 
 	/**
-	 * Blend with Other using Alpha, same as Lerp 
+	 * Blend with Other using Alpha, same as Lerp
 	 */
-	void LerpTo(const FBaseBlendedCurve& Other, float Alpha)
+	template<typename OtherAllocator, typename OtherElementType>
+	void LerpTo(const TBaseBlendedCurve<OtherAllocator, OtherElementType>& Other, float Alpha)
 	{
-		check(Num() == Other.Num());
+		CURVE_PROFILE_CYCLE_COUNTER(TBaseBlendedCurve_LerpTo);
+		
 		if (!FAnimWeight::IsRelevant(FMath::Abs(Alpha)))
 		{
 			return;
@@ -562,257 +759,259 @@ struct FBaseBlendedCurve
 		}
 		else
 		{
-			// Only consider curve elements where either or both elements are valid.
-			for (TConstDualEitherSetBitIterator<InAllocator, InAllocator> It(ValidCurveWeights, Other.ValidCurveWeights); It; ++It)
+			// Combine using Lerp. Result is a merged set of curves in 'this'. 
+			UE::Anim::FNamedValueArrayUtils::Union(*this, Other, [&Alpha](ElementType& InOutThisElement, const OtherElementType& InOtherElement, UE::Anim::ENamedValueUnionFlags InFlags)
 			{
-				int32 Idx = It.GetIndex();
-				ValidCurveWeights[Idx] = true;
-				CurveWeights[Idx] = FMath::Lerp(CurveWeights[Idx], Other.CurveWeights[Idx], Alpha);
-			}
+				InOutThisElement.Value = FMath::Lerp(InOutThisElement.Value, InOtherElement.Value, Alpha);
+				InOutThisElement.Flags |= InOtherElement.Flags;
+			});
 		}
 	}
+
 	/**
-	 * Convert current curves to Additive (this - BaseCurve) if same found
+	 * Convert current curves to Additive (this - BaseCurve) if overlapping entries are found
 	 */
-	void ConvertToAdditive(const FBaseBlendedCurve& BaseCurve)
+	template<typename OtherAllocator, typename OtherElementType>
+	void ConvertToAdditive(const TBaseBlendedCurve<OtherAllocator, OtherElementType>& BaseCurve)
 	{
-		check(bInitialized);
-		check(Num() == BaseCurve.Num());
-
-		for (TConstDualEitherSetBitIterator<InAllocator, InAllocator> It(ValidCurveWeights, BaseCurve.ValidCurveWeights); It; ++It)
+		CURVE_PROFILE_CYCLE_COUNTER(TBaseBlendedCurve_ConvertToAdditive);
+		
+		UE::Anim::FNamedValueArrayUtils::Union(*this, BaseCurve, [](ElementType& InOutThisElement, const OtherElementType& InBaseCurveElement, UE::Anim::ENamedValueUnionFlags InFlags)
 		{
-			int32 Idx = It.GetIndex();
-			ValidCurveWeights[Idx] = true;
-			CurveWeights[Idx] -= BaseCurve.CurveWeights[Idx];
-		}
+			InOutThisElement.Value -= InBaseCurveElement.Value;
+			InOutThisElement.Flags |= InBaseCurveElement.Flags;
+		});
 	}
+
 	/**
 	 * Accumulate the input curve with input Weight
 	 */
-	void Accumulate(const FBaseBlendedCurve& AdditiveCurve, float Weight)
+	template<typename OtherAllocator, typename OtherElementType>
+	void Accumulate(const TBaseBlendedCurve<OtherAllocator, OtherElementType>& AdditiveCurve, float Weight)
 	{
-		check(bInitialized);
-		check(Num() == AdditiveCurve.Num());
-
+		CURVE_PROFILE_CYCLE_COUNTER(TBaseBlendedCurve_Accumulate);
+		
 		if (FAnimWeight::IsRelevant(Weight))
 		{
-			for (TConstDualEitherSetBitIterator<InAllocator, InAllocator> It(ValidCurveWeights, AdditiveCurve.ValidCurveWeights); It; ++It)
+			UE::Anim::FNamedValueArrayUtils::Union(*this, AdditiveCurve, [Weight](ElementType& InOutThisElement, const OtherElementType& InAdditiveCurveElement, UE::Anim::ENamedValueUnionFlags InFlags)
 			{
-				int32 Idx = It.GetIndex();
-				ValidCurveWeights[Idx] = true;
-				CurveWeights[Idx] += AdditiveCurve.CurveWeights[Idx] * Weight;
-			}
+				InOutThisElement.Value += InAdditiveCurveElement.Value * Weight;
+				InOutThisElement.Flags |= InAdditiveCurveElement.Flags;
+			});
 		}
 	}
 
 	/**
-	 * This doesn't blend but combine MAX(current weight, curvetocombine weight)
+	 * This doesn't blend but combines MAX(current value, CurveToCombine value)
 	 */
-	void UseMaxValue(const FBaseBlendedCurve& CurveToCombine)
+	template<typename OtherAllocator, typename OtherElementType>
+	void UseMaxValue(const TBaseBlendedCurve<OtherAllocator, OtherElementType>& CurveToCombine)
 	{
-		check(bInitialized);
-		check(Num() == CurveToCombine.Num());
-
-		for (TConstDualEitherSetBitIterator<InAllocator, InAllocator> It(ValidCurveWeights, CurveToCombine.ValidCurveWeights); It; ++It)
+		CURVE_PROFILE_CYCLE_COUNTER(TBaseBlendedCurve_UseMaxValue);
+		
+		UE::Anim::FNamedValueArrayUtils::Union(*this, CurveToCombine, [](ElementType& InOutThisElement, const OtherElementType& InCurveToCombineElement, UE::Anim::ENamedValueUnionFlags InFlags)
 		{
-			int32 Idx = It.GetIndex();
-			if (!ValidCurveWeights[Idx])
+			if(EnumHasAllFlags(InFlags, UE::Anim::ENamedValueUnionFlags::BothArgsValid))
 			{
-				CurveWeights[Idx] = CurveToCombine.CurveWeights[Idx];
-				ValidCurveWeights[Idx] = true;
+				InOutThisElement.Value = FMath::Max(InOutThisElement.Value, InCurveToCombineElement.Value);
+				InOutThisElement.Flags |= InCurveToCombineElement.Flags;
 			}
-			else if (CurveToCombine.ValidCurveWeights[Idx])
+			else if(EnumHasAllFlags(InFlags, UE::Anim::ENamedValueUnionFlags::ValidArg1))
 			{
-				CurveWeights[Idx] = FMath::Max(CurveWeights[Idx], CurveToCombine.CurveWeights[Idx]);
+				InOutThisElement.Value = InCurveToCombineElement.Value;
+				InOutThisElement.Flags = InCurveToCombineElement.Flags;
 			}
-		}
+		});
 	}
 
 	/**
-	 * This doesn't blend but combine MIN(current weight, curvetocombine weight)
+	 * This doesn't blend but combines MIN(current weight, CurveToCombine weight)
 	 */
-	void UseMinValue(const FBaseBlendedCurve& CurveToCombine)
+	template<typename OtherAllocator, typename OtherElementType>
+	void UseMinValue(const TBaseBlendedCurve<OtherAllocator, OtherElementType>& CurveToCombine)
 	{
-		check(bInitialized);
-		check(Num() == CurveToCombine.Num());
-
-		for (TConstDualEitherSetBitIterator<InAllocator, InAllocator> It(ValidCurveWeights, CurveToCombine.ValidCurveWeights); It; ++It)
+		CURVE_PROFILE_CYCLE_COUNTER(TBaseBlendedCurve_UseMinValue);
+		
+		UE::Anim::FNamedValueArrayUtils::Union(*this, CurveToCombine, [](ElementType& InOutThisElement, const OtherElementType& InCurveToCombineElement, UE::Anim::ENamedValueUnionFlags InFlags)
 		{
-			int32 Idx = It.GetIndex();
-			if (!ValidCurveWeights[Idx])
+			if(EnumHasAllFlags(InFlags, UE::Anim::ENamedValueUnionFlags::BothArgsValid))
 			{
-				CurveWeights[Idx] = CurveToCombine.CurveWeights[Idx];
-				ValidCurveWeights[Idx] = true;
+				InOutThisElement.Value = FMath::Min(InOutThisElement.Value, InCurveToCombineElement.Value);
+				InOutThisElement.Flags |= InCurveToCombineElement.Flags;
 			}
-			else if (CurveToCombine.ValidCurveWeights[Idx])
+			else if(EnumHasAllFlags(InFlags, UE::Anim::ENamedValueUnionFlags::ValidArg1))
 			{
-				CurveWeights[Idx] = FMath::Min(CurveWeights[Idx], CurveToCombine.CurveWeights[Idx]);
+				InOutThisElement.Value = InCurveToCombineElement.Value;
+				InOutThisElement.Flags = InCurveToCombineElement.Flags;
 			}
-		}
+		});
 	}
 
 	/**
-	 * This combines IF the input does not contain valid curve
+	 * If 'this' does not contain a valid element, then the value in 'this' is set, otherwise the value
+	 * is not modified.
 	 */
-	void CombinePreserved(const FBaseBlendedCurve& CurveToCombine)
+	template<typename OtherAllocator, typename OtherElementType>
+	void CombinePreserved(const TBaseBlendedCurve<OtherAllocator, OtherElementType>& CurveToCombine)
 	{
-		check(bInitialized);
-		check(Num() == CurveToCombine.Num());
-
-		for (TConstSetBitIterator<InAllocator> It(CurveToCombine.ValidCurveWeights); It; ++It)
+		CURVE_PROFILE_CYCLE_COUNTER(TBaseBlendedCurve_CombinePreserved);
+		
+		UE::Anim::FNamedValueArrayUtils::Union(*this, CurveToCombine, [](ElementType& InOutThisElement, const OtherElementType& InCurveToCombineElement, UE::Anim::ENamedValueUnionFlags InFlags)
 		{
-			int32 Idx = It.GetIndex();
-			if (!ValidCurveWeights[Idx])
+			if(!EnumHasAnyFlags(InFlags, UE::Anim::ENamedValueUnionFlags::ValidArg0))
 			{
-				CurveWeights[Idx] = CurveToCombine.CurveWeights[Idx];
-				ValidCurveWeights[Idx] = true;
+				InOutThisElement.Value = InCurveToCombineElement.Value;
+				InOutThisElement.Flags = InCurveToCombineElement.Flags;
 			}
-		}
+		});
 	}
+
 	/**
-	 * This doesn't blend but later pose with valid curve value overrides
+	 * If CurveToCombine contains a valid element, then the value in 'this' is overridden, otherwise the value is not
+	 * modified.
 	 */
-	void Combine(const FBaseBlendedCurve& CurveToCombine)
+	template<typename OtherAllocator, typename OtherElementType>
+	void Combine(const TBaseBlendedCurve<OtherAllocator, OtherElementType>& CurveToCombine)
 	{
-		check(bInitialized);
-		check(Num() == CurveToCombine.Num());
-
-		for (TConstSetBitIterator<InAllocator> It(CurveToCombine.ValidCurveWeights); It; ++It)
+		CURVE_PROFILE_CYCLE_COUNTER(TBaseBlendedCurve_Combine);
+		
+		UE::Anim::FNamedValueArrayUtils::Union(*this, CurveToCombine, [](ElementType& InOutThisElement, const OtherElementType& InCurveToCombineElement, UE::Anim::ENamedValueUnionFlags InFlags)
 		{
-			int32 Idx = It.GetIndex();
-			CurveWeights[Idx] = CurveToCombine.CurveWeights[Idx];
-			ValidCurveWeights[Idx] = true;
-		}
+			if(EnumHasAnyFlags(InFlags, UE::Anim::ENamedValueUnionFlags::ValidArg1))
+			{
+				InOutThisElement.Value = InCurveToCombineElement.Value;
+				InOutThisElement.Flags = InCurveToCombineElement.Flags;
+			}
+		});
 	}
 
 	/**
-	 * Override with inupt curve * weight
+	 * Override with input curve * weight
 	 */
-	void Override(const FBaseBlendedCurve& CurveToOverrideFrom, float Weight)
+	template<typename OtherAllocator, typename OtherElementType>
+	void Override(const TBaseBlendedCurve<OtherAllocator, OtherElementType>& CurveToOverrideFrom, float Weight)
 	{
+		CURVE_PROFILE_CYCLE_COUNTER(TBaseBlendedCurve_Override_Weighted);
+		
 		InitFrom(CurveToOverrideFrom);
 
-		CurveWeights = CurveToOverrideFrom.CurveWeights;
-		ValidCurveWeights = CurveToOverrideFrom.ValidCurveWeights;
-
-		if (!FMath::IsNearlyEqual(Weight, 1.f))
+		if (!FMath::IsNearlyEqual(Weight, 1.0f))
 		{
-			for (TConstSetBitIterator<InAllocator> It(ValidCurveWeights); It; ++It)
+			for (ElementType& Element : Super::Elements)
 			{
-				CurveWeights[It.GetIndex()] *= Weight;
+				Element.Value *= Weight;
 			}
 		}
 	}
 
+public:
 	/**
 	 * Override with input curve 
 	 */
-	void Override(const FBaseBlendedCurve& CurveToOverrideFrom)
+	template<typename OtherAllocator, typename OtherElementType>
+	void Override(const TBaseBlendedCurve<OtherAllocator, OtherElementType>& CurveToOverrideFrom)
 	{
-		// make sure this doesn't happen
-		if (ensure(&CurveToOverrideFrom != this))
-		{
-			check(CurveToOverrideFrom.UIDToArrayIndexLUT != nullptr);
-			UIDToArrayIndexLUT = CurveToOverrideFrom.UIDToArrayIndexLUT;
-			NumValidCurveCount = CurveToOverrideFrom.NumValidCurveCount;
-			CurveWeights = CurveToOverrideFrom.CurveWeights;
-			ValidCurveWeights = CurveToOverrideFrom.ValidCurveWeights;
-			bInitialized = true;
-		}
+		CURVE_PROFILE_CYCLE_COUNTER(TBaseBlendedCurve_Override);
+		
+		InitFrom(CurveToOverrideFrom);
 	}
-
+	
 	/**
 	 * Override with input curve, leaving input curve invalid
 	 */
-	void OverrideMove(FBaseBlendedCurve& CurveToOverrideFrom)
+	void OverrideMove(TBaseBlendedCurve& CurveToOverrideFrom)
 	{
 		// make sure this doesn't happen
 		if (ensure(&CurveToOverrideFrom != this))
 		{
-			check(CurveToOverrideFrom.UIDToArrayIndexLUT != nullptr);
-			UIDToArrayIndexLUT = CurveToOverrideFrom.UIDToArrayIndexLUT;
-			CurveToOverrideFrom.UIDToArrayIndexLUT = nullptr;
-			NumValidCurveCount = CurveToOverrideFrom.NumValidCurveCount;
-			CurveToOverrideFrom.NumValidCurveCount = 0;
-			CurveWeights = MoveTemp(CurveToOverrideFrom.CurveWeights);
-			ValidCurveWeights = MoveTemp(CurveToOverrideFrom.ValidCurveWeights);
-			bInitialized = true;
-			CurveToOverrideFrom.bInitialized = false;
+			Super::Elements = MoveTemp(CurveToOverrideFrom.Elements);
+			Super::bSorted = CurveToOverrideFrom.bSorted;
+			CurveToOverrideFrom.bSorted = false;
 		}
 	}
 
-	/** Return number of elements */
-	int32 Num() const { return CurveWeights.Num(); }
-
-	/** Return number of valid elements */
-	int32 NumValid() const { return ValidCurveWeights.CountSetBits(); }
-
-	/** CopyFrom as expected. */
-	template <typename OtherAllocator>
-	void CopyFrom(const FBaseBlendedCurve<OtherAllocator>& CurveToCopyFrom)
-	{
-		checkf(CurveToCopyFrom.IsValid(), TEXT("Copying data from an invalid curve UIDToArrayIndexLUT: 0x%x  (Sizes %i/%i)"), CurveToCopyFrom.UIDToArrayIndexLUT, (CurveToCopyFrom.UIDToArrayIndexLUT ? CurveToCopyFrom.UIDToArrayIndexLUT->Num() : -1), CurveToCopyFrom.CurveWeights.Num());
-		UIDToArrayIndexLUT = CurveToCopyFrom.UIDToArrayIndexLUT;
-		CurveWeights = CurveToCopyFrom.CurveWeights;
-		ValidCurveWeights = CurveToCopyFrom.ValidCurveWeights;
-		NumValidCurveCount = CurveToCopyFrom.NumValidCurveCount;
-		bInitialized = true;
-	}
+	UE_DEPRECATED(5.3, "Please use Num(). Element validity/enabled state is now handled via UE::Anim::FCurveFilter.")
+	int32 NumValid() const { return Super::Elements.Num(); }
 
-	void CopyFrom(const FBaseBlendedCurve<Allocator>& CurveToCopyFrom)
+	/** Copy elements between curves that have different allocators & element types. Does not copy filter. */
+	template <typename OtherAllocator, typename OtherElementType>
+	void CopyFrom(const TBaseBlendedCurve<OtherAllocator, OtherElementType>& InCurveToCopyFrom)
 	{
-		if (&CurveToCopyFrom != this)
+		CURVE_PROFILE_CYCLE_COUNTER(TBaseBlendedCurve_CopyFrom);
+		
+		if constexpr(std::is_same<ElementType, OtherElementType>::value)
+		{
+			Super::Elements = InCurveToCopyFrom.Elements;
+			Super::bSorted = InCurveToCopyFrom.bSorted;
+		}
+		else
 		{
-			checkf(CurveToCopyFrom.IsValid(), TEXT("Copying data from an invalid curve UIDToArrayIndexLUT: 0x%x  (Sizes %i/%i)"), CurveToCopyFrom.UIDToArrayIndexLUT, (CurveToCopyFrom.UIDToArrayIndexLUT ? CurveToCopyFrom.UIDToArrayIndexLUT->Num() : -1), CurveToCopyFrom.CurveWeights.Num());
-			UIDToArrayIndexLUT = CurveToCopyFrom.UIDToArrayIndexLUT;
-			CurveWeights = CurveToCopyFrom.CurveWeights;
-			ValidCurveWeights = CurveToCopyFrom.ValidCurveWeights;
-			NumValidCurveCount = CurveToCopyFrom.NumValidCurveCount;
-			bInitialized = true;
+			Super::Elements.Reset();
+
+			UE::Anim::FNamedValueArrayUtils::Union(*this, InCurveToCopyFrom,
+				[](ElementType& InOutResultElement, const OtherElementType& InParamElement, UE::Anim::ENamedValueUnionFlags InFlags)
+				{
+					InOutResultElement.Value = InParamElement.Value;
+					InOutResultElement.Flags = InParamElement.Flags;
+				});
+
+			Super::bSorted = InCurveToCopyFrom.bSorted;
 		}
 	}
 
 	/** Once moved, source is invalid */
-	void MoveFrom(FBaseBlendedCurve<Allocator>& CurveToMoveFrom)
+	void MoveFrom(TBaseBlendedCurve& CurveToMoveFrom)
 	{
-		UIDToArrayIndexLUT = CurveToMoveFrom.UIDToArrayIndexLUT;
-		CurveToMoveFrom.UIDToArrayIndexLUT = nullptr;
-		NumValidCurveCount = CurveToMoveFrom.NumValidCurveCount;
-		CurveToMoveFrom.NumValidCurveCount = 0;
-		CurveWeights = MoveTemp(CurveToMoveFrom.CurveWeights);
-		ValidCurveWeights = MoveTemp(CurveToMoveFrom.ValidCurveWeights);
-		bInitialized = true;
-		CurveToMoveFrom.bInitialized = false;
+		Super::Elements = MoveTemp(CurveToMoveFrom.Elements);
+		Super::bSorted = CurveToMoveFrom.bSorted;
+		CurveToMoveFrom.bSorted = false;
 	}
 
-	/** Empty */
-	void Empty()
+	/** Reserves memory for InNumElements */
+	void Reserve(int32 InNumElements)
 	{
-		// Set to nullptr as we only received a ptr reference from USkeleton
-		UIDToArrayIndexLUT = nullptr;
-		CurveWeights.Reset();
-		ValidCurveWeights.Reset();
-		NumValidCurveCount = 0;
-		bInitialized = false;
+		CURVE_PROFILE_CYCLE_COUNTER(TBaseBlendedCurve_Reserve);
+		Super::Elements.Reserve(InNumElements);
+	}
+	
+	UE_DEPRECATED(5.3, "Curves are now always valid.")
+	bool IsValid() const
+	{
+		return true;
 	}
 
-	/**  Whether initialized or not */
-	bool bInitialized;
+	/** Set the filter referenced by this curve. */
+	void SetFilter(const UE::Anim::FCurveFilter* InFilter)
+	{
+		Filter = InFilter;
+	}
 
-	// Only checks bare minimal validity. (namely that we have a UID list and that it 
-	// is the same size as our element list
-	bool IsValid() const
+	/** Get the filter referenced by this curve. */
+	const UE::Anim::FCurveFilter* GetFilter() const
 	{
-		return UIDToArrayIndexLUT != nullptr && (CurveWeights.Num() == NumValidCurveCount);
+		return Filter;
 	}
 };
 
-struct ENGINE_API FBlendedCurve : public FBaseBlendedCurve<FAnimStackAllocator>
+template<typename AllocatorType>
+struct UE_DEPRECATED(5.3, "FBaseBlendedCurve was renamed to TBaseBlendedCurve.") FBaseBlendedCurve : public TBaseBlendedCurve<AllocatorType>
 {
 };
 
-struct ENGINE_API FBlendedHeapCurve : public FBaseBlendedCurve<FDefaultAllocator>
+struct ENGINE_API FBlendedCurve : public TBaseBlendedCurve<FAnimStackAllocator>
 {
+	using TBaseBlendedCurve<FAnimStackAllocator>::InitFrom; 
+	
+	/** Initialize from bone container's filtered curves */
+	void InitFrom(const FBoneContainer& InBoneContainer);
+};
+
+struct ENGINE_API FBlendedHeapCurve : public TBaseBlendedCurve<FDefaultAllocator>
+{
+	using TBaseBlendedCurve<FDefaultAllocator>::InitFrom;
+	
+	/** Initialize from bone container's filtered curves */
+	void InitFrom(const FBoneContainer& InBoneContainer);
 };
 
 UENUM(BlueprintType)
@@ -866,30 +1065,46 @@ struct FRawCurveTracks
 	/**
 	* Add new float curve from the given UID if not existing and add the key with time/value
 	*/
-	ENGINE_API void AddFloatCurveKey(const FSmartName& NewCurve, int32 CurveFlags, float Time, float Value);
+	ENGINE_API void AddFloatCurveKey(const FName& NewCurve, int32 CurveFlags, float Time, float Value);
+	
+	UE_DEPRECATED(5.3, "Please use AddFloatCurveKey that takes a FName.")
+	void AddFloatCurveKey(const FSmartName& NewCurve, int32 CurveFlags, float Time, float Value) { AddFloatCurveKey(NewCurve.DisplayName, CurveFlags, Time, Value); }
+
+	
 	ENGINE_API void RemoveRedundantKeys(float Tolerance = UE_SMALL_NUMBER, FFrameRate SampleRate = FFrameRate(0,0));
 
 #endif // WITH_EDITOR
 	/**
 	 * Find curve data based on the curve UID
 	 */
-	ENGINE_API FAnimCurveBase * GetCurveData(SkeletonAnimCurveUID Uid, ERawCurveTrackTypes SupportedCurveType = ERawCurveTrackTypes::RCT_Float);
+	UE_DEPRECATED(5.3, "Please use GetCurveData that takes an FName.")
+	FAnimCurveBase * GetCurveData(SmartName::UID_Type Uid, ERawCurveTrackTypes SupportedCurveType = ERawCurveTrackTypes::RCT_Float) { return nullptr; }
+	
+	UE_DEPRECATED(5.3, "Please use GetCurveData that takes an FName.")
+	const FAnimCurveBase * GetCurveData(SmartName::UID_Type Uid, ERawCurveTrackTypes SupportedCurveType = ERawCurveTrackTypes::RCT_Float) const { return nullptr; }
 
 	/**
-	* Find curve data based on the curve UID
+	* Find curve data based on the curve name
 	*/
-	ENGINE_API const FAnimCurveBase * GetCurveData(SkeletonAnimCurveUID Uid, ERawCurveTrackTypes SupportedCurveType = ERawCurveTrackTypes::RCT_Float) const;
-
+	ENGINE_API FAnimCurveBase * GetCurveData(FName Name, ERawCurveTrackTypes SupportedCurveType = ERawCurveTrackTypes::RCT_Float);
+	ENGINE_API const FAnimCurveBase * GetCurveData(FName Name, ERawCurveTrackTypes SupportedCurveType = ERawCurveTrackTypes::RCT_Float) const;
+	
 	/**
 	 * Add new curve from the provided UID and return true if success
 	 * bVectorInterpCurve == true, then it will create FVectorCuve, otherwise, FFloatCurve
 	 */
-	ENGINE_API bool AddCurveData(const FSmartName& NewCurve, int32 CurveFlags = AACF_DefaultCurve, ERawCurveTrackTypes SupportedCurveType = ERawCurveTrackTypes::RCT_Float);
+	ENGINE_API bool AddCurveData(const FName& NewCurve, int32 CurveFlags = AACF_DefaultCurve, ERawCurveTrackTypes SupportedCurveType = ERawCurveTrackTypes::RCT_Float);
 
+	UE_DEPRECATED(5.3, "Please use AddCurveData that takes a FName.")
+	bool AddCurveData(const FSmartName& NewCurve, int32 CurveFlags = AACF_DefaultCurve, ERawCurveTrackTypes SupportedCurveType = ERawCurveTrackTypes::RCT_Float) { return AddCurveData(NewCurve.DisplayName, CurveFlags, SupportedCurveType); }
+	
 	/**
 	 * Delete curve data 
 	 */
-	ENGINE_API bool DeleteCurveData(const FSmartName& CurveToDelete, ERawCurveTrackTypes SupportedCurveType = ERawCurveTrackTypes::RCT_Float);
+	ENGINE_API bool DeleteCurveData(const FName& CurveToDelete, ERawCurveTrackTypes SupportedCurveType = ERawCurveTrackTypes::RCT_Float);
+	
+	UE_DEPRECATED(5.3, "Please use DeleteCurveData that takes a FName.")
+	bool DeleteCurveData(const FSmartName& CurveToDelete, ERawCurveTrackTypes SupportedCurveType = ERawCurveTrackTypes::RCT_Float) { return DeleteCurveData(CurveToDelete.DisplayName, SupportedCurveType); }
 
 	/**
 	 * Delete all curve data 
@@ -898,19 +1113,17 @@ struct FRawCurveTracks
 
 	/**
 	 * Duplicate curve data
-	 * 
-	 */
-	ENGINE_API bool DuplicateCurveData(const FSmartName& CurveToCopy, const FSmartName& NewCurve, ERawCurveTrackTypes SupportedCurveType = ERawCurveTrackTypes::RCT_Float);
-
-	/**
-	 * Updates the DisplayName field of the curves from the provided name container
 	 */
-	ENGINE_API void RefreshName(const FSmartNameMapping* NameMapping, ERawCurveTrackTypes SupportedCurveType = ERawCurveTrackTypes::RCT_Float);
+	ENGINE_API bool DuplicateCurveData(const FName& CurveToCopy, const FName& NewCurve, ERawCurveTrackTypes SupportedCurveType = ERawCurveTrackTypes::RCT_Float);
+	
+	UE_DEPRECATED(5.3, "Please use DuplicateCurveData that takes FNames as curve parameters.")
+	bool DuplicateCurveData(const FSmartName& CurveToCopy, const FSmartName& NewCurve, ERawCurveTrackTypes SupportedCurveType = ERawCurveTrackTypes::RCT_Float) { return DuplicateCurveData(CurveToCopy.DisplayName, NewCurve.DisplayName, SupportedCurveType); }
+	
+	UE_DEPRECATED(5.3, "This function is no longer used")
+	ENGINE_API void RefreshName(const FSmartNameMapping* NameMapping, ERawCurveTrackTypes SupportedCurveType = ERawCurveTrackTypes::RCT_Float) {}
 
-	/** 
-	 * Serialize
-	 */
-	void PostSerialize(FArchive& Ar);
+	// This allows loading data that was saved between VER_UE4_SKELETON_ADD_SMARTNAMES and FFrameworkObjectVersion::SmartNameRefactor
+	void PostSerializeFixup(FArchive& Ar);
 
 	/*
 	 * resize curve length. If longer, it doesn't do any. If shorter, remove previous keys and add new key to the end of the frame. 
@@ -934,40 +1147,34 @@ private:
 	 */
 
 	/**
-	 * Find curve data based on the curve UID
-	 */
+	* Find curve data based on the curve name
+	*/
 	template <typename DataType>
-	DataType * GetCurveDataImpl(TArray<DataType>& Curves, SkeletonAnimCurveUID Uid);
-
+	DataType * GetCurveDataImpl(TArray<DataType>& Curves, FName Name);
+	
 	/**
-	* Find curve data based on the curve UID
+	* Find curve data based on the curve name
 	*/
 	template <typename DataType>
-	const DataType * GetCurveDataImpl(const TArray<DataType>& Curves, SkeletonAnimCurveUID Uid) const;
-
+	const DataType * GetCurveDataImpl(const TArray<DataType>& Curves, FName Name) const;
+	
 	/**
 	 * Add new curve from the provided UID and return true if success
 	 * bVectorInterpCurve == true, then it will create FVectorCuve, otherwise, FFloatCurve
 	 */
 	template <typename DataType>
-	bool AddCurveDataImpl(TArray<DataType>& Curves, const FSmartName& NewCurve, int32 CurveFlags);
+	bool AddCurveDataImpl(TArray<DataType>& Curves, const FName& NewCurve, int32 CurveFlags);
 	/**
 	 * Delete curve data 
 	 */
 	template <typename DataType>
-	bool DeleteCurveDataImpl(TArray<DataType>& Curves, const FSmartName& CurveToDelete);
+	bool DeleteCurveDataImpl(TArray<DataType>& Curves, const FName& CurveToDelete);
 	/**
 	 * Duplicate curve data
 	 * 
 	 */
 	template <typename DataType>
-	bool DuplicateCurveDataImpl(TArray<DataType>& Curves, const FSmartName& CurveToCopy, const FSmartName& NewCurve);
-
-	/**
-	 * Updates the DisplayName field of the curves from the provided name container
-	 */
-	template <typename DataType>
-	void UpdateLastObservedNamesImpl(TArray<DataType>& Curves, const FSmartNameMapping* NameMapping);
+	bool DuplicateCurveDataImpl(TArray<DataType>& Curves, const FName& CurveToCopy, const FName& NewCurve);
 };
 
 FArchive& operator<<(FArchive& Ar, FRawCurveTracks& D);
diff --git a/Engine/Source/Runtime/Engine/Public/Animation/AnimInstanceProxy.h b/Engine/Source/Runtime/Engine/Public/Animation/AnimInstanceProxy.h
--- a/Engine/Source/Runtime/Engine/Public/Animation/AnimInstanceProxy.h
+++ b/Engine/Source/Runtime/Engine/Public/Animation/AnimInstanceProxy.h
@@ -197,9 +197,14 @@ public:
 		return ActualAnimClass ? Cast<UAnimBlueprint>(ActualAnimClass->ClassGeneratedBy) : nullptr;
 	}
 
-	// Record pose for node of ID LinkID if it is currently being watched
+	UE_DEPRECATED(5.3, "Please use RegisterWatchedPose that takes a pose and a curve")
 	void RegisterWatchedPose(const FCompactPose& Pose, int32 LinkID);
+	UE_DEPRECATED(5.3, "Please use RegisterWatchedPose that takes a pose and a curve")
 	void RegisterWatchedPose(const FCSPose<FCompactPose>& Pose, int32 LinkID);
+
+	// Record pose for node of ID LinkID if it is currently being watched
+	void RegisterWatchedPose(const FCompactPose& Pose, const FBlendedCurve& InCurve, int32 LinkID);
+	void RegisterWatchedPose(const FCSPose<FCompactPose>& Pose, const FBlendedCurve& InCurve, int32 LinkID);
 #endif
 
 	UE_DEPRECATED(5.0, "Function renamed to FlipBufferWriteIndex as this no longer deals with sync groups.")
@@ -748,10 +753,13 @@ protected:
 	/** 
 	 * Recalculate required curve list for animation - if you call RecalcRequiredBones, this should already happen
 	 */
-	void RecalcRequiredCurves(const FCurveEvaluationOption& CurveEvalOption);
+	void RecalcRequiredCurves(const UE::Anim::FCurveFilterSettings& CurveFilterSettings);
 
-	/** Update the material parameters of the supplied component from this instance */
-	void UpdateCurvesToComponents(USkeletalMeshComponent* Component);
+	UE_DEPRECATED(5.3, "Please use RecalcRequiredCurves that takes a FCurveEvaluationOption.")
+	void RecalcRequiredCurves(const FCurveEvaluationOption& CurveEvalOption);
+	
+	UE_DEPRECATED(5.3, "This function is no longer used.")
+	void UpdateCurvesToComponents(USkeletalMeshComponent* Component) {}
 
 	/** Get Currently active montage evaluation state.
 		Note that there might be multiple Active at the same time. This will only return the first active one it finds. **/
@@ -935,9 +943,13 @@ protected:
 	/** Check whether we have any active curves */
 	bool HasActiveCurves() const;
 
-	/** Add a curve value */
-	void AddCurveValue(const FSmartNameMapping& Mapping, const FName& CurveName, float Value);
+	
+	UE_DEPRECATED(5.3, "Please use AddCurveValue that does not take a FSmartNameMapping.")
+	void AddCurveValue(const FSmartNameMapping& Mapping, const FName& CurveName, float Value) { AddCurveValue(CurveName, Value); }
 
+	/** Add a curve value */
+	void AddCurveValue(const FName& CurveName, float Value, bool bMorphtarget = false, bool bMaterial = false);
+	
 	/** Custom proxy Init/Cache/Update/Evaluate functions */
 	static void InitializeInputProxy(FAnimInstanceProxy* InputProxy, UAnimInstance* InAnimInstance);
 	static void GatherInputProxyDebugData(FAnimInstanceProxy* InputProxy, FNodeDebugData& DebugData);
diff --git a/Engine/Source/Runtime/Engine/Public/Animation/AnimTrace.h b/Engine/Source/Runtime/Engine/Public/Animation/AnimTrace.h
--- a/Engine/Source/Runtime/Engine/Public/Animation/AnimTrace.h
+++ b/Engine/Source/Runtime/Engine/Public/Animation/AnimTrace.h
@@ -34,6 +34,7 @@ struct FAnimNotifyEvent;
 struct FPassedMarker;
 struct FAnimSyncMarker;
 struct FAnimMontageInstance;
+class UPoseWatchPoseElement;
 
 struct FAnimTrace
 {
@@ -148,7 +149,7 @@ struct FAnimTrace
 			OutputAnimNodeAttribute(*InContext.AnimInstanceProxy, *InContext.AnimInstanceProxy, InTargetNodeId, InSourceNodeId, UE::Anim::FAttributes::Attributes);
 		}
 
-		if(InContext.Curve.NumValid() > 0)
+		if(InContext.Curve.Num() > 0)
 		{
 			OutputAnimNodeAttribute(*InContext.AnimInstanceProxy, *InContext.AnimInstanceProxy, InTargetNodeId, InSourceNodeId, UE::Anim::FAttributes::Curves);
 		}
@@ -191,7 +192,7 @@ struct FAnimTrace
 	ENGINE_API static void OutputSync(const FAnimInstanceProxy& InSourceProxy, int32 InSourceNodeId, FName InGroupName);
 
 	/** Helper function to output a pose watch record */
-	ENGINE_API static void OutputPoseWatch(const FAnimInstanceProxy& InSourceProxy, int32 InPoseWatchId, const TArray<FTransform>& BoneTransforms, const TArray<FBoneIndexType>& RequiredBones, const FTransform& WorldTransform, const bool bIsEnabled);
+	ENGINE_API static void OutputPoseWatch(const FAnimInstanceProxy& InSourceProxy, UPoseWatchPoseElement* InPoseWatchElement, int32 InPoseWatchId, const TArray<FTransform>& BoneTransforms, const FBlendedHeapCurve& InCurves, const TArray<FBoneIndexType>& RequiredBones, const FTransform& WorldTransform, const bool bIsEnabled);
 };
 
 #define TRACE_ANIM_TICK_RECORD(Context, TickRecord) \
@@ -245,8 +246,8 @@ struct FAnimTrace
 #define TRACE_ANIM_NODE_SYNC(SourceProxy, SourceNodeId, GroupName) \
 	FAnimTrace::OutputSync(SourceProxy, SourceNodeId, GroupName);
 
-#define TRACE_ANIM_POSE_WATCH(SourceProxy, PoseWatchId, BoneTransforms, RequiredBones, WorldTransform, bIsEnabled) \
-	FAnimTrace::OutputPoseWatch(SourceProxy, PoseWatchId, BoneTransforms, RequiredBones, WorldTransform, bIsEnabled);
+#define TRACE_ANIM_POSE_WATCH(SourceProxy, PoseWatchElement, PoseWatchId, BoneTransforms, Curves, RequiredBones, WorldTransform, bIsEnabled) \
+	FAnimTrace::OutputPoseWatch(SourceProxy, PoseWatchElement, PoseWatchId, BoneTransforms, Curves, RequiredBones, WorldTransform, bIsEnabled);
 
 #else
 
@@ -267,6 +268,6 @@ struct FAnimTrace
 #define TRACE_ANIM_SYNC_MARKER(AnimInstance, SyncMarker)
 #define TRACE_ANIM_MONTAGE(AnimInstance, MontageInstance)
 #define TRACE_ANIM_NODE_SYNC(SourceProxy, SourceNodeId, GroupName)
-#define TRACE_ANIM_POSE_WATCH(SourceProxy, PoseWatchId, BoneTransforms, RequiredBones, WorldTransform, bIsEnabled)
+#define TRACE_ANIM_POSE_WATCH(SourceProxy, PoseWatchElement, PoseWatchId, BoneTransforms, Curves, RequiredBones, WorldTransform, bIsEnabled)
 
 #endif
diff --git a/Engine/Source/Runtime/Engine/Public/Animation/AnimTypes.h b/Engine/Source/Runtime/Engine/Public/Animation/AnimTypes.h
--- a/Engine/Source/Runtime/Engine/Public/Animation/AnimTypes.h
+++ b/Engine/Source/Runtime/Engine/Public/Animation/AnimTypes.h
@@ -426,7 +426,7 @@ public:
 	/** Returns true is this AnimNotify is a BranchingPoint */
 	ENGINE_API bool IsBranchingPoint() const;
 
-	/** Returns true if this is blueprint derived notifies **/
+	/** Returns true if this is blueprint or native class-based notify (i.e. not a named notify) */
 	bool IsBlueprintNotify() const
 	{
 		return Notify != nullptr || NotifyStateClass != nullptr;
diff --git a/Engine/Source/Runtime/Engine/Public/Animation/NamedValueArray.h b/Engine/Source/Runtime/Engine/Public/Animation/NamedValueArray.h
new file mode 100644
--- /dev/null
+++ b/Engine/Source/Runtime/Engine/Public/Animation/NamedValueArray.h
@@ -0,0 +1,660 @@
+﻿// Copyright Epic Games, Inc. All Rights Reserved.
+
+#pragma once
+
+#include "CoreMinimal.h"
+#include "Algo/IsSorted.h"
+
+#define ENABLE_ANIM_CURVE_PROFILING 0
+
+#if ENABLE_ANIM_CURVE_PROFILING
+#include "Stats/Stats.h"
+#endif
+
+#if ENABLE_ANIM_CURVE_PROFILING
+#define CURVE_PROFILE_CYCLE_COUNTER(Stat) QUICK_SCOPE_CYCLE_COUNTER(Stat)
+#else
+#define CURVE_PROFILE_CYCLE_COUNTER(Stat)
+#endif
+
+#define DO_ANIM_NAMED_VALUE_SORTING_CHECKS		(DO_GUARD_SLOW || WITH_DEV_AUTOMATION_TESTS)
+#define DO_ANIM_NAMED_VALUE_DUPLICATE_CHECKS	(DO_GUARD_SLOW || WITH_DEV_AUTOMATION_TESTS)
+
+namespace UE::Anim
+{
+
+struct FNamedValueArrayUtils;
+
+/**
+ * Container of lazily-sorted name/value pairs.
+ * Used to perform efficient merge operations.
+ * Assumes that InElementType has an accessible member: FName Name. 
+ */
+template<typename InAllocatorType, typename InElementType>
+struct TNamedValueArray
+{
+	typedef InAllocatorType AllocatorType;
+	typedef InElementType ElementType;
+
+	friend struct FNamedValueArrayUtils;
+
+	/**
+	 * Add a named element.
+	 * Note that this should only really be used when building a fresh value array, as using this during runtime can
+	 * introduce duplicate values.
+	 * Asserts in debug builds if duplicate values are present
+	 */
+	template<typename... ArgTypes>
+	void Add(ArgTypes&&... Args)
+	{
+		Elements.Emplace(Forward<ArgTypes>(Args)...);
+		bSorted = false;
+
+		CheckDuplicates();
+	}
+
+	/**
+	 * Add an array of named elements.
+	 * Note that this should only really be used when building a fresh array, as using this during runtime can
+	 * introduce duplicate values.
+	 * Asserts in debug builds if duplicate values are present
+	 */
+	void AppendNames(TConstArrayView<FName> InNameArray)
+	{
+		Elements.Reserve(Elements.Num() + InNameArray.Num());
+		for(const FName& Name : InNameArray)
+		{
+			Elements.Emplace(Name);
+		}
+		bSorted = false;
+
+		CheckDuplicates();
+	}
+
+	/**
+	 * Add an array of named elements.
+	 * Note that this should only really be used when building a fresh array, as using this during runtime can
+	 * introduce duplicate values.
+	 * Asserts in debug builds if duplicate values are present
+	 */
+	void AppendNames(std::initializer_list<const FName> InInputArgs)
+	{
+		Elements.Reserve(Elements.Num() + InInputArgs.size());
+		for(const FName& Name : InInputArgs)
+		{
+			Elements.Emplace(Name);
+		}
+		bSorted = false;
+
+		CheckDuplicates();
+	}
+
+	/** Reset the internal allocations */
+	void Empty()
+	{
+		Elements.Reset();
+		bSorted = false;
+	}
+
+	/** Reserves memory for InNumElements */
+	void Reserve(int32 InNumElements)
+	{
+		Elements.Reserve(InNumElements);
+	}
+
+	/**
+	 * Check whether an element is present for the supplied name
+	 * Note that this performs a binary search per-call.
+	 * @param	InName	the name of the element to check
+	 * @return  true if an element with the supplied name is present
+	 */
+	bool HasElement(FName InName) const
+	{
+		return Find(InName) != nullptr;
+	}
+
+	/**
+	 * Iterate over each element calling InPredicate for each.
+	 * Predicate: (const ElementType&) -> void
+	 */
+	template<typename PredicateType>
+	void ForEachElement(const PredicateType& InPredicate) const
+	{
+		for(const ElementType& Element : Elements)
+		{
+			InPredicate(Element);
+		}
+	}
+	
+	/** @returns the number of elements */
+	int32 Num() const
+	{
+		return Elements.Num();
+	}
+	
+protected:
+	// Sort by FName - Note: this is not stable across serialization
+	struct FElementSortPredicate
+	{
+		FORCEINLINE bool operator()(const ElementType& InElement0, const ElementType& InElement1) const
+		{
+			return InElement0.Name.FastLess(InElement1.Name);
+		}
+	};
+
+	// Sorts the elements if they are not yet sorted
+	void SortElementsIfRequired() const
+	{
+		if(!bSorted)
+		{
+			CURVE_PROFILE_CYCLE_COUNTER(SortElementsIfRequired);
+
+			Algo::Sort(Elements, FElementSortPredicate());
+			bSorted = true;
+		}
+	}
+
+	// Checks whether the sorting invariant is correct
+	void CheckSorted() const
+	{
+#if DO_ANIM_NAMED_VALUE_SORTING_CHECKS
+		if(bSorted)
+		{
+			check(Algo::IsSorted(Elements, FElementSortPredicate()));
+		}
+#endif
+	}
+
+	// Checks whether the 'no duplicates' invariant is correct
+	void CheckDuplicates() const
+	{
+#if DO_ANIM_NAMED_VALUE_DUPLICATE_CHECKS
+		for(int32 Index0 = 0; Index0 < Elements.Num(); ++Index0)
+		{
+			for(int32 Index1 = 0; Index1 < Elements.Num(); ++Index1)
+			{
+				if(Index0 != Index1 && Elements[Index0].Name == Elements[Index1].Name)
+				{
+					checkf(false, TEXT("Duplicate curve entry found: %s"), *Elements[Index0].Name.ToString());
+				}
+			}
+		}
+#endif
+	}
+
+	/** Finds index of the element with the specified name, disregarding enabled state */
+	int32 IndexOf(FName InName) const
+	{
+		SortElementsIfRequired();
+
+		return Algo::BinarySearchBy(Elements, ElementType(InName), &ElementType::Name, FElementSortPredicate());
+	}
+
+	/** Finds the element with the specified name (const) */
+	const ElementType* Find(FName InName) const
+	{
+		const int32 ElementIndex = IndexOf(InName);
+		if(ElementIndex != INDEX_NONE)
+		{
+			return &Elements[ElementIndex]; 
+		}
+		return nullptr;
+	}
+
+	/** Finds the element with the specified name */
+	ElementType* Find(FName InName)
+	{
+		const int32 ElementIndex = IndexOf(InName);
+		if(ElementIndex != INDEX_NONE)
+		{
+			return &Elements[ElementIndex]; 
+		}
+		return nullptr;
+	}
+
+protected:
+	// Named elements, sorted by name
+	mutable TArray<ElementType, AllocatorType> Elements;
+
+	// Whether the elements are sorted
+	mutable bool bSorted = false;
+};
+
+// Flags passed during union operations
+enum class ENamedValueUnionFlags : uint8
+{
+	// No flags
+	None		= 0,
+
+	// First argument is valid
+	ValidArg0	= 0x01,
+
+	// Second argument is valid
+	ValidArg1	= 0x02,
+
+	// Both arguments are valid
+	BothArgsValid = ValidArg0 | ValidArg1,
+};
+
+struct FNamedValueArrayUtils
+{
+	// Helper function
+	// Uses a simple 'tape merge'
+	// Performs an operation per-element on the two value arrays. Writes result to InOutValueArray0.
+	// InValueArray0 will be the union of the two value arrays after the operation is completed (i.e.
+	// new elements in InValueArray1 are added to InValueArray0)
+	// InPredicate is called on all elements that are added to or already existing in InOutValueArray0, with
+	// appropriate flags.
+	template<typename PredicateType, typename AllocatorTypeResult, typename ElementTypeResult, typename AllocatorTypeParam, typename ElementTypeParam>
+	static void Union(TNamedValueArray<AllocatorTypeResult, ElementTypeResult>& InOutValueArray0, const TNamedValueArray<AllocatorTypeParam, ElementTypeParam>& InValueArray1, const PredicateType& InPredicate)
+	{
+		int32 NumElements0 = InOutValueArray0.Num();	// ValueArray1 elements remain constant, but ValueArray0 can have entries added.
+		const int32 NumElements1 = InValueArray1.Num();
+
+		// Early out if we have no elements to union
+		if(NumElements1 == 0)
+		{
+			return;
+		}
+		
+		// Sort both input arrays if required
+		InOutValueArray0.SortElementsIfRequired();
+		InValueArray1.SortElementsIfRequired();
+
+		// Reserve memory for at least MAX(NumElements0, NumElements1)
+		InOutValueArray0.Reserve(FMath::Max(NumElements0, NumElements1));
+
+		int32 ElementIndex0 = 0;
+		int32 ElementIndex1 = 0;
+
+		// Perform dual-iteration on the two sorted arrays
+		while(true)
+		{
+			if(ElementIndex0 == NumElements0 && ElementIndex1 < NumElements1)
+			{
+				// Reached end of ValueArray0 with remaining in ValueArray1, we can just copy the remainder of ValueArray1
+				InOutValueArray0.Elements.Reserve(InOutValueArray0.Elements.Num() + (NumElements1 - ElementIndex1));
+				for( ; ElementIndex1 < NumElements1; ++ElementIndex1)
+				{
+					const ElementTypeParam& Element1 = InValueArray1.Elements[ElementIndex1];
+					ElementTypeResult& Element0 = InOutValueArray0.Elements.AddDefaulted_GetRef();
+					Element0.Name = Element1.Name;
+
+					InPredicate(Element0, Element1, UE::Anim::ENamedValueUnionFlags::ValidArg1);
+				}
+				break;
+			}
+			else if(ElementIndex1 == NumElements1 && ElementIndex0 < NumElements0)
+			{
+				// Reached end of ValueArray1 with remaining in ValueArray0, we can just early out
+				break;
+			}
+			else if(ElementIndex0 == NumElements0 && ElementIndex1 == NumElements1)
+			{
+				// Reached end of both, exit
+				break;
+			}
+			
+			ElementTypeResult& Element0 = InOutValueArray0.Elements[ElementIndex0];
+			const ElementTypeParam& Element1 = InValueArray1.Elements[ElementIndex1];
+			
+			if(Element0.Name == Element1.Name)
+			{
+				// Elements match, run predicate and increment both indices
+				InPredicate(Element0, Element1, UE::Anim::ENamedValueUnionFlags::BothArgsValid);
+				++ElementIndex0;
+				++ElementIndex1;
+			}
+			else if(Element0.Name.FastLess(Element1.Name))
+			{
+				// ValueArray0 element is earlier, so run predicate with only ValueArray0 and increment ValueArray0
+				ElementTypeParam DefaultElement;
+				DefaultElement.Name = Element0.Name;
+				InPredicate(Element0, DefaultElement, UE::Anim::ENamedValueUnionFlags::ValidArg0);
+				++ElementIndex0;
+			}
+			else
+			{
+				// ValueArray1 element is earlier, so add to ValueArray0, run predicate with only second and increment ValueArray1
+				ElementTypeResult& NewElement = InOutValueArray0.Elements.InsertDefaulted_GetRef(ElementIndex0++); // increment beyond new element
+				NumElements0 = InOutValueArray0.Num();
+				NewElement.Name = Element1.Name;
+				InPredicate(NewElement, Element1, UE::Anim::ENamedValueUnionFlags::ValidArg1);
+				++ElementIndex1;
+			}
+		}
+
+		InOutValueArray0.CheckSorted();
+	}
+
+	// Helper function
+	// Uses a simple 'tape merge'
+	// Performs an operation per-element on the two value arrays. Writes result to InOutValueArray0.
+	// InValueArray0 will be the union of the two value arrays after the operation is completed (i.e.
+	// new elements in InValueArray1 are added to InValueArray0)
+	// Performs a simple copy for each element
+	template<typename AllocatorTypeResult, typename ElementType, typename AllocatorTypeParam>
+	static void Union(TNamedValueArray<AllocatorTypeResult, ElementType>& InOutValueArray0, const TNamedValueArray<AllocatorTypeParam, ElementType>& InValueArray1)
+	{
+		// Early out if we just want to perform a simple copy
+		if(InOutValueArray0.Num() == 0 && InValueArray1.Num() > 0)
+		{
+			InOutValueArray0.Elements = InValueArray1.Elements;
+			InOutValueArray0.bSorted = InValueArray1.bSorted;
+			return;
+		}
+
+		Union(InOutValueArray0, InValueArray1, [](ElementType& Element0, const ElementType& Element1, UE::Anim::ENamedValueUnionFlags InFlags)
+		{
+			if(EnumHasAnyFlags(InFlags, UE::Anim::ENamedValueUnionFlags::ValidArg1))
+			{
+				Element0 = Element1;
+			}
+		});
+	}
+	
+	// Helper function
+	// Uses a simple 'tape merge'
+	// Performs an operation per-element on the two value arrays. Writes result to OutResultValueArray.
+	// OutResultValueArray will be the union of the two value arrays after the operation is completed.
+	// InPredicate is called on all elements that are added to OutResultValueArray, with appropriate flags.
+	template<typename PredicateType, typename AllocatorTypeResult, typename ElementTypeResult, typename AllocatorType0, typename ElementType0, typename AllocatorType1, typename ElementType1>
+	static void Union(
+		TNamedValueArray<AllocatorTypeResult, ElementTypeResult>& OutResultValueArray,
+		const TNamedValueArray<AllocatorType0, ElementType0>& InValueArray0,
+		const TNamedValueArray<AllocatorType1, ElementType1>& InValueArray1,
+		const PredicateType& InPredicate)
+	{
+		// Make sure result is clear
+		OutResultValueArray.Elements.Reset();
+
+		const int32 NumElements0 = InValueArray0.Num();
+		const int32 NumElements1 = InValueArray1.Num();
+
+		// Sort both input arrays if required
+		InValueArray0.SortElementsIfRequired();
+		InValueArray1.SortElementsIfRequired();
+
+		// Reserve memory for at least MAX(NumElements0, NumElements1)
+		OutResultValueArray.Reserve(FMath::Max(NumElements0, NumElements1));
+		
+		int32 ElementIndex0 = 0;
+		int32 ElementIndex1 = 0;
+
+		// Perform dual-iteration on the two sorted arrays
+		while(true)
+		{
+			if(ElementIndex0 == NumElements0 && ElementIndex1 < NumElements1)
+			{
+				// Reached end of ValueArray0 with remaining in ValueArray1, we can just copy the remainder of ValueArray1
+				const int32 NumResults = OutResultValueArray.Elements.Num();
+				OutResultValueArray.Elements.Reserve(OutResultValueArray.Elements.Num() + (NumElements1 - ElementIndex1));
+				for(int32 ResultIndex = NumResults; ElementIndex1 < NumElements1; ++ResultIndex, ++ElementIndex1)
+				{
+					const ElementType1& Element1 = InValueArray1.Elements[ElementIndex1];
+					ElementType0 DefaultElement0;
+					DefaultElement0.Name = Element1.Name;
+					ElementTypeResult& ResultElement = OutResultValueArray.Elements.AddDefaulted_GetRef();
+					ResultElement.Name = Element1.Name;
+
+					InPredicate(ResultElement, DefaultElement0, Element1, UE::Anim::ENamedValueUnionFlags::ValidArg1);
+				}
+				break;
+			}
+			else if(ElementIndex1 == NumElements1 && ElementIndex0 < NumElements0)
+			{
+				// Reached end of ValueArray0 with remaining in ValueArray0, we can just copy the remainder of ValueArray0
+				const int32 NumResults = OutResultValueArray.Elements.Num();
+				OutResultValueArray.Elements.Reserve(OutResultValueArray.Elements.Num() + (NumElements0 - ElementIndex0));
+				for(int32 ResultIndex = NumResults; ElementIndex0 < NumElements0; ++ResultIndex, ++ElementIndex0)
+				{
+					const ElementType0& Element0 = InValueArray0.Elements[ElementIndex0];
+					ElementType1 DefaultElement1;
+					DefaultElement1.Name = Element0.Name;
+					ElementTypeResult& ResultElement = OutResultValueArray.Elements.AddDefaulted_GetRef();
+					ResultElement.Name = Element0.Name;
+					
+					InPredicate(OutResultValueArray.Elements[ResultIndex], Element0, DefaultElement1, UE::Anim::ENamedValueUnionFlags::ValidArg0);
+				}
+				break;
+			}
+			else if(ElementIndex0 == NumElements0 && ElementIndex1 == NumElements1)
+			{
+				// Reached end of both, exit
+				break;
+			}
+			
+			const ElementType0& Element0 = InValueArray0.Elements[ElementIndex0];
+			const ElementType1& Element1 = InValueArray1.Elements[ElementIndex1];
+			
+			if(Element0.Name == Element1.Name)
+			{
+				// Elements match, run predicate and increment both indices
+				ElementTypeResult& NewResultElement = OutResultValueArray.Elements.AddDefaulted_GetRef();
+				NewResultElement.Name = Element0.Name;
+				InPredicate(NewResultElement, Element0, Element1, UE::Anim::ENamedValueUnionFlags::BothArgsValid);
+
+				++ElementIndex0;
+				++ElementIndex1;
+			}
+			else if(Element0.Name.FastLess(Element1.Name))
+			{
+				// ValueArray0 element is earlier, so run predicate with only ValueArray0 and increment ValueArray0
+				ElementTypeResult& NewResultElement = OutResultValueArray.Elements.AddDefaulted_GetRef();
+				NewResultElement.Name = Element0.Name;
+				
+				ElementType1 DefaultElement;
+				DefaultElement.Name = Element0.Name;
+				
+				InPredicate(NewResultElement, Element0, DefaultElement, UE::Anim::ENamedValueUnionFlags::ValidArg0);
+				++ElementIndex0;
+			}
+			else
+			{
+				// ValueArray1 element is earlier, so so run predicate with only ValueArray1 and increment ValueArray1
+				ElementTypeResult& NewResultElement = OutResultValueArray.Elements.AddDefaulted_GetRef();
+				NewResultElement.Name = Element1.Name;
+				
+				ElementType0 DefaultElement;
+				DefaultElement.Name = Element1.Name;
+				
+				InPredicate(NewResultElement, DefaultElement, Element1, UE::Anim::ENamedValueUnionFlags::ValidArg1);
+				++ElementIndex1;
+			}
+		}
+
+		// Insertion always proceeds in sorted order, so result is sorted by default
+		OutResultValueArray.bSorted = true;
+
+		OutResultValueArray.CheckSorted();
+	}
+
+	// Helper function
+	// Calls predicate on all elements in the two passed-in value arrays.
+	template<typename PredicateType, typename AllocatorType0, typename ElementType0, typename AllocatorType1, typename ElementType1>
+	static void Union(const TNamedValueArray<AllocatorType0, ElementType0>& InValueArray0, const TNamedValueArray<AllocatorType1, ElementType1>& InValueArray1, const PredicateType& InPredicate)
+	{
+		// Sort both input arrays if required
+		InValueArray0.SortElementsIfRequired();
+		InValueArray1.SortElementsIfRequired();
+		
+		const int32 NumElements0 = InValueArray0.Num();
+		const int32 NumElements1 = InValueArray1.Num();
+
+		int32 ElementIndex0 = 0;
+		int32 ElementIndex1 = 0;
+
+		// Perform dual-iteration on the two sorted arrays
+		while(true)
+		{
+			if(ElementIndex0 == NumElements0 && ElementIndex1 < NumElements1)
+			{
+				// Reached end of ValueArray0 with remaining in ValueArray1, we can just iterate over the remainder of ValueArray1
+				for( ; ElementIndex1 < NumElements1; ++ElementIndex1)
+				{
+					const ElementType1& Element1 = InValueArray1.Elements[ElementIndex1];
+					ElementType0 DefaultElement0;
+					DefaultElement0.Name = Element1.Name;
+
+					InPredicate(DefaultElement0, Element1, UE::Anim::ENamedValueUnionFlags::ValidArg1);
+				}
+				break;
+			}
+			else if(ElementIndex1 == NumElements1 && ElementIndex0 < NumElements0)
+			{
+				// Reached end of ValueArray1 with remaining in ValueArray0, we can just iterate over the remainder of ValueArray0
+				for( ; ElementIndex0 < NumElements0; ++ElementIndex0)
+				{
+					const ElementType0 Element0 = InValueArray0.Elements[ElementIndex0];
+					ElementType1 DefaultElement1;
+					DefaultElement1.Name = Element0.Name;
+
+					InPredicate(Element0, DefaultElement1, UE::Anim::ENamedValueUnionFlags::ValidArg0);
+				}	
+				break;
+			}
+			else if(ElementIndex0 == NumElements0 && ElementIndex1 == NumElements1)
+			{
+				// Reached end of both, exit
+				break;
+			}
+			
+			const ElementType0& Element0 = InValueArray0.Elements[ElementIndex0];
+			const ElementType1& Element1 = InValueArray1.Elements[ElementIndex1];
+			
+			if(Element0.Name == Element1.Name)
+			{
+				// Elements match, run predicate and increment both indices
+				InPredicate(Element0, Element1, UE::Anim::ENamedValueUnionFlags::BothArgsValid);
+				++ElementIndex0;
+				++ElementIndex1;
+			}
+			else if(Element0.Name.FastLess(Element1.Name))
+			{
+				// ValueArray0 element is earlier, so run predicate with only ValueArray0 and increment ElementIndex0
+				ElementType1 DefaultElement1;
+				DefaultElement1.Name = Element0.Name;
+				InPredicate(Element0, DefaultElement1, UE::Anim::ENamedValueUnionFlags::ValidArg0);
+				++ElementIndex0;
+			}
+			else
+			{
+				// ValueArray1 element is earlier, so run predicate with only ValueArray1 and increment ElementIndex1
+				ElementType0 DefaultElement0;
+				DefaultElement0.Name = Element1.Name;
+				InPredicate(DefaultElement0, Element1, UE::Anim::ENamedValueUnionFlags::ValidArg1);
+				++ElementIndex1;
+			}
+		}
+	}
+	
+	/**
+	 * Calls predicate on all matching elements in the two passed-in value arrays.
+	 * ValuePredicateType is a function of signature: (const ElementType0& InElement0, const ElementType1& InElement1) -> void
+	 **/
+	template<typename AllocatorType0, typename ElementType0, typename AllocatorType1, typename ElementType1, typename ValuePredicateType>
+	static void Intersection(const TNamedValueArray<AllocatorType0, ElementType0>& InNamedValues0, const TNamedValueArray<AllocatorType1, ElementType1>& InNamedValues1, const ValuePredicateType& InValuePredicate)
+	{
+		// Sort both inputs if required
+		InNamedValues0.SortElementsIfRequired();
+		InNamedValues1.SortElementsIfRequired();
+
+		const int32 NumElements0 = InNamedValues0.Num();
+		const int32 NumElements1 = InNamedValues1.Num();
+		
+		// Perform dual-iteration on the two sorted arrays
+		int32 ElementIndex0 = 0;
+		int32 ElementIndex1 = 0;
+
+		while(true)
+		{
+			if(ElementIndex0 == NumElements0 && ElementIndex1 < NumElements1)
+			{
+				// Reached end of array with remaining in bulk, we can just early out
+				break;
+			}
+			else if(ElementIndex1 == NumElements1 && ElementIndex0 < NumElements0)
+			{
+				// Reached end of bulk with remaining in array, we can just early out
+				break;
+			}
+			else if(ElementIndex0 == NumElements0 && ElementIndex1 == NumElements1)
+			{
+				// All elements exhausted, exit
+				break;
+			}
+
+			const ElementType0& Element0 = InNamedValues0.Elements[ElementIndex0];
+			const ElementType1& Element1 = InNamedValues1.Elements[ElementIndex1];
+			
+			if(Element0.Name == Element1.Name)
+			{
+				// Elements match so extract value
+				InValuePredicate(Element0, Element1);
+				++ElementIndex0;
+				++ElementIndex1;
+			}
+			else if(Element0.Name.FastLess(Element1.Name))
+			{
+				// Element exists only in array, skip
+				++ElementIndex0;
+			}
+			else
+			{
+				// Element exists only in bulk, skip
+				++ElementIndex1;
+			}
+		}
+	}
+
+	/**
+	 * Removes elements in InOutValueArray0 that match InValueArray1 if predicate returns false
+	 **/
+	template<typename AllocatorType0, typename ElementType0, typename AllocatorType1, typename ElementType1, typename PredicateType>
+	static void RemoveByPredicate(TNamedValueArray<AllocatorType0, ElementType0>& InOutValueArray0, const TNamedValueArray<AllocatorType1, ElementType1>& InValueArray1, const PredicateType& InPredicate)
+	{
+		// Sort both input arrays if required
+		InOutValueArray0.SortElementsIfRequired();
+		InValueArray1.SortElementsIfRequired();
+
+		// Perform dual-iteration on the two sorted arrays
+		int32 ElementIndex0 = 0;
+		int32 ElementIndex1 = 0;
+
+		while(ElementIndex0 < InOutValueArray0.Num() && ElementIndex1 < InValueArray1.Num())
+		{
+			const ElementType0& Element0 = InOutValueArray0.Elements[ElementIndex0];
+			const ElementType1& Element1 = InValueArray1.Elements[ElementIndex1];
+
+			if(Element0.Name == Element1.Name)
+			{
+				// Elements match so check filter flags to see if it should be removed from InOutValueArray0
+				if(!InPredicate(Element0, Element1))
+				{
+					InOutValueArray0.Elements.RemoveAt(ElementIndex0, 1, false);
+					++ElementIndex1;
+				}
+				else
+				{
+					++ElementIndex0;
+					++ElementIndex1;
+				}
+			}
+			else if(Element0.Name.FastLess(Element1.Name))
+			{
+				++ElementIndex0;
+			}
+			else
+			{
+				++ElementIndex1;
+			}
+		}
+
+		InOutValueArray0.Elements.Shrink();
+		InOutValueArray0.CheckSorted();
+	}
+};
+
+}
+
+ENUM_CLASS_FLAGS(UE::Anim::ENamedValueUnionFlags);
\ No newline at end of file
