Animation Curve Runtime & Editor Improvements

diff --git a/Engine/Source/Runtime/Engine/Public/BoneContainer.h b/Engine/Source/Runtime/Engine/Public/BoneContainer.h
--- a/Engine/Source/Runtime/Engine/Public/BoneContainer.h
+++ b/Engine/Source/Runtime/Engine/Public/BoneContainer.h
@@ -7,8 +7,11 @@
 #include "UObject/WeakObjectPtr.h"
 #include "BoneIndices.h"
 #include "ReferenceSkeleton.h"
+#include "Animation/AnimCurveFilter.h"
 #include "Animation/AnimTypes.h"
-#include "BoneContainer.generated.h"
+#include "Animation/BoneReference.h"
+#include "Animation/AnimCurveMetadata.h"
+#include "Animation/AnimBulkCurves.h"
 
 class USkeletalMesh;
 class USkeleton;
@@ -16,6 +19,12 @@ class USkeletalMesh;
 struct FBoneContainer;
 struct FSkeletonRemapping;
 struct FBlendedCurve;
+struct FAnimCurveType;
+
+namespace SmartName
+{
+typedef uint16 UID_Type;
+}
 
 /** Struct used to store per-component ref pose override */
 struct FSkelMeshRefPoseOverride
@@ -26,27 +35,6 @@ struct FSkelMeshRefPoseOverride
 	TArray<FTransform> RefBonePoses;
 };
 
-/** in the future if we need more bools, please convert to bitfield 
- * These are not saved in asset but per skeleton. 
- */
-USTRUCT()
-struct ENGINE_API FAnimCurveType
-{
-	GENERATED_USTRUCT_BODY()
-
-	UPROPERTY()
-	bool bMaterial;
-
-	UPROPERTY()
-	bool bMorphtarget;
-
-	FAnimCurveType(bool bInMorphtarget = false, bool bInMaterial = false)
-		: bMaterial(bInMaterial)
-		, bMorphtarget(bInMorphtarget)
-	{
-	}
-};
-
 /**
 * This is a native transient structure. Used to store virtual bone mappings for compact poses
 **/
@@ -66,15 +54,14 @@ struct FVirtualBoneCompactPoseData
 	{}
 };
 
-/**
- * This is curve evaluation options for bone container
- */
+// DEPRECATED - please use UE::Anim::FCurveFilterSettings
 struct FCurveEvaluationOption
 {
 	bool bAllowCurveEvaluation;
 	const TArray<FName>* DisallowedList;
 	int32 LODIndex;
-	
+
+	UE_DEPRECATED(5.3, "Please use UE::Anim::FCurveFilterSettings.")
 	FCurveEvaluationOption(bool bInAllowCurveEvaluation = true, const TArray<FName>* InDisallowedList = nullptr, int32 InLODIndex = 0)
 		: bAllowCurveEvaluation(bInAllowCurveEvaluation)
 		, DisallowedList(InDisallowedList)
@@ -83,6 +70,30 @@ struct FCurveEvaluationOption
 	}
 };
 
+namespace UE::Anim
+{
+
+/** Options for filtering curves. Used to initialize a bone container with a filter curve. */
+struct FCurveFilterSettings
+{
+	FCurveFilterSettings(UE::Anim::ECurveFilterMode InFilterMode = UE::Anim::ECurveFilterMode::None, const TArray<FName>* InFilterCurves = nullptr, int32 InLODIndex = 0)
+		: FilterCurves(InFilterCurves)
+		, FilterMode(InFilterMode)
+		, LODIndex(InLODIndex)
+	{}
+
+	// Filter curves. Application of these curves is dependent on FilterMode
+	const TArray<FName>* FilterCurves = nullptr;
+	
+	// Filtering mode - allows curves to be disable entirely or allow/deny-listed
+	UE::Anim::ECurveFilterMode FilterMode = UE::Anim::ECurveFilterMode::None;
+
+	// LOD index
+	int32 LODIndex = 0;
+};
+
+}
+
 /** Stores cached data for Orient and Scale bone translation retargeting */
 struct FOrientAndScaleRetargetingCachedData
 {
@@ -116,48 +127,6 @@ struct FCachedSkeletonCurveMapping
 	bool bIsDirty = true; /** Specifies whether we need to rebuild this cached data or not. */
 };
 
-/**
- * A scoped curve remapping.
- * This object is used in a RAII pattern. It initializes the curve remapping on the provided curve.
- * When the object gets destructed it will automatically revert the remapping on the curve object.
- * 
- * Usage would be something like this:
- * \code{.cpp}
- * FSkeletonRemapping* Mapping = TargetSkeleton->GetSkeletonRemapping(SourceSkeleton);
- * if (Mapping)
- * {
- *		FSkeletonRemappingCurve Context(BlendedCurve, RequiredBones, Mapping);
- *		EvaluateCurveData(BlendedCurve);
- * }
- * \endcode
- * 
- * @param InCurve The curve object that will be used during sampling of the animation data. We will modify its UIDToArrayIndexLUT member to point to a remapped version.
- * @param InBoneContainer The bone container that we're using. This stores the remappings.
- * @param InSkeletonMapping The skeleton mapping to use. This can be requested with USkeleton::GetSkeletonRemapping(SourceSkeleton).
- */
-struct FSkeletonRemappingCurve
-{
-public:
-	FSkeletonRemappingCurve() = delete;
-	FSkeletonRemappingCurve(const FSkeletonRemappingCurve&) = delete;
-	FSkeletonRemappingCurve(FSkeletonRemappingCurve&&) = delete;
-	FSkeletonRemappingCurve(FBlendedCurve& InCurve, FBoneContainer& InBoneContainer, const FSkeletonRemapping* InSkeletonMapping);
-	FSkeletonRemappingCurve(FBlendedCurve& InCurve, FBoneContainer& InBoneContainer, const USkeleton* SourceSkeleton);
-	~FSkeletonRemappingCurve();
-
-	FSkeletonRemappingCurve& operator = (const FSkeletonRemappingCurve&) = delete;
-	FSkeletonRemappingCurve& operator = (FSkeletonRemappingCurve&&) = delete;
-
-	FBlendedCurve& GetCurve() { return Curve;  }
-	const FBlendedCurve& GetCurve() const { return Curve; }
-
-private:
-	FBlendedCurve& Curve;
-	FBoneContainer& BoneContainer;
-	bool bIsRemapping = false;
-};
-
-
 /** Retargeting cached data for a specific Retarget Source */
 struct FRetargetSourceCachedData
 {
@@ -236,12 +205,12 @@ private:
 	// Array of cached virtual bone data so that animations running from raw data can generate them.
 	TArray<FVirtualBoneCompactPoseData> VirtualBoneCompactPoseData;
 
-	/** Look up table of UID to array index UIDToArrayIndexLUT[InUID] = ArrayIndex of order. If MAX_uint16, it is invalid.*/
-	TArray<uint16> UIDToArrayIndexLUT;
-
-	/** Number of valid entries in UIDToArrayIndexLUT. I.e. a count of entries whose value does not equal to  MAX_uint16. */
-	int32 UIDToArrayIndexLUTValidCount;
+	/** Curve used for filtering by LOD/bone */
+	UE::Anim::FCurveFilter CurveFilter;
 
+	/** Curve flags to apply (built from curve metadata) */
+	UE::Anim::FBulkCurveFlags CurveFlags;
+	
 	TSharedPtr<FSkelMeshRefPoseOverride> RefPoseOverride;
 	
 	// The serial number of this bone container. This is incremented each time the container is regenerated and can
@@ -260,14 +229,19 @@ private:
 	bool bUseRAWData;
 	/** Use Source Data that is imported that are not compressed. */
 	bool bUseSourceData;
-	
+
 public:
 
 	FBoneContainer();
 
+	FBoneContainer(const TArray<FBoneIndexType>& InRequiredBoneIndexArray, const UE::Anim::FCurveFilterSettings& InCurveFilterSettings, UObject& InAsset);
+	
+	void InitializeTo(const TArray<FBoneIndexType>& InRequiredBoneIndexArray, const UE::Anim::FCurveFilterSettings& InCurveFilterSettings, UObject& InAsset);
+	
+	UE_DEPRECATED(5.3, "Please use the constructor that takes a FCurveFilterSettings.")
 	FBoneContainer(const TArray<FBoneIndexType>& InRequiredBoneIndexArray, const FCurveEvaluationOption& CurveEvalOption, UObject& InAsset);
 
-	/** Initialize BoneContainer to a new Asset, RequiredBonesArray and RefPoseArray. */
+	UE_DEPRECATED(5.3, "Please use InitializeTo that takes a FCurveFilterSettings.")
 	void InitializeTo(const TArray<FBoneIndexType>& InRequiredBoneIndexArray, const FCurveEvaluationOption& CurveEvalOption, UObject& InAsset);
 
 	/** Returns true if FBoneContainer is Valid. Needs an Asset, a RefPoseArray, and a RequiredBonesArray. */
@@ -449,19 +423,24 @@ public:
 	/** Returns true if bone is child of for current asset. */
 	bool BoneIsChildOf(const FCompactPoseBoneIndex& BoneIndex, const FCompactPoseBoneIndex& ParentBoneIndex) const;
 
-	/** Get UID To Array look up table */
-	TArray<uint16> const& GetUIDToArrayLookupTable() const
+	/** Get filter used for filtering by LOD/bone */
+	const UE::Anim::FCurveFilter& GetCurveFilter() const
 	{
-		return UIDToArrayIndexLUT;
+		return CurveFilter;
 	}
 
-	/** Returns the number of valid entries in the GetUIDToArrayLookupTable result array */
-	int32 GetUIDToArrayIndexLookupTableValidCount() const
+	/** Get flags to apply to curves (built from metadata) */
+	const UE::Anim::FBulkCurveFlags& GetCurveFlags() const
 	{
-		return UIDToArrayIndexLUTValidCount;
+		return CurveFlags;
 	}
 
-	/** DEPRECATED: Get UID To Name look up table */
+	UE_DEPRECATED(5.3, "GetUIDToArrayLookupTable is deprecated, it is no longer used.")
+	TArray<uint16> const& GetUIDToArrayLookupTable() const;
+
+	UE_DEPRECATED(5.3, "GetUIDToArrayIndexLookupTableValidCount is deprecated, it is no longer used.")
+	int32 GetUIDToArrayIndexLookupTableValidCount() const;
+
 	UE_DEPRECATED(5.0, "GetUIDToNameLookupTable is deprecated, please access from the SmartNameMapping directly via GetSkeletonAsset()->GetSmartNameContainer(USkeleton::AnimCurveMappingName)")
 	TArray<FName> const& GetUIDToNameLookupTable() const
 	{
@@ -469,20 +448,11 @@ public:
 		return Dummy;
 	}
 
-	/** DEPRECATED: Get UID To curve type look up table */
 	UE_DEPRECATED(5.0, "GetUIDToCurveTypeLookupTable is deprecated, please access from the SmartNameMapping directly via GetSkeletonAsset()->GetSmartNameContainer(USkeleton::AnimCurveMappingName)")
-	TArray<FAnimCurveType> const& GetUIDToCurveTypeLookupTable() const
-	{
-		static TArray<FAnimCurveType> Dummy;
-		return Dummy;
-	}
+	TArray<FAnimCurveType> const& GetUIDToCurveTypeLookupTable() const;
 
-	
-	/** Get the array that maps UIDs to array indexes. */
-	TArray<SmartName::UID_Type> const& GetUIDToArrayLookupTableBackup() const
-	{
-		return UIDToArrayIndexLUTBackup;
-	}
+	UE_DEPRECATED(5.3, "GetUIDToArrayLookupTableBackup is deprecated, it is no longer used.")
+	TArray<SmartName::UID_Type> const& GetUIDToArrayLookupTableBackup() const;
 
 	/**
 	* Serializes the bones
@@ -676,9 +646,16 @@ public:
 		return FCompactPoseBoneIndex(GetBoneIndicesArray().IndexOfByKey(BoneIndex.GetInt()));
 	}
 
-	/** Cache required Anim Curve Uids */
-	void CacheRequiredAnimCurveUids(const FCurveEvaluationOption& CurveEvalOption);
+	UE_DEPRECATED(5.3, "Please use CacheRequiredAnimCurves with a FCurveFilterSettings.")
+	void CacheRequiredAnimCurveUids(const FCurveEvaluationOption& CurveEvalOption)
+	{
+		const UE::Anim::FCurveFilterSettings CurveFilterSettings(CurveEvalOption.bAllowCurveEvaluation ? UE::Anim::ECurveFilterMode::None : UE::Anim::ECurveFilterMode::DisallowAll, CurveEvalOption.DisallowedList, CurveEvalOption.LODIndex);
+		CacheRequiredAnimCurves(CurveFilterSettings);
+	}
 
+	/** Cache required Anim Curves */
+	void CacheRequiredAnimCurves(const UE::Anim::FCurveFilterSettings& InCurveFilterSettings);
+	
 	const FRetargetSourceCachedData& GetRetargetSourceCachedData(const FName& InRetargetSource) const;
 	const FRetargetSourceCachedData& GetRetargetSourceCachedData(const FName& InSourceName, const FSkeletonRemapping& InRemapping, const TArray<FTransform>& InRetargetTransforms) const;
 
@@ -693,14 +670,8 @@ public:
 
 	// Get the serial number of this bone container. @see SerialNumber
 	uint16 GetSerialNumber() const { return SerialNumber; }
-	
-private:
-	/** The map of cached curve mapping indexes, which is used for skeleton remapping. The key of this table is the source skeleton of the asset we are sampling curve values from. */
-	mutable TMap<USkeleton*, FCachedSkeletonCurveMapping> CachedCurveMappingTable;
-
-	/** The backup of the original UIDToArrayIndexLUT array. This is needed for skeleton remapping curves, as we have to modify this array and later need to restore it again. */
-	TArray<SmartName::UID_Type> UIDToArrayIndexLUTBackup;
 
+private:
 	/** 
 	 * Runtime cached data for retargeting from a specific RetargetSource to this current SkelMesh LOD.
 	 * @todo: We could also cache this once per skelmesh per lod, rather than creating it at runtime for each skelmesh instance.
@@ -708,7 +679,7 @@ private:
 	mutable TMap<FName, FRetargetSourceCachedData> RetargetSourceCachedDataLUT;
 
 	/** Initialize FBoneContainer. */
-	void Initialize(const FCurveEvaluationOption& CurveEvalOption);
+	void Initialize(const UE::Anim::FCurveFilterSettings& CurveFilterSettings);
 
 	/** Cache remapping data if current Asset is a SkeletalMesh, with all compatible Skeletons. */
 	void RemapFromSkelMesh(USkeletalMesh const & SourceSkeletalMesh, USkeleton& TargetSkeleton);
@@ -721,142 +692,3 @@ private:
 };
 
 
-USTRUCT()
-struct FBoneReference
-{
-	GENERATED_USTRUCT_BODY()
-
-	/** Name of bone to control. This is the main bone chain to modify from. **/
-	UPROPERTY(EditAnywhere, Category = BoneReference)
-	FName BoneName;
-
-	/** Cached bone index for run time - right now bone index of skeleton **/
-	int32 BoneIndex:31;
-
-	/** Change this to Bitfield if we have more than one bool 
-	 * This specifies whether or not this indices is mesh or skeleton
-	 */
-	uint32 bUseSkeletonIndex:1;
-
-	FCompactPoseBoneIndex CachedCompactPoseIndex;
-
-	FBoneReference()
-		: BoneName(NAME_None)
-		, BoneIndex(INDEX_NONE)
-		, bUseSkeletonIndex(false)
-		, CachedCompactPoseIndex(INDEX_NONE)
-	{
-	}
-
-	FBoneReference(const FName& InBoneName)
-		: BoneName(InBoneName)
-		, BoneIndex(INDEX_NONE)
-		, bUseSkeletonIndex(false)
-		, CachedCompactPoseIndex(INDEX_NONE)
-	{
-	}
-
-	bool operator==(const FBoneReference& Other) const
-	{
-		return BoneName == Other.BoneName;
-	}
-
-	bool operator!=(const FBoneReference& Other) const
-	{
-		return BoneName != Other.BoneName;
-	}
-
-	/** Initialize Bone Reference, return TRUE if success, otherwise, return false **/
-	ENGINE_API bool Initialize(const FBoneContainer& RequiredBones);
-
-	// only used by blendspace 'PerBoneBlend'. This is skeleton indices since the input is skeleton
-	// @note, if you use this function, it won't work with GetCompactPoseIndex, GetMeshPoseIndex;
-	// it triggers ensure in those functions
-	ENGINE_API bool Initialize(const USkeleton* Skeleton);
-
-
-	/** return true if it has valid set up */
-	bool HasValidSetup() const
-	{
-		return (BoneIndex != INDEX_NONE);
-	}
-
-	/** return true if has valid index, and required bones contain it **/
-	ENGINE_API bool IsValidToEvaluate(const FBoneContainer& RequiredBones) const;
-	/** return true if has valid compact index. This will return invalid if you're using skeleton index */
-	ENGINE_API bool IsValidToEvaluate() const
-	{
-		return (!bUseSkeletonIndex && CachedCompactPoseIndex != INDEX_NONE);
-	}
-
-	void InvalidateCachedBoneIndex()
-	{
-		BoneIndex = INDEX_NONE;
-		CachedCompactPoseIndex = FCompactPoseBoneIndex(INDEX_NONE);
-	}
-
-	FSkeletonPoseBoneIndex GetSkeletonPoseIndex(const FBoneContainer& RequiredBones) const
-	{ 
-		// accessing array with invalid index would cause crash, so we have to check here
-		if (BoneIndex != INDEX_NONE)
-		{
-			if (bUseSkeletonIndex)
-			{
-				return FSkeletonPoseBoneIndex(BoneIndex);
-			}
-			else
-			{
-				return RequiredBones.GetSkeletonPoseIndexFromMeshPoseIndex(FMeshPoseBoneIndex(BoneIndex));
-			}
-		}
-
-		return FSkeletonPoseBoneIndex(INDEX_NONE);
-	}
-	
-	FMeshPoseBoneIndex GetMeshPoseIndex(const FBoneContainer& RequiredBones) const
-	{ 
-		// accessing array with invalid index would cause crash, so we have to check here
-		if (BoneIndex != INDEX_NONE)
-		{
-			if (bUseSkeletonIndex)
-			{
-				return RequiredBones.GetMeshPoseIndexFromSkeletonPoseIndex(FSkeletonPoseBoneIndex(BoneIndex));
-			}
-			else
-			{
-				return FMeshPoseBoneIndex(BoneIndex);
-			}
-		}
-
-		return FMeshPoseBoneIndex(INDEX_NONE);
-	}
-
-	FCompactPoseBoneIndex GetCompactPoseIndex(const FBoneContainer& RequiredBones) const 
-	{ 
-		if (bUseSkeletonIndex)
-		{
-			//If we were initialized with a skeleton we wont have a cached index.
-			if (BoneIndex != INDEX_NONE)
-			{
-				// accessing array with invalid index would cause crash, so we have to check here
-				return RequiredBones.GetCompactPoseIndexFromSkeletonPoseIndex(FSkeletonPoseBoneIndex(BoneIndex));
-			}
-			return FCompactPoseBoneIndex(INDEX_NONE);
-		}
-		
-		return CachedCompactPoseIndex;
-	}
-
-	// need this because of BoneReference being used in CurveMetaData and that is in SmartName
-	friend FArchive& operator<<(FArchive& Ar, FBoneReference& B)
-	{
-		Ar << B.BoneName;
-		return Ar;
-	}
-
-	bool Serialize(FArchive& Ar)
-	{
-		Ar << *this;
-		return true;
-	}
-};
diff --git a/Engine/Source/Runtime/LiveLinkAnimationCore/Private/LiveLinkRetargetAsset.cpp b/Engine/Source/Runtime/LiveLinkAnimationCore/Private/LiveLinkRetargetAsset.cpp
--- a/Engine/Source/Runtime/LiveLinkAnimationCore/Private/LiveLinkRetargetAsset.cpp
+++ b/Engine/Source/Runtime/LiveLinkAnimationCore/Private/LiveLinkRetargetAsset.cpp
@@ -18,11 +18,7 @@ ULiveLinkRetargetAsset::ULiveLinkRetargetAsset(const FObjectInitializer& ObjectI
 
 void ULiveLinkRetargetAsset::ApplyCurveValue(const USkeleton* Skeleton, const FName CurveName, const float CurveValue, FBlendedCurve& OutCurve) const
 {
-	SmartName::UID_Type UID = Skeleton->GetUIDByName(USkeleton::AnimCurveMappingName, CurveName);
-	if (UID != SmartName::MaxUID)
-	{
-		OutCurve.Set(UID, CurveValue);
-	}
+	OutCurve.Set(CurveName, CurveValue);
 }
 
 void ULiveLinkRetargetAsset::BuildCurveData(const FLiveLinkSkeletonStaticData* InSkeletonData, const FLiveLinkAnimationFrameData* InFrameData, const FCompactPose& InPose, FBlendedCurve& OutCurve) const
diff --git a/Engine/Source/Runtime/MovieSceneTracks/Private/Tracks/MovieSceneSkeletalAnimationTrack.cpp b/Engine/Source/Runtime/MovieSceneTracks/Private/Tracks/MovieSceneSkeletalAnimationTrack.cpp
--- a/Engine/Source/Runtime/MovieSceneTracks/Private/Tracks/MovieSceneSkeletalAnimationTrack.cpp
+++ b/Engine/Source/Runtime/MovieSceneTracks/Private/Tracks/MovieSceneSkeletalAnimationTrack.cpp
@@ -467,14 +467,13 @@ TOptional<FTransform>  UMovieSceneSkeletalAnimationTrack::GetRootMotion(FFrameTi
 				FMemMark Mark(FMemStack::Get());
 				FCompactPose OutPose;
 				TArray<FBoneIndexType> RequiredBoneIndexArray;
-				const FCurveEvaluationOption CurveEvalOption;
 				RequiredBoneIndexArray.AddUninitialized(ValidAnimSequence->GetSkeleton()->GetReferenceSkeleton().GetNum());
 				for (int32 BoneIndex = 0; BoneIndex < RequiredBoneIndexArray.Num(); ++BoneIndex)
 				{
 					RequiredBoneIndexArray[BoneIndex] = BoneIndex;
 				}
 
-				FBoneContainer BoneContainer(RequiredBoneIndexArray, CurveEvalOption, *ValidAnimSequence->GetSkeleton());
+				FBoneContainer BoneContainer(RequiredBoneIndexArray, UE::Anim::FCurveFilterSettings(UE::Anim::ECurveFilterMode::None), *ValidAnimSequence->GetSkeleton());
 				OutPose.ResetToRefPose(BoneContainer);
 				FBlendedCurve OutCurve;
 				OutCurve.InitFrom(BoneContainer);
@@ -574,14 +573,13 @@ void UMovieSceneSkeletalAnimationTrack::SetUpRootMotions(bool bForce)
 							FMemMark Mark(FMemStack::Get());
 							FCompactPose OutPose;
 							TArray<FBoneIndexType> RequiredBoneIndexArray;
-							const FCurveEvaluationOption CurveEvalOption;
 							RequiredBoneIndexArray.AddUninitialized(PrevAnimSequence->GetSkeleton()->GetReferenceSkeleton().GetNum());
 							for (int32 BoneIndex = 0; BoneIndex < RequiredBoneIndexArray.Num(); ++BoneIndex)
 							{
 								RequiredBoneIndexArray[BoneIndex] = BoneIndex;
 							}
 
-							FBoneContainer BoneContainer(RequiredBoneIndexArray, CurveEvalOption, *PrevAnimSequence->GetSkeleton());
+							FBoneContainer BoneContainer(RequiredBoneIndexArray, UE::Anim::FCurveFilterSettings(UE::Anim::ECurveFilterMode::None), *PrevAnimSequence->GetSkeleton());
 							OutPose.ResetToRefPose(BoneContainer);
 							FBlendedCurve OutCurve;
 							OutCurve.InitFrom(BoneContainer);
@@ -651,14 +649,14 @@ void UMovieSceneSkeletalAnimationTrack::SetUpRootMotions(bool bForce)
 		if (ValidAnimSequence)
 		{
 			TArray<FBoneIndexType> RequiredBoneIndexArray;
-			const FCurveEvaluationOption CurveEvalOption;
+			const UE::Anim::FCurveFilterSettings CurveFilterSettings;
 			RequiredBoneIndexArray.AddUninitialized(ValidAnimSequence->GetSkeleton()->GetReferenceSkeleton().GetNum());
 			for (int32 BoneIndex = 0; BoneIndex < RequiredBoneIndexArray.Num(); ++BoneIndex)
 			{
 				RequiredBoneIndexArray[BoneIndex] = BoneIndex;
 			}
 
-			FBoneContainer BoneContainer(RequiredBoneIndexArray, CurveEvalOption, *ValidAnimSequence->GetSkeleton());
+			FBoneContainer BoneContainer(RequiredBoneIndexArray, CurveFilterSettings, *ValidAnimSequence->GetSkeleton());
 			OutPose.ResetToRefPose(BoneContainer);
 			FBlendedCurve OutCurve;
 			OutCurve.InitFrom(BoneContainer);
@@ -741,14 +739,14 @@ static FTransform GetTransformForBoneRelativeToIndex(UAnimSequence* AnimSequence
 	FMemMark Mark(FMemStack::Get());
 	FCompactPose OutPose;
 	TArray<FBoneIndexType> RequiredBoneIndexArray;
-	const FCurveEvaluationOption CurveEvalOption;
+	const UE::Anim::FCurveFilterSettings CurveFilterSettings;
 	RequiredBoneIndexArray.AddUninitialized(AnimSequence->GetSkeleton()->GetReferenceSkeleton().GetNum());
 	for (int32 BoneIndex = 0; BoneIndex < RequiredBoneIndexArray.Num(); ++BoneIndex)
 	{
 		RequiredBoneIndexArray[BoneIndex] = BoneIndex;
 	}
 
-	FBoneContainer BoneContainer(RequiredBoneIndexArray, CurveEvalOption, *AnimSequence->GetSkeleton());
+	FBoneContainer BoneContainer(RequiredBoneIndexArray, CurveFilterSettings, *AnimSequence->GetSkeleton());
 	OutPose.ResetToRefPose(BoneContainer);
 	FBlendedCurve OutCurve;
 	OutCurve.InitFrom(BoneContainer);
