Fixed static analysis warnings exposed by the log macro refactor

#rnx

diff --git a/Engine/Plugins/Experimental/WaveFunctionCollapse/Source/WaveFunctionCollapse/Private/WaveFunctionCollapseBPLibrary.cpp b/Engine/Plugins/Experimental/WaveFunctionCollapse/Source/WaveFunctionCollapse/Private/WaveFunctionCollapseBPLibrary.cpp
--- a/Engine/Plugins/Experimental/WaveFunctionCollapse/Source/WaveFunctionCollapse/Private/WaveFunctionCollapseBPLibrary.cpp
+++ b/Engine/Plugins/Experimental/WaveFunctionCollapse/Source/WaveFunctionCollapse/Private/WaveFunctionCollapseBPLibrary.cpp
@@ -525,7 +525,7 @@ bool UWaveFunctionCollapseBPLibrary::GetPositionToOptionMapFromActor(AActor* Act
 	// Check if Model is valid
 	if (!Actor)
 	{
-		UE_LOG(LogWFC, Display, TEXT("%s is an invalid Actor"), *Actor->GetFName().ToString());
+		UE_LOG(LogWFC, Display, TEXT("GetPositionToOptionMapFromActor called with a null Actor"));
 		return false;
 	}
 
diff --git a/Engine/Plugins/FX/Niagara/Source/Niagara/Private/NiagaraComponentPool.cpp b/Engine/Plugins/FX/Niagara/Source/Niagara/Private/NiagaraComponentPool.cpp
--- a/Engine/Plugins/FX/Niagara/Source/Niagara/Private/NiagaraComponentPool.cpp
+++ b/Engine/Plugins/FX/Niagara/Source/Niagara/Private/NiagaraComponentPool.cpp
@@ -98,7 +98,7 @@ UNiagaraComponent* FNCPool::Acquire(UWorld* World, UNiagaraSystem* Template, ENC
 		if (!RetElem.Component || !IsValidChecked(RetElem.Component))
 		{			
 			// Possible someone still has a reference to our NC and destroyed it while it was sat in the pool. Or possibly a teardown edgecase path that is GCing components from the pool be
-			UE_LOG(LogNiagara, Warning, TEXT("Pooled NC has been destroyed or is pending kill! Possibly via a DestroyComponent() call. You should not destroy pooled components manually. \nJust deactivate them and allow them to destroy themselves or be reclaimed by the pool. | NC: %p |\t System: %s"), RetElem.Component, *Template->GetFullName());
+			UE_LOG(LogNiagara, Warning, TEXT("Pooled NC has been destroyed or is pending kill! Possibly via a DestroyComponent() call. You should not destroy pooled components manually. \nJust deactivate them and allow them to destroy themselves or be reclaimed by the pool. | NC: %p |\t System: %s"), RetElem.Component.Get(), *Template->GetFullName());
 			RetElem = FNCPoolElement();
 		}
 		else
diff --git a/Engine/Plugins/FX/Niagara/Source/Niagara/Private/NiagaraScalabilityManager.cpp b/Engine/Plugins/FX/Niagara/Source/Niagara/Private/NiagaraScalabilityManager.cpp
--- a/Engine/Plugins/FX/Niagara/Source/Niagara/Private/NiagaraScalabilityManager.cpp
+++ b/Engine/Plugins/FX/Niagara/Source/Niagara/Private/NiagaraScalabilityManager.cpp
@@ -217,8 +217,8 @@ bool FNiagaraScalabilityManager::EvaluateCullState(FNiagaraWorldManager* WorldMa
 		UNiagaraSystem* System = Component->GetAsset();
 		if (System == nullptr)
 		{
-			UE_LOG(LogNiagara, Warning, TEXT("Niagara System has been destroyed with components still registered to the scalability manager. Unregistering this component.\nComponent: 0x%P - %s\nEffectType: 0x%P - %s"),
-				Component, *Component->GetName(), EffectType, *EffectType->GetName());
+			UE_LOG(LogNiagara, Warning, TEXT("Niagara System has been destroyed with components still registered to the scalability manager. Unregistering this component.\nComponent: 0x%p - %s\nEffectType: 0x%p - %s"),
+				Component, *Component->GetName(), EffectType.Get(), *EffectType->GetName());
 			Unregister(Component);
 			return false;
 		}
diff --git a/Engine/Plugins/Online/OnlineFramework/Source/Party/Private/Party/PartyMember.cpp b/Engine/Plugins/Online/OnlineFramework/Source/Party/Private/Party/PartyMember.cpp
--- a/Engine/Plugins/Online/OnlineFramework/Source/Party/Private/Party/PartyMember.cpp
+++ b/Engine/Plugins/Online/OnlineFramework/Source/Party/Private/Party/PartyMember.cpp
@@ -93,13 +93,13 @@ void UPartyMember::InitializePartyMember(const FOnlinePartyMemberConstRef& InOss
 			}, false);
 		check(SocialUser);
 
-		UE_LOG(LogParty, Log, TEXT("%s - QUDA returned SocialUser %s (%x)"), ANSI_TO_TCHAR(__FUNCTION__), *GetFullNameSafe(SocialUser), SocialUser);
+		UE_LOG(LogParty, Log, TEXT("%hs - QUDA returned SocialUser %s (%p)"), __FUNCTION__, *GetFullNameSafe(SocialUser), SocialUser.Get());
 
 		// Local player already has all the data they need, everyone else we want to wait for
 		bHasReceivedInitialData = IsLocalPlayer();
 
 		OnInitializationComplete().Add(OnInitComplete);
-		UE_LOG(LogParty, Log, TEXT("%s - Registering Init Complete Handler for [%s], SocialUser: %s (%x)"), ANSI_TO_TCHAR(__FUNCTION__), *ToDebugString(), *GetFullNameSafe(SocialUser), SocialUser);
+		UE_LOG(LogParty, Log, TEXT("%hs - Registering Init Complete Handler for [%s], SocialUser: %s (%p)"), __FUNCTION__, *ToDebugString(), *GetFullNameSafe(SocialUser), SocialUser.Get());
 		SocialUser->RegisterInitCompleteHandler(FOnNewSocialUserInitialized::CreateUObject(this, &UPartyMember::HandleSocialUserInitialized));
 
 		UE_LOG(LogParty, Verbose, TEXT("Created new party member [%s]"), *ToDebugString());
diff --git a/Engine/Plugins/Runtime/MassGameplay/Source/MassSpawner/Private/MassEntityTemplateRegistry.cpp b/Engine/Plugins/Runtime/MassGameplay/Source/MassSpawner/Private/MassEntityTemplateRegistry.cpp
--- a/Engine/Plugins/Runtime/MassGameplay/Source/MassSpawner/Private/MassEntityTemplateRegistry.cpp
+++ b/Engine/Plugins/Runtime/MassGameplay/Source/MassSpawner/Private/MassEntityTemplateRegistry.cpp
@@ -176,10 +176,8 @@ bool FMassEntityTemplateBuildContext::ValidateBuildContext(const UWorld& World)
 		{
 			if (!bHeaderOutputed)
 			{
-				check(CurrentStruct);
-				check(CurrentTrait);
-				UE_LOG(LogMass, Error, TEXT("Fragment(%s) was added multiple time and can only be added by one trait. Fragment was added by:"), *CurrentStruct->GetName());
-				UE_LOG(LogMass, Error, TEXT("\t\t%s"), *CurrentTrait->GetClass()->GetName());
+				UE_LOG(LogMass, Error, TEXT("Fragment(%s) was added multiple time and can only be added by one trait. Fragment was added by:"), CurrentStruct ? *CurrentStruct->GetName() : TEXT("null"));
+				UE_LOG(LogMass, Error, TEXT("\t\t%s"), CurrentTrait ? *CurrentTrait->GetClass()->GetName() : TEXT("null"));
 				bHeaderOutputed = true;
 			}
 			UE_LOG(LogMass, Error, TEXT("\t\t%s"), *Pair.Value->GetClass()->GetName());
diff --git a/Engine/Plugins/Runtime/nDisplay/Source/DisplayClusterMedia/Private/Platforms/Windows/SharedMemoryMediaPlatformWindowsD3D12.cpp b/Engine/Plugins/Runtime/nDisplay/Source/DisplayClusterMedia/Private/Platforms/Windows/SharedMemoryMediaPlatformWindowsD3D12.cpp
--- a/Engine/Plugins/Runtime/nDisplay/Source/DisplayClusterMedia/Private/Platforms/Windows/SharedMemoryMediaPlatformWindowsD3D12.cpp
+++ b/Engine/Plugins/Runtime/nDisplay/Source/DisplayClusterMedia/Private/Platforms/Windows/SharedMemoryMediaPlatformWindowsD3D12.cpp
@@ -135,7 +135,7 @@ FTextureRHIRef FSharedMemoryMediaPlatformWindowsD3D12::CreateSharedCrossGpuTextu
 	{
 		UE_LOG(LogDisplayClusterMedia, Error, TEXT(
 			"D3D12Device->CreateCommittedResource failed when creating a cross GPU texture:\n"
-			"0x%X - %s. Texture was %dx%d, EPixelFormat %d and DXGI type %d"),
+			"0x%lX - %s. Texture was %dx%d, EPixelFormat %d and DXGI type %d"),
 			HResult, *GetD3D12ComErrorDescription(HResult), Width, Height, Format, DxgiFormat);
 
 		return nullptr;
@@ -191,7 +191,7 @@ FTextureRHIRef FSharedMemoryMediaPlatformWindowsD3D12::OpenSharedCrossGpuTexture
 
 	if (FAILED(HResult))
 	{
-		UE_LOG(LogDisplayClusterMedia, Error, TEXT("D3D12Device->OpenSharedHandleByName(%s) failed: 0x%X - %s"),
+		UE_LOG(LogDisplayClusterMedia, Error, TEXT("D3D12Device->OpenSharedHandleByName(%s) failed: 0x%lX - %s"),
 			*SharedGpuTextureName, HResult, *GetD3D12ComErrorDescription(HResult));
 
 		return nullptr;
@@ -199,13 +199,13 @@ FTextureRHIRef FSharedMemoryMediaPlatformWindowsD3D12::OpenSharedCrossGpuTexture
 
 	HResult = D3D12Device->OpenSharedHandle(NamedSharedGpuTextureHandle, IID_PPV_ARGS(&SharedCrossGpuTexture));
 
+	UE_CLOG(FAILED(HResult), LogDisplayClusterMedia, Error, TEXT("D3D12Device->OpenSharedHandle(0x%p) failed: 0x%lX - %s"),
+		NamedSharedGpuTextureHandle, HResult, *GetD3D12ComErrorDescription(HResult));
+
 	::CloseHandle(NamedSharedGpuTextureHandle);
 
 	if (FAILED(HResult))
 	{
-		UE_LOG(LogDisplayClusterMedia, Error, TEXT("D3D12Device->OpenSharedHandle(0x%X) failed: 0x%X - %s"),
-			NamedSharedGpuTextureHandle, HResult, *GetD3D12ComErrorDescription(HResult));
-
 		return nullptr;
 	}
 
diff --git a/Engine/Source/Developer/Apple/MetalShaderFormat/Private/MetalDerivedData.cpp b/Engine/Source/Developer/Apple/MetalShaderFormat/Private/MetalDerivedData.cpp
--- a/Engine/Source/Developer/Apple/MetalShaderFormat/Private/MetalDerivedData.cpp
+++ b/Engine/Source/Developer/Apple/MetalShaderFormat/Private/MetalDerivedData.cpp
@@ -1423,7 +1423,6 @@ bool DoCompileMetalShader(
 			}
 			GLog->Flush();
 		}
+		return false;
 	}
-
-	return Result != 0;
 }
diff --git a/Engine/Source/Editor/UnrealEd/Private/Commandlets/WorldPartitionDataLayerToAssetCommandlet.cpp b/Engine/Source/Editor/UnrealEd/Private/Commandlets/WorldPartitionDataLayerToAssetCommandlet.cpp
--- a/Engine/Source/Editor/UnrealEd/Private/Commandlets/WorldPartitionDataLayerToAssetCommandlet.cpp
+++ b/Engine/Source/Editor/UnrealEd/Private/Commandlets/WorldPartitionDataLayerToAssetCommandlet.cpp
@@ -598,7 +598,6 @@ bool UDataLayerToAssetCommandlet::CreateDataLayerInstances(TStrongObjectPtr<UDat
 {
 	UE_SCOPED_TIMER(TEXT("Creating Data Layer Instances"), LogDataLayerToAssetCommandlet, Display);
 
-	uint32 ErrorCount = 0;
 	AWorldDataLayers* WorldDataLayers = MainWorld->GetWorldDataLayers();
 	for (const TWeakObjectPtr<UDataLayerConversionInfo>& ConvertingInfo : CommandletContext->GetConvertingDataLayerConversionInfo())
 	{
@@ -614,12 +613,11 @@ bool UDataLayerToAssetCommandlet::CreateDataLayerInstances(TStrongObjectPtr<UDat
 		WorldDataLayers->DeprecatedDataLayerNameToDataLayerInstance.Add(ConvertingInfo->DataLayerToConvert->GetDataLayerFName(), ConvertingInfo->DataLayerInstance);
 	}
 
-	if (ErrorCount == 0)
+	uint32 ErrorCount = 0;
+
+	if (!RebuildDataLayerHierarchies(CommandletContext))
 	{
-		if (!RebuildDataLayerHierarchies(CommandletContext))
-		{
-			ErrorCount++;
-		}
+		ErrorCount++;
 	}
 
 	return ErrorCount == 0;
diff --git a/Engine/Source/Runtime/Engine/Private/GameInstance.cpp b/Engine/Source/Runtime/Engine/Private/GameInstance.cpp
--- a/Engine/Source/Runtime/Engine/Private/GameInstance.cpp
+++ b/Engine/Source/Runtime/Engine/Private/GameInstance.cpp
@@ -910,7 +910,7 @@ int32 UGameInstance::AddLocalPlayer(ULocalPlayer* NewLocalPlayer, FPlatformUserI
 		GetGameViewportClient()->NotifyPlayerAdded(InsertIndex, NewLocalPlayer);
 	}
 
-	UE_LOG(LogPlayerManagement, Log, TEXT("UGameInstance::AddLocalPlayer: Added player %s with PlatformUserId %d at index %d (%d remaining players)"), *NewLocalPlayer->GetName(), NewLocalPlayer->GetPlatformUserId(), InsertIndex, LocalPlayers.Num());
+	UE_LOG(LogPlayerManagement, Log, TEXT("UGameInstance::AddLocalPlayer: Added player %s with PlatformUserId %d at index %d (%d remaining players)"), *NewLocalPlayer->GetName(), NewLocalPlayer->GetPlatformUserId().GetInternalId(), InsertIndex, LocalPlayers.Num());
 
 	OnLocalPlayerAddedEvent.Broadcast(NewLocalPlayer);
 
diff --git a/Engine/Source/Runtime/Engine/Private/PlayerController.cpp b/Engine/Source/Runtime/Engine/Private/PlayerController.cpp
--- a/Engine/Source/Runtime/Engine/Private/PlayerController.cpp
+++ b/Engine/Source/Runtime/Engine/Private/PlayerController.cpp
@@ -4024,19 +4024,19 @@ APlayerState* APlayerController::GetSplitscreenPlayerByIndex(int32 PlayerIndex)
 					}
 					else
 					{
-						UE_LOG(LogPlayerController, Log, TEXT("(%s) APlayerController::%s:%s: %s IS NOT THE PRIMARY CONNECTION AND HAS NO CHILD CONNECTIONS!"), *GetFName().ToString(), *GetStateName().ToString(), Player);
+						UE_LOG(LogPlayerController, Log, TEXT("(%s) APlayerController::%s:GetSplitscreenPlayerByIndex: %s IS NOT THE PRIMARY CONNECTION AND HAS NO CHILD CONNECTIONS!"), *GetName(), *GetStateName().ToString(), *Player->GetName());
 					}
 				}
 			}
 			else
 			{
-				UE_LOG(LogPlayerController, Log, TEXT("(%s) APlayerController::%s:%s: %s IS NOT A ULocalPlayer* AND NOT A RemoteConnection! (No valid UPlayer* reference)"), *GetFName().ToString(), *GetStateName().ToString(), Player);
+				UE_LOG(LogPlayerController, Log, TEXT("(%s) APlayerController::%s:GetSplitscreenPlayerByIndex: %s IS NOT A ULocalPlayer* AND NOT A RemoteConnection! (No valid UPlayer* reference)"), *GetName(), *GetStateName().ToString(), *Player->GetName());
 			}
 		}
 	}
 	else
 	{
-		UE_LOG(LogPlayerController, Log, TEXT("(%s) APlayerController::%s:%s: %s"), *GetFName().ToString(), *GetStateName().ToString(),  "NULL value for Player!");
+		UE_LOG(LogPlayerController, Log, TEXT("(%s) APlayerController::%s:GetSplitscreenPlayerByIndex: %s"), *GetName(), *GetStateName().ToString(), TEXT("NULL value for Player!"));
 	}
 
 	return Result;
@@ -4070,12 +4070,12 @@ int32 APlayerController::GetSplitscreenPlayerCount() const
 			}
 			else
 			{
-				UE_LOG(LogPlayerController, Log, TEXT("(%s) APlayerController::%s:%s NOT A ULocalPlayer* AND NOT A RemoteConnection!"), *GetFName().ToString(), *GetStateName().ToString());
+				UE_LOG(LogPlayerController, Log, TEXT("(%s) APlayerController::%s:GetSplitscreenPlayerCount NOT A ULocalPlayer* AND NOT A RemoteConnection!"), *GetName(), *GetStateName().ToString());
 			}
 		}
 		else
 		{
-			UE_LOG(LogPlayerController, Log, TEXT("(%s) APlayerController::%s:%s called without a valid UPlayer* value!"), *GetFName().ToString(), *GetStateName().ToString());
+			UE_LOG(LogPlayerController, Log, TEXT("(%s) APlayerController::%s:GetSplitscreenPlayerCount called without a valid UPlayer* value!"), *GetName(), *GetStateName().ToString());
 		}
 	}
 
