PR #10119: Fix a enum mistake in UMaterialInterface::GetFontParameterValue (Contributed by [REDACTED_PERSON])

diff --git a/Engine/Source/Runtime/Engine/Private/Materials/MaterialInterface.cpp b/Engine/Source/Runtime/Engine/Private/Materials/MaterialInterface.cpp
--- a/Engine/Source/Runtime/Engine/Private/Materials/MaterialInterface.cpp
+++ b/Engine/Source/Runtime/Engine/Private/Materials/MaterialInterface.cpp
@@ -923,7 +923,7 @@ bool UMaterialInterface::GetTextureParameterChannelNames(const FHashedMaterialPa
 bool UMaterialInterface::GetFontParameterValue(const FHashedMaterialParameterInfo& ParameterInfo, class UFont*& OutFontValue, int32& OutFontPage, bool bOveriddenOnly) const
 {
 	FMaterialParameterMetadata Result;
-	if (GetParameterValue(EMaterialParameterType::RuntimeVirtualTexture, ParameterInfo, Result, MakeParameterValueFlags(bOveriddenOnly)))
+	if (GetParameterValue(EMaterialParameterType::Font, ParameterInfo, Result, MakeParameterValueFlags(bOveriddenOnly)))
 	{
 		OutFontValue = Result.Value.Font.Value;
 		OutFontPage = Result.Value.Font.Page;
