Fix for uniform expression cache assert when changing Engine Scalability settings in editor.  Need to flush any in-flight uniform expression cache operations, when material quality changes.  The material quality level affects the result returned by calls to FMaterialRenderProxy::GetMaterialNoFallback, and in-flight operations will have been using a uniform buffer layout from the old version of the material.

#rnx

diff --git a/Engine/Source/Runtime/Engine/Private/UnrealEngine.cpp b/Engine/Source/Runtime/Engine/Private/UnrealEngine.cpp
--- a/Engine/Source/Runtime/Engine/Private/UnrealEngine.cpp
+++ b/Engine/Source/Runtime/Engine/Private/UnrealEngine.cpp
@@ -169,6 +169,7 @@ UnrealEngine.cpp: Implements the UEngine class and helpers.
 #include "Widgets/Notifications/SNotificationList.h"
 #include "Engine/UserInterfaceSettings.h"
 #include "ComponentRecreateRenderStateContext.h"
+#include "Materials/MaterialRenderProxy.h"
 
 #include "IMessageRpcClient.h"
 #include "IMessagingRpcModule.h"
@@ -766,6 +767,18 @@ void ScalabilityCVarsSinkCallback()
 		return;
 	}
 
+	if (LocalScalabilityCVars.MaterialQualityLevel != GCachedScalabilityCVars.MaterialQualityLevel)
+	{
+		// Need to flush any in-flight uniform expression cache operations, when material quality changes.  The material quality level affects
+		// the result returned by calls to FMaterialRenderProxy::GetMaterialNoFallback, and in-flight operations will have been using a uniform
+		// buffer layout from the old version of the material, leading to an assert.
+		ENQUEUE_RENDER_COMMAND(UpdateDeferredCachedUniformExpressions)(
+			[](FRHICommandListImmediate& RHICmdList)
+			{
+				FMaterialRenderProxy::UpdateDeferredCachedUniformExpressions();
+			});
+	}
+
 	FlushRenderingCommands();
 
 	if (!GCachedScalabilityCVars.bInitialized)
