Fix usage of SendMachineInfoForAccount.
Add comment about sending SessionMachineInfo with a blank account id.

diff --git a/Engine/Source/Runtime/Engine/Private/EngineAnalytics.cpp b/Engine/Source/Runtime/Engine/Private/EngineAnalytics.cpp
--- a/Engine/Source/Runtime/Engine/Private/EngineAnalytics.cpp
+++ b/Engine/Source/Runtime/Engine/Private/EngineAnalytics.cpp
@@ -303,26 +303,19 @@ void FEngineAnalytics::AppendMachineStats(TArray<FAnalyticsEventAttribute>& Even
 
 void FEngineAnalytics::SendMachineInfoForAccount(const FString& EpicAccountId)
 {
-	if (!SessionEpicAccountIds.Contains(EpicAccountId))
+	// Note: EpicAccountId may be empty when the user has not signed in to the epic games launcher.
+	// The intention here is to only send SessionMachineStats once per unique user including when
+	// no user has logged in.
+	if (Analytics && !SessionEpicAccountIds.Contains(EpicAccountId))
 	{
+		SessionEpicAccountIds.Add(EpicAccountId);
+
 		TArray<FAnalyticsEventAttribute> EventAttributes;
 		AppendMachineStats(EventAttributes);
-		SendMachineInfoForAccount(EpicAccountId, EventAttributes);
-	}
-}
-
-void FEngineAnalytics::SendMachineInfoForAccount(const FString& EpicAccountId, const TArray<FAnalyticsEventAttribute>& EventAttributes)
-{
-	// Only send SessionMachineStats when the user is changed to a user that has not been seen in the session.
-	if (!SessionEpicAccountIds.Contains(EpicAccountId))
-	{
-		SessionEpicAccountIds.Add(EpicAccountId);
 
 		// When the user id is changed send a SessionMachineStats event.
 		static const FString SZEventName = TEXT("SessionMachineStats");
-		TArray<FAnalyticsEventAttribute> Params;
-		AppendMachineStats(Params);
-		Analytics->RecordEvent(SZEventName, Params);
+		Analytics->RecordEvent(SZEventName, EventAttributes);
 	}
 }
 
diff --git a/Engine/Source/Runtime/Engine/Public/EngineAnalytics.h b/Engine/Source/Runtime/Engine/Public/EngineAnalytics.h
--- a/Engine/Source/Runtime/Engine/Public/EngineAnalytics.h
+++ b/Engine/Source/Runtime/Engine/Public/EngineAnalytics.h
@@ -61,7 +61,6 @@ public:
 private:
 	static void AppendMachineStats(TArray<FAnalyticsEventAttribute>& EventAttributes);
 	static void SendMachineInfoForAccount(const FString& EpicAccountId);
-	static void SendMachineInfoForAccount(const FString& EpicAccountId, const TArray<FAnalyticsEventAttribute>& EventAttributes);
 	static void OnEpicAccountIdChanged(const FString& EpicAccountId);
 
 	static bool bIsInitialized;
