CharacterMovementComponent: Fixed diagnostic strings label

diff --git a/Engine/Source/Runtime/Engine/Private/Components/CharacterMovementComponent.cpp b/Engine/Source/Runtime/Engine/Private/Components/CharacterMovementComponent.cpp
--- a/Engine/Source/Runtime/Engine/Private/Components/CharacterMovementComponent.cpp
+++ b/Engine/Source/Runtime/Engine/Private/Components/CharacterMovementComponent.cpp
@@ -469,17 +469,17 @@ void FCharacterMovementComponentPostPhysicsTickFunction::ExecuteTick(float Delta
 
 FString FCharacterMovementComponentPostPhysicsTickFunction::DiagnosticMessage()
 {
-	return Target->GetFullName() + TEXT("[UCharacterMovementComponent::PreClothTick]");
+	return Target->GetFullName() + TEXT("[UCharacterMovementComponent::PostPhysicsTick]");
 }
 
 FName FCharacterMovementComponentPostPhysicsTickFunction::DiagnosticContext(bool bDetailed)
 {
 	if (bDetailed)
 	{
-		return FName(*FString::Printf(TEXT("SkeletalMeshComponentClothTick/%s"), *GetFullNameSafe(Target)));
+		return FName(*FString::Printf(TEXT("CharacterMovementComponentPostPhysicsTick/%s"), *GetFullNameSafe(Target)));
 	}
 
-	return FName(TEXT("SkeletalMeshComponentClothTick"));
+	return FName(TEXT("CharacterMovementComponentPostPhysicsTick"));
 }
 
 void FCharacterMovementComponentPrePhysicsTickFunction::ExecuteTick(float DeltaTime, ELevelTick TickType, ENamedThreads::Type CurrentThread, const FGraphEventRef& MyCompletionGraphEvent)
