Deprecation cleanup

diff --git a/Engine/Source/Runtime/CoreUObject/Public/UObject/CoreNet.h b/Engine/Source/Runtime/CoreUObject/Public/UObject/CoreNet.h
--- a/Engine/Source/Runtime/CoreUObject/Public/UObject/CoreNet.h
+++ b/Engine/Source/Runtime/CoreUObject/Public/UObject/CoreNet.h
@@ -543,9 +543,6 @@ public:
 	UE_DEPRECATED(5.3, "Please use FPropertyConditions::SetActiveOverride instead.")
 	virtual void SetCustomIsActiveOverride(UObject* OwningObject, const uint16 RepIndex, const bool bIsActive) = 0;
 
-	UE_DEPRECATED(5.0, "Please use UReplaySubsystem::SetExternalDataForObject instead.")
-	virtual void SetExternalData(const uint8* Src, const int32 NumBits) = 0;
-
 	/**
 	* Used when tracking memory to gather the total size of a given instance.
 	* This should include the dynamically allocated data, as well as the classes size.
diff --git a/Engine/Source/Runtime/Engine/Classes/Engine/NetConnection.h b/Engine/Source/Runtime/Engine/Classes/Engine/NetConnection.h
--- a/Engine/Source/Runtime/Engine/Classes/Engine/NetConnection.h
+++ b/Engine/Source/Runtime/Engine/Classes/Engine/NetConnection.h
@@ -402,12 +402,11 @@ public:
 	/** Clears the actor starvation map */
 	void ResetActorsStarvedByClassTimeMap() { ActorsStarvedByClassTimeMap.Empty(); }
 
-public:
+private:
 	// Connection information.
-
-	UE_DEPRECATED(5.0, "EConnectionState State will be made private. Use GetConnectionState and SetConnectionState to access State instead.")
 	EConnectionState	State;					// State this connection is in.
 
+public:
 	ENGINE_API const EConnectionState GetConnectionState() const;
 	ENGINE_API void SetConnectionState(EConnectionState ConnectionState);
 	
@@ -1667,10 +1666,6 @@ private:
 
 	/** Out of order packet tracking/correction */
 
-	/** Stat tracking for the total number of out of order packets, for this connection */
-	UE_DEPRECATED(5.0, "Use GetTotalOutOfOrderPackets instead.")
-	int32 TotalOutOfOrderPackets = 0;
-
 	/** Stat tracking for the total number of out of order packets lost */
 	int32 TotalOutOfOrderPacketsLost = 0;
 
diff --git a/Engine/Source/Runtime/Engine/Classes/Engine/NetDriver.h b/Engine/Source/Runtime/Engine/Classes/Engine/NetDriver.h
--- a/Engine/Source/Runtime/Engine/Classes/Engine/NetDriver.h
+++ b/Engine/Source/Runtime/Engine/Classes/Engine/NetDriver.h
@@ -1033,10 +1033,6 @@ public:
 	uint32						OutPacketsLost;
 	/** Total packets lost that have been sent by the server since the net driver's creation  */
 	uint32						OutTotalPacketsLost;
-	UE_DEPRECATED(5.0, "Use GetTotalOutOfOrderPackets instead.")
-	uint32						InOutOfOrderPackets = 0;
-	UE_DEPRECATED(5.0, "OutOutOfOrderPackets is not updated and is now deprecated.")
-	uint32						OutOutOfOrderPackets = 0;
 	/** Tracks the total number of voice packets sent */
 	uint32						VoicePacketsSent;
 	/** Tracks the total number of voice bytes sent */
@@ -1887,10 +1883,6 @@ public:
 	void IncreaseTotalOutOfOrderPacketsLost(int32 Count=1)
 	{
 		TotalOutOfOrderPacketsLost += Count;
-
-		PRAGMA_DISABLE_DEPRECATION_WARNINGS
-		InOutOfOrderPackets += Count;
-		PRAGMA_ENABLE_DEPRECATION_WARNINGS
 	}
 
 	/**
diff --git a/Engine/Source/Runtime/Engine/Classes/Engine/World.h b/Engine/Source/Runtime/Engine/Classes/Engine/World.h
--- a/Engine/Source/Runtime/Engine/Classes/Engine/World.h
+++ b/Engine/Source/Runtime/Engine/Classes/Engine/World.h
@@ -3374,14 +3374,6 @@ public:
 	// Start listening for connections.
 	bool Listen( FURL& InURL );
 
-	/** @return true if this level is a client */
-	UE_DEPRECATED(5.0, "Use GetNetMode or IsNetMode instead for more accurate results.")
-	bool IsClient() const;
-
-	/** @return true if this level is a server */
-	UE_DEPRECATED(5.0, "Use GetNetMode or IsNetMode instead for more accurate results")
-	bool IsServer() const;
-
 	/** @return true if the world is in the paused state */
 	bool IsPaused() const;
 
diff --git a/Engine/Source/Runtime/Engine/Classes/GameFramework/PlayerState.h b/Engine/Source/Runtime/Engine/Classes/GameFramework/PlayerState.h
--- a/Engine/Source/Runtime/Engine/Classes/GameFramework/PlayerState.h
+++ b/Engine/Source/Runtime/Engine/Classes/GameFramework/PlayerState.h
@@ -340,9 +340,6 @@ public:
 		return CompressedPing;
 	}
 
-	UE_DEPRECATED(5.0, "Use GetPingInMilliseconds() or GetCompressedPing() instead")
-	uint8 GetPing() const { return GetCompressedPing(); }
-
 	/** Sets the value of CompressedPing without causing other side effects to this instance. */
 	void SetCompressedPing(const uint8 NewPing);
 
@@ -426,13 +423,6 @@ public:
 	UFUNCTION(BlueprintCallable, Category = "PlayerState", meta = (DisplayName = "Get Unique Net Id"))
 	FUniqueNetIdRepl BP_GetUniqueId() const;
 
-	/**
-	 * Associate an online unique id with this player
-	 * @param InUniqueId the unique id associated with this player
-	 */
-	UE_DEPRECATED(5.0, "Use SetUniqueId(FUniqueNetIdRepl) instead")
-	virtual void SetUniqueId(const FUniqueNetIdPtr& InUniqueId);
-
 	/**
 	 * Associate an online unique id with this player
 	 * @param InUniqueId the unique id associated with this player
diff --git a/Engine/Source/Runtime/Engine/Private/DemoNetDriver.cpp b/Engine/Source/Runtime/Engine/Private/DemoNetDriver.cpp
--- a/Engine/Source/Runtime/Engine/Private/DemoNetDriver.cpp
+++ b/Engine/Source/Runtime/Engine/Private/DemoNetDriver.cpp
@@ -2316,9 +2316,7 @@ bool UDemoNetDriver::ReplicatePrioritizedActor(const FActorPriority& ActorPriori
 
 		const bool bDidReplicateActor = DemoReplicateActor(Actor, Params.Connection, false);
 
-		const bool bUpdatedExternalData = ReplayHelper.UpdateExternalDataForObject(Params.Connection, Actor);
-
-		if (bDidReplicateActor || bUpdatedExternalData)
+		if (bDidReplicateActor)
 		{
 			// Choose an optimal time, we choose 70% of the actual rate to allow frequency to go up if needed
 			ActorInfo->OptimalNetUpdateDelta = FMath::Clamp(LastReplicateDelta * 0.7f, MinOptimalDelta, MaxOptimalDelta);
diff --git a/Engine/Source/Runtime/Engine/Private/NetConnection.cpp b/Engine/Source/Runtime/Engine/Private/NetConnection.cpp
--- a/Engine/Source/Runtime/Engine/Private/NetConnection.cpp
+++ b/Engine/Source/Runtime/Engine/Private/NetConnection.cpp
@@ -2851,10 +2851,6 @@ void UNetConnection::ReceivedPacket( FBitReader& Reader, bool bIsReinjectedPacke
 				}
 				else
 				{
-					PRAGMA_DISABLE_DEPRECATION_WARNINGS
-					TotalOutOfOrderPackets++;
-					PRAGMA_ENABLE_DEPRECATION_WARNINGS
-
 					TotalOutOfOrderPacketsDuplicate++;
 					AnalyticsVars.IncreaseOutOfOrderPacketsDuplicateCount();
 					Driver->IncreaseTotalOutOfOrderPacketsDuplicate();
@@ -2878,10 +2874,6 @@ void UNetConnection::ReceivedPacket( FBitReader& Reader, bool bIsReinjectedPacke
 		}
 		else
 		{
-			PRAGMA_DISABLE_DEPRECATION_WARNINGS
-			TotalOutOfOrderPackets++;
-			PRAGMA_ENABLE_DEPRECATION_WARNINGS
-
 			TotalOutOfOrderPacketsLost++;
 			AnalyticsVars.IncreaseOutOfOrderPacketsLostCount();
 			Driver->IncreaseTotalOutOfOrderPacketsLost();
diff --git a/Engine/Source/Runtime/Engine/Private/PlayerState.cpp b/Engine/Source/Runtime/Engine/Private/PlayerState.cpp
--- a/Engine/Source/Runtime/Engine/Private/PlayerState.cpp
+++ b/Engine/Source/Runtime/Engine/Private/PlayerState.cpp
@@ -505,13 +505,6 @@ FUniqueNetIdRepl APlayerState::BP_GetUniqueId() const
 	return GetUniqueId();
 }
 
-void APlayerState::SetUniqueId(const FUniqueNetIdPtr& InUniqueId)
-{
-	MARK_PROPERTY_DIRTY_FROM_NAME(APlayerState, UniqueId, this);
-	UniqueId.SetUniqueNetId(InUniqueId);
-	OnSetUniqueId();
-}
-
 void APlayerState::SetUniqueId(const FUniqueNetIdRepl& NewUniqueId)
 {
 	MARK_PROPERTY_DIRTY_FROM_NAME(APlayerState, UniqueId, this);
diff --git a/Engine/Source/Runtime/Engine/Private/RepLayout.cpp b/Engine/Source/Runtime/Engine/Private/RepLayout.cpp
--- a/Engine/Source/Runtime/Engine/Private/RepLayout.cpp
+++ b/Engine/Source/Runtime/Engine/Private/RepLayout.cpp
@@ -2668,21 +2668,6 @@ void FRepSerializationSharedInfo::CountBytes(FArchive& Ar) const
 	);
 }
 
-PRAGMA_DISABLE_DEPRECATION_WARNINGS
-const FRepSerializedPropertyInfo* FRepSerializationSharedInfo::WriteSharedProperty(
-	const FRepLayoutCmd& Cmd,
-	const FGuid& PropertyGuid,
-	const int32 CmdIndex,
-	const uint16 Handle,
-	const FConstRepObjectDataBuffer Data,
-	const bool bWriteHandle,
-	const bool bDoChecksum)
-{
-	FRepSharedPropertyKey PropertyKey(PropertyGuid.A, PropertyGuid.B, PropertyGuid.C, (void*)Data.Data);
-	return WriteSharedProperty(Cmd, PropertyKey, CmdIndex, Handle, Data, bWriteHandle, bDoChecksum);
-}
-PRAGMA_ENABLE_DEPRECATION_WARNINGS
-
 const FRepSerializedPropertyInfo* FRepSerializationSharedInfo::WriteSharedProperty(
 	const FRepLayoutCmd& Cmd,
 	const FRepSharedPropertyKey& PropertyKey,
diff --git a/Engine/Source/Runtime/Engine/Private/ReplayHelper.cpp b/Engine/Source/Runtime/Engine/Private/ReplayHelper.cpp
--- a/Engine/Source/Runtime/Engine/Private/ReplayHelper.cpp
+++ b/Engine/Source/Runtime/Engine/Private/ReplayHelper.cpp
@@ -1461,29 +1461,6 @@ void FReplayHelper::SaveExternalData(UNetConnection* Connection, FArchive& Ar)
 					}
 				}
 			}
-			else
-			{
-				PRAGMA_DISABLE_DEPRECATION_WARNINGS
-				TSharedPtr<FRepChangedPropertyTracker> PropertyTracker = Connection->Driver->FindRepChangedPropertyTracker(Object);
-				check(PropertyTracker.IsValid());
-
-				uint32 ExternalDataNumBits = PropertyTracker->ExternalDataNumBits;
-				if (ExternalDataNumBits > 0)
-				{
-					// Save payload size (in bits)
-					Ar.SerializeIntPacked(ExternalDataNumBits);
-
-					// Save GUID
-					Ar << Pair.Value;
-
-					// Save payload
-					Ar.Serialize(PropertyTracker->ExternalData.GetData(), PropertyTracker->ExternalData.Num());
-
-					PropertyTracker->ExternalData.Empty();
-					PropertyTracker->ExternalDataNumBits = 0;
-				}
-				PRAGMA_ENABLE_DEPRECATION_WARNINGS
-			}
 		}
 	}
 
@@ -1595,8 +1572,6 @@ bool FReplayHelper::ReplicateCheckpointActor(AActor* ToReplicate, UNetConnection
 			ActorChannel->ConditionalCleanUp(false, EChannelCloseReason::Dormancy);
 		}
 
-		UpdateExternalDataForObject(Connection, ToReplicate);
-
 		const double CheckpointTime = FPlatformTime::Seconds();
 
 		if (Params.CheckpointMaxUploadTimePerFrame > 0 && CheckpointTime - Params.StartCheckpointTime > Params.CheckpointMaxUploadTimePerFrame)
@@ -1643,41 +1618,6 @@ void FReplayHelper::LoadExternalData(FArchive& Ar, const float TimeSeconds)
 	}
 }
 
-PRAGMA_DISABLE_DEPRECATION_WARNINGS
-bool FReplayHelper::UpdateExternalDataForObject(UNetConnection* Connection, UObject* OwningObject)
-{
-	check(Connection && Connection->Driver);
-
-	TSharedPtr<FRepChangedPropertyTracker> PropertyTracker = Connection->Driver->FindRepChangedPropertyTracker(OwningObject);
-	if (!PropertyTracker.IsValid())
-	{
-		return false;
-	}
-
-	if (PropertyTracker->ExternalData.Num() == 0)
-	{
-		return false;
-	}
-
-	if (FNetworkGUID* NetworkGUID = Connection->Driver->GuidCache->NetGUIDLookup.Find(OwningObject))
-	{
-		ObjectsWithExternalDataMap.Add(OwningObject, *NetworkGUID);
-
-		return true;
-	}
-	else
-	{
-		UE_LOG(LogDemo, Warning, TEXT("UpdateExternalDataForObject: Discarding external data for object with no net guid: %s"), *GetNameSafe(OwningObject));
-
-		// Clear external data if the actor has never replicated yet (and doesn't have a net guid)
-		PropertyTracker->ExternalData.Reset();
-		PropertyTracker->ExternalDataNumBits = 0;
-
-		return false;
-	}
-}
-PRAGMA_ENABLE_DEPRECATION_WARNINGS
-
 // only allow one chunk of data per object at a time (first wins), when recording ticks it will store/clear it
 bool FReplayHelper::SetExternalDataForObject(UNetConnection* Connection, UObject* OwningObject, const uint8* Src, const int32 NumBits)
 {
diff --git a/Engine/Source/Runtime/Engine/Private/ReplayNetConnection.cpp b/Engine/Source/Runtime/Engine/Private/ReplayNetConnection.cpp
--- a/Engine/Source/Runtime/Engine/Private/ReplayNetConnection.cpp
+++ b/Engine/Source/Runtime/Engine/Private/ReplayNetConnection.cpp
@@ -134,13 +134,6 @@ void UReplayNetConnection::LowLevelSend(void* Data, int32 CountBits, FOutPacketT
 
 		if (AActor* Actor = GetRepContextActor())
 		{
-			//@todo: do we still call this during checkpoints?
-			if (!Actor->IsPendingKillPending())
-			{
-				//@todo: unique this in tick?
-				ReplayHelper.UpdateExternalDataForObject(this, Actor);
-			}
-
 			if (!bCheckpoint && ReplayHelper.bHasDeltaCheckpoints && Driver)
 			{
 				Driver->GetNetworkObjectList().MarkDirtyForReplay(Actor);
diff --git a/Engine/Source/Runtime/Engine/Private/World.cpp b/Engine/Source/Runtime/Engine/Private/World.cpp
--- a/Engine/Source/Runtime/Engine/Private/World.cpp
+++ b/Engine/Source/Runtime/Engine/Private/World.cpp
@@ -6464,29 +6464,6 @@ bool UWorld::Listen( FURL& InURL )
 #endif // WITH_SERVER_CODE
 }
 
-bool UWorld::IsClient() const
-{
-	return GIsClient;
-}
-
-bool UWorld::IsServer() const
-{
-	if ( NetDriver != NULL )
-	{
-		return NetDriver->IsServer();
-	}
-
-	PRAGMA_DISABLE_DEPRECATION_WARNINGS
-	if ( DemoNetDriver != NULL )
-	{
-		return DemoNetDriver->IsServer();
-	}
-	PRAGMA_ENABLE_DEPRECATION_WARNINGS
-
-	return true;
-}
-
-PRAGMA_DISABLE_DEPRECATION_WARNINGS
 bool UWorld::IsPlayingReplay() const
 {
 	return (DemoNetDriver && DemoNetDriver->IsPlaying());
@@ -6507,7 +6484,6 @@ bool UWorld::IsRecordingReplay() const
 		return (DemoNetDriver && DemoNetDriver->IsServer());
 	}
 }
-PRAGMA_ENABLE_DEPRECATION_WARNINGS
 
 void UWorld::PrepareMapChange(const TArray<FName>& LevelNames)
 {
diff --git a/Engine/Source/Runtime/Engine/Public/Net/RepLayout.h b/Engine/Source/Runtime/Engine/Public/Net/RepLayout.h
--- a/Engine/Source/Runtime/Engine/Public/Net/RepLayout.h
+++ b/Engine/Source/Runtime/Engine/Public/Net/RepLayout.h
@@ -119,29 +119,6 @@ typedef TRepDataBuffer<ERepDataBufferType::ShadowBuffer> FRepShadowDataBuffer;
 typedef TConstRepDataBuffer<ERepDataBufferType::ObjectBuffer> FConstRepObjectDataBuffer;
 typedef TConstRepDataBuffer<ERepDataBufferType::ShadowBuffer> FConstRepShadowDataBuffer;
 
-/** Stores meta data about a given Replicated property. */
-class UE_DEPRECATED(5.0, "No longer used") FRepChangedParent
-{
-public:
-	FRepChangedParent():
-		Active(1),
-		OldActive(1),
-		IsConditional(0)
-	{}
-
-	/** Whether or not this property is currently Active (i.e., considered for replication). */
-	uint32 Active: 1;
-
-	/** The last updated state of Active, used to track when the Active state changes. */
-	uint32 OldActive: 1;
-
-	/**
-	 * Whether or not this property has conditions that may exclude it from replicating to a given connection.
-	 * @see FRepState::ConditionMap.
-	 */
-	uint32 IsConditional: 1;
-};
-
 /** FRepChangedPropertyTracker moved to NetCore module */
 
 class FRepLayout;
@@ -198,10 +175,6 @@ struct FRepSerializedPropertyInfo
 		PropBitLength(0)
 	{}
 
-	/** Unique identifier for this property, may include array index and depth. */
-	UE_DEPRECATED(4.27, "No longer used, please use PropertyKey instead.")
-	FGuid Guid;
-
 	/** Unique identifier for this property */
 	FRepSharedPropertyKey PropertyKey;
 
@@ -254,27 +227,6 @@ struct FRepSerializationSharedInfo
 		}
 	}
 
-	/**
-	 * Creates a new SharedPropertyInfo and adds it to the SharedPropertyInfo list.
-	 *
-	 * @param Cmd				The command that represents the property we want to share.
-	 * @param PropertyGuid		A guid used to identify the property.
-	 * @param CmdIndex			Index of the property command. Only used if bDoChecksum is true.
-	 * @param Handle			Relative Handle of the property command. Only used if bWriteHandle is true.
-	 * @param Data				Pointer to the raw property memory that will be serialized.
-	 * @param bWriteHandle		Whether or not we should write Command handles into the serialized data.
-	 * @param bDoChecksum		Whether or not we should do checksums. Only used if ENABLE_PROPERTY_CHECKSUMS is enabled.
-	 */
-	UE_DEPRECATED(4.27, "Now passing a FRepSharedPropertyKey instead of FGuid")
-	const FRepSerializedPropertyInfo* WriteSharedProperty(
-		const FRepLayoutCmd& Cmd,
-		const FGuid& PropertyGuid,
-		const int32 CmdIndex,
-		const uint16 Handle,
-		const FConstRepObjectDataBuffer Data,
-		const bool bWriteHandle,
-		const bool bDoChecksum);
-
 	/**
 	 * Creates a new SharedPropertyInfo and adds it to the SharedPropertyInfo list.
 	 *
diff --git a/Engine/Source/Runtime/Engine/Public/ReplayHelper.h b/Engine/Source/Runtime/Engine/Public/ReplayHelper.h
--- a/Engine/Source/Runtime/Engine/Public/ReplayHelper.h
+++ b/Engine/Source/Runtime/Engine/Public/ReplayHelper.h
@@ -199,7 +199,6 @@ private:
 	void SaveExternalData(UNetConnection* Connection, FArchive& Ar);
 	void LoadExternalData(FArchive& Ar, const float TimeSeconds);
 
-	bool UpdateExternalDataForObject(UNetConnection* Connection, UObject* OwningObject);
 	bool SetExternalDataForObject(UNetConnection* Connection, UObject* OwningObject, const uint8* Src, const int32 NumBits);
 
 	// Cached replay URL
diff --git a/Engine/Source/Runtime/Net/Core/Private/Net/Core/PropertyConditions/RepChangedPropertyTracker.cpp b/Engine/Source/Runtime/Net/Core/Private/Net/Core/PropertyConditions/RepChangedPropertyTracker.cpp
--- a/Engine/Source/Runtime/Net/Core/Private/Net/Core/PropertyConditions/RepChangedPropertyTracker.cpp
+++ b/Engine/Source/Runtime/Net/Core/Private/Net/Core/PropertyConditions/RepChangedPropertyTracker.cpp
@@ -4,17 +4,14 @@
 #include "Net/Core/PropertyConditions/PropertyConditionsDelegates.h"
 #include "Net/Core/PushModel/PushModel.h"
 
-PRAGMA_DISABLE_DEPRECATION_WARNINGS
 FRepChangedPropertyTracker::FRepChangedPropertyTracker(FCustomPropertyConditionState&& InActiveState)
 	: ActiveState(MoveTemp(InActiveState))
-	, ExternalDataNumBits(0)
 {
 }
 
 FRepChangedPropertyTracker::FRepChangedPropertyTracker(const bool InbIsReplay, const bool InbIsClientReplayRecording)
 	: FRepChangedPropertyTracker(0)
 {}
-PRAGMA_ENABLE_DEPRECATION_WARNINGS
 
 void FRepChangedPropertyTracker::SetCustomIsActiveOverride(UObject* OwningObject, const uint16 RepIndex, const bool bIsActive)
 {
@@ -61,24 +58,9 @@ void FRepChangedPropertyTracker::CallSetDynamicCondition(const UObject* OwningOb
 #endif
 }
 
-PRAGMA_DISABLE_DEPRECATION_WARNINGS
-void FRepChangedPropertyTracker::SetExternalData(const uint8* Src, const int32 NumBits)
-{
-	ExternalDataNumBits = NumBits;
-	const int32 NumBytes = (NumBits + 7) >> 3;
-	ExternalData.Reset(NumBytes);
-	ExternalData.AddUninitialized(NumBytes);
-	FMemory::Memcpy(ExternalData.GetData(), Src, NumBytes);
-}
-PRAGMA_ENABLE_DEPRECATION_WARNINGS
-
 void FRepChangedPropertyTracker::CountBytes(FArchive& Ar) const
 {
 	// Include our size here, because the caller won't know.
 	Ar.CountBytes(sizeof(FRepChangedPropertyTracker), sizeof(FRepChangedPropertyTracker));
 	ActiveState.CountBytes(Ar);
-
-	PRAGMA_DISABLE_DEPRECATION_WARNINGS
-	ExternalData.CountBytes(Ar);
-	PRAGMA_ENABLE_DEPRECATION_WARNINGS
 }
diff --git a/Engine/Source/Runtime/Net/Core/Public/Net/Core/PropertyConditions/RepChangedPropertyTracker.h b/Engine/Source/Runtime/Net/Core/Public/Net/Core/PropertyConditions/RepChangedPropertyTracker.h
--- a/Engine/Source/Runtime/Net/Core/Public/Net/Core/PropertyConditions/RepChangedPropertyTracker.h
+++ b/Engine/Source/Runtime/Net/Core/Public/Net/Core/PropertyConditions/RepChangedPropertyTracker.h
@@ -28,9 +28,7 @@ public:
 	UE_DEPRECATED(5.1, "Replay arguments no longer used.")
 	FRepChangedPropertyTracker(const bool InbIsReplay, const bool InbIsClientReplayRecording);
 
-	PRAGMA_DISABLE_DEPRECATION_WARNINGS
 	virtual ~FRepChangedPropertyTracker() = default;
-	PRAGMA_ENABLE_DEPRECATION_WARNINGS
 
 	//~ Begin IRepChangedPropertyTracker Interface.
 	/**
@@ -46,17 +44,6 @@ public:
 	UE_DEPRECATED(5.2, "Please use FNetPropertyConditionManager::SetPropertyActiveOverride instead.")
 	virtual void SetCustomIsActiveOverride(UObject* OwningObject, const uint16 RepIndex, const bool bIsActive) override;
 
-	/**
-	 * Sets (or resets) the External Data.
-	 * External Data is primarily used for Replays, and is used to track additional non-replicated
-	 * data or state about an object.
-	 *
-	 * @param Src		Memory containing the external data.
-	 * @param NumBits	Size of the memory, in bits.
-	 */
-	UE_DEPRECATED(5.0, "Please use UReplaySubsystem::SetExternalDataForObject instead.")
-	virtual void SetExternalData(const uint8* Src, const int32 NumBits) override;
-
 	virtual void CountBytes(FArchive& Ar) const override;
 	//~ End IRepChangedPropertyTracker Interface
 
@@ -91,11 +78,4 @@ private:
 
 	/** Activation data for top level Properties on the given Actor / Object. */
 	FCustomPropertyConditionState ActiveState;
-
-public:
-	UE_DEPRECATED(5.0, "No longer used, see UReplaySubsystem::SetExternalDataForObject")
-	TArray<uint8> ExternalData;
-
-	UE_DEPRECATED(5.0, "No longer used, see UReplaySubsystem::SetExternalDataForObject")
-	uint32 ExternalDataNumBits;
 };
