World Partition Helpers:
- Avoid triggering an initial GC when calling FWorldPartitionHelpers::HasExceededMaxMemory for the first time long after a GC.
- Bring back commandlets behavior that force collecting garbage for standalone objects.

diff --git a/Engine/Source/Runtime/Engine/Private/WorldPartition/WorldPartitionHelpers.cpp b/Engine/Source/Runtime/Engine/Private/WorldPartition/WorldPartitionHelpers.cpp
--- a/Engine/Source/Runtime/Engine/Private/WorldPartition/WorldPartitionHelpers.cpp
+++ b/Engine/Source/Runtime/Engine/Private/WorldPartition/WorldPartitionHelpers.cpp
@@ -234,6 +234,15 @@ void FWorldPartitionHelpers::ForEachActorWithLoading(UWorldPartition* WorldParti
 
 bool FWorldPartitionHelpers::HasExceededMaxMemory()
 {
+	// Even if we're not exhausting memory, GC should be run at periodic intervals
+	static double LastTime = FPlatformTime::Seconds();
+	double ThisTime = FPlatformTime::Seconds();
+	if ((ThisTime - LastTime) > 30)
+	{
+		LastTime = ThisTime;
+		return true;
+	}
+
 	const FPlatformMemoryStats MemStats = FPlatformMemory::GetStats();
 
 	const uint64 MemoryMinFreePhysical = 1llu * 1024 * 1024 * 1024;
@@ -243,14 +252,13 @@ bool FWorldPartitionHelpers::HasExceededMaxMemory()
 	const bool bHasExceededMaxUsedPhysical = MemStats.UsedPhysical >= MemoryMaxUsedPhysical;
 	const bool bHasExceededMaxMemory = bHasExceededMinFreePhysical || bHasExceededMaxUsedPhysical;
 
-	// Even if we're not exhausting memory, GC should be run at periodic intervals
-	return bHasExceededMaxMemory || (FPlatformTime::Seconds() - GetLastGCTime()) > 30;
+	return bHasExceededMaxMemory;
 };
 
 void FWorldPartitionHelpers::DoCollectGarbage()
 {
 	const FPlatformMemoryStats MemStatsBefore = FPlatformMemory::GetStats();
-	CollectGarbage(GARBAGE_COLLECTION_KEEPFLAGS, true);
+	CollectGarbage(IsRunningCommandlet() ? RF_NoFlags : GARBAGE_COLLECTION_KEEPFLAGS, true);
 	const FPlatformMemoryStats MemStatsAfter = FPlatformMemory::GetStats();
 
 	UE_LOG(LogWorldPartition, Log, TEXT("GC Performed - Available Physical: %.2fGB, Available Virtual: %.2fGB"),
