fix Platform Texture name remaps should be done before Conditional Prefix
this way having the OODLE_ or TFO_ prefix enabled doesn't break ETC or ASTC remaps in the platform
fix TextureFormatRemovePrefixFromName incorrectly assuming format names don't have underscores in them
Platform remap config no longer needs prefixed versions
this also turns on Oodle for Android DXT flavor

diff --git a/Engine/Config/Android/AndroidEngine.ini b/Engine/Config/Android/AndroidEngine.ini
--- a/Engine/Config/Android/AndroidEngine.ini
+++ b/Engine/Config/Android/AndroidEngine.ini
@@ -97,8 +97,9 @@ r.DistanceFields=0
 r.TemporalAA.Mobile.UseCompute=0
 
 [AlternateTextureCompression]
-TextureCompressionFormat=""
-TextureFormatPrefix=""
+;disabling TextureCompressionFormat for Android/IOS no longer needed
+;TextureCompressionFormat=""
+;TextureFormatPrefix=""
 
 [SystemSettings]
 ;Perplatform to PerQualityLevel conversion mapping for platform groups
diff --git a/Engine/Config/IOS/IOSEngine.ini b/Engine/Config/IOS/IOSEngine.ini
--- a/Engine/Config/IOS/IOSEngine.ini
+++ b/Engine/Config/IOS/IOSEngine.ini
@@ -87,5 +87,6 @@ r.DistanceFields=0
 r.TemporalAA.Mobile.UseCompute=0
 
 [AlternateTextureCompression]
-TextureCompressionFormat=""
-TextureFormatPrefix=""
+;disabling TextureCompressionFormat for Android/IOS no longer needed
+;TextureCompressionFormat=""
+;TextureFormatPrefix=""
diff --git a/Engine/Source/Developer/Android/AndroidTargetPlatform/Private/AndroidTargetPlatform.cpp b/Engine/Source/Developer/Android/AndroidTargetPlatform/Private/AndroidTargetPlatform.cpp
--- a/Engine/Source/Developer/Android/AndroidTargetPlatform/Private/AndroidTargetPlatform.cpp
+++ b/Engine/Source/Developer/Android/AndroidTargetPlatform/Private/AndroidTargetPlatform.cpp
@@ -490,11 +490,16 @@ void FAndroidTargetPlatform::GetTextureFormats( const UTexture* Texture, TArray<
 
 	for (FName& TextureFormatName : LayerFormats)
 	{
+		// @todo Oodle: this should not be here
+		//	should be in GetDefaultTextureFormatNamePerLayer
+		//	so that 4x4 checks can be applied correctly, etc.
 		if (Texture->LODGroup == TEXTUREGROUP_Terrain_Weightmap && bCompressLandscapeWeightMaps)
 		{
 			TextureFormatName = AndroidTexFormat::NameAutoDXT;
 		}
 		
+		// @todo : IsA then CastChecked ?  Just Cast<> ?
+		//	also use ETextureClass instead
 		if (Texture->IsA(UTextureCube::StaticClass()))
 		{
 			const UTextureCube* Cube = CastChecked<UTextureCube>(Texture);
diff --git a/Engine/Source/Developer/IOS/IOSTargetPlatform/Private/IOSTargetPlatform.cpp b/Engine/Source/Developer/IOS/IOSTargetPlatform/Private/IOSTargetPlatform.cpp
--- a/Engine/Source/Developer/IOS/IOSTargetPlatform/Private/IOSTargetPlatform.cpp
+++ b/Engine/Source/Developer/IOS/IOSTargetPlatform/Private/IOSTargetPlatform.cpp
@@ -586,6 +586,9 @@ void FIOSTargetPlatform::GetTextureFormats( const UTexture* Texture, TArray< TAr
 	TextureFormatNames.Reserve(NumLayers);
 
 	// optionaly compress landscape weightmaps for a mobile rendering
+	// @todo Oodle: this should not be here
+	//	should be in GetDefaultTextureFormatNamePerLayer
+	//	so that 4x4 checks can be applied correctly, etc.
 	static const auto CompressLandscapeWeightMapsVar = IConsoleManager::Get().FindTConsoleVariableDataInt(TEXT("r.Mobile.CompressLandscapeWeightMaps"));
 	static const bool bCompressLandscapeWeightMaps = (CompressLandscapeWeightMapsVar && CompressLandscapeWeightMapsVar->GetValueOnAnyThread() != 0);
 	if (Texture->LODGroup == TEXTUREGROUP_Terrain_Weightmap && bCompressLandscapeWeightMaps)
diff --git a/Engine/Source/Developer/TextureBuildUtilities/Private/TextureBuildUtilities.cpp b/Engine/Source/Developer/TextureBuildUtilities/Private/TextureBuildUtilities.cpp
--- a/Engine/Source/Developer/TextureBuildUtilities/Private/TextureBuildUtilities.cpp
+++ b/Engine/Source/Developer/TextureBuildUtilities/Private/TextureBuildUtilities.cpp
@@ -5,6 +5,7 @@
 #include "Misc/DataDrivenPlatformInfoRegistry.h"
 #include "Serialization/CompactBinary.h"
 #include "Serialization/CompactBinaryWriter.h"
+//#include "EngineLogs.h" // can't use from SCW
 
 
 namespace UE
@@ -32,7 +33,7 @@ TEXTUREBUILDUTILITIES_API bool TextureFormatIsHdr(FName const& InName)
 	return false;
 }
 
-TEXTUREBUILDUTILITIES_API const FName TextureFormatRemovePrefixFromName(FName const& InName, FName& OutPrefix)
+TEXTUREBUILDUTILITIES_API const FName TextureFormatRemovePlatformPrefixFromName(FName const& InName)
 {
 	FString NameString = InName.ToString();
 
@@ -45,6 +46,13 @@ TEXTUREBUILDUTILITIES_API const FName TextureFormatRemovePrefixFromName(FName co
 	// Then we detect a non-platform prefix (such as codec name)
 	// and split the result into  explicit FORMAT and PREFIX parts.
 
+	// fast(ish) early out if there are no underscores in InName :
+	int32 UnderscoreIndexIgnored = INDEX_NONE;
+	if ( ! NameString.FindChar(TCHAR('_'), UnderscoreIndexIgnored))
+	{
+		return InName;
+	}
+
 	for (FName PlatformName : FDataDrivenPlatformInfoRegistry::GetSortedPlatformNames(EPlatformInfoType::AllPlatformInfos))
 	{
 		FString PlatformTextureFormatPrefix = PlatformName.ToString();
@@ -52,25 +60,57 @@ TEXTUREBUILDUTILITIES_API const FName TextureFormatRemovePrefixFromName(FName co
 		if (NameString.StartsWith(PlatformTextureFormatPrefix, ESearchCase::IgnoreCase))
 		{
 			// Remove platform prefix and proceed with non-platform prefix detection.
-			NameString = NameString.RightChop(PlatformTextureFormatPrefix.Len());
-			break;
+			FString PlatformRemoved = NameString.RightChop(PlatformTextureFormatPrefix.Len());
+			return FName( PlatformRemoved );
 		}
 	}
+	
+	return InName;
+}
+	
+TEXTUREBUILDUTILITIES_API const FName TextureFormatRemovePrefixFromName(FName const& InNameWithPlatform, FName& OutPrefix)
+{
+	// first remove platform prefix :
+	FName NameWithoutPlatform = TextureFormatRemovePlatformPrefixFromName( InNameWithPlatform );
+	FString NameString = NameWithoutPlatform.ToString();
 
+	// then see if there's another underscore separated prefix :
 	int32 UnderscoreIndex = INDEX_NONE;
-	if (NameString.FindChar(TCHAR('_'), UnderscoreIndex))
+	if ( ! NameString.FindChar(TCHAR('_'), UnderscoreIndex))
 	{
-		// Non-platform prefix, we want to keep these
-		OutPrefix = *NameString.Left(UnderscoreIndex + 1);
-		return *NameString.RightChop(UnderscoreIndex + 1);
+		return FName( NameString );
 	}
 
-	return *NameString;
+	// texture format names can have underscores in them (eg. ETC2_RG11)
+	//	so need to differentiate between that and a conditional prefix :
+
+	// found an underscore; is it a composite texture name, or an "Alternate" prefix?
+	FString Prefix = NameString.Left(UnderscoreIndex + 1);
+	if ( Prefix == "OODLE_" || Prefix == "TFO_" )
+	{
+		// Alternate prefix
+		OutPrefix = FName( Prefix );
+		return FName( NameString.RightChop(UnderscoreIndex + 1) );
+	}
+	else if ( Prefix == "ASTC_" || Prefix == "ETC2_" )
+	{
+		// composite format, don't split
+		return NameWithoutPlatform;
+	}
+	else
+	{
+		// prefix not recognized
+		// LogTexture doesn't exist in SCW
+		UE_LOG(LogCore,Warning,TEXT("Texture Format Prefix not recognized: %s [%s]"),*Prefix,*InNameWithPlatform.ToString());
+		
+		return NameWithoutPlatform;
+	}
 }
 
 
 TEXTUREBUILDUTILITIES_API ERawImageFormat::Type GetVirtualTextureBuildIntermediateFormat(const FTextureBuildSettings& BuildSettings)
 {
+	// Platform prefix should have already been removed, also remove any Oodle prefix:
 	const FName TextureFormatName = TextureFormatRemovePrefixFromName(BuildSettings.TextureFormatName);
 
 	// note: using RGBA16F when the Source is HDR but the output is not HDR is not needed
diff --git a/Engine/Source/Developer/TextureBuildUtilities/Public/TextureBuildUtilities.h b/Engine/Source/Developer/TextureBuildUtilities/Public/TextureBuildUtilities.h
--- a/Engine/Source/Developer/TextureBuildUtilities/Public/TextureBuildUtilities.h
+++ b/Engine/Source/Developer/TextureBuildUtilities/Public/TextureBuildUtilities.h
@@ -65,12 +65,17 @@ TEXTUREBUILDUTILITIES_API bool TextureFormatIsHdr(FName const& InName);
 // i.e. PLAT_BLAH_AutoDXT returns AutoDXT and writes BLAH_ to OutPrefix.
 TEXTUREBUILDUTILITIES_API const FName TextureFormatRemovePrefixFromName(FName const& InName, FName& OutPrefix);
 
+// removes platform prefix but leaves other custom prefixes :
+TEXTUREBUILDUTILITIES_API const FName TextureFormatRemovePlatformPrefixFromName(FName const& InName);
+
 FORCEINLINE const FName TextureFormatRemovePrefixFromName(FName const& InName)
 {
 	FName OutPrefix;
 	return TextureFormatRemovePrefixFromName(InName,OutPrefix);
 }
 
+// Get the format to use for output of the VT Intermediate stage, cutting into tiles and processing
+//	  the next step will then encode from this format to the desired output format
 TEXTUREBUILDUTILITIES_API ERawImageFormat::Type GetVirtualTextureBuildIntermediateFormat(const FTextureBuildSettings& BuildSettings);
 
 } // namespace TextureBuildUtilities
diff --git a/Engine/Source/Runtime/Engine/Classes/Engine/Texture.h b/Engine/Source/Runtime/Engine/Classes/Engine/Texture.h
--- a/Engine/Source/Runtime/Engine/Classes/Engine/Texture.h
+++ b/Engine/Source/Runtime/Engine/Classes/Engine/Texture.h
@@ -1613,6 +1613,13 @@ public:
 	*  You should do this any time the texture is modified, such as on reimport, since the bits are changing anyway.
 	*/
 	ENGINE_API virtual void UpdateOodleTextureSdkVersionToLatest(void);
+	
+	/** Get TextureFormatName with platform remaps and conditional prefix
+	 *   this is the entry point API for getting the final texture format name
+	 *  OutFormats will be resized to the number of sub-flavors of the platform (typically just 1)
+	 *  OutFormats[i] gets an array of format names, one per layer
+	 */
+	ENGINE_API void GetPlatformTextureFormatNamesWithPrefix(const class ITargetPlatform* TargetPlatform,TArray< TArray<FName> >& OutFormats) const;
 #endif // WITH_EDITOR
 
 	ENGINE_API EGammaSpace GetGammaSpace() const
diff --git a/Engine/Source/Runtime/Engine/Private/Texture.cpp b/Engine/Source/Runtime/Engine/Private/Texture.cpp
--- a/Engine/Source/Runtime/Engine/Private/Texture.cpp
+++ b/Engine/Source/Runtime/Engine/Private/Texture.cpp
@@ -40,6 +40,7 @@
 #include "DerivedDataBuildVersion.h"
 #include "GuardedInt.h"
 #include "TextureCompiler.h"
+#include "TextureBuildUtilities.h"
 #include "Misc/ScopeRWLock.h"
 #endif
 
@@ -3160,23 +3161,25 @@ static FName ConditionalGetPrefixedFormat(FName TextureFormatName, const ITarget
 					
 			if ( TextureFormat )
 			{
-			FString FormatPrefix = TextureFormat->GetAlternateTextureFormatPrefix();
-			check( ! FormatPrefix.IsEmpty() );
+				FString FormatPrefix = TextureFormat->GetAlternateTextureFormatPrefix();
+				check( ! FormatPrefix.IsEmpty() );
 			
-			FName NewFormatName(FormatPrefix + TextureFormatName.ToString());
+				FName NewFormatName(FormatPrefix + TextureFormatName.ToString());
 
-			TArray<FName> SupportedFormats;
-			TextureFormat->GetSupportedFormats(SupportedFormats);
+				// check that prefixed name is one we support
+				// only apply prefix if it is in list
+				TArray<FName> SupportedFormats;
+				TextureFormat->GetSupportedFormats(SupportedFormats);
 
-			if (SupportedFormats.Contains(NewFormatName))
-			{
-				return NewFormatName;
+				if (SupportedFormats.Contains(NewFormatName))
+				{
+					return NewFormatName;
+				}
 			}
-		}
-		else
-		{
-			UE_CALL_ONCE( [&](){
-					UE_LOG(LogTexture, Warning, TEXT("AlternateTextureCompression specified, Module found, but no TextureFormat : %s."), *TextureCompressionFormat);
+			else
+			{
+				UE_CALL_ONCE( [&](){
+						UE_LOG(LogTexture, Warning, TEXT("AlternateTextureCompression specified, Module found, but no TextureFormat : %s."), *TextureCompressionFormat);
 				} );
 			}
 		}
@@ -3192,11 +3195,12 @@ static FName ConditionalGetPrefixedFormat(FName TextureFormatName, const ITarget
 	return TextureFormatName;
 }
 
+// this should not be called directly; it is called from TargetPlatform GetTextureFormats
+//	entry point API is GetPlatformTextureFormatNamesWithPrefix
 FName GetDefaultTextureFormatName( const ITargetPlatform* TargetPlatform, const UTexture* Texture, int32 LayerIndex, 
 	bool bSupportCompressedVolumeTexture, int32 Unused_BlockSize, bool bSupportFilteredFloat32Textures )
 {
 	FName TextureFormatName = NAME_None;
-	bool bOodleTextureSdkVersionIsNone = true;
 
 	/**
 	 * IF you add a format to this function don't forget to update GetAllDefaultTextureFormats 
@@ -3273,7 +3277,7 @@ FName GetDefaultTextureFormatName( const ITargetPlatform* TargetPlatform, const
 			SizeX = SizeY = FMath::Max(1 << FMath::FloorLog2(SizeX / 2), 32);
 		}
 
-		//we need to really have the actual top mip size of output platformdata
+		//@todo Oodle: we need to really have the actual top mip size of output platformdata
 		//	(hence the LODBias check below)
 		// trying to reproduce here exactly what TextureCompressor + serialization will do = brittle
 
@@ -3505,13 +3509,77 @@ FName GetDefaultTextureFormatName( const ITargetPlatform* TargetPlatform, const
 		}
 	}
 
-	bOodleTextureSdkVersionIsNone = Texture->OodleTextureSdkVersion.IsNone();
 #endif //WITH_EDITOR
 
-	FName Result = ConditionalGetPrefixedFormat(TextureFormatName, TargetPlatform, bOodleTextureSdkVersionIsNone);
+	//FName Result = ConditionalGetPrefixedFormat(TextureFormatName, TargetPlatform, bOodleTextureSdkVersionIsNone);
+	// changed : now NOT prefixed here, so that targetplatform remaps receive format without prefix
+	//	will be done by GetPlatformTextureFormatNamesWithPrefix, after target platform remaps
 
-	return Result;
+	return TextureFormatName;
+}
+
+#if WITH_EDITOR
+// this is the entry point API to get the texture format name
+void UTexture::GetPlatformTextureFormatNamesWithPrefix(const class ITargetPlatform* TargetPlatform,TArray< TArray<FName> >& OutFormats) const
+{
+	check( OutFormats.Num() == 0 );
+
+	// does GetDefaultTextureFormatNamePerLayer then per-platform remaps , then (optionally) adds platform prefix
+	//   (platform prefix implies tiling)
+	TargetPlatform->GetTextureFormats(this,OutFormats);
+
+	// almost always == 1, except for Android_Multi, which makes an array of layer formats per variant
+	check( OutFormats.Num() >= 1 );
+	
+	// now do optional Oodle prefix after platform remaps :
+
+	bool bOodleTextureSdkVersionIsNone = OodleTextureSdkVersion.IsNone();
+
+	for( TArray<FName> & LayerFormats : OutFormats )
+	{
+		// conditional prefix only applies the prefix if the prefixed format is supported.
+		//	eg. AutoDXT becomes OODLE_AutoDXT but ETC and RGBA8 do not change.
+		// The Oodle prefixing & support test should always be last, and no more format remapping should be done after that point.
+
+		for (int32 LayerIndex = 0; LayerIndex < LayerFormats.Num(); ++LayerIndex)
+		{
+			FName StartName = LayerFormats[LayerIndex];
+	
+			// can NOT just split name using underscore to get the Platform prefix and Format suffix
+			//	because format names have underscores in them (ASTC_ and ETC2_)
+
+			FName WithoutPlatformPrefixName = UE::TextureBuildUtilities::TextureFormatRemovePlatformPrefixFromName(StartName);
+			FString WithoutPlatformPrefix = WithoutPlatformPrefixName.ToString();
+			checkSlow( StartName.ToString().Find( WithoutPlatformPrefix ) != INDEX_NONE );
+
+			FName ModifiedName = ConditionalGetPrefixedFormat(WithoutPlatformPrefixName, TargetPlatform, bOodleTextureSdkVersionIsNone);
+			if ( WithoutPlatformPrefixName != ModifiedName )
+			{
+				// Got Oodle prefix, now put platform prefix back on :
+			
+				FString StartNameString = StartName.ToString();
+				if ( WithoutPlatformPrefix.Len() < StartNameString.Len() )
+				{
+					FString PlatformPrefix = StartNameString.Left( StartNameString.Len() - WithoutPlatformPrefix.Len() );
+					check( PlatformPrefix + WithoutPlatformPrefix == StartName );
+
+					LayerFormats[LayerIndex] = FName( PlatformPrefix + ModifiedName.ToString() );
+					// should now be like PLATFORM_PREFIX_FORMAT
+				}
+				else
+				{
+					LayerFormats[LayerIndex] = ModifiedName;
+					// should now be like PREFIX_FORMAT
+				}
+			}
+		}
+	}
+
+	// result should be :
+	//	PLATFORM_PREFIX_FORMAT
+	//  (platform and prefix both optional)
 }
+#endif
 
 void GetDefaultTextureFormatNamePerLayer(TArray<FName>& OutFormatNames, const class ITargetPlatform* TargetPlatform, const class UTexture* Texture, 
 	bool bSupportCompressedVolumeTexture, int32 Unused_BlockSize, bool bSupportFilteredFloat32Textures )
diff --git a/Engine/Source/Runtime/Engine/Private/TextureDerivedData.cpp b/Engine/Source/Runtime/Engine/Private/TextureDerivedData.cpp
--- a/Engine/Source/Runtime/Engine/Private/TextureDerivedData.cpp
+++ b/Engine/Source/Runtime/Engine/Private/TextureDerivedData.cpp
@@ -1118,8 +1118,11 @@ static void GetBuildSettingsForRunningPlatform(
 		GetTextureBuildSettings(Texture, *LODSettings, *TargetPlatform, InEncodeSpeed, SourceBuildSettings, &SourceMetadata);
 
 		TArray< TArray<FName> > PlatformFormats;
-		TargetPlatform->GetTextureFormats(&Texture, PlatformFormats);
-		check(PlatformFormats.Num() > 0);
+		Texture.GetPlatformTextureFormatNamesWithPrefix(TargetPlatform,PlatformFormats);
+
+		// this code only uses PlatformFormats[0] , so it would be wrong for Android_Multi
+		//	but it's only used for the platform running the Editor
+		check(PlatformFormats.Num() == 1);
 
 		const int32 NumLayers = Texture.Source.GetNumLayers();
 		check(PlatformFormats[0].Num() == NumLayers);
@@ -1158,7 +1161,10 @@ static void GetBuildSettingsPerFormat(
 	const int32 NumLayers = Texture.Source.GetNumLayers();
 
 	TArray< TArray<FName> > PlatformFormats;
-	TargetPlatform->GetTextureFormats(&Texture, PlatformFormats);
+	Texture.GetPlatformTextureFormatNamesWithPrefix(TargetPlatform,PlatformFormats);
+
+	// almost always == 1, except for Android_Multi, which makes an array of layer formats per variant
+	check( PlatformFormats.Num() >= 1 );
 
 	OutBuildSettingsPerFormat.Reserve(PlatformFormats.Num());
 	if (OutResultMetadataPerFormat)
@@ -1185,21 +1191,9 @@ static void GetBuildSettingsPerFormat(
 
 			if (OutSettings.bVirtualStreamable)
 			{
-				// Virtual textures always strip the child format prefix prior to actual encode since
-				// VTs never tile. We do this here so that we can end up with the same DDC keys and
-				// avoid re-encoding the same texture N times under different keys.
-				// This is the same code used at the end of FVirtualTextureDataBuilder::BuildSourcePixels
-				FName TextureFormatPrefix;
-				FName TextureFormatName = UE::TextureBuildUtilities::TextureFormatRemovePrefixFromName(OutSettings.TextureFormatName, TextureFormatPrefix);
-
-				if (TextureFormatPrefix.IsNone())
-				{
-					OutSettings.TextureFormatName = TextureFormatName;
-				}
-				else
-				{
-					OutSettings.TextureFormatName = *(TextureFormatPrefix.ToString() + TextureFormatName.ToString());
-				}
+				// Virtual textures always strip the child format prefix prior to actual encode since VTs never tile.
+				// must match VirtualTextureDataBuilder.cpp
+				OutSettings.TextureFormatName = UE::TextureBuildUtilities::TextureFormatRemovePlatformPrefixFromName(OutSettings.TextureFormatName);
 			}
 
 			FTexturePlatformData::FTextureEncodeResultMetadata* OutResultMetadata = nullptr;
diff --git a/Engine/Source/Runtime/Engine/Private/VT/VirtualTextureDataBuilder.cpp b/Engine/Source/Runtime/Engine/Private/VT/VirtualTextureDataBuilder.cpp
--- a/Engine/Source/Runtime/Engine/Private/VT/VirtualTextureDataBuilder.cpp
+++ b/Engine/Source/Runtime/Engine/Private/VT/VirtualTextureDataBuilder.cpp
@@ -1149,24 +1149,10 @@ void FVirtualTextureDataBuilder::BuildSourcePixels(const FTextureSourceData& Sou
 		const FTextureBuildSettings& BuildSettingsForLayer = SettingsPerLayer[LayerIndex];
 		FVirtualTextureSourceLayerData& LayerData = SourceLayers[LayerIndex];
 
-		// Note that this code is duplicated at GetBuildSettingsPerFormat in TextureDerivedData.cpp so that the DDC key for 
-		// the VT ends up the same for child platforms, preventing re-encoding of the same texture data over and over. 
-		// As a result, this code shouldn't be doing anything when used via the ddc texture build/cook path - the prefix 
-		// should already be removed.
-		// 
 		// Don't want platform specific swizzling for VT tile data, this tends to add extra padding for textures with odd dimensions
 		// (VT physical tiles generally not power-of-2 after adding border)
-		FName TextureFormatPrefix;
-		FName TextureFormatName = UE::TextureBuildUtilities::TextureFormatRemovePrefixFromName(BuildSettingsForLayer.TextureFormatName, TextureFormatPrefix);
-
-		if ( TextureFormatPrefix.IsNone())
-		{
-			LayerData.TextureFormatName = TextureFormatName;
-		}
-		else
-		{
-			LayerData.TextureFormatName = *(TextureFormatPrefix.ToString() + TextureFormatName.ToString());
-		}
+		// Must match TextureDerivedData.cpp
+		LayerData.TextureFormatName = UE::TextureBuildUtilities::TextureFormatRemovePlatformPrefixFromName(BuildSettingsForLayer.TextureFormatName);
 	}
 }
 
diff --git a/Engine/Source/Runtime/Engine/Public/TextureResource.h b/Engine/Source/Runtime/Engine/Public/TextureResource.h
--- a/Engine/Source/Runtime/Engine/Public/TextureResource.h
+++ b/Engine/Source/Runtime/Engine/Public/TextureResource.h
@@ -622,6 +622,11 @@ private:
 	ECubeFace CurrentTargetFace;
 };
 
+/**
+ * Do not call these (GetDefaultTextureFormatName) directly, use GetPlatformTextureFormatNamesWithPrefix instead
+ * this should only be called by TargetPlatform::GetTextureFormats()
+ */
+
 /** Gets the name of a format for the given LayerIndex */
 ENGINE_API FName GetDefaultTextureFormatName( const class ITargetPlatform* TargetPlatform, const class UTexture* Texture, int32 LayerIndex, bool bSupportCompressedVolumeTexture, int32 Unused_BlockSize, bool bSupportFilteredFloat32Textures);
 
