Fixes to make COTF recompile shaders work correctly for mesh and material shaders, and also fixed ODSC to properly compile material shader permutations.

diff --git a/Engine/Source/Editor/MaterialEditor/Private/MaterialStatsCommon.cpp b/Engine/Source/Editor/MaterialEditor/Private/MaterialStatsCommon.cpp
--- a/Engine/Source/Editor/MaterialEditor/Private/MaterialStatsCommon.cpp
+++ b/Engine/Source/Editor/MaterialEditor/Private/MaterialStatsCommon.cpp
@@ -14,7 +14,7 @@
 /***********************************************************************************************************************/
 /*begin FMaterialResourceStats functions*/
 
-void FMaterialResourceStats::SetupExtaCompilationSettings(const EShaderPlatform Platform, FExtraShaderCompilerSettings& Settings) const
+void FMaterialResourceStats::SetupExtraCompilationSettings(const EShaderPlatform Platform, FExtraShaderCompilerSettings& Settings) const
 {
 	Settings.bExtractShaderSource = true;
 	Settings.OfflineCompilerPath = FMaterialStatsUtils::GetPlatformOfflineCompilerPath(Platform);
diff --git a/Engine/Source/Editor/MaterialEditor/Public/MaterialStatsCommon.h b/Engine/Source/Editor/MaterialEditor/Public/MaterialStatsCommon.h
--- a/Engine/Source/Editor/MaterialEditor/Public/MaterialStatsCommon.h
+++ b/Engine/Source/Editor/MaterialEditor/Public/MaterialStatsCommon.h
@@ -30,7 +30,7 @@ public:
 	virtual bool RequiresSynchronousCompilation() const override { return false; }
 
 	/** this will enable shader source extraction and pass paths to (eventual) offline shader compilers */
-	virtual void SetupExtaCompilationSettings(const EShaderPlatform Platform, FExtraShaderCompilerSettings& Settings) const override;
+	virtual void SetupExtraCompilationSettings(const EShaderPlatform Platform, FExtraShaderCompilerSettings& Settings) const override;
 };
 
 /** enumeration used to group shader platforms */
diff --git a/Engine/Source/Runtime/Engine/Classes/Materials/Material.h b/Engine/Source/Runtime/Engine/Classes/Materials/Material.h
--- a/Engine/Source/Runtime/Engine/Classes/Materials/Material.h
+++ b/Engine/Source/Runtime/Engine/Classes/Materials/Material.h
@@ -1487,7 +1487,7 @@ private:
 	 * Caller is responsible for deleting OutCachedMaterialResources.
 	 * Note: This modifies material variables used for rendering and is assumed to be called within a FMaterialUpdateContext!
 	 */
-	void CacheResourceShadersForCooking(EShaderPlatform Platform, TArray<FMaterialResource*>& OutCachedMaterialResources, const ITargetPlatform* TargetPlatform = nullptr);
+	void CacheResourceShadersForCooking(EShaderPlatform Platform, TArray<FMaterialResource*>& OutCachedMaterialResources, const ITargetPlatform* TargetPlatform = nullptr, bool bBlocking = false);
 
 	void GetNewResources(EShaderPlatform ShaderPlatform, TArray<FMaterialResource*>& NewResourcesToCache);
 
diff --git a/Engine/Source/Runtime/Engine/Classes/Materials/MaterialInstance.h b/Engine/Source/Runtime/Engine/Classes/Materials/MaterialInstance.h
--- a/Engine/Source/Runtime/Engine/Classes/Materials/MaterialInstance.h
+++ b/Engine/Source/Runtime/Engine/Classes/Materials/MaterialInstance.h
@@ -849,7 +849,13 @@ public:
 	 * The results will be applied to this FMaterial in the renderer when they are finished compiling.
 	 * Note: This modifies material variables used for rendering and is assumed to be called within a FMaterialUpdateContext!
 	 */
-	void CacheResourceShadersForCooking(EShaderPlatform ShaderPlatform, TArray<FMaterialResource*>& OutCachedMaterialResources, EMaterialShaderPrecompileMode PrecompileMode = EMaterialShaderPrecompileMode::Default, const ITargetPlatform* TargetPlatform = nullptr);
+	void CacheResourceShadersForCooking(
+		EShaderPlatform ShaderPlatform,
+		TArray<FMaterialResource*>& OutCachedMaterialResources,
+		EMaterialShaderPrecompileMode PrecompileMode = EMaterialShaderPrecompileMode::Default,
+		const ITargetPlatform* TargetPlatform = nullptr,
+		bool bBlocking = false
+	);
 
 	/** 
 	 * Gathers actively used shader maps from all material resources used by this material instance
diff --git a/Engine/Source/Runtime/Engine/Private/Materials/Material.cpp b/Engine/Source/Runtime/Engine/Private/Materials/Material.cpp
--- a/Engine/Source/Runtime/Engine/Private/Materials/Material.cpp
+++ b/Engine/Source/Runtime/Engine/Private/Materials/Material.cpp
@@ -2395,16 +2395,25 @@ void UMaterial::CacheResourceShadersForRendering(bool bRegenerateId, EMaterialSh
 	FMaterial::DeferredDeleteArray(ResourcesToFree);
 }
 
-void UMaterial::CacheResourceShadersForCooking(EShaderPlatform ShaderPlatform, TArray<FMaterialResource*>& OutCachedMaterialResources, const ITargetPlatform* TargetPlatform)
+void UMaterial::CacheResourceShadersForCooking(EShaderPlatform ShaderPlatform, TArray<FMaterialResource*>& OutCachedMaterialResources, const ITargetPlatform* TargetPlatform, bool bBlocking)
 {
 	TArray<FMaterialResource*> NewResourcesToCache;	// only new resources need to have CacheShaders() called on them, whereas OutCachedMaterialResources may already contain resources for another shader platform
 	GetNewResources(ShaderPlatform, NewResourcesToCache);
 
 #if WITH_EDITOR
-	// For cooking, we can call the begin function and it will be completed as part of the polling mechanism.
-	BeginCacheShadersForResources(ShaderPlatform, NewResourcesToCache, EMaterialShaderPrecompileMode::Background, TargetPlatform);
-#else
-	CacheShadersForResources(ShaderPlatform, NewResourcesToCache, EMaterialShaderPrecompileMode::Background, TargetPlatform);
+	// The editor needs to block if the caching call comes from cook on the fly, where the polling mechanisms are not active.
+	// This is important so that the jobs finish and the CacheShadersCompletion() callback is triggered via FinishCacheShaders()!
+	if (bBlocking)
+#endif
+	{
+		CacheShadersForResources(ShaderPlatform, NewResourcesToCache, EMaterialShaderPrecompileMode::Background, TargetPlatform);
+	}
+#if WITH_EDITOR
+	else
+	{
+		// For cooking, we can call the begin function and it will be completed as part of the polling mechanism.
+		BeginCacheShadersForResources(ShaderPlatform, NewResourcesToCache, EMaterialShaderPrecompileMode::Background, TargetPlatform);
+	}
 #endif
 
 	OutCachedMaterialResources.Append(NewResourcesToCache);
@@ -4224,12 +4233,12 @@ void UMaterial::BeginCacheForCookedPlatformData( const ITargetPlatform *TargetPl
 
 	TArray<FMaterialResource*> *CachedMaterialResourcesForPlatform = CachedMaterialResourcesForCooking.Find( TargetPlatform );
 
-	if ( CachedMaterialResourcesForPlatform == NULL )
+	if (CachedMaterialResourcesForPlatform == nullptr)
 	{
 		CachedMaterialResourcesForCooking.Add( TargetPlatform );
 		CachedMaterialResourcesForPlatform = CachedMaterialResourcesForCooking.Find( TargetPlatform );
 
-		check( CachedMaterialResourcesForPlatform != NULL );
+		check(CachedMaterialResourcesForPlatform != nullptr);
 
 		if (DesiredShaderFormats.Num())
 		{
@@ -4252,7 +4261,7 @@ bool UMaterial::IsCachedCookedPlatformDataLoaded( const ITargetPlatform* TargetP
 	LLM_SCOPE(ELLMTag::Materials);
 	const TArray<FMaterialResource*>* CachedMaterialResourcesForPlatform = CachedMaterialResourcesForCooking.Find( TargetPlatform );
 
-	if ( CachedMaterialResourcesForPlatform != NULL ) // this should always succeed if begincacheforcookedcplatformdata is called first
+	if (CachedMaterialResourcesForPlatform != nullptr) // this should always succeed if BeginCacheForCookedPlatformData is called first
 	{
 		for ( const auto& MaterialResource : *CachedMaterialResourcesForPlatform )
 		{
@@ -5328,12 +5337,12 @@ void UMaterial::CompileMaterialsForRemoteRecompile(
 		if (MaterialInstance && MaterialInstance->bHasStaticPermutationResource)
 		{
 			TArray<FMaterialResource*>& ResourceArray = CompilingResources.Add(Material->GetPathName(), TArray<FMaterialResource*>());
-			MaterialInstance->CacheResourceShadersForCooking(ShaderPlatform, ResourceArray, EMaterialShaderPrecompileMode::Default, TargetPlatform);
+			MaterialInstance->CacheResourceShadersForCooking(ShaderPlatform, ResourceArray, EMaterialShaderPrecompileMode::Default, TargetPlatform, true /* Blocking */);
 		}
 		else if (BaseMaterial)
 		{
 			TArray<FMaterialResource*>& ResourceArray = CompilingResources.Add(Material->GetPathName(), TArray<FMaterialResource*>());
-			BaseMaterial->CacheResourceShadersForCooking(ShaderPlatform, ResourceArray, TargetPlatform);
+			BaseMaterial->CacheResourceShadersForCooking(ShaderPlatform, ResourceArray, TargetPlatform, true /* Blocking */);
 		}
 	}
 
diff --git a/Engine/Source/Runtime/Engine/Private/Materials/MaterialInstance.cpp b/Engine/Source/Runtime/Engine/Private/Materials/MaterialInstance.cpp
--- a/Engine/Source/Runtime/Engine/Private/Materials/MaterialInstance.cpp
+++ b/Engine/Source/Runtime/Engine/Private/Materials/MaterialInstance.cpp
@@ -2290,7 +2290,13 @@ void UMaterialInstance::CacheResourceShadersForRendering(EMaterialShaderPrecompi
 	FMaterial::DeferredDeleteArray(ResourcesToFree);
 }
 
-void UMaterialInstance::CacheResourceShadersForCooking(EShaderPlatform ShaderPlatform, TArray<FMaterialResource*>& OutCachedMaterialResources, EMaterialShaderPrecompileMode PrecompileMode, const ITargetPlatform* TargetPlatform)
+void UMaterialInstance::CacheResourceShadersForCooking(
+	EShaderPlatform ShaderPlatform,
+	TArray<FMaterialResource*>& OutCachedMaterialResources,
+	EMaterialShaderPrecompileMode PrecompileMode,
+	const ITargetPlatform* TargetPlatform,
+	bool bBlocking
+)
 {
 	if (bHasStaticPermutationResource)
 	{
@@ -2332,12 +2338,21 @@ void UMaterialInstance::CacheResourceShadersForCooking(EShaderPlatform ShaderPla
 			NewResourcesToCache.Add(NewResource);
 		}
 
-#if WITH_EDITOR
-		// For cooking, we can call the begin function and it will be completed as part of the polling mechanism.
-		BeginCacheShadersForResources(ShaderPlatform, NewResourcesToCache, PrecompileMode, TargetPlatform);
-#else
-		CacheShadersForResources(ShaderPlatform, NewResourcesToCache, PrecompileMode, TargetPlatform);
-#endif
+	#if WITH_EDITOR
+		// The editor needs to block if the caching call comes from cook on the fly, where the polling mechanisms are not active.
+		// This is important so that the jobs finish and the CacheShadersCompletion() callback is triggered via FinishCacheShaders()!
+		if (bBlocking)
+	#endif
+		{
+			CacheShadersForResources(ShaderPlatform, NewResourcesToCache, PrecompileMode, TargetPlatform);
+		}
+	#if WITH_EDITOR
+		else
+		{
+			// For cooking, we can call the begin function and it will be completed as part of the polling mechanism.
+			BeginCacheShadersForResources(ShaderPlatform, NewResourcesToCache, PrecompileMode, TargetPlatform);
+		}
+	#endif
 
 		OutCachedMaterialResources.Append(NewResourcesToCache);
 	}
diff --git a/Engine/Source/Runtime/Engine/Private/Materials/MaterialShader.cpp b/Engine/Source/Runtime/Engine/Private/Materials/MaterialShader.cpp
--- a/Engine/Source/Runtime/Engine/Private/Materials/MaterialShader.cpp
+++ b/Engine/Source/Runtime/Engine/Private/Materials/MaterialShader.cpp
@@ -1145,10 +1145,15 @@ void FMaterialShaderMapId::SetShaderDependencies(const TArray<FShaderType*>& Sha
 	{
 		for (int32 ShaderTypeIndex = 0; ShaderTypeIndex < ShaderTypes.Num(); ShaderTypeIndex++)
 		{
+			FShaderType* ShaderType = ShaderTypes[ShaderTypeIndex];
 			FShaderTypeDependency Dependency;
-			Dependency.ShaderTypeName = ShaderTypes[ShaderTypeIndex]->GetHashedName();
-			Dependency.SourceHash = ShaderTypes[ShaderTypeIndex]->GetSourceHash(ShaderPlatform);
-			ShaderTypeDependencies.Add(Dependency);
+			Dependency.ShaderTypeName = ShaderType->GetHashedName();
+			Dependency.SourceHash = ShaderType->GetSourceHash(ShaderPlatform);
+			for (int32 Id = 0; Id < ShaderType->GetPermutationCount(); ++Id)
+			{
+				Dependency.PermutationId = Id;
+				ShaderTypeDependencies.Add(Dependency);
+			}
 		}
 
 		for (int32 VFTypeIndex = 0; VFTypeIndex < VFTypes.Num(); VFTypeIndex++)
@@ -1192,7 +1197,7 @@ static void PrepareMaterialShaderCompileJob(EShaderPlatform Platform,
 	//update material shader stats
 	UpdateMaterialShaderCompilingStats(Material);
 
-	Material->SetupExtaCompilationSettings(Platform, NewJob->Input.ExtraSettings);
+	Material->SetupExtraCompilationSettings(Platform, NewJob->Input.ExtraSettings);
 
 	// Allow the shader type to modify the compile environment.
 	ShaderType->SetupCompileEnvironment(Platform, Material, Key.PermutationId, PermutationFlags, ShaderEnvironment);
diff --git a/Engine/Source/Runtime/Engine/Private/Materials/MaterialShared.cpp b/Engine/Source/Runtime/Engine/Private/Materials/MaterialShared.cpp
--- a/Engine/Source/Runtime/Engine/Private/Materials/MaterialShared.cpp
+++ b/Engine/Source/Runtime/Engine/Private/Materials/MaterialShared.cpp
@@ -2945,18 +2945,33 @@ void FMaterial::CacheGivenTypes(EShaderPlatform Platform, const TArray<const FVe
 			}
 			else if (ShaderType->GetTypeForDynamicCast() == FShaderType::EShaderTypeForDynamicCast::Material)
 			{
-				ShaderType->AsMaterialShaderType()->BeginCompileShader(
-					EShaderCompileJobPriority::ForceLocal,
-					GetGameThreadCompilingShaderMapId(),
-					0,
-					this,
-					GameThreadShaderMap->GetShaderMapId(),
-					GameThreadPendingCompilerEnvironment,
-					Platform,
-					GameThreadShaderMap->GetPermutationFlags(),
-					CompileJobs,
-					nullptr,
-					nullptr);
+				const EShaderPermutationFlags ShaderPermutation = GameThreadShaderMap->GetPermutationFlags();
+
+				const uint32 CompilingShaderMapId = GetGameThreadCompilingShaderMapId();
+				const FMaterialShaderMapId& ShaderMapId = GameThreadShaderMap->GetShaderMapId();
+				const FMaterialShaderType* MaterialShaderType = ShaderType->GetMaterialShaderType();
+
+				for (int32 PermutationId = 0; PermutationId < ShaderType->GetPermutationCount(); ++PermutationId)
+				{
+					const bool bShaderShouldCompile = MaterialShaderType->ShouldCompilePermutation(Platform, this, PermutationId, ShaderPermutation);
+					if (!bShaderShouldCompile)
+					{
+						continue;
+					}
+
+					MaterialShaderType->BeginCompileShader(
+						EShaderCompileJobPriority::ForceLocal,
+						CompilingShaderMapId,
+						PermutationId,
+						this,
+						ShaderMapId,
+						GameThreadPendingCompilerEnvironment,
+						Platform,
+						ShaderPermutation,
+						CompileJobs,
+						nullptr,
+						nullptr);
+				}
 			}
 			else if (ShaderType->GetTypeForDynamicCast() == FShaderType::EShaderTypeForDynamicCast::MeshMaterial)
 			{
@@ -3331,7 +3346,7 @@ bool FMaterial::TryGetShaders(const FMaterialShaderTypes& InTypes, const FVertex
 							ShaderStageNamesToCompile.Add(ShaderType->GetName());
 						}
 
-						GODSCManager->AddThreadedShaderPipelineRequest(ShaderPlatform, GetFeatureLevel(), GetQualityLevel(), MaterialName, VFTypeName, PipelineName, ShaderStageNamesToCompile);
+						GODSCManager->AddThreadedShaderPipelineRequest(ShaderPlatform, GetFeatureLevel(), GetQualityLevel(), MaterialName, VFTypeName, PipelineName, ShaderStageNamesToCompile, kUniqueShaderPermutationId);
 					}
 				}
 				else
@@ -3410,7 +3425,7 @@ bool FMaterial::TryGetShaders(const FMaterialShaderTypes& InTypes, const FVertex
 							TArray<FString> ShaderStageNamesToCompile;
 							ShaderStageNamesToCompile.Add(ShaderType->GetName());
 
-							GODSCManager->AddThreadedShaderPipelineRequest(ShaderPlatform, GetFeatureLevel(), GetQualityLevel(), MaterialName, VFTypeName, PipelineName, ShaderStageNamesToCompile);
+							GODSCManager->AddThreadedShaderPipelineRequest(ShaderPlatform, GetFeatureLevel(), GetQualityLevel(), MaterialName, VFTypeName, PipelineName, ShaderStageNamesToCompile, PermutationId);
 						}
 					}
 					else
@@ -3425,32 +3440,32 @@ bool FMaterial::TryGetShaders(const FMaterialShaderTypes& InTypes, const FVertex
 								if (InVertexFactoryType)
 								{
 									ShaderType->AsMeshMaterialShaderType()->BeginCompileShader(
-										EShaderCompileJobPriority::ForceLocal, 
-										CompilingShaderMapId, 
-										PermutationId, 
-										ShaderPlatform, 
-										PermutationFlags, 
+										EShaderCompileJobPriority::ForceLocal,
+										CompilingShaderMapId,
+										PermutationId,
+										ShaderPlatform,
+										PermutationFlags,
 										this,
 										ShaderMap->GetShaderMapId(),
-										RenderingThreadPendingCompilerEnvironment, 
-										InVertexFactoryType, 
-										CompileJobs, 
-										nullptr, 
+										RenderingThreadPendingCompilerEnvironment,
+										InVertexFactoryType,
+										CompileJobs,
+										nullptr,
 										nullptr);
 								}
 								else
 								{
 									ShaderType->AsMaterialShaderType()->BeginCompileShader(
-										EShaderCompileJobPriority::ForceLocal, 
-										CompilingShaderMapId, 
-										PermutationId, 
+										EShaderCompileJobPriority::ForceLocal,
+										CompilingShaderMapId,
+										PermutationId,
 										this,
 										ShaderMap->GetShaderMapId(),
-										RenderingThreadPendingCompilerEnvironment, 
-										ShaderPlatform, 
-										PermutationFlags, 
-										CompileJobs, 
-										nullptr, 
+										RenderingThreadPendingCompilerEnvironment,
+										ShaderPlatform,
+										PermutationFlags,
+										CompileJobs,
+										nullptr,
 										nullptr);
 								}
 							}
diff --git a/Engine/Source/Runtime/Engine/Private/Materials/MeshMaterialShader.cpp b/Engine/Source/Runtime/Engine/Private/Materials/MeshMaterialShader.cpp
--- a/Engine/Source/Runtime/Engine/Private/Materials/MeshMaterialShader.cpp
+++ b/Engine/Source/Runtime/Engine/Private/Materials/MeshMaterialShader.cpp
@@ -49,7 +49,7 @@ static void PrepareMeshMaterialShaderCompileJob(EShaderPlatform Platform,
 	check(VertexFactoryType);
 	VertexFactoryType->ModifyCompilationEnvironment(FVertexFactoryShaderPermutationParameters(Platform, MaterialParameters, VertexFactoryType, ShaderType, PermutationFlags), ShaderEnvironment);
 
-	Material->SetupExtaCompilationSettings(Platform, NewJob->Input.ExtraSettings);
+	Material->SetupExtraCompilationSettings(Platform, NewJob->Input.ExtraSettings);
 
 	//update material shader stats
 	UpdateMaterialShaderCompilingStats(Material);
diff --git a/Engine/Source/Runtime/Engine/Private/ODSC/ODSCManager.cpp b/Engine/Source/Runtime/Engine/Private/ODSC/ODSCManager.cpp
--- a/Engine/Source/Runtime/Engine/Private/ODSC/ODSCManager.cpp
+++ b/Engine/Source/Runtime/Engine/Private/ODSC/ODSCManager.cpp
@@ -69,7 +69,14 @@ bool FODSCManager::Tick(float DeltaSeconds)
 	return false;
 }
 
-void FODSCManager::AddThreadedRequest(const TArray<FString>& MaterialsToCompile, const FString& ShaderTypesToLoad, EShaderPlatform ShaderPlatform, ERHIFeatureLevel::Type FeatureLevel, EMaterialQualityLevel::Type QualityLevel, ODSCRecompileCommand RecompileCommandType)
+void FODSCManager::AddThreadedRequest(
+	const TArray<FString>& MaterialsToCompile,
+	const FString& ShaderTypesToLoad,
+	EShaderPlatform ShaderPlatform,
+	ERHIFeatureLevel::Type FeatureLevel,
+	EMaterialQualityLevel::Type QualityLevel,
+	ODSCRecompileCommand RecompileCommandType
+)
 {
 	if (IsHandlingRequests())
 	{
@@ -77,10 +84,19 @@ void FODSCManager::AddThreadedRequest(const TArray<FString>& MaterialsToCompile,
 	}
 }
 
-void FODSCManager::AddThreadedShaderPipelineRequest(EShaderPlatform ShaderPlatform, ERHIFeatureLevel::Type FeatureLevel, EMaterialQualityLevel::Type QualityLevel, const FString& MaterialName, const FString& VertexFactoryName, const FString& PipelineName, const TArray<FString>& ShaderTypeNames)
+void FODSCManager::AddThreadedShaderPipelineRequest(
+	EShaderPlatform ShaderPlatform,
+	ERHIFeatureLevel::Type FeatureLevel,
+	EMaterialQualityLevel::Type QualityLevel,
+	const FString& MaterialName,
+	const FString& VertexFactoryName,
+	const FString& PipelineName,
+	const TArray<FString>& ShaderTypeNames,
+	int32 PermutationId
+)
 {
 	if (IsHandlingRequests())
 	{
-		Thread->AddShaderPipelineRequest(ShaderPlatform, FeatureLevel, QualityLevel, MaterialName, VertexFactoryName, PipelineName, ShaderTypeNames);
+		Thread->AddShaderPipelineRequest(ShaderPlatform, FeatureLevel, QualityLevel, MaterialName, VertexFactoryName, PipelineName, ShaderTypeNames, PermutationId);
 	}
 }
diff --git a/Engine/Source/Runtime/Engine/Private/ODSC/ODSCThread.cpp b/Engine/Source/Runtime/Engine/Private/ODSC/ODSCThread.cpp
--- a/Engine/Source/Runtime/Engine/Private/ODSC/ODSCThread.cpp
+++ b/Engine/Source/Runtime/Engine/Private/ODSC/ODSCThread.cpp
@@ -4,8 +4,26 @@
 #include "ODSCLog.h"
 #include "HAL/FileManager.h"
 
-FODSCRequestPayload::FODSCRequestPayload(EShaderPlatform InShaderPlatform, ERHIFeatureLevel::Type InFeatureLevel, EMaterialQualityLevel::Type InQualityLevel, const FString& InMaterialName, const FString& InVertexFactoryName, const FString& InPipelineName, const TArray<FString>& InShaderTypeNames, const FString& InRequestHash)
-	: ShaderPlatform(InShaderPlatform), FeatureLevel(InFeatureLevel), QualityLevel(InQualityLevel), MaterialName(InMaterialName), VertexFactoryName(InVertexFactoryName), PipelineName(InPipelineName), ShaderTypeNames(std::move(InShaderTypeNames)), RequestHash(InRequestHash)
+FODSCRequestPayload::FODSCRequestPayload(
+	EShaderPlatform InShaderPlatform,
+	ERHIFeatureLevel::Type InFeatureLevel,
+	EMaterialQualityLevel::Type InQualityLevel,
+	const FString& InMaterialName,
+	const FString& InVertexFactoryName,
+	const FString& InPipelineName,
+	const TArray<FString>& InShaderTypeNames,
+	int32 InPermutationId,
+	const FString& InRequestHash
+)
+: ShaderPlatform(InShaderPlatform)
+, FeatureLevel(InFeatureLevel)
+, QualityLevel(InQualityLevel)
+, MaterialName(InMaterialName)
+, VertexFactoryName(InVertexFactoryName)
+, PipelineName(InPipelineName)
+, ShaderTypeNames(std::move(InShaderTypeNames))
+, PermutationId(InPermutationId)
+, RequestHash(InRequestHash)
 {
 
 }
@@ -120,8 +138,21 @@ void FODSCThread::AddRequest(const TArray<FString>& MaterialsToCompile, const FS
 	PendingMaterialThreadedRequests.Enqueue(new FODSCMessageHandler(MaterialsToCompile, ShaderTypesToLoad, ShaderPlatform, FeatureLevel, QualityLevel, RecompileCommandType));
 }
 
-void FODSCThread::AddShaderPipelineRequest(EShaderPlatform ShaderPlatform, ERHIFeatureLevel::Type FeatureLevel, EMaterialQualityLevel::Type QualityLevel, const FString& MaterialName, const FString& VertexFactoryName, const FString& PipelineName, const TArray<FString>& ShaderTypeNames)
+void FODSCThread::AddShaderPipelineRequest(
+	EShaderPlatform ShaderPlatform,
+	ERHIFeatureLevel::Type FeatureLevel,
+	EMaterialQualityLevel::Type QualityLevel,
+	const FString& MaterialName,
+	const FString& VertexFactoryName,
+	const FString& PipelineName,
+	const TArray<FString>& ShaderTypeNames,
+	int32 PermutationId
+)
 {
+	// TODO: Requests for individual permutations come in here, but a single coalesced payload is submitted to the server since 
+	// we compile all material shader permutations encountered for the moment. Consider batching up requested permutations and 
+	// have the server skip compiling those not in the list. Ensure that DDC key and shader map assumptions are correct!
+
 	FString RequestString = (MaterialName + VertexFactoryName + PipelineName);
 	for (const auto& ShaderTypeName : ShaderTypeNames)
 	{
@@ -132,7 +163,7 @@ void FODSCThread::AddShaderPipelineRequest(EShaderPlatform ShaderPlatform, ERHIF
 	FScopeLock Lock(&RequestHashCriticalSection);
 	if (!RequestHashes.Contains(RequestHash))
 	{
-		PendingMeshMaterialThreadedRequests.Enqueue(FODSCRequestPayload(ShaderPlatform, FeatureLevel, QualityLevel, MaterialName, VertexFactoryName, PipelineName, ShaderTypeNames, RequestHash));
+		PendingMeshMaterialThreadedRequests.Enqueue(FODSCRequestPayload(ShaderPlatform, FeatureLevel, QualityLevel, MaterialName, VertexFactoryName, PipelineName, ShaderTypeNames, PermutationId, RequestHash));
 		RequestHashes.Add(RequestHash);
 	}
 }
diff --git a/Engine/Source/Runtime/Engine/Private/ODSC/ODSCThread.h b/Engine/Source/Runtime/Engine/Private/ODSC/ODSCThread.h
--- a/Engine/Source/Runtime/Engine/Private/ODSC/ODSCThread.h
+++ b/Engine/Source/Runtime/Engine/Private/ODSC/ODSCThread.h
@@ -112,10 +112,20 @@ public:
 	 * @param VertexFactoryName - The name of the vertex factory type we should compile.
 	 * @param PipelineName - The name of the shader pipeline we should compile.
 	 * @param ShaderTypeNames - The shader type names of all the shader stages in the pipeline.
+	 * @param PermutationId - The permutation ID of the shader we should compile.
 	 *
 	 * @return false if no longer needs ticking
 	 */
-	void AddShaderPipelineRequest(EShaderPlatform ShaderPlatform, ERHIFeatureLevel::Type FeatureLevel, EMaterialQualityLevel::Type QualityLevel, const FString& MaterialName, const FString& VertexFactoryName, const FString& PipelineName, const TArray<FString>& ShaderTypeNames);
+	void AddShaderPipelineRequest(
+		EShaderPlatform ShaderPlatform,
+		ERHIFeatureLevel::Type FeatureLevel,
+		EMaterialQualityLevel::Type QualityLevel,
+		const FString& MaterialName,
+		const FString& VertexFactoryName,
+		const FString& PipelineName,
+		const TArray<FString>& ShaderTypeNames,
+		int32 PermutationId
+	);
 
 	/**
 	 * Get completed requests.  Clears internal arrays.  Called on Game thread.
diff --git a/Engine/Source/Runtime/Engine/Private/ShaderCompiler/ShaderCompiler.cpp b/Engine/Source/Runtime/Engine/Private/ShaderCompiler/ShaderCompiler.cpp
--- a/Engine/Source/Runtime/Engine/Private/ShaderCompiler/ShaderCompiler.cpp
+++ b/Engine/Source/Runtime/Engine/Private/ShaderCompiler/ShaderCompiler.cpp
@@ -8102,6 +8102,7 @@ FArchive& operator<<(FArchive& Ar, FODSCRequestPayload& Payload)
 	Ar << Payload.VertexFactoryName;
 	Ar << Payload.PipelineName;
 	Ar << Payload.ShaderTypeNames;
+	Ar << Payload.PermutationId;
 	Ar << Payload.RequestHash;
 
 	if (Ar.IsLoading())
diff --git a/Engine/Source/Runtime/Engine/Public/MaterialShared.h b/Engine/Source/Runtime/Engine/Public/MaterialShared.h
--- a/Engine/Source/Runtime/Engine/Public/MaterialShared.h
+++ b/Engine/Source/Runtime/Engine/Public/MaterialShared.h
@@ -2190,7 +2190,7 @@ public:
 	ENGINE_API bool WritesEveryPixel(bool bShadowPass = false) const;
 
 	/** call during shader compilation jobs setup to fill additional settings that may be required by classes who inherit from this */
-	virtual void SetupExtaCompilationSettings(const EShaderPlatform Platform, FExtraShaderCompilerSettings& Settings) const
+	virtual void SetupExtraCompilationSettings(const EShaderPlatform Platform, FExtraShaderCompilerSettings& Settings) const
 	{}
 
 	void DumpDebugInfo(FOutputDevice& OutputDevice);
diff --git a/Engine/Source/Runtime/Engine/Public/ODSC/ODSCManager.h b/Engine/Source/Runtime/Engine/Public/ODSC/ODSCManager.h
--- a/Engine/Source/Runtime/Engine/Public/ODSC/ODSCManager.h
+++ b/Engine/Source/Runtime/Engine/Public/ODSC/ODSCManager.h
@@ -54,7 +54,7 @@ public:
 	void AddThreadedRequest(const TArray<FString>& MaterialsToCompile, const FString& ShaderTypesToLoad, EShaderPlatform ShaderPlatform, ERHIFeatureLevel::Type FeatureLevel, EMaterialQualityLevel::Type QualityLevel, ODSCRecompileCommand RecompileCommandType);
 
 	/**
-	 * Add a request to compile a pipeline (VS/PS) of shaders.  The results are submitted and processed in an async manner.
+	 * Add a request to compile a pipeline of shaders.  The results are submitted and processed in an async manner.
 	 *
 	 * @param ShaderPlatform - Which shader platform to compile for.
 	 * @param FeatureLevel - Which feature level to compile for.
@@ -63,10 +63,20 @@ public:
 	 * @param VertexFactoryName - The name of the vertex factory type we should compile.
 	 * @param PipelineName - The name of the shader pipeline we should compile.
 	 * @param ShaderTypeNames - The shader type names of all the shader stages in the pipeline.
+	 * @param PermutationId - The permutation ID of the shader we should compile.
 	 *
 	 * @return false if no longer needs ticking
 	 */
-	void AddThreadedShaderPipelineRequest(EShaderPlatform ShaderPlatform, ERHIFeatureLevel::Type FeatureLevel, EMaterialQualityLevel::Type QualityLevel, const FString& MaterialName, const FString& VertexFactoryName, const FString& PipelineName, const TArray<FString>& ShaderTypeNames);
+	void AddThreadedShaderPipelineRequest(
+		EShaderPlatform ShaderPlatform,
+		ERHIFeatureLevel::Type FeatureLevel,
+		EMaterialQualityLevel::Type QualityLevel,
+		const FString& MaterialName,
+		const FString& VertexFactoryName,
+		const FString& PipelineName,
+		const TArray<FString>& ShaderTypeNames,
+		int32 PermutationId
+	);
 
 	/** Returns true if we would actually add a request when calling AddThreadedShaderPipelineRequest. */
 	inline bool IsHandlingRequests() const { return Thread != nullptr; }
diff --git a/Engine/Source/Runtime/Engine/Public/ShaderCompiler.h b/Engine/Source/Runtime/Engine/Public/ShaderCompiler.h
--- a/Engine/Source/Runtime/Engine/Public/ShaderCompiler.h
+++ b/Engine/Source/Runtime/Engine/Public/ShaderCompiler.h
@@ -1058,11 +1058,24 @@ struct FODSCRequestPayload
 	/** An array of shader type names for each stage in the Pipeline. */
 	TArray<FString> ShaderTypeNames;
 
+	/** The permutation ID to compile. */
+	int32 PermutationId;
+
 	/** A hash of the above information to uniquely identify a Request. */
 	FString RequestHash;
 
 	ENGINE_API FODSCRequestPayload() {};
-	ENGINE_API FODSCRequestPayload(EShaderPlatform InShaderPlatform, ERHIFeatureLevel::Type InFeatureLevel, EMaterialQualityLevel::Type InQualityLevel, const FString& InMaterialName, const FString& InVertexFactoryName, const FString& InPipelineName, const TArray<FString>& InShaderTypeNames, const FString& InRequestHash);
+	ENGINE_API FODSCRequestPayload(
+		EShaderPlatform InShaderPlatform,
+		ERHIFeatureLevel::Type InFeatureLevel,
+		EMaterialQualityLevel::Type InQualityLevel,
+		const FString& InMaterialName,
+		const FString& InVertexFactoryName,
+		const FString& InPipelineName,
+		const TArray<FString>& InShaderTypeNames,
+		int32 InPermutationId,
+		const FString& InRequestHash
+	);
 
 	/**
 	* Serializes FODSCRequestPayload value from or into this archive.
diff --git a/Engine/Source/Runtime/RenderCore/Public/ShaderCompilerCore.h b/Engine/Source/Runtime/RenderCore/Public/ShaderCompilerCore.h
--- a/Engine/Source/Runtime/RenderCore/Public/ShaderCompilerCore.h
+++ b/Engine/Source/Runtime/RenderCore/Public/ShaderCompilerCore.h
@@ -229,8 +229,8 @@ struct FShaderCompilerInput
 	const FShaderParametersMetadata* RootParametersStructure = nullptr;
 
 
-	// Additional compilation settings that can be filled by FMaterial::SetupExtaCompilationSettings
-	// FMaterial::SetupExtaCompilationSettings is usually called by each (*)MaterialShaderType::BeginCompileShader() function
+	// Additional compilation settings that can be filled by FMaterial::SetupExtraCompilationSettings
+	// FMaterial::SetupExtraCompilationSettings is usually called by each (*)MaterialShaderType::BeginCompileShader() function
 	FExtraShaderCompilerSettings ExtraSettings;
 
 	/** Oodle-specific compression algorithm - used if CompressionFormat is set to NAME_Oodle. */
