Add GetPlatformUserIndex function to LocalPlayer and fix comments to explain the different indices

diff --git a/Engine/Source/Runtime/Engine/Classes/Engine/LocalPlayer.h b/Engine/Source/Runtime/Engine/Classes/Engine/LocalPlayer.h
--- a/Engine/Source/Runtime/Engine/Classes/Engine/LocalPlayer.h
+++ b/Engine/Source/Runtime/Engine/Classes/Engine/LocalPlayer.h
@@ -441,7 +441,9 @@ public:
 	virtual void SetControllerId(int32 NewControllerId);
 
 	/**
-	 * Returns the physical controller ID for the player
+	 * Returns the controller ID for the player.
+	 * This is a legacy identifier corresponding to the primary physical controller used by the player.
+	 * You may want to use GetPlatformUserIndex or GetLocalPlayerIndex instead.
 	 */
 	int32 GetControllerId() const { return ControllerId; }
 
@@ -456,8 +458,15 @@ public:
 	FPlatformUserId GetPlatformUserId() const { return PlatformUserId; }
 
 	/**
-	 * Returns the logical local player index where 0 is the first player
-	 * By default, this uses index in the game instance array
+	 * Converts the platform user id to an index where 0 is the first logged in user.
+	 * This index is used for platform functions like save games, user selection, and slate input.
+	 */
+	virtual int32 GetPlatformUserIndex() const;
+
+	/**
+	 * Returns the logical local player index where 0 is the first LocalPlayer that was created.
+	 * By default, this uses index in the game instance array.
+	 * This index is used for gameplay purposes but will not be correct for platform functions.
 	 */
 	virtual int32 GetLocalPlayerIndex() const;
 
diff --git a/Engine/Source/Runtime/Engine/Private/LocalPlayer.cpp b/Engine/Source/Runtime/Engine/Private/LocalPlayer.cpp
--- a/Engine/Source/Runtime/Engine/Private/LocalPlayer.cpp
+++ b/Engine/Source/Runtime/Engine/Private/LocalPlayer.cpp
@@ -21,6 +21,7 @@
 #include "SceneManagement.h"
 #include "Rendering/SkeletalMeshRenderData.h"
 #include "HAL/PlatformApplicationMisc.h"
+#include "GenericPlatform/GenericPlatformInputDeviceMapper.h"
 #include "Framework/Application/SlateApplication.h"
 
 #include "IXRTrackingSystem.h"
@@ -1654,7 +1655,7 @@ void ULocalPlayer::SetControllerId( int32 NewControllerId )
 		if (GEngine->IsControllerIdUsingPlatformUserId())
 		{
 			// This won't recurse back because we've already modified ControllerId
-			SetPlatformUserId(FPlatformMisc::GetPlatformUserForUserIndex(NewControllerId));
+			SetPlatformUserId(IPlatformInputDeviceMapper::Get().GetPlatformUserForUserIndex(NewControllerId));
 		}
 	}
 }
@@ -1677,14 +1678,19 @@ void ULocalPlayer::SetPlatformUserId(FPlatformUserId NewPlatformUserId)
 
 		if (GEngine->IsControllerIdUsingPlatformUserId())
 		{
-			SetControllerId(FPlatformMisc::GetUserIndexForPlatformUser(PlatformUserId));
+			SetControllerId(IPlatformInputDeviceMapper::Get().GetUserIndexForPlatformUser(PlatformUserId));
 		}
 	}
 }
 
+int32 ULocalPlayer::GetPlatformUserIndex() const
+{
+	FPlatformUserId UserId = GetPlatformUserId();
+	return IPlatformInputDeviceMapper::Get().GetUserIndexForPlatformUser(UserId);
+}
+
 int32 ULocalPlayer::GetLocalPlayerIndex() const
 {
-	// TODO: Add caching
 	return GetIndexInGameInstance();
 }
 
