Only run FSpanAllocator::Consolidate when no free space is immediately available

It is possible for there to be freed spans that are pending consolidation, however no consolidation can actually occur. In this case every Allocate() will still pay the cost of looking for blocks to consolidate, making it much more expensive.
Simple solution is to only run consolidation if there is no available space otherwise. The downside is that this can theoretically result in worse fragmentation over time.

Significantly improves FLumenSceneData::ProcessLumenSurfaceCacheRequests performance in some situations.

diff --git a/Engine/Source/Runtime/Engine/Public/SpanAllocator.h b/Engine/Source/Runtime/Engine/Public/SpanAllocator.h
--- a/Engine/Source/Runtime/Engine/Public/SpanAllocator.h
+++ b/Engine/Source/Runtime/Engine/Public/SpanAllocator.h
@@ -22,14 +22,14 @@ public:
 	// Allocate a range.  Returns allocated StartOffset.
 	inline int32 Allocate(int32 Num = 1)
 	{
-		// Auto consolidate if free has been called
-		if (!PendingFreeSpans.IsEmpty())
+		int32 FoundIndex = SearchFreeList(Num, FirstNonEmptySpan);
+
+		if (FoundIndex == INDEX_NONE && !PendingFreeSpans.IsEmpty())
 		{
 			Consolidate();
+			FoundIndex = SearchFreeList(Num, FirstNonEmptySpan);
 		}
 
-		const int32 FoundIndex = SearchFreeList(Num, FirstNonEmptySpan);
-
 		// Use an existing free span if one is found
 		if (FoundIndex != INDEX_NONE)
 		{
