fix FinalizeVirtualTextureLayerFormat is run after other TextureFormatName logic and incorrectly assumes that it gets non-prefixed names
don't try to move it into the GetPlatformTextureFormatNamesWithPrefix phase yet, just remove prefix
fixes VT format remapping for Android_DXT

diff --git a/Engine/Source/Developer/TextureBuildUtilities/Private/TextureBuildUtilities.cpp b/Engine/Source/Developer/TextureBuildUtilities/Private/TextureBuildUtilities.cpp
--- a/Engine/Source/Developer/TextureBuildUtilities/Private/TextureBuildUtilities.cpp
+++ b/Engine/Source/Developer/TextureBuildUtilities/Private/TextureBuildUtilities.cpp
@@ -78,7 +78,7 @@ TEXTUREBUILDUTILITIES_API const FName TextureFormatRemovePrefixFromName(FName co
 	int32 UnderscoreIndex = INDEX_NONE;
 	if ( ! NameString.FindChar(TCHAR('_'), UnderscoreIndex))
 	{
-		return FName( NameString );
+		return NameWithoutPlatform;
 	}
 
 	// texture format names can have underscores in them (eg. ETC2_RG11)
diff --git a/Engine/Source/Runtime/Engine/Private/Texture.cpp b/Engine/Source/Runtime/Engine/Private/Texture.cpp
--- a/Engine/Source/Runtime/Engine/Private/Texture.cpp
+++ b/Engine/Source/Runtime/Engine/Private/Texture.cpp
@@ -3531,6 +3531,8 @@ void UTexture::GetPlatformTextureFormatNamesWithPrefix(const class ITargetPlatfo
 	// almost always == 1, except for Android_Multi, which makes an array of layer formats per variant
 	check( OutFormats.Num() >= 1 );
 	
+	// note : FinalizeVirtualTextureLayerFormat can do another set of remaps later; consider moving that here
+
 	// now do optional Oodle prefix after platform remaps :
 
 	bool bOodleTextureSdkVersionIsNone = OodleTextureSdkVersion.IsNone();
diff --git a/Engine/Source/Runtime/Engine/Private/TextureDerivedData.cpp b/Engine/Source/Runtime/Engine/Private/TextureDerivedData.cpp
--- a/Engine/Source/Runtime/Engine/Private/TextureDerivedData.cpp
+++ b/Engine/Source/Runtime/Engine/Private/TextureDerivedData.cpp
@@ -643,9 +643,27 @@ static void FinalizeBuildSettingsForLayer(
 		OutSettings.bReplicateRed = true;
 	}
 
-	if (OutSettings.bVirtualStreamable)
+	// this is called once per Texture with OutSettings.TextureFormatName == None
+	//	and then called again (per Layer) with OutSettings.TextureFormatName filled out
+
+	if (OutSettings.bVirtualStreamable && ! OutSettings.TextureFormatName.IsNone())
 	{
-		OutSettings.TextureFormatName = TargetPlatform->FinalizeVirtualTextureLayerFormat(OutSettings.TextureFormatName);
+		// note : FinalizeVirtualTextureLayerFormat is run outside of the normal TextureFormatName set up ; fix?
+		//	should be done inside GetPlatformTextureFormatNamesWithPrefix
+		//	this is only used by Android & iOS
+		//  the reason to do it here is we now have bVirtualStreamable, which is not available at the earlier call
+		
+		// FinalizeVirtualTextureLayerFormat assumes (incorrectly) that it gets non-prefixed names, so remove them :
+
+		// VT does not tile so should never have a platform prefix, but could have an Oodle prefix
+		checkSlow( OutSettings.TextureFormatName == UE::TextureBuildUtilities::TextureFormatRemovePlatformPrefixFromName(OutSettings.TextureFormatName) );
+		
+		FName NameWithoutPrefix = UE::TextureBuildUtilities::TextureFormatRemovePrefixFromName(OutSettings.TextureFormatName);
+		FName ModifiedName = TargetPlatform->FinalizeVirtualTextureLayerFormat(NameWithoutPrefix);
+		if ( NameWithoutPrefix != ModifiedName )
+		{
+			OutSettings.TextureFormatName = ModifiedName;
+		}
 	}
 
 	// Now that we know the texture format, we can make decisions based on it.
