Fix GetNetMode to work properly in PIE for functions like InitGame called halfway through world startup. This restores it to work more like how it did before CL 18423916
Fix EnableListenServer to properly work for dedicated servers in PIE, it only accidentally worked before because GetNetMode was wrong
Fix PIE server worlds to be properly marked as dedicated server during a hard server travel, before they were turning into listen servers until the net driver spawned

diff --git a/Engine/Source/Editor/UnrealEd/Private/PlayLevel.cpp b/Engine/Source/Editor/UnrealEd/Private/PlayLevel.cpp
--- a/Engine/Source/Editor/UnrealEd/Private/PlayLevel.cpp
+++ b/Engine/Source/Editor/UnrealEd/Private/PlayLevel.cpp
@@ -2465,6 +2465,14 @@ void UEditorEngine::PostCreatePIEWorld(UWorld *NewPIEWorld)
 	// make sure we can clean up this world!
 	NewPIEWorld->ClearFlags(RF_Standalone);
 
+	// Force the new world to use a dedicated server net mode if needed
+	// The other types will correctly derive it from the URL as it changes during play
+	FWorldContext* const Context = GetWorldContextFromWorld(NewPIEWorld);
+	if (Context && Context->RunAsDedicated)
+	{
+		NewPIEWorld->SetPlayInEditorInitialNetMode(NM_DedicatedServer);
+	}
+	
 	// Init the PIE world
 	NewPIEWorld->InitWorld();
 	UE_LOG(LogPlayLevel, Log, TEXT("PIE: World Init took: (%fs)"),  float(FPlatformTime::Seconds() - WorldInitStart));
diff --git a/Engine/Source/Runtime/Engine/Private/GameInstance.cpp b/Engine/Source/Runtime/Engine/Private/GameInstance.cpp
--- a/Engine/Source/Runtime/Engine/Private/GameInstance.cpp
+++ b/Engine/Source/Runtime/Engine/Private/GameInstance.cpp
@@ -1290,9 +1290,9 @@ bool UGameInstance::EnableListenServer(bool bEnable, int32 PortOverride /*= 0*/)
 
 	ENetMode ExistingMode = World->GetNetMode();
 
-	if (ExistingMode == NM_Client || ExistingMode == NM_DedicatedServer)
+	if (ExistingMode == NM_Client)
 	{
-		// Clients and dedicated servers cannot change to listen!
+		// Clients cannot change to listen!
 		return false;
 	}
 
@@ -1306,7 +1306,7 @@ bool UGameInstance::EnableListenServer(bool bEnable, int32 PortOverride /*= 0*/)
 		}
 		WorldContext->LastURL.AddOption(TEXT("Listen"));
 
-		if (ExistingMode == NM_Standalone)
+		if (!World->GetNetDriver())
 		{
 			// This actually opens the port
 			FURL ListenURL = WorldContext->LastURL;
@@ -1320,13 +1320,19 @@ bool UGameInstance::EnableListenServer(bool bEnable, int32 PortOverride /*= 0*/)
 	}
 	else
 	{
+		if (ExistingMode == NM_DedicatedServer)
+		{
+			UE_LOG(LogGameSession, Warning, TEXT("EnableListenServer: Dedicated servers always listen for connections"));
+			return false;
+		}
+
 		WorldContext->LastURL.RemoveOption(TEXT("Listen"));
 		WorldContext->LastURL.Port = FURL::UrlConfig.DefaultPort;
 
 		if (ExistingMode == NM_ListenServer)
 		{
 			// What to do in this case is very game-specific
-			UE_LOG(LogGameSession, Warning, TEXT("Disabling a listen server with active connections does not disconnect existing players by default"));
+			UE_LOG(LogGameSession, Warning, TEXT("EnableListenServer: Disabling a listen server with active connections does not disconnect existing players by default"));
 		}
 
 		return true;
diff --git a/Engine/Source/Runtime/Engine/Private/World.cpp b/Engine/Source/Runtime/Engine/Private/World.cpp
--- a/Engine/Source/Runtime/Engine/Private/World.cpp
+++ b/Engine/Source/Runtime/Engine/Private/World.cpp
@@ -8093,11 +8093,11 @@ ENetMode UWorld::InternalGetNetMode() const
 
 	ENetMode URLNetMode = AttemptDeriveFromURL();
 #if WITH_EDITOR
-	if (WorldType == EWorldType::PIE && URLNetMode == NM_Standalone && !AreActorsInitialized())
+	if (WorldType == EWorldType::PIE && (URLNetMode == NM_Standalone || PlayInEditorNetMode == NM_DedicatedServer))
 	{
-		// If we're too early in startup and it defaults to Standalone, use the mode we were created with
-		// Once the world has been initialized the URL will be correct so we will use that
-		// This is required for dedicated server worlds so it is correct for InitWorld
+		// If we're early in startup before the net driver exists and there is no URL override
+		// or this is a dedicated server, use the mode we were first created with
+		// This is required for dedicated server/listen worlds so it is correct for InitWorld
 		return PlayInEditorNetMode;
 	}
 #endif
