Minor optimizations in the rendering code that made were after profiling in the editor.

diff --git a/Engine/Source/Runtime/Engine/Private/HierarchicalInstancedStaticMesh.cpp b/Engine/Source/Runtime/Engine/Private/HierarchicalInstancedStaticMesh.cpp
--- a/Engine/Source/Runtime/Engine/Private/HierarchicalInstancedStaticMesh.cpp
+++ b/Engine/Source/Runtime/Engine/Private/HierarchicalInstancedStaticMesh.cpp
@@ -1261,10 +1261,11 @@ void FHierarchicalStaticMeshSceneProxy::FillDynamicMeshElements(FMeshElementColl
 	for (int32 LODIndex = FirstLOD; LODIndex < LastLODPlusOne; LODIndex++)
 	{
 		const FStaticMeshLODResources& LODModel = RenderData->LODResources[LODIndex];
+		const int32 TotalNumSections = LODModel.Sections.Num();
 
 		for (int32 SelectionGroupIndex = 0; SelectionGroupIndex < ElementParams.NumSelectionGroups; SelectionGroupIndex++)
 		{
-			for (int32 SectionIndex = 0; SectionIndex < LODModel.Sections.Num(); SectionIndex++)
+			for (int32 SectionIndex = 0; SectionIndex < TotalNumSections; SectionIndex++)
 			{
 				const FLODInfo& ProxyLODInfo = LODs[LODIndex];
 				UMaterialInterface* Material = ProxyLODInfo.Sections[SectionIndex].Material;
@@ -1313,7 +1314,7 @@ void FHierarchicalStaticMeshSceneProxy::FillDynamicMeshElements(FMeshElementColl
 							if (ElementParams.FinalCullDistance > 9.9E8)
 							{
 								UE_LOG(LogStaticMesh, Display, TEXT("lod:%1d/%1d   sel:%1d   section:%1d/%1d   runs:%4d   inst:%8d   tris:%9lld   cast shadow:%1d   cull:-NONE!!-   shadow:%1d     %s %s"),
-									LODIndex, InstancedRenderData.VertexFactories.Num(), SelectionGroupIndex, SectionIndex, LODModel.Sections.Num(), RunArray.Num() / 2,
+									LODIndex, InstancedRenderData.VertexFactories.Num(), SelectionGroupIndex, SectionIndex, TotalNumSections, RunArray.Num() / 2,
 									TotalInstances, Tris, (int)MeshBatch.CastShadow, ElementParams.ShadowFrustum,
 									*StaticMesh->GetPathName(),
 									*MeshBatch.MaterialRenderProxy->GetIncompleteMaterialWithFallback(ElementParams.FeatureLevel).GetFriendlyName());
@@ -1321,7 +1322,7 @@ void FHierarchicalStaticMeshSceneProxy::FillDynamicMeshElements(FMeshElementColl
 							else
 							{
 								UE_LOG(LogStaticMesh, Display, TEXT("lod:%1d/%1d   sel:%1d   section:%1d/%1d   runs:%4d   inst:%8d   tris:%9lld   cast shadow:%1d   cull:%8.0f   shadow:%1d     %s %s"),
-									LODIndex, InstancedRenderData.VertexFactories.Num(), SelectionGroupIndex, SectionIndex, LODModel.Sections.Num(), RunArray.Num() / 2,
+									LODIndex, InstancedRenderData.VertexFactories.Num(), SelectionGroupIndex, SectionIndex, TotalNumSections, RunArray.Num() / 2,
 									TotalInstances, Tris, (int)MeshBatch.CastShadow, ElementParams.FinalCullDistance, ElementParams.ShadowFrustum,
 									*StaticMesh->GetPathName(),
 									*MeshBatch.MaterialRenderProxy->GetIncompleteMaterialWithFallback(ElementParams.FeatureLevel).GetFriendlyName());
@@ -1355,7 +1356,7 @@ void FHierarchicalStaticMeshSceneProxy::FillDynamicMeshElements(FMeshElementColl
 							OverlayMeshBatch.bSelectable = false;
 							OverlayMeshBatch.MaterialRenderProxy = OverlayMaterial->GetRenderProxy();
 							// make sure overlay is always rendered on top of base mesh
-							OverlayMeshBatch.MeshIdInPrimitive += LODModel.Sections.Num();
+							OverlayMeshBatch.MeshIdInPrimitive += TotalNumSections;
 							Collector.AddMesh(ElementParams.ViewIndex, OverlayMeshBatch);
 						}
 					}
@@ -1414,7 +1415,7 @@ void FHierarchicalStaticMeshSceneProxy::FillDynamicMeshElements(FMeshElementColl
 								if (ElementParams.FinalCullDistance > 9.9E8)
 								{
 									UE_LOG(LogStaticMesh, Display, TEXT("lod:%1d/%1d   sel:%1d   section:%1d/%1d   runs:%4d   inst:%8d   tris:%9lld   cast shadow:%1d   cull:-NONE!!-   shadow:%1d     %s %s"),
-										LODIndex, InstancedRenderData.VertexFactories.Num(), SelectionGroupIndex, SectionIndex, LODModel.Sections.Num(), RunArray.Num() / 2,
+										LODIndex, InstancedRenderData.VertexFactories.Num(), SelectionGroupIndex, SectionIndex, TotalNumSections, RunArray.Num() / 2,
 										RemainingInstances, Tris, (int)MeshElement.CastShadow, ElementParams.ShadowFrustum,
 										*StaticMesh->GetPathName(),
 										*MeshElement.MaterialRenderProxy->GetIncompleteMaterialWithFallback(ElementParams.FeatureLevel).GetFriendlyName());
@@ -1422,7 +1423,7 @@ void FHierarchicalStaticMeshSceneProxy::FillDynamicMeshElements(FMeshElementColl
 								else
 								{
 									UE_LOG(LogStaticMesh, Display, TEXT("lod:%1d/%1d   sel:%1d   section:%1d/%1d   runs:%4d   inst:%8d   tris:%9lld   cast shadow:%1d   cull:%8.0f   shadow:%1d     %s %s"),
-										LODIndex, InstancedRenderData.VertexFactories.Num(), SelectionGroupIndex, SectionIndex, LODModel.Sections.Num(), RunArray.Num() / 2,
+										LODIndex, InstancedRenderData.VertexFactories.Num(), SelectionGroupIndex, SectionIndex, TotalNumSections, RunArray.Num() / 2,
 										RemainingInstances, Tris, (int)MeshElement.CastShadow, ElementParams.FinalCullDistance, ElementParams.ShadowFrustum,
 										*StaticMesh->GetPathName(),
 										*MeshElement.MaterialRenderProxy->GetIncompleteMaterialWithFallback(ElementParams.FeatureLevel).GetFriendlyName());
@@ -1485,7 +1486,7 @@ void FHierarchicalStaticMeshSceneProxy::FillDynamicMeshElements(FMeshElementColl
 								OverlayMeshBatch.bSelectable = false;
 								OverlayMeshBatch.MaterialRenderProxy = OverlayMaterial->GetRenderProxy();
 								// make sure overlay is always rendered on top of base mesh
-								OverlayMeshBatch.MeshIdInPrimitive += LODModel.Sections.Num();
+								OverlayMeshBatch.MeshIdInPrimitive += TotalNumSections;
 								Collector.AddMesh(ElementParams.ViewIndex, OverlayMeshBatch);
 							}
 						}
@@ -1519,6 +1520,8 @@ void FHierarchicalStaticMeshSceneProxy::GetDynamicMeshElements(const TArray<cons
 
 	int32 MinVertsToSplitNode = CVarMinVertsToSplitNode.GetValueOnRenderThread();
 
+	const FMatrix WorldToLocal = GetLocalToWorld().Inverse();
+
 	for (int32 ViewIndex = 0; ViewIndex < Views.Num(); ViewIndex++)
 	{
 		if (VisibilityMap & (1 << ViewIndex))
@@ -1547,7 +1550,6 @@ void FHierarchicalStaticMeshSceneProxy::GetDynamicMeshElements(const TArray<cons
 
 				InstanceParams.View = View;
 
-				FMatrix WorldToLocal = GetLocalToWorld().Inverse();
 				bool bUseVectorCull = GUseVectorCull;
 				bool bIsOrtho = false;
 
@@ -1837,7 +1839,6 @@ void FHierarchicalStaticMeshSceneProxy::GetDynamicMeshElements(const TArray<cons
 					}
 					else
 					{
-						FMatrix WorldToLocal = GetLocalToWorld().Inverse();
 						FVector ViewOriginInLocalZero = WorldToLocal.TransformPosition(View->GetTemporalLODOrigin(0, bMultipleSections));
 						FVector ViewOriginInLocalOne  = WorldToLocal.TransformPosition(View->GetTemporalLODOrigin(1, bMultipleSections));
 						float LODPlanesMax[MAX_STATIC_MESH_LODS];
diff --git a/Engine/Source/Runtime/Engine/Private/Materials/Material.cpp b/Engine/Source/Runtime/Engine/Private/Materials/Material.cpp
--- a/Engine/Source/Runtime/Engine/Private/Materials/Material.cpp
+++ b/Engine/Source/Runtime/Engine/Private/Materials/Material.cpp
@@ -861,7 +861,7 @@ extern FMaterialResource* FindMaterialResource(const TArray<FMaterialResource*>&
 FMaterialResource* FindMaterialResource(TArray<FMaterialResource*>& MaterialResources, ERHIFeatureLevel::Type InFeatureLevel, EMaterialQualityLevel::Type QualityLevel, bool bAllowDefaultQuality)
 {
 	FMaterialResource* DefaultResource = nullptr;
-	for (int32 ResourceIndex = 0; ResourceIndex < MaterialResources.Num(); ++ResourceIndex)
+	for (int32 ResourceIndex = 0, NumMaterialResources = MaterialResources.Num(); ResourceIndex < NumMaterialResources; ++ResourceIndex)
 	{
 		FMaterialResource* CurrentResource = MaterialResources[ResourceIndex];
 		if (CurrentResource->GetFeatureLevel() == InFeatureLevel)
