Add comments on TC_EncodedReflectionCapture to indicate that it is vestigial
behavior nop
Use of it was removed in :
CL 22025079 Change the reflection capture's encoded HDR data to RG11B10 format.

diff --git a/Engine/Source/Developer/Android/AndroidTargetPlatform/Private/AndroidTargetPlatform.cpp b/Engine/Source/Developer/Android/AndroidTargetPlatform/Private/AndroidTargetPlatform.cpp
--- a/Engine/Source/Developer/Android/AndroidTargetPlatform/Private/AndroidTargetPlatform.cpp
+++ b/Engine/Source/Developer/Android/AndroidTargetPlatform/Private/AndroidTargetPlatform.cpp
@@ -502,6 +502,7 @@ void FAndroidTargetPlatform::GetTextureFormats( const UTexture* Texture, TArray<
 		{
 			FTextureFormatSettings FormatSettings;
 			Texture->GetDefaultFormatSettings(FormatSettings);
+			// TC_EncodedReflectionCapture is no longer used and could be deleted
 			if (FormatSettings.CompressionSettings == TC_EncodedReflectionCapture && !FormatSettings.CompressionNone)
 			{
 				TextureFormatName = FName(TEXT("ETC2_RGBA"));
diff --git a/Engine/Source/Developer/IOS/IOSTargetPlatform/Private/IOSTargetPlatform.cpp b/Engine/Source/Developer/IOS/IOSTargetPlatform/Private/IOSTargetPlatform.cpp
--- a/Engine/Source/Developer/IOS/IOSTargetPlatform/Private/IOSTargetPlatform.cpp
+++ b/Engine/Source/Developer/IOS/IOSTargetPlatform/Private/IOSTargetPlatform.cpp
@@ -638,6 +638,7 @@ void FIOSTargetPlatform::GetTextureFormats( const UTexture* Texture, TArray< TAr
 		{
 			FTextureFormatSettings FormatSettings;
 			Texture->GetDefaultFormatSettings(FormatSettings);
+			// TC_EncodedReflectionCapture is no longer used and could be deleted
 			if (FormatSettings.CompressionSettings == TC_EncodedReflectionCapture && !FormatSettings.CompressionNone)
 			{
 				TextureFormatName = FName(TEXT("ETC2_RGBA"));
diff --git a/Engine/Source/Runtime/Engine/Classes/Engine/TextureDefines.h b/Engine/Source/Runtime/Engine/Classes/Engine/TextureDefines.h
--- a/Engine/Source/Runtime/Engine/Classes/Engine/TextureDefines.h
+++ b/Engine/Source/Runtime/Engine/Classes/Engine/TextureDefines.h
@@ -325,6 +325,7 @@ enum ETextureSourceFormat : int
 };
 
 // This needs to be mirrored in EditorFactories.cpp.
+// TC_EncodedReflectionCapture is no longer used and could be deleted
 UENUM()
 enum TextureCompressionSettings : int
 {
@@ -342,7 +343,7 @@ enum TextureCompressionSettings : int
 	TC_BC7						UMETA(DisplayName = "BC7 (DX11, optional A)"),
 	TC_HalfFloat				UMETA(DisplayName = "Half Float (R16F)"),
 	TC_LQ				        UMETA(Hidden, DisplayName = "Low Quality (BGR565/BGR555A1)", ToolTip = "BGR565/BGR555A1, fallback to DXT1/DXT5 on Mac platform"),
-	TC_EncodedReflectionCapture	UMETA(Hidden),
+	TC_EncodedReflectionCapture	UMETA(Hidden), 
 	TC_SingleFloat				UMETA(DisplayName = "Single Float (R32F)"),
 	TC_HDR_F32					UMETA(DisplayName = "HDR High Precision (RGBA32F)"),
 	TC_MAX,
diff --git a/Engine/Source/Runtime/Engine/Private/TextureCube.cpp b/Engine/Source/Runtime/Engine/Private/TextureCube.cpp
--- a/Engine/Source/Runtime/Engine/Private/TextureCube.cpp
+++ b/Engine/Source/Runtime/Engine/Private/TextureCube.cpp
@@ -658,7 +658,7 @@ void UTextureCube::GetResourceSizeEx(FResourceSizeEx& CumulativeResourceSize)
 
 bool UTextureCube::NeedsLoadForTargetPlatform(const ITargetPlatform* TargetPlatform) const
 {
-	// TC_EncodedReflectionCapture is used only for EncodedCapture for Mobile and is hidden for the user.
+	// TC_EncodedReflectionCapture is no longer used and could be deleted
 	if (CompressionSettings == TC_EncodedReflectionCapture)
 	{
 		const static FName EncodedHDR(TEXT("EncodedHDR"));
