Adding batched parameter paths for more shared shader types.

diff --git a/Engine/Plugins/Runtime/Oculus/OculusVR/Source/OculusHMD/Private/OculusHMD_CustomPresent.cpp b/Engine/Plugins/Runtime/Oculus/OculusVR/Source/OculusHMD/Private/OculusHMD_CustomPresent.cpp
--- a/Engine/Plugins/Runtime/Oculus/OculusVR/Source/OculusHMD/Private/OculusHMD_CustomPresent.cpp
+++ b/Engine/Plugins/Runtime/Oculus/OculusVR/Source/OculusHMD/Private/OculusHMD_CustomPresent.cpp
@@ -502,8 +502,9 @@ void FCustomPresent::CopyTexture_RenderThread(FRHICommandListImmediate& RHICmdLi
 				TShaderMapRef<FOculusCubemapPS> PixelShader(ShaderMap);
 				GraphicsPSOInit.BoundShaderState.PixelShaderRHI = PixelShader.GetPixelShader();
 				SetGraphicsPipelineState(RHICmdList, GraphicsPSOInit, 0);
+
 				FRHISamplerState* SamplerState = DstRect.Size() == SrcRect.Size() ? TStaticSamplerState<SF_Point>::GetRHI() : TStaticSamplerState<SF_Bilinear>::GetRHI();
-				PixelShader->SetParameters(RHICmdList, SamplerState, SrcTextureRHI, FaceIndex);
+				SetShaderParametersLegacyPS(RHICmdList, PixelShader, SamplerState, SrcTextureRHI, FaceIndex);
 
 				RHICmdList.SetViewport(DstRect.Min.X, DstRect.Min.Y, 0.0f, DstRect.Max.X, DstRect.Max.Y, 1.0f);
 
diff --git a/Engine/Source/Runtime/D3D12RHI/Private/D3D12RenderTarget.cpp b/Engine/Source/Runtime/D3D12RHI/Private/D3D12RenderTarget.cpp
--- a/Engine/Source/Runtime/D3D12RHI/Private/D3D12RenderTarget.cpp
+++ b/Engine/Source/Runtime/D3D12RHI/Private/D3D12RenderTarget.cpp
@@ -257,8 +257,8 @@ void FD3D12CommandContext::ResolveTextureUsingShader(
 			SetGraphicsPipelineState(RHICmdList, GraphicsPSOInit, 0);
 			RHICmdList.SetBlendFactor(FLinearColor::White);
 
-			ResolveVertexShader->SetParameters(RHICmdList, SourceRect, DestRect, ResolveTargetDesc.Width, ResolveTargetDesc.Height);
-			ResolvePixelShader->SetParameters(RHICmdList, PixelShaderParameter);
+			SetShaderParametersLegacyVS(RHICmdList, ResolveVertexShader, SourceRect, DestRect, ResolveTargetDesc.Width, ResolveTargetDesc.Height);
+			SetShaderParametersLegacyPS(RHICmdList, ResolvePixelShader, PixelShaderParameter);
 
 			const uint32 TextureIndex = ResolvePixelShader->UnresolvedSurface.GetBaseIndex();
 
diff --git a/Engine/Source/Runtime/Engine/Private/CubemapUnwrapUtils.cpp b/Engine/Source/Runtime/Engine/Private/CubemapUnwrapUtils.cpp
--- a/Engine/Source/Runtime/Engine/Private/CubemapUnwrapUtils.cpp
+++ b/Engine/Source/Runtime/Engine/Private/CubemapUnwrapUtils.cpp
@@ -140,7 +140,7 @@ void FCubemapTexturePropertiesVS::SetParameters(FRHIBatchedShaderParameters& Bat
 void FCubemapTexturePropertiesVS::SetParameters(FRHICommandList& RHICmdList, const FMatrix& TransformValue)
 {
 	FRHIBatchedShaderParameters& BatchedParameters = RHICmdList.GetScratchShaderParameters();
-	SetShaderValue(RHICmdList, RHICmdList.GetBoundVertexShader(), Transform, (FMatrix44f)TransformValue);
+	SetParameters(BatchedParameters, TransformValue);
 	RHICmdList.SetBatchedShaderParameters(RHICmdList.GetBoundVertexShader(), BatchedParameters);
 }
 
diff --git a/Engine/Source/Runtime/Engine/Public/ScreenRendering.h b/Engine/Source/Runtime/Engine/Public/ScreenRendering.h
--- a/Engine/Source/Runtime/Engine/Public/ScreenRendering.h
+++ b/Engine/Source/Runtime/Engine/Public/ScreenRendering.h
@@ -76,14 +76,30 @@ public:
 	}
 	FScreenPS() {}
 
+	void SetParameters(FRHIBatchedShaderParameters& BatchedParameters, const FTexture* Texture)
+	{
+		SetTextureParameter(BatchedParameters,InTexture,InTextureSampler,Texture);
+	}
+
+	void SetParameters(FRHIBatchedShaderParameters& BatchedParameters, FRHISamplerState* SamplerStateRHI, FRHITexture* TextureRHI)
+	{
+		SetTextureParameter(BatchedParameters,InTexture,InTextureSampler,SamplerStateRHI,TextureRHI);
+	}
+
+	//UE_DEPRECATED(5.3, "SetParameters with FRHIBatchedShaderParameters should be used.")
 	void SetParameters(FRHICommandList& RHICmdList, const FTexture* Texture)
 	{
-		SetTextureParameter(RHICmdList, RHICmdList.GetBoundPixelShader(),InTexture,InTextureSampler,Texture);
+		FRHIBatchedShaderParameters& BatchedParameters = RHICmdList.GetScratchShaderParameters();
+		SetParameters(BatchedParameters, Texture);
+		RHICmdList.SetBatchedShaderParameters(RHICmdList.GetBoundPixelShader(), BatchedParameters);
 	}
 
+	//UE_DEPRECATED(5.3, "SetParameters with FRHIBatchedShaderParameters should be used.")
 	void SetParameters(FRHICommandList& RHICmdList, FRHISamplerState* SamplerStateRHI, FRHITexture* TextureRHI)
 	{
-		SetTextureParameter(RHICmdList, RHICmdList.GetBoundPixelShader(),InTexture,InTextureSampler,SamplerStateRHI,TextureRHI);
+		FRHIBatchedShaderParameters& BatchedParameters = RHICmdList.GetScratchShaderParameters();
+		SetParameters(BatchedParameters, SamplerStateRHI, TextureRHI);
+		RHICmdList.SetBatchedShaderParameters(RHICmdList.GetBoundPixelShader(), BatchedParameters);
 	}
 
 private:
@@ -120,16 +136,6 @@ public:
 	{
 	}
 	FScreenFromSlice0PS() {}
-
-	void SetParameters(FRHICommandList& RHICmdList, const FTexture* Texture)
-	{
-		FScreenPS::SetParameters(RHICmdList, Texture);
-	}
-
-	void SetParameters(FRHICommandList& RHICmdList, FRHISamplerState* SamplerStateRHI, FRHITexture* TextureRHI)
-	{
-		FScreenPS::SetParameters(RHICmdList, SamplerStateRHI, TextureRHI);
-	}
 };
 
 
@@ -151,14 +157,30 @@ public:
 	}
 	FScreenPSInvertAlpha() {}
 
+	void SetParameters(FRHIBatchedShaderParameters& BatchedParameters, const FTexture* Texture)
+	{
+		SetTextureParameter(BatchedParameters, InTexture, InTextureSampler, Texture);
+	}
+
+	void SetParameters(FRHIBatchedShaderParameters& BatchedParameters, FRHISamplerState* SamplerStateRHI, FRHITexture* TextureRHI)
+	{
+		SetTextureParameter(BatchedParameters, InTexture, InTextureSampler, SamplerStateRHI, TextureRHI);
+	}
+
+	//UE_DEPRECATED(5.3, "SetParameters with FRHIBatchedShaderParameters should be used.")
 	void SetParameters(FRHICommandList& RHICmdList, const FTexture* Texture)
 	{
-		SetTextureParameter(RHICmdList, RHICmdList.GetBoundPixelShader(), InTexture, InTextureSampler, Texture);
+		FRHIBatchedShaderParameters& BatchedParameters = RHICmdList.GetScratchShaderParameters();
+		SetParameters(BatchedParameters, Texture);
+		RHICmdList.SetBatchedShaderParameters(RHICmdList.GetBoundPixelShader(), BatchedParameters);
 	}
 
+	//UE_DEPRECATED(5.3, "SetParameters with FRHIBatchedShaderParameters should be used.")
 	void SetParameters(FRHICommandList& RHICmdList, FRHISamplerState* SamplerStateRHI, FRHITexture* TextureRHI)
 	{
-		SetTextureParameter(RHICmdList, RHICmdList.GetBoundPixelShader(), InTexture, InTextureSampler, SamplerStateRHI, TextureRHI);
+		FRHIBatchedShaderParameters& BatchedParameters = RHICmdList.GetScratchShaderParameters();
+		SetParameters(BatchedParameters, SamplerStateRHI, TextureRHI);
+		RHICmdList.SetBatchedShaderParameters(RHICmdList.GetBoundPixelShader(), BatchedParameters);
 	}
 
 private:
@@ -184,14 +206,30 @@ public:
 	}
 	FScreenPSsRGBSource() {}
 
+	void SetParameters(FRHIBatchedShaderParameters& BatchedParameters, const FTexture* Texture)
+	{
+		SetTextureParameter(BatchedParameters, InTexture, InTextureSampler, Texture);
+	}
+
+	void SetParameters(FRHIBatchedShaderParameters& BatchedParameters, FRHISamplerState* SamplerStateRHI, FRHITexture* TextureRHI)
+	{
+		SetTextureParameter(BatchedParameters, InTexture, InTextureSampler, SamplerStateRHI, TextureRHI);
+	}
+
+	//UE_DEPRECATED(5.3, "SetParameters with FRHIBatchedShaderParameters should be used.")
 	void SetParameters(FRHICommandList& RHICmdList, const FTexture* Texture)
 	{
-		SetTextureParameter(RHICmdList, RHICmdList.GetBoundPixelShader(), InTexture, InTextureSampler, Texture);
+		FRHIBatchedShaderParameters& BatchedParameters = RHICmdList.GetScratchShaderParameters();
+		SetParameters(BatchedParameters, Texture);
+		RHICmdList.SetBatchedShaderParameters(RHICmdList.GetBoundPixelShader(), BatchedParameters);
 	}
 
+	//UE_DEPRECATED(5.3, "SetParameters with FRHIBatchedShaderParameters should be used.")
 	void SetParameters(FRHICommandList& RHICmdList, FRHISamplerState* SamplerStateRHI, FRHITexture* TextureRHI)
 	{
-		SetTextureParameter(RHICmdList, RHICmdList.GetBoundPixelShader(), InTexture, InTextureSampler, SamplerStateRHI, TextureRHI);
+		FRHIBatchedShaderParameters& BatchedParameters = RHICmdList.GetScratchShaderParameters();
+		SetParameters(BatchedParameters, SamplerStateRHI, TextureRHI);
+		RHICmdList.SetBatchedShaderParameters(RHICmdList.GetBoundPixelShader(), BatchedParameters);
 	}
 
 private:
@@ -218,16 +256,32 @@ public:
 	}
 	FScreenPSMipLevel() {}
 
+	void SetParameters(FRHIBatchedShaderParameters& BatchedParameters, const FTexture* Texture, int MipLevel = 0)
+	{
+		SetTextureParameter(BatchedParameters, InTexture, InTextureSampler, Texture);
+		SetShaderValue(BatchedParameters, InMipLevelParameter, MipLevel);
+	}
+
+	void SetParameters(FRHIBatchedShaderParameters& BatchedParameters, FRHISamplerState* SamplerStateRHI, FRHITexture* TextureRHI, int MipLevel = 0)
+	{
+		SetTextureParameter(BatchedParameters, InTexture, InTextureSampler, SamplerStateRHI, TextureRHI);
+		SetShaderValue(BatchedParameters, InMipLevelParameter, MipLevel);
+	}
+
+	//UE_DEPRECATED(5.3, "SetParameters with FRHIBatchedShaderParameters should be used.")
 	void SetParameters(FRHICommandList& RHICmdList, const FTexture* Texture, int MipLevel = 0)
 	{
-		SetTextureParameter(RHICmdList, RHICmdList.GetBoundPixelShader(), InTexture, InTextureSampler, Texture);
-		SetShaderValue(RHICmdList, RHICmdList.GetBoundPixelShader(), InMipLevelParameter, MipLevel);
+		FRHIBatchedShaderParameters& BatchedParameters = RHICmdList.GetScratchShaderParameters();
+		SetParameters(BatchedParameters, Texture, MipLevel);
+		RHICmdList.SetBatchedShaderParameters(RHICmdList.GetBoundPixelShader(), BatchedParameters);
 	}
 
+	//UE_DEPRECATED(5.3, "SetParameters with FRHIBatchedShaderParameters should be used.")
 	void SetParameters(FRHICommandList& RHICmdList, FRHISamplerState* SamplerStateRHI, FRHITexture* TextureRHI, int MipLevel = 0)
 	{
-		SetTextureParameter(RHICmdList, RHICmdList.GetBoundPixelShader(), InTexture, InTextureSampler, SamplerStateRHI, TextureRHI);
-		SetShaderValue(RHICmdList, RHICmdList.GetBoundPixelShader(), InMipLevelParameter, MipLevel);
+		FRHIBatchedShaderParameters& BatchedParameters = RHICmdList.GetScratchShaderParameters();
+		SetParameters(BatchedParameters, SamplerStateRHI, TextureRHI, MipLevel);
+		RHICmdList.SetBatchedShaderParameters(RHICmdList.GetBoundPixelShader(), BatchedParameters);
 	}
 
 private:
@@ -255,16 +309,32 @@ public:
 	}
 	FScreenPSsRGBSourceMipLevel() {}
 
+	void SetParameters(FRHIBatchedShaderParameters& BatchedParameters, const FTexture* Texture, int MipLevel = 0)
+	{
+		SetTextureParameter(BatchedParameters, InTexture, InTextureSampler, Texture);
+		SetShaderValue(BatchedParameters, InMipLevelParameter, MipLevel);
+	}
+
+	void SetParameters(FRHIBatchedShaderParameters& BatchedParameters, FRHISamplerState* SamplerStateRHI, FRHITexture* TextureRHI, int MipLevel = 0)
+	{
+		SetTextureParameter(BatchedParameters, InTexture, InTextureSampler, SamplerStateRHI, TextureRHI);
+		SetShaderValue(BatchedParameters, InMipLevelParameter, MipLevel);
+	}
+
+	//UE_DEPRECATED(5.3, "SetParameters with FRHIBatchedShaderParameters should be used.")
 	void SetParameters(FRHICommandList& RHICmdList, const FTexture* Texture, int MipLevel = 0)
 	{
-		SetTextureParameter(RHICmdList, RHICmdList.GetBoundPixelShader(), InTexture, InTextureSampler, Texture);
-		SetShaderValue(RHICmdList, RHICmdList.GetBoundPixelShader(), InMipLevelParameter, MipLevel);
+		FRHIBatchedShaderParameters& BatchedParameters = RHICmdList.GetScratchShaderParameters();
+		SetParameters(BatchedParameters, Texture, MipLevel);
+		RHICmdList.SetBatchedShaderParameters(RHICmdList.GetBoundPixelShader(), BatchedParameters);
 	}
 
+	//UE_DEPRECATED(5.3, "SetParameters with FRHIBatchedShaderParameters should be used.")
 	void SetParameters(FRHICommandList& RHICmdList, FRHISamplerState* SamplerStateRHI, FRHITexture* TextureRHI, int MipLevel = 0)
 	{
-		SetTextureParameter(RHICmdList, RHICmdList.GetBoundPixelShader(), InTexture, InTextureSampler, SamplerStateRHI, TextureRHI);
-		SetShaderValue(RHICmdList, RHICmdList.GetBoundPixelShader(), InMipLevelParameter, MipLevel);
+		FRHIBatchedShaderParameters& BatchedParameters = RHICmdList.GetScratchShaderParameters();
+		SetParameters(BatchedParameters, SamplerStateRHI, TextureRHI, MipLevel);
+		RHICmdList.SetBatchedShaderParameters(RHICmdList.GetBoundPixelShader(), BatchedParameters);
 	}
 
 private:
@@ -289,15 +359,31 @@ public:
 
     FScreenPS_OSE() {}
 
-    void SetParameters(FRHICommandList& RHICmdList, const FTexture* Texture)
+    void SetParameters(FRHIBatchedShaderParameters& BatchedParameters, const FTexture* Texture)
     {
-        SetTextureParameter(RHICmdList, RHICmdList.GetBoundPixelShader(),InTexture,InTextureSampler,Texture);
-    }
+		SetTextureParameter(BatchedParameters, InTexture, InTextureSampler, Texture);
+	}
 
-    void SetParameters(FRHICommandList& RHICmdList, FRHISamplerState* SamplerStateRHI, FRHITexture* TextureRHI)
+    void SetParameters(FRHIBatchedShaderParameters& BatchedParameters, FRHISamplerState* SamplerStateRHI, FRHITexture* TextureRHI)
     {
-        SetTextureParameter(RHICmdList, RHICmdList.GetBoundPixelShader(),InTexture,InTextureSampler,SamplerStateRHI,TextureRHI);
-    }
+		SetTextureParameter(BatchedParameters, InTexture, InTextureSampler, SamplerStateRHI, TextureRHI);
+	}
+
+	//UE_DEPRECATED(5.3, "SetParameters with FRHIBatchedShaderParameters should be used.")
+	void SetParameters(FRHICommandList& RHICmdList, const FTexture* Texture)
+	{
+		FRHIBatchedShaderParameters& BatchedParameters = RHICmdList.GetScratchShaderParameters();
+		SetParameters(BatchedParameters, Texture);
+		RHICmdList.SetBatchedShaderParameters(RHICmdList.GetBoundPixelShader(), BatchedParameters);
+	}
+
+	//UE_DEPRECATED(5.3, "SetParameters with FRHIBatchedShaderParameters should be used.")
+	void SetParameters(FRHICommandList& RHICmdList, FRHISamplerState* SamplerStateRHI, FRHITexture* TextureRHI)
+	{
+		FRHIBatchedShaderParameters& BatchedParameters = RHICmdList.GetScratchShaderParameters();
+		SetParameters(BatchedParameters, SamplerStateRHI, TextureRHI);
+		RHICmdList.SetBatchedShaderParameters(RHICmdList.GetBoundPixelShader(), BatchedParameters);
+	}
 
 private:
 	LAYOUT_FIELD(FShaderResourceParameter, InTexture);
@@ -320,6 +406,12 @@ public:
 	}
 	FScreenVS() {}
 
+	void SetParameters(FRHIBatchedShaderParameters& BatchedParameters, FRHIUniformBuffer* ViewUniformBuffer)
+	{
+		FGlobalShader::SetParameters<FViewUniformShaderParameters>(BatchedParameters, ViewUniformBuffer);
+	}
+
+	//UE_DEPRECATED(5.3, "SetParameters with FRHIBatchedShaderParameters should be used.")
 	void SetParameters(FRHICommandList& RHICmdList, FRHIUniformBuffer* ViewUniformBuffer)
 	{
 		FGlobalShader::SetParameters<FViewUniformShaderParameters>(RHICmdList, RHICmdList.GetBoundVertexShader(), ViewUniformBuffer);
diff --git a/Engine/Source/Runtime/Engine/Public/StereoLayerRendering.h b/Engine/Source/Runtime/Engine/Public/StereoLayerRendering.h
--- a/Engine/Source/Runtime/Engine/Public/StereoLayerRendering.h
+++ b/Engine/Source/Runtime/Engine/Public/StereoLayerRendering.h
@@ -32,13 +32,11 @@ public:
 	}
 	FStereoLayerVS() {}
 
-	void SetParameters(FRHICommandList& RHICmdList, FVector2D QuadSize, FBox2D UVRect, const FMatrix& ViewProjection, const FMatrix& World)
+	void SetParameters(FRHIBatchedShaderParameters& BatchedParameters, FVector2D QuadSize, FBox2D UVRect, const FMatrix& ViewProjection, const FMatrix& World)
 	{
-		FRHIVertexShader* VS = RHICmdList.GetBoundVertexShader();
-
 		if (InQuadAdjust.IsBound())
 		{
-			SetShaderValue(RHICmdList, VS, InQuadAdjust, FVector2f(QuadSize));
+			SetShaderValue(BatchedParameters, InQuadAdjust, FVector2f(QuadSize));
 		}
 
 		if (InUVAdjust.IsBound())
@@ -48,20 +46,28 @@ public:
 			UVAdjust.Y = static_cast<float>(UVRect.Min.Y);
 			UVAdjust.Z = static_cast<float>(UVRect.Max.X - UVRect.Min.X);
 			UVAdjust.W = static_cast<float>(UVRect.Max.Y - UVRect.Min.Y);
-			SetShaderValue(RHICmdList, VS, InUVAdjust, UVAdjust);
+			SetShaderValue(BatchedParameters, InUVAdjust, UVAdjust);
 		}
 
 		if (InViewProjection.IsBound())
 		{
-			SetShaderValue(RHICmdList, VS, InViewProjection, (FMatrix44f)ViewProjection);
+			SetShaderValue(BatchedParameters, InViewProjection, (FMatrix44f)ViewProjection);
 		}
 
 		if (InWorld.IsBound())
 		{
-			SetShaderValue(RHICmdList, VS, InWorld, (FMatrix44f)World);
+			SetShaderValue(BatchedParameters, InWorld, (FMatrix44f)World);
 		}
 	}
 
+	UE_DEPRECATED(5.3, "SetParameters with FRHIBatchedShaderParameters should be used.")
+	void SetParameters(FRHICommandList& RHICmdList, FVector2D QuadSize, FBox2D UVRect, const FMatrix& ViewProjection, const FMatrix& World)
+	{
+		FRHIBatchedShaderParameters& BatchedParameters = RHICmdList.GetScratchShaderParameters();
+		SetParameters(BatchedParameters, QuadSize, UVRect, ViewProjection, World);
+		RHICmdList.SetBatchedShaderParameters(RHICmdList.GetBoundVertexShader(), BatchedParameters);
+	}
+
 private:
 	LAYOUT_FIELD(FShaderParameter, InQuadAdjust);
 	LAYOUT_FIELD(FShaderParameter, InUVAdjust);
@@ -76,19 +82,25 @@ public:
 
 	static bool ShouldCompilePermutation(const FGlobalShaderPermutationParameters& Parameters) { return true; }
 
-	void SetParameters(FRHICommandList& RHICmdList, FRHISamplerState* SamplerStateRHI, FRHITexture* TextureRHI, bool bIsOpaque)
+	void SetParameters(FRHIBatchedShaderParameters& BatchedParameters, FRHISamplerState* SamplerStateRHI, FRHITexture* TextureRHI, bool bIsOpaque)
 	{
-		FRHIPixelShader* PS = RHICmdList.GetBoundPixelShader();
-
-		SetTextureParameter(RHICmdList, PS, InTexture, InTextureSampler, SamplerStateRHI, TextureRHI);
+		SetTextureParameter(BatchedParameters, InTexture, InTextureSampler, SamplerStateRHI, TextureRHI);
 
 		if (InIsOpaque.IsBound())
 		{
 			const float OpaqueVal = bIsOpaque ? 1.0 : 0.0;
-			SetShaderValue(RHICmdList, PS, InIsOpaque, OpaqueVal);
+			SetShaderValue(BatchedParameters, InIsOpaque, OpaqueVal);
 		}
 	}
 
+	UE_DEPRECATED(5.3, "SetParameters with FRHIBatchedShaderParameters should be used.")
+	void SetParameters(FRHICommandList& RHICmdList, FRHISamplerState* SamplerStateRHI, FRHITexture* TextureRHI, bool bIsOpaque)
+	{
+		FRHIBatchedShaderParameters& BatchedParameters = RHICmdList.GetScratchShaderParameters();
+		SetParameters(BatchedParameters, SamplerStateRHI, TextureRHI, bIsOpaque);
+		RHICmdList.SetBatchedShaderParameters(RHICmdList.GetBoundPixelShader(), BatchedParameters);
+	}
+
 protected:
 	FStereoLayerPS_Base(const ShaderMetaType::CompiledShaderInitializerType& Initializer, const TCHAR* TextureParamName) :
 		FGlobalShader(Initializer) 
diff --git a/Engine/Source/Runtime/Engine/Public/UpdateTextureShaders.h b/Engine/Source/Runtime/Engine/Public/UpdateTextureShaders.h
--- a/Engine/Source/Runtime/Engine/Public/UpdateTextureShaders.h
+++ b/Engine/Source/Runtime/Engine/Public/UpdateTextureShaders.h
@@ -30,6 +30,29 @@ public:
 		DestTexture.Bind(Initializer.ParameterMap, TEXT("TDestTexture"), SPF_Mandatory);
 	}
 
+	static const TCHAR* GetSourceFilename() { return TEXT("/Engine/Private/UpdateTextureShaders.usf"); }
+	static const TCHAR* GetFunctionName() { return TEXT("TUpdateTexture2DSubresourceCS"); }
+
+	static bool ShouldCompilePermutation(const FGlobalShaderPermutationParameters& Parameters)
+	{
+		return IsFeatureLevelSupported(Parameters.Platform, ERHIFeatureLevel::SM5);
+	}
+
+	void SetParameters(FRHIBatchedShaderParameters& BatchedParameters, const FUpdateTextureRegion2D& UpdateRegionInBlocks, uint32 SrcElementPitch)
+	{
+		const uint32 DstPosSize[4] =
+		{
+			UpdateRegionInBlocks.DestX,
+			UpdateRegionInBlocks.DestY,
+			UpdateRegionInBlocks.Width,
+			UpdateRegionInBlocks.Height
+		};
+		FIntVector SrcPosPitch(UpdateRegionInBlocks.SrcX, UpdateRegionInBlocks.SrcY, SrcElementPitch);
+
+		SetShaderValue(BatchedParameters, DestPosSizeParameter, DstPosSize);
+		SetShaderValue(BatchedParameters, SrcPosPitchParameter, SrcPosPitch);
+	}
+
 	LAYOUT_FIELD(FShaderParameter, SrcPosPitchParameter);
 	LAYOUT_FIELD(FShaderResourceParameter, SrcBuffer);
 	LAYOUT_FIELD(FShaderParameter, DestPosSizeParameter);
@@ -52,10 +75,6 @@ public:
 		: FUpdateTexture2DSubresourceCS(Initializer)
 	{}
 
-	static const TCHAR* GetSourceFilename() { return TEXT("/Engine/Private/UpdateTextureShaders.usf"); }
-	
-	static const TCHAR* GetFunctionName() { return TEXT("TUpdateTexture2DSubresourceCS"); }
-
 	static void ModifyCompilationEnvironment(const FGlobalShaderPermutationParameters& Parameters, FShaderCompilerEnvironment& OutEnvironment)
 	{
 		const TCHAR* Type = nullptr;
@@ -67,11 +86,6 @@ public:
 		case EUpdateTextureValueType::Uint32: OutEnvironment.SetDefine(TEXT("COMPONENT_TYPE"), TEXT("uint4")); break;
 		}
 	}
-
-	static bool ShouldCompilePermutation(const FGlobalShaderPermutationParameters& Parameters)
-	{
-		return IsFeatureLevelSupported(Parameters.Platform, ERHIFeatureLevel::SM5);
-	}
 };
 
 template <typename ShaderType>
diff --git a/Engine/Source/Runtime/HeadMountedDisplay/Private/DefaultStereoLayers.cpp b/Engine/Source/Runtime/HeadMountedDisplay/Private/DefaultStereoLayers.cpp
--- a/Engine/Source/Runtime/HeadMountedDisplay/Private/DefaultStereoLayers.cpp
+++ b/Engine/Source/Runtime/HeadMountedDisplay/Private/DefaultStereoLayers.cpp
@@ -141,15 +141,17 @@ void FDefaultStereoLayers::StereoLayerRender(FRHICommandListImmediate& RHICmdLis
 		}
 
 		// Set shader uniforms
-		VertexShader->SetParameters(
+		SetShaderParametersLegacyVS(
 			RHICmdList,
+			VertexShader,
 			QuadSize,
 			Layer.UVRect,
 			RenderParams.RenderMatrices[static_cast<int>(Layer.PositionType)],
 			LayerMatrix);
 
-		PixelShader->SetParameters(
+		SetShaderParametersLegacyPS(
 			RHICmdList,
+			PixelShader,
 			TStaticSamplerState<SF_Trilinear>::GetRHI(),
 			Layer.Texture,
 			bIsOpaque);
diff --git a/Engine/Source/Runtime/RenderCore/Private/ResolveShader.cpp b/Engine/Source/Runtime/RenderCore/Private/ResolveShader.cpp
--- a/Engine/Source/Runtime/RenderCore/Private/ResolveShader.cpp
+++ b/Engine/Source/Runtime/RenderCore/Private/ResolveShader.cpp
@@ -31,6 +31,10 @@ void FResolveDepthPS::ModifyCompilationEnvironment(const FGlobalShaderPermutatio
 	FGlobalShader::ModifyCompilationEnvironment(Parameters, OutEnvironment);
 }
 
+void FResolveDepthPS::SetParameters(FRHIBatchedShaderParameters& BatchedParameters, FParameter)
+{
+}
+
 void FResolveDepthPS::SetParameters(FRHICommandList& RHICmdList, FParameter)
 {
 }
@@ -183,9 +187,16 @@ bool FResolveSingleSamplePS::ShouldCompilePermutation(const FGlobalShaderPermuta
 	return FDataDrivenShaderPlatformInfo::GetIsLanguageD3D(Parameters.Platform);
 }
 
+void FResolveSingleSamplePS::SetParameters(FRHIBatchedShaderParameters& BatchedParameters, uint32 SingleSampleIndexValue)
+{
+	SetShaderValue(BatchedParameters, SingleSampleIndex, SingleSampleIndexValue);
+}
+
 void FResolveSingleSamplePS::SetParameters(FRHICommandList& RHICmdList, uint32 SingleSampleIndexValue)
 {
-	SetShaderValue(RHICmdList, RHICmdList.GetBoundPixelShader(),SingleSampleIndex,SingleSampleIndexValue);
+	FRHIBatchedShaderParameters& BatchedParameters = RHICmdList.GetScratchShaderParameters();
+	SetParameters(BatchedParameters, SingleSampleIndexValue);
+	RHICmdList.SetBatchedShaderParameters(RHICmdList.GetBoundPixelShader(), BatchedParameters);
 }
 
 // FResolveVS
@@ -203,7 +214,7 @@ bool FResolveVS::ShouldCompilePermutation(const FGlobalShaderPermutationParamete
 	return true;
 }
 
-void FResolveVS::SetParameters(FRHICommandList& RHICmdList, const FResolveRect& SrcBounds, const FResolveRect& DstBounds, uint32 DstSurfaceWidth, uint32 DstSurfaceHeight)
+void FResolveVS::SetParameters(FRHIBatchedShaderParameters& BatchedParameters, const FResolveRect& SrcBounds, const FResolveRect& DstBounds, uint32 DstSurfaceWidth, uint32 DstSurfaceHeight)
 {
 	// Generate the vertices used to copy from the source surface to the destination surface.
 	const float MinU = (float)SrcBounds.X1;
@@ -215,8 +226,15 @@ void FResolveVS::SetParameters(FRHICommandList& RHICmdList, const FResolveRect&
 	const float MaxX = -1.f + DstBounds.X2 / ((float)DstSurfaceWidth * 0.5f);
 	const float MaxY = +1.f - DstBounds.Y2 / ((float)DstSurfaceHeight * 0.5f);
 
-	SetShaderValue(RHICmdList, RHICmdList.GetBoundVertexShader(), PositionMinMax, FVector4f(MinX, MinY, MaxX, MaxY));
-	SetShaderValue(RHICmdList, RHICmdList.GetBoundVertexShader(), UVMinMax, FVector4f(MinU, MinV, MaxU, MaxV));
+	SetShaderValue(BatchedParameters, PositionMinMax, FVector4f(MinX, MinY, MaxX, MaxY));
+	SetShaderValue(BatchedParameters, UVMinMax, FVector4f(MinU, MinV, MaxU, MaxV));
+}
+
+void FResolveVS::SetParameters(FRHICommandList& RHICmdList, const FResolveRect& SrcBounds, const FResolveRect& DstBounds, uint32 DstSurfaceWidth, uint32 DstSurfaceHeight)
+{
+	FRHIBatchedShaderParameters& BatchedParameters = RHICmdList.GetScratchShaderParameters();
+	SetParameters(BatchedParameters, SrcBounds, DstBounds, DstSurfaceWidth, DstSurfaceHeight);
+	RHICmdList.SetBatchedShaderParameters(RHICmdList.GetBoundVertexShader(), BatchedParameters);
 }
 
 // FResolveArrayVS
@@ -243,11 +261,6 @@ void FResolveArrayVS::ModifyCompilationEnvironment(const FGlobalShaderPermutatio
 	OutEnvironment.SetDefine(TEXT("DEPTH_RESOLVE_TEXTUREARRAY"), 1);
 }
 
-void FResolveArrayVS::SetParameters(FRHICommandList& RHICmdList, const FResolveRect& SrcBounds, const FResolveRect& DstBounds, uint32 DstSurfaceWidth, uint32 DstSurfaceHeight)
-{
-	return FResolveVS::SetParameters(RHICmdList, SrcBounds, DstBounds, DstSurfaceWidth, DstSurfaceHeight);
-}
-
 void CreateResolveShaders()
 {
 	ForceInitGlobalShaderType<FResolveDepthPS>();
diff --git a/Engine/Source/Runtime/RenderCore/Public/MediaShaders.h b/Engine/Source/Runtime/RenderCore/Public/MediaShaders.h
--- a/Engine/Source/Runtime/RenderCore/Public/MediaShaders.h
+++ b/Engine/Source/Runtime/RenderCore/Public/MediaShaders.h
@@ -178,6 +178,7 @@ public:
 		: FGlobalShader(Initializer)
 	{ }
 
+	RENDERCORE_API void SetParameters(FRHIBatchedShaderParameters& BatchedParameters, TRefCountPtr<FRHITexture2D> AYUVTexture, const FMatrix& ColorTransform, const FVector& YUVOffset, bool SrgbToLinear);
 	RENDERCORE_API void SetParameters(FRHICommandList& RHICmdList, TRefCountPtr<FRHITexture2D> AYUVTexture, const FMatrix& ColorTransform, const FVector& YUVOffset, bool SrgbToLinear);
 };
 
@@ -199,6 +200,7 @@ public:
 		: FGlobalShader(Initializer)
 	{ }
 
+	RENDERCORE_API void SetParameters(FRHIBatchedShaderParameters& BatchedParameters, TRefCountPtr<FRHITexture2D> BMPTexture, const FIntPoint& OutputDimensions, bool SrgbToLinear);
 	RENDERCORE_API void SetParameters(FRHICommandList& RHICmdList, TRefCountPtr<FRHITexture2D> BMPTexture, const FIntPoint& OutputDimensions, bool SrgbToLinear);
 };
 
@@ -222,6 +224,7 @@ public:
 		: FGlobalShader(Initializer)
 	{ }
 
+	RENDERCORE_API void SetParameters(FRHIBatchedShaderParameters& BatchedParameters, const FIntPoint & TexDim, FShaderResourceViewRHIRef SRV_Y, FShaderResourceViewRHIRef SRV_UV, const FIntPoint& OutputDimensions, const FMatrix& ColorTransform, const FVector& YUVOffset, bool SrgbToLinear);
 	RENDERCORE_API void SetParameters(FRHICommandList& RHICmdList, const FIntPoint & TexDim, FShaderResourceViewRHIRef SRV_Y, FShaderResourceViewRHIRef SRV_UV, const FIntPoint& OutputDimensions, const FMatrix& ColorTransform, const FVector& YUVOffset, bool SrgbToLinear);
 };
 
@@ -246,6 +249,7 @@ public:
 		: FGlobalShader(Initializer)
 	{ }
 
+	RENDERCORE_API void SetParameters(FRHIBatchedShaderParameters& BatchedParameters, TRefCountPtr<FRHITexture2D> NV12Texture, const FIntPoint& OutputDimensions, const FMatrix& ColorTransform, const FVector& YUVOffset, bool SrgbToLinear);
 	RENDERCORE_API void SetParameters(FRHICommandList& RHICmdList, TRefCountPtr<FRHITexture2D> NV12Texture, const FIntPoint& OutputDimensions, const FMatrix& ColorTransform, const FVector& YUVOffset, bool SrgbToLinear);
 };
 
@@ -269,6 +273,7 @@ public:
 		: FGlobalShader(Initializer)
 	{ }
 
+	RENDERCORE_API void SetParameters(FRHIBatchedShaderParameters& BatchedParameters, TRefCountPtr<FRHITexture2D> NV21Texture, const FIntPoint& OutputDimensions, const FMatrix& ColorTransform, const FVector& YUVOffset, bool SrgbToLinear);
 	RENDERCORE_API void SetParameters(FRHICommandList& RHICmdList, TRefCountPtr<FRHITexture2D> NV21Texture, const FIntPoint& OutputDimensions, const FMatrix& ColorTransform, const FVector& YUVOffset, bool SrgbToLinear);
 };
 
@@ -288,6 +293,7 @@ public:
 		: FGlobalShader(Initializer)
 	{ }
 
+	RENDERCORE_API void SetParameters(FRHIBatchedShaderParameters& BatchedParameters, const FIntPoint& TexDim, FShaderResourceViewRHIRef SRV_Y, FShaderResourceViewRHIRef SRV_UV, const FIntPoint& OutputDimensions, const FMatrix44f& ColorTransform, const FMatrix44f& CSTransform, bool bIsST2084);
 	RENDERCORE_API void SetParameters(FRHICommandList& RHICmdList, const FIntPoint& TexDim, FShaderResourceViewRHIRef SRV_Y, FShaderResourceViewRHIRef SRV_UV, const FIntPoint& OutputDimensions, const FMatrix44f& ColorTransform, const FMatrix44f& CSTransform, bool bIsST2084);
 };
 
@@ -307,6 +313,7 @@ public:
 		: FGlobalShader(Initializer)
 	{ }
 
+	RENDERCORE_API void SetParameters(FRHIBatchedShaderParameters& BatchedParameters, const FIntPoint& TexDim, TRefCountPtr<FRHITexture2D> NV12Texture, const FIntPoint& OutputDimensions, const FMatrix44f& ColorTransform, const FMatrix44f& CSTransform, bool bIsST2084);
 	RENDERCORE_API void SetParameters(FRHICommandList& RHICmdList, const FIntPoint& TexDim, TRefCountPtr<FRHITexture2D> NV12Texture, const FIntPoint& OutputDimensions, const FMatrix44f& ColorTransform, const FMatrix44f& CSTransform, bool bIsST2084);
 };
 
@@ -326,6 +333,7 @@ public:
 		: FGlobalShader(Initializer)
 	{ }
 
+	RENDERCORE_API void SetParameters(FRHIBatchedShaderParameters& BatchedParameters, const FIntPoint& TexDim, FShaderResourceViewRHIRef SRV_Y, FShaderResourceViewRHIRef SRV_U, FShaderResourceViewRHIRef SRV_V, const FIntPoint& OutputDimensions, const FMatrix44f& ColorTransform, const FMatrix44f& CSTransform, bool bIsST2084);
 	RENDERCORE_API void SetParameters(FRHICommandList& RHICmdList, const FIntPoint& TexDim, FShaderResourceViewRHIRef SRV_Y, FShaderResourceViewRHIRef SRV_U, FShaderResourceViewRHIRef SRV_V, const FIntPoint& OutputDimensions, const FMatrix44f& ColorTransform, const FMatrix44f& CSTransform, bool bIsST2084);
 };
 
@@ -348,6 +356,7 @@ public:
 	{ }
 
 
+	RENDERCORE_API void SetParameters(FRHIBatchedShaderParameters& BatchedParameters, TRefCountPtr<FRHITexture2D> RGBTexture, const FIntPoint& OutputDimensions, bool SrgbToLinear, bool bIsST2084, const FMatrix44f& CSTransform);
 	RENDERCORE_API void SetParameters(FRHICommandList& RHICmdList, TRefCountPtr<FRHITexture2D> RGBTexture, const FIntPoint& OutputDimensions, bool SrgbToLinear, bool bIsST2084, const FMatrix44f& CSTransform);
 };
 
@@ -373,6 +382,7 @@ public:
 		: FGlobalShader(Initializer)
 	{ }
 
+	RENDERCORE_API void SetParameters(FRHIBatchedShaderParameters& BatchedParameters, TRefCountPtr<FRHITexture2D> LumaTexture, TRefCountPtr<FRHITexture2D> CbCrTexture, const FMatrix& ColorTransform, const FVector& YUVOffset, bool SrgbToLinear);
 	RENDERCORE_API void SetParameters(FRHICommandList& RHICmdList, TRefCountPtr<FRHITexture2D> LumaTexture, TRefCountPtr<FRHITexture2D> CbCrTexture, const FMatrix& ColorTransform, const FVector& YUVOffset, bool SrgbToLinear);
 };
 
@@ -409,6 +419,7 @@ public:
 		: FGlobalShader(Initializer)
 	{ }
 
+	RENDERCORE_API void SetParameters(FRHIBatchedShaderParameters& BatchedParameters, TRefCountPtr<FRHITexture2D> UYVYTexture, const FMatrix& ColorTransform, const FVector& YUVOffset, bool SrgbToLinear);
 	RENDERCORE_API void SetParameters(FRHICommandList& RHICmdList, TRefCountPtr<FRHITexture2D> UYVYTexture, const FMatrix& ColorTransform, const FVector& YUVOffset, bool SrgbToLinear);
 };
 
@@ -432,6 +443,7 @@ public:
 	{ }
 
 
+	RENDERCORE_API void SetParameters(FRHIBatchedShaderParameters& BatchedParameters, TRefCountPtr<FRHITexture2D> YTexture, TRefCountPtr<FRHITexture2D> UTexture, TRefCountPtr<FRHITexture2D> VTexture, const FIntPoint& OutputDimensions, const FMatrix& ColorTransform, const FVector& YUVOffset, bool SrgbToLinear);
 	RENDERCORE_API void SetParameters(FRHICommandList& RHICmdList, TRefCountPtr<FRHITexture2D> YTexture, TRefCountPtr<FRHITexture2D> UTexture, TRefCountPtr<FRHITexture2D> VTexture, const FIntPoint& OutputDimensions, const FMatrix& ColorTransform, const FVector& YUVOffset, bool SrgbToLinear);
 };
 
@@ -453,6 +465,7 @@ public:
 		: FGlobalShader(Initializer)
 	{ }
 
+	RENDERCORE_API void SetParameters(FRHIBatchedShaderParameters& BatchedParameters, TRefCountPtr<FRHITexture2D> YUVTexture, const FIntPoint& OutputDimensions, const FMatrix& ColorTransform, const FVector& YUVOffset, bool SrgbToLinear);
 	RENDERCORE_API void SetParameters(FRHICommandList& RHICmdList, TRefCountPtr<FRHITexture2D> YUVTexture, const FIntPoint& OutputDimensions, const FMatrix& ColorTransform, const FVector& YUVOffset, bool SrgbToLinear);
 };
 
@@ -474,6 +487,7 @@ public:
 		: FGlobalShader(Initializer)
 	{ }
 
+	RENDERCORE_API void SetParameters(FRHIBatchedShaderParameters& BatchedParameters, FShaderResourceViewRHIRef SRV_Y, const FMatrix& ColorTransform, const FVector& YUVOffset, bool SrgbToLinear);
 	RENDERCORE_API void SetParameters(FRHICommandList& RHICmdList, FShaderResourceViewRHIRef SRV_Y, const FMatrix& ColorTransform, const FVector& YUVOffset, bool SrgbToLinear);
 };
 
@@ -498,6 +512,7 @@ public:
 		: FGlobalShader(Initializer)
 	{ }
 
+	RENDERCORE_API void SetParameters(FRHIBatchedShaderParameters& BatchedParameters, TRefCountPtr<FRHITexture2D> YUY2Texture, const FIntPoint& OutputDimensions, const FMatrix& ColorTransform, const FVector& YUVOffset, bool SrgbToLinear);
 	RENDERCORE_API void SetParameters(FRHICommandList& RHICmdList, TRefCountPtr<FRHITexture2D> YUY2Texture, const FIntPoint& OutputDimensions, const FMatrix& ColorTransform, const FVector& YUVOffset, bool SrgbToLinear);
 };
 
@@ -522,6 +537,7 @@ public:
 		: FGlobalShader(Initializer)
 	{ }
 
+	RENDERCORE_API void SetParameters(FRHIBatchedShaderParameters& BatchedParameters, TRefCountPtr<FRHITexture2D> YVYUTexture, const FMatrix& ColorTransform, const FVector& YUVOffset, bool SrgbToLinear);
 	RENDERCORE_API void SetParameters(FRHICommandList& RHICmdList, TRefCountPtr<FRHITexture2D> YVYUTexture, const FMatrix& ColorTransform, const FVector& YUVOffset, bool SrgbToLinear);
 };
 
@@ -543,6 +559,7 @@ public:
 		: FGlobalShader(Initializer)
 	{ }
 
+	RENDERCORE_API void SetParameters(FRHIBatchedShaderParameters& BatchedParameters, TRefCountPtr<FRHITexture2D> RGBATexture, const FVector4f& ColorTransform, bool LinearToSrgb);
 	RENDERCORE_API void SetParameters(FRHICommandList& RHICmdList, TRefCountPtr<FRHITexture2D> RGBATexture, const FVector4f& ColorTransform, bool LinearToSrgb);
 };
 
@@ -562,6 +579,7 @@ public:
 		: FGlobalShader(Initializer)
 	{ }
 
+	RENDERCORE_API void SetParameters(FRHIBatchedShaderParameters& BatchedParameters, FTextureRHIRef TextureExt, FSamplerStateRHIRef SamplerState, const FLinearColor & ScaleRotation, const FLinearColor & Offset);
 	RENDERCORE_API void SetParameters(FRHICommandList& RHICmdList, FTextureRHIRef TextureExt, FSamplerStateRHIRef SamplerState, const FLinearColor & ScaleRotation, const FLinearColor & Offset);
 };
 
diff --git a/Engine/Source/Runtime/RenderCore/Public/OculusShaders.h b/Engine/Source/Runtime/RenderCore/Public/OculusShaders.h
--- a/Engine/Source/Runtime/RenderCore/Public/OculusShaders.h
+++ b/Engine/Source/Runtime/RenderCore/Public/OculusShaders.h
@@ -75,14 +75,14 @@ public:
 	}
 	FOculusAlphaInverseShader() {}
 
-	RENDERCORE_API void SetParameters(FRHICommandList& RHICmdList, const FTexture* Texture)
+	void SetParameters(FRHIBatchedShaderParameters& BatchedParameters, const FTexture* Texture)
 	{
-		SetTextureParameter(RHICmdList, RHICmdList.GetBoundPixelShader(), InTexture, InTextureSampler, Texture);
+		SetTextureParameter(BatchedParameters, InTexture, InTextureSampler, Texture);
 	}
 
-	RENDERCORE_API void SetParameters(FRHICommandList& RHICmdList, FRHISamplerState* SamplerStateRHI, FRHITexture* TextureRHI)
+	void SetParameters(FRHIBatchedShaderParameters& BatchedParameters, FRHISamplerState* SamplerStateRHI, FRHITexture* TextureRHI)
 	{
-		SetTextureParameter(RHICmdList, RHICmdList.GetBoundPixelShader(), InTexture, InTextureSampler, SamplerStateRHI, TextureRHI);
+		SetTextureParameter(BatchedParameters, InTexture, InTextureSampler, SamplerStateRHI, TextureRHI);
 	}
 
 private:
@@ -110,16 +110,16 @@ public:
 	}
 	FOculusCubemapPS() {}
 
-	void SetParameters(FRHICommandList& RHICmdList, const FTexture* Texture, int FaceIndex)
+	void SetParameters(FRHIBatchedShaderParameters& BatchedParameters, const FTexture* Texture, int FaceIndex)
 	{
-		SetTextureParameter(RHICmdList, RHICmdList.GetBoundPixelShader(), InTexture, InTextureSampler, Texture);
-		SetShaderValue(RHICmdList, RHICmdList.GetBoundPixelShader(), InFaceIndexParameter, FaceIndex);
+		SetTextureParameter(BatchedParameters, InTexture, InTextureSampler, Texture);
+		SetShaderValue(BatchedParameters, InFaceIndexParameter, FaceIndex);
 	}
 
-	void SetParameters(FRHICommandList& RHICmdList, FRHISamplerState* SamplerStateRHI, FRHITexture* TextureRHI, int FaceIndex)
+	void SetParameters(FRHIBatchedShaderParameters& BatchedParameters, FRHISamplerState* SamplerStateRHI, FRHITexture* TextureRHI, int FaceIndex)
 	{
-		SetTextureParameter(RHICmdList, RHICmdList.GetBoundPixelShader(), InTexture, InTextureSampler, SamplerStateRHI, TextureRHI);
-		SetShaderValue(RHICmdList, RHICmdList.GetBoundPixelShader(), InFaceIndexParameter, FaceIndex);
+		SetTextureParameter(BatchedParameters, InTexture, InTextureSampler, SamplerStateRHI, TextureRHI);
+		SetShaderValue(BatchedParameters, InFaceIndexParameter, FaceIndex);
 	}
 
 private:
diff --git a/Engine/Source/Runtime/RenderCore/Public/ResolveShader.h b/Engine/Source/Runtime/RenderCore/Public/ResolveShader.h
--- a/Engine/Source/Runtime/RenderCore/Public/ResolveShader.h
+++ b/Engine/Source/Runtime/RenderCore/Public/ResolveShader.h
@@ -27,6 +27,9 @@ public:
 	FResolveDepthPS(const ShaderMetaType::CompiledShaderInitializerType& Initializer);
 
 	static void ModifyCompilationEnvironment(const FGlobalShaderPermutationParameters& Parameters, FShaderCompilerEnvironment& OutEnvironment);
+
+	RENDERCORE_API void SetParameters(FRHIBatchedShaderParameters& BatchedParameters, FParameter);
+	UE_DEPRECATED(5.3, "SetParameters with FRHIBatchedShaderParameters should be used.")
 	RENDERCORE_API void SetParameters(FRHICommandList& RHICmdList, FParameter);
 
 	LAYOUT_FIELD(FShaderResourceParameter, UnresolvedSurface);
@@ -133,6 +136,9 @@ public:
 	FResolveSingleSamplePS(const ShaderMetaType::CompiledShaderInitializerType& Initializer);
 	
 	static bool ShouldCompilePermutation(const FGlobalShaderPermutationParameters& Parameters);
+
+	RENDERCORE_API void SetParameters(FRHIBatchedShaderParameters& BatchedParameters, uint32 SingleSampleIndexValue);
+	UE_DEPRECATED(5.3, "SetParameters with FRHIBatchedShaderParameters should be used.")
 	RENDERCORE_API void SetParameters(FRHICommandList& RHICmdList, uint32 SingleSampleIndexValue);
 	
 	LAYOUT_FIELD(FShaderResourceParameter, UnresolvedSurface);
@@ -150,6 +156,9 @@ public:
 	FResolveVS(const ShaderMetaType::CompiledShaderInitializerType& Initializer);
 	
 	static bool ShouldCompilePermutation(const FGlobalShaderPermutationParameters& Parameters);
+
+	RENDERCORE_API void SetParameters(FRHIBatchedShaderParameters& BatchedParameters, const FResolveRect& SrcBounds, const FResolveRect& DstBounds, uint32 DstSurfaceWidth, uint32 DstSurfaceHeight);
+	UE_DEPRECATED(5.3, "SetParameters with FRHIBatchedShaderParameters should be used.")
 	RENDERCORE_API void SetParameters(FRHICommandList& RHICmdList, const FResolveRect& SrcBounds, const FResolveRect& DstBounds, uint32 DstSurfaceWidth, uint32 DstSurfaceHeight);
 
 	LAYOUT_FIELD(FShaderParameter, PositionMinMax);
@@ -165,6 +174,4 @@ public:
 
 	static bool ShouldCompilePermutation(const FGlobalShaderPermutationParameters& Parameters);
 	static void ModifyCompilationEnvironment(const FGlobalShaderPermutationParameters& Parameters, FShaderCompilerEnvironment& OutEnvironment);
-
-	RENDERCORE_API void SetParameters(FRHICommandList& RHICmdList, const FResolveRect& SrcBounds, const FResolveRect& DstBounds, uint32 DstSurfaceWidth, uint32 DstSurfaceHeight);
 };
diff --git a/Engine/Source/Runtime/Renderer/Private/SceneRendering.cpp b/Engine/Source/Runtime/Renderer/Private/SceneRendering.cpp
--- a/Engine/Source/Runtime/Renderer/Private/SceneRendering.cpp
+++ b/Engine/Source/Runtime/Renderer/Private/SceneRendering.cpp
@@ -5503,7 +5503,7 @@ void AddResolveSceneDepthPass(FRDGBuilder& GraphBuilder, const FViewInfo& View,
 		RHICmdList.SetBlendFactor(FLinearColor::White);
 		RHICmdList.SetShaderTexture(ResolvePixelShader, TextureIndex, SourceTextureRHI);
 
-		ResolveVertexShader->SetParameters(RHICmdList, ResolveRect, ResolveRect, DepthExtent.X, DepthExtent.Y);
+		SetShaderParametersLegacyVS(RHICmdList, ResolveVertexShader, ResolveRect, ResolveRect, DepthExtent.X, DepthExtent.Y);
 
 		RHICmdList.SetStreamSource(0, nullptr, 0);
 		RHICmdList.DrawPrimitive(0, 2, bArrayResolve ? 2 : 1);
diff --git a/Engine/Source/Runtime/Windows/D3D11RHI/Private/D3D11RenderTarget.cpp b/Engine/Source/Runtime/Windows/D3D11RHI/Private/D3D11RenderTarget.cpp
--- a/Engine/Source/Runtime/Windows/D3D11RHI/Private/D3D11RenderTarget.cpp
+++ b/Engine/Source/Runtime/Windows/D3D11RHI/Private/D3D11RenderTarget.cpp
@@ -152,8 +152,8 @@ void FD3D11DynamicRHI::ResolveTextureUsingShader(
 		SetGraphicsPipelineState(RHICmdList, GraphicsPSOInit, 0);
 		RHICmdList.SetBlendFactor(FLinearColor::White);
 
-		ResolveVertexShader->SetParameters(RHICmdList, SourceRect, DestRect, ResolveTargetDesc.Width, ResolveTargetDesc.Height);
-		ResolvePixelShader->SetParameters(RHICmdList, PixelShaderParameter);
+		SetShaderParametersLegacyVS(RHICmdList, ResolveVertexShader, SourceRect, DestRect, ResolveTargetDesc.Width, ResolveTargetDesc.Height);
+		SetShaderParametersLegacyPS(RHICmdList, ResolvePixelShader, PixelShaderParameter);
 
 		// Set the source texture.
 		const uint32 TextureIndex = ResolvePixelShader->UnresolvedSurface.GetBaseIndex();
