[VK] Fix for Read Only sound cues not playing any sound in Editor
#pf 63f952b3ae54ee4ce928c8ba

diff --git a/Engine/Source/Runtime/Engine/Private/SoundWave.cpp b/Engine/Source/Runtime/Engine/Private/SoundWave.cpp
--- a/Engine/Source/Runtime/Engine/Private/SoundWave.cpp
+++ b/Engine/Source/Runtime/Engine/Private/SoundWave.cpp
@@ -175,11 +175,14 @@ FName FSoundWaveData::FindRuntimeFormat(const USoundWave& InWave) const
 {		
 #if WITH_EDITOR	
 	// If this is an editor build, we can ask ITargetPlatform.
-	if (ITargetPlatformManagerModule* TPM = GetTargetPlatformManager())
+	if (!InWave.IsLoadedFromCookedData())
 	{
-		if (ITargetPlatform* RunningTarget = TPM->GetRunningTargetPlatform())
+		if (ITargetPlatformManagerModule* TPM = GetTargetPlatformManager())
 		{
-			return RunningTarget->GetWaveFormat(&InWave);
+			if (ITargetPlatform* RunningTarget = TPM->GetRunningTargetPlatform())
+			{
+				return RunningTarget->GetWaveFormat(&InWave);
+			}
 		}
 	}
 #endif //WITH_EDITOR
