Only use demo driver net mode during playback or the recording tick

diff --git a/Engine/Source/Runtime/Engine/Private/World.cpp b/Engine/Source/Runtime/Engine/Private/World.cpp
--- a/Engine/Source/Runtime/Engine/Private/World.cpp
+++ b/Engine/Source/Runtime/Engine/Private/World.cpp
@@ -537,9 +537,7 @@ void UWorld::AddReferencedObjects(UObject* InThis, FReferenceCollector& Collecto
 		Collector.AddReferencedObjects(This->Levels, This);
 		Collector.AddReferencedObject( This->CurrentLevel, This );
 		Collector.AddReferencedObject( This->NetDriver, This );
-		PRAGMA_DISABLE_DEPRECATION_WARNINGS
 		Collector.AddReferencedObject( This->DemoNetDriver, This );
-		PRAGMA_ENABLE_DEPRECATION_WARNINGS
 		Collector.AddReferencedObject( This->LineBatcher, This );
 		Collector.AddReferencedObject( This->PersistentLineBatcher, This );
 		Collector.AddReferencedObject( This->ForegroundLineBatcher, This );
@@ -4370,7 +4368,6 @@ void UWorld::HandleTimelineScrubbed()
 	}
 }
 
-PRAGMA_DISABLE_DEPRECATION_WARNINGS
 bool UWorld::HandleDemoScrubCommand(const TCHAR* Cmd, FOutputDevice& Ar, UWorld* InWorld)
 {
 	FString TimeString;
@@ -4415,7 +4412,6 @@ bool UWorld::HandleDemoPauseCommand(const TCHAR* Cmd, FOutputDevice& Ar, UWorld*
 	}
 	return true;
 }
-PRAGMA_ENABLE_DEPRECATION_WARNINGS
 
 bool UWorld::HandleDemoSpeedCommand(const TCHAR* Cmd, FOutputDevice& Ar, UWorld* InWorld)
 {
@@ -4642,7 +4638,6 @@ bool UWorld::HandleDemoStopCommand( const TCHAR* Cmd, FOutputDevice& Ar, UWorld*
 	return true;
 }
 
-PRAGMA_DISABLE_DEPRECATION_WARNINGS
 void UWorld::DestroyDemoNetDriver()
 {
 	if ( DemoNetDriver != nullptr )
@@ -4677,7 +4672,6 @@ void UWorld::ClearDemoNetDriver()
 
 	DemoNetDriver = nullptr;
 }
-PRAGMA_ENABLE_DEPRECATION_WARNINGS
 
 void UWorld::ClearNetDriver(UNetDriver* Driver)
 {
@@ -8078,18 +8072,17 @@ ULevel* UWorld::GetCurrentLevel() const
 
 ENetMode UWorld::InternalGetNetMode() const
 {
-	if ( NetDriver != NULL )
+	if (NetDriver != nullptr)
 	{
 		const bool bIsClientOnly = IsRunningClientOnly();
 		return bIsClientOnly ? NM_Client : NetDriver->GetNetMode();
 	}
 
-	PRAGMA_DISABLE_DEPRECATION_WARNINGS
-	if ( DemoNetDriver )
+	// Use replay driver's net mode if we're in playback or ticking recording
+	if (DemoNetDriver && (DemoNetDriver->IsPlaying() || (DemoNetDriver->IsRecording() && DemoNetDriver->IsInTick())))
 	{
 		return DemoNetDriver->GetNetMode();
 	}
-	PRAGMA_ENABLE_DEPRECATION_WARNINGS
 
 	ENetMode URLNetMode = AttemptDeriveFromURL();
 #if WITH_EDITOR
@@ -8104,7 +8097,6 @@ ENetMode UWorld::InternalGetNetMode() const
 	return URLNetMode;
 }
 
-PRAGMA_DISABLE_DEPRECATION_WARNINGS
 bool UWorld::IsRecordingClientReplay() const
 {
 	if (GetNetDriver() != nullptr && !GetNetDriver()->IsServer())
@@ -8122,7 +8114,6 @@ bool UWorld::IsPlayingClientReplay() const
 {
 	return (DemoNetDriver != nullptr && DemoNetDriver->IsPlayingClientReplay());
 }
-PRAGMA_ENABLE_DEPRECATION_WARNINGS
 
 ENetMode UWorld::AttemptDeriveFromURL() const
 {
@@ -8409,9 +8400,7 @@ void UWorld::SetActiveLevelCollection(int32 LevelCollectionIndex)
 #endif
 	GameState = ActiveLevelCollection->GetGameState();
 	NetDriver = ActiveLevelCollection->GetNetDriver();
-	PRAGMA_DISABLE_DEPRECATION_WARNINGS
 	DemoNetDriver = ActiveLevelCollection->GetDemoNetDriver();
-	PRAGMA_ENABLE_DEPRECATION_WARNINGS
 
 	// TODO: START TEMP FIX FOR UE-42508
 	if (NetDriver && NetDriver->NetDriverName != NAME_None)
@@ -8424,7 +8413,6 @@ void UWorld::SetActiveLevelCollection(int32 LevelCollectionIndex)
 		}
 	}
 
-	PRAGMA_DISABLE_DEPRECATION_WARNINGS
 	if (DemoNetDriver && DemoNetDriver->NetDriverName != NAME_None)
 	{
 		UDemoNetDriver* TempDemoNetDriver = Cast<UDemoNetDriver>(GEngine->FindNamedNetDriver(this, DemoNetDriver->NetDriverName));
@@ -8434,7 +8422,6 @@ void UWorld::SetActiveLevelCollection(int32 LevelCollectionIndex)
 			DemoNetDriver = TempDemoNetDriver;
 		}
 	}
-	PRAGMA_ENABLE_DEPRECATION_WARNINGS
 	// TODO: END TEMP FIX FOR UE-42508
 }
 
