accessor to get the cluster union component on a cluster union actor

diff --git a/Engine/Source/Runtime/Engine/Classes/PhysicsEngine/ClusterUnionActor.h b/Engine/Source/Runtime/Engine/Classes/PhysicsEngine/ClusterUnionActor.h
--- a/Engine/Source/Runtime/Engine/Classes/PhysicsEngine/ClusterUnionActor.h
+++ b/Engine/Source/Runtime/Engine/Classes/PhysicsEngine/ClusterUnionActor.h
@@ -17,6 +17,9 @@ class ENGINE_API AClusterUnionActor : public AActor
 public:
 	AClusterUnionActor(const FObjectInitializer& ObjectInitializer);
 
+	UFUNCTION()
+	UClusterUnionComponent* GetClusterUnionComponent() const { return ClusterUnion; }
+
 protected:
 	/** The pivot used while building. */
 	UPROPERTY(EditAnywhere, BlueprintReadOnly, Category = "Cluster Union")
