[Rendering] Pre-allocate receiving array when adding chunks of instances of static meshes.

diff --git a/Engine/Source/Runtime/Engine/Private/InstancedStaticMesh.cpp b/Engine/Source/Runtime/Engine/Private/InstancedStaticMesh.cpp
--- a/Engine/Source/Runtime/Engine/Private/InstancedStaticMesh.cpp
+++ b/Engine/Source/Runtime/Engine/Private/InstancedStaticMesh.cpp
@@ -3243,6 +3243,8 @@ TArray<int32> UInstancedStaticMeshComponent::AddInstancesInternal(const TArray<F
 	SelectedInstances.Add(false, Count);
 #endif
 
+	PerInstanceSMData.Reserve(InstanceIndex + Count);
+
 	for (const FTransform& InstanceTransform : InstanceTransforms)
 	{
 		FInstancedStaticMeshInstanceData& NewInstanceData = PerInstanceSMData.AddDefaulted_GetRef();
