Fix out-of-bounds access in PoseAssets that have additional poses (with respect to curves)

Build health: Fixes CIS issue with quick smoke

diff --git a/Engine/Source/Runtime/Engine/Classes/Animation/PoseAsset.h b/Engine/Source/Runtime/Engine/Classes/Animation/PoseAsset.h
--- a/Engine/Source/Runtime/Engine/Classes/Animation/PoseAsset.h
+++ b/Engine/Source/Runtime/Engine/Classes/Animation/PoseAsset.h
@@ -102,8 +102,8 @@ private:
 	UPROPERTY()
 	TArray<FName> PoseFNames;
 
-	// Sorted pose name indices
-	TArray<int32> SortedPoseIndices;
+	// Sorted curve name indices
+	TArray<int32> SortedCurveIndices;
 	
 	// this is list of tracks - vertical data
 	UPROPERTY()
diff --git a/Engine/Source/Runtime/Engine/Private/Animation/PoseAsset.cpp b/Engine/Source/Runtime/Engine/Private/Animation/PoseAsset.cpp
--- a/Engine/Source/Runtime/Engine/Private/Animation/PoseAsset.cpp
+++ b/Engine/Source/Runtime/Engine/Private/Animation/PoseAsset.cpp
@@ -64,7 +64,7 @@ void FPoseDataContainer::Reset()
 {
 	// clear everything
 	PoseFNames.Reset();
-	SortedPoseIndices.Reset();
+	SortedCurveIndices.Reset();
 	Poses.Reset();
 	Tracks.Reset();
 	TrackBoneIndices.Reset();
@@ -80,12 +80,12 @@ void FPoseDataContainer::GetPoseCurve(const FPoseData* PoseData, FBlendedCurve&
 
 		auto GetNameFromIndex = [this](int32 InCurveIndex)
 		{
-			return PoseFNames[SortedPoseIndices[InCurveIndex]];
+			return Curves[SortedCurveIndices[InCurveIndex]].GetName();
 		};
 
 		auto GetValueFromIndex = [this, &CurveValues](int32 InCurveIndex)
 		{
-			return CurveValues[SortedPoseIndices[InCurveIndex]];
+			return CurveValues[SortedCurveIndices[InCurveIndex]];
 		};
 
 		UE::Anim::FCurveUtils::BuildSorted(OutCurve, CurveValues.Num(), GetNameFromIndex, GetValueFromIndex, OutCurve.GetFilter());
@@ -387,15 +387,15 @@ void FPoseDataContainer::ConvertToAdditivePose(const TArray<FTransform>& InBaseP
 void FPoseDataContainer::RebuildCurveIndexTable()
 {
 	// Recreate sorted curve index table
-	SortedPoseIndices.SetNumUninitialized(PoseFNames.Num());
-	for(int32 NameIndex = 0; NameIndex < SortedPoseIndices.Num(); ++NameIndex)
+	SortedCurveIndices.SetNumUninitialized(Curves.Num());
+	for(int32 NameIndex = 0; NameIndex < SortedCurveIndices.Num(); ++NameIndex)
 	{
-		SortedPoseIndices[NameIndex] = NameIndex;
+		SortedCurveIndices[NameIndex] = NameIndex;
 	}
 
-	SortedPoseIndices.Sort([&PoseFNames = PoseFNames](int32 LHS, int32 RHS)
+	SortedCurveIndices.Sort([&Curves = Curves](int32 LHS, int32 RHS)
 	{
-		return PoseFNames[LHS].FastLess(PoseFNames[RHS]);
+		return Curves[LHS].GetName().FastLess(Curves[RHS].GetName());
 	});
 }
 /////////////////////////////////////////////////////
