Fix static mesh GetResourceSize not taking into account size of position vertices ( only amount ) when calculating memory usage.

diff --git a/Engine/Source/Runtime/Engine/Private/StaticMesh.cpp b/Engine/Source/Runtime/Engine/Private/StaticMesh.cpp
--- a/Engine/Source/Runtime/Engine/Private/StaticMesh.cpp
+++ b/Engine/Source/Runtime/Engine/Private/StaticMesh.cpp
@@ -852,7 +852,7 @@ void FStaticMeshLODResources::Serialize(FArchive& Ar, UObject* Owner, int32 Inde
 void FStaticMeshLODResources::GetResourceSizeEx(FResourceSizeEx& CumulativeResourceSize) const
 {
 	CumulativeResourceSize.AddUnknownMemoryBytes(TEXT("TexcoordBuffer and TangentBuffer"), VertexBuffers.StaticMeshVertexBuffer.GetResourceSize());
-	CumulativeResourceSize.AddUnknownMemoryBytes(TEXT("PositionVertexBuffer"), VertexBuffers.PositionVertexBuffer.GetNumVertices());
+	CumulativeResourceSize.AddUnknownMemoryBytes(TEXT("PositionVertexBuffer"), VertexBuffers.PositionVertexBuffer.GetStride() * VertexBuffers.PositionVertexBuffer.GetNumVertices());
 	CumulativeResourceSize.AddUnknownMemoryBytes(TEXT("ColorVertexBuffer"), VertexBuffers.ColorVertexBuffer.GetStride() * VertexBuffers.ColorVertexBuffer.GetNumVertices());
 
 	const int32 IndexStride = (IndexBuffer.Is32Bit() ? 4 : 2);
