Moved SkeletionRemapping.h and SkeletonRemappingRegistry from private to public, to allow including these files in plugins

diff --git a/Engine/Source/Runtime/Engine/Private/Animation/AnimationDecompression.cpp b/Engine/Source/Runtime/Engine/Private/Animation/AnimationDecompression.cpp
--- a/Engine/Source/Runtime/Engine/Private/Animation/AnimationDecompression.cpp
+++ b/Engine/Source/Runtime/Engine/Private/Animation/AnimationDecompression.cpp
@@ -7,8 +7,8 @@
 
 #include "AnimationRuntime.h"
 #include "AnimEncoding.h"
-#include "SkeletonRemappingRegistry.h"
-#include "SkeletonRemapping.h"
+#include "Animation/SkeletonRemappingRegistry.h"
+#include "Animation/SkeletonRemapping.h"
 
 CSV_DECLARE_CATEGORY_MODULE_EXTERN(ENGINE_API, Animation);
 DECLARE_CYCLE_STAT(TEXT("Build Anim Track Pairs"), STAT_BuildAnimTrackPairs, STATGROUP_Anim);
diff --git a/Engine/Source/Runtime/Engine/Private/Animation/AnimationRuntime.cpp b/Engine/Source/Runtime/Engine/Private/Animation/AnimationRuntime.cpp
--- a/Engine/Source/Runtime/Engine/Private/Animation/AnimationRuntime.cpp
+++ b/Engine/Source/Runtime/Engine/Private/Animation/AnimationRuntime.cpp
@@ -16,8 +16,8 @@
 #include "Animation/AnimationPoseData.h"
 #include "Animation/BlendProfile.h"
 #include "Animation/MirrorDataTable.h"
-#include "SkeletonRemappingRegistry.h"
-#include "SkeletonRemapping.h"
+#include "Animation/SkeletonRemappingRegistry.h"
+#include "Animation/SkeletonRemapping.h"
 
 DEFINE_LOG_CATEGORY(LogAnimation);
 DEFINE_LOG_CATEGORY(LogRootMotion);
diff --git a/Engine/Source/Runtime/Engine/Private/Animation/PoseAsset.cpp b/Engine/Source/Runtime/Engine/Private/Animation/PoseAsset.cpp
--- a/Engine/Source/Runtime/Engine/Private/Animation/PoseAsset.cpp
+++ b/Engine/Source/Runtime/Engine/Private/Animation/PoseAsset.cpp
@@ -18,7 +18,7 @@
 #include "UObject/ObjectSaveContext.h"
 #include "UObject/Package.h"
 #include "UObject/UnrealType.h"
-#include "SkeletonRemappingRegistry.h"
+#include "Animation/SkeletonRemappingRegistry.h"
 #include "Animation/SkeletonRemapping.h"
 #include "UObject/UE5MainStreamObjectVersion.h"
 #include "UObject/UE5ReleaseStreamObjectVersion.h"
diff --git a/Engine/Source/Runtime/Engine/Private/Animation/Skeleton.cpp b/Engine/Source/Runtime/Engine/Private/Animation/Skeleton.cpp
--- a/Engine/Source/Runtime/Engine/Private/Animation/Skeleton.cpp
+++ b/Engine/Source/Runtime/Engine/Private/Animation/Skeleton.cpp
@@ -27,7 +27,7 @@
 #include "Animation/AnimBlueprint.h"
 #include "UObject/AnimObjectVersion.h"
 #include "EngineUtils.h"
-#include "SkeletonRemappingRegistry.h"
+#include "Animation/SkeletonRemappingRegistry.h"
 #include "UObject/UE5MainStreamObjectVersion.h"
 
 #include UE_INLINE_GENERATED_CPP_BY_NAME(Skeleton)
diff --git a/Engine/Source/Runtime/Engine/Private/Animation/SkeletonRemapping.cpp b/Engine/Source/Runtime/Engine/Private/Animation/SkeletonRemapping.cpp
--- a/Engine/Source/Runtime/Engine/Private/Animation/SkeletonRemapping.cpp
+++ b/Engine/Source/Runtime/Engine/Private/Animation/SkeletonRemapping.cpp
@@ -1,6 +1,6 @@
 // Copyright Epic Games, Inc. All Rights Reserved.
 
-#include "SkeletonRemapping.h"
+#include "Animation/SkeletonRemapping.h"
 
 #include "AnimationRuntime.h"
 #include "Animation/Skeleton.h"
diff --git a/Engine/Source/Runtime/Engine/Private/Animation/SkeletonRemapping.h b/Engine/Source/Runtime/Engine/Private/Animation/SkeletonRemapping.h
deleted file mode 100644
diff --git a/Engine/Source/Runtime/Engine/Private/Animation/SkeletonRemappingRegistry.cpp b/Engine/Source/Runtime/Engine/Private/Animation/SkeletonRemappingRegistry.cpp
--- a/Engine/Source/Runtime/Engine/Private/Animation/SkeletonRemappingRegistry.cpp
+++ b/Engine/Source/Runtime/Engine/Private/Animation/SkeletonRemappingRegistry.cpp
@@ -1,8 +1,8 @@
 // Copyright Epic Games, Inc. All Rights Reserved.
 
-#include "SkeletonRemappingRegistry.h"
+#include "Animation/SkeletonRemappingRegistry.h"
 #include "Misc/ScopeRWLock.h"
-#include "SkeletonRemapping.h"
+#include "Animation/SkeletonRemapping.h"
 #include "Animation/Skeleton.h"
 
 namespace UE::Anim
diff --git a/Engine/Source/Runtime/Engine/Private/Animation/SkeletonRemappingRegistry.h b/Engine/Source/Runtime/Engine/Private/Animation/SkeletonRemappingRegistry.h
deleted file mode 100644
diff --git a/Engine/Source/Runtime/Engine/Public/Animation/SkeletonRemapping.h b/Engine/Source/Runtime/Engine/Public/Animation/SkeletonRemapping.h
new file mode 100644
--- /dev/null
+++ b/Engine/Source/Runtime/Engine/Public/Animation/SkeletonRemapping.h
@@ -0,0 +1,182 @@
+// Copyright Epic Games, Inc. All Rights Reserved.
+
+#pragma once
+
+#include "CoreMinimal.h"
+#include "Animation/AnimTypes.h"
+
+class USkeleton;
+
+struct FSkeletonRemapping
+{
+	FSkeletonRemapping() = default;
+	FSkeletonRemapping(const USkeleton* InSourceSkeleton, const USkeleton* InTargetSkeleton);
+
+	const TWeakObjectPtr<const USkeleton>& GetSourceSkeleton() const { return SourceSkeleton; }
+	const TWeakObjectPtr<const USkeleton>& GetTargetSkeleton() const { return TargetSkeleton; }
+
+	/**
+	 * Compose this remapping with another remapping in place.  The other remapping's source skeleton must match
+	 * this remapping's target skeleton.  The result will map from this remapping's source skeleton to the other
+	 * remapping's target skeleton
+	 *
+	 * @param	OtherSkeletonRemapping		Skeleton remapping to compose into this remapping
+	 */
+	void ComposeWith(const FSkeletonRemapping& OtherSkeletonRemapping);
+
+	/**
+	 * Get the target skeleton bone index that corresponds to the specified bone on the source skeleton
+	 *
+	 * @param	SourceSkeletonBoneIndex		Skeleton bone index on the source skeleton
+	 * @return								Skeleton bone index on the target skeleton (or INDEX_NONE)
+	 */
+	int32 GetTargetSkeletonBoneIndex(int32 SourceSkeletonBoneIndex) const
+	{
+		return (SourceToTargetBoneIndexes.IsValidIndex(SourceSkeletonBoneIndex))
+				   ? SourceToTargetBoneIndexes[SourceSkeletonBoneIndex]
+				   : INDEX_NONE;
+	}
+
+	/**
+	 * Get the source skeleton bone index that corresponds to the specified bone on the target skeleton
+	 *
+	 * @param	TargetSkeletonBoneIndex		Skeleton bone index on the target skeleton
+	 * @return								Skeleton bone index on the source skeleton (or INDEX_NONE)
+	 */
+	int32 GetSourceSkeletonBoneIndex(int32 TargetSkeletonBoneIndex) const
+	{
+		return (TargetToSourceBoneIndexes.IsValidIndex(TargetSkeletonBoneIndex))
+				   ? TargetToSourceBoneIndexes[TargetSkeletonBoneIndex]
+				   : INDEX_NONE;
+	}
+
+	/**
+	 * Get the specified bone transform retargeted from the source skeleton onto the target skeleton, corrected
+	 * for differences between source and target rest poses
+	 *
+	 * @param	TargetSkeletonBoneIndex		Skeleton bone index on the target skeleton
+	 * @param	SourceTransform				Bone transform from the corresponding bone on the source skeleton
+	 * @return								Transform mapped onto the target skeleton
+	 */
+	FTransform RetargetBoneTransformToTargetSkeleton(int32 TargetSkeletonBoneIndex, const FTransform& SourceTransform) const
+	{
+		return FTransform(
+			RetargetBoneRotationToTargetSkeleton(TargetSkeletonBoneIndex, SourceTransform.GetRotation()),
+			RetargetBoneTranslationToTargetSkeleton(TargetSkeletonBoneIndex, SourceTransform.GetTranslation()),
+			SourceTransform.GetScale3D());
+	}
+
+	/**
+	 * Get the specified bone translation retargeted from the source skeleton onto the target skeleton, corrected
+	 * for differences between source and target rest poses
+	 *
+	 * @param	TargetSkeletonBoneIndex		Skeleton bone index on the target skeleton
+	 * @param	SourceTranslation			Bone translation from the corresponding bone on the source skeleton
+	 * @return								Translation mapped onto the target skeleton
+	 */
+	FVector RetargetBoneTranslationToTargetSkeleton(int32 TargetSkeletonBoneIndex, const FVector& SourceTranslation) const
+	{
+		// Compute the translation part of FTransform(Q1) * Source * FTransform(Q0)
+		const TTuple<FQuat, FQuat>& QQ = RetargetingTable[TargetSkeletonBoneIndex];
+		return QQ.Get<0>().RotateVector(SourceTranslation);
+	}
+
+	/**
+	 * Get the specified bone rotation retargeted from the source skeleton onto the target skeleton, corrected
+	 * for differences between source and target rest poses
+	 *
+	 * @param	TargetSkeletonBoneIndex		Skeleton bone index on the target skeleton
+	 * @param	SourceRotation				Bone rotation from the corresponding bone on the source skeleton
+	 * @return								Rotation mapped onto the target skeleton
+	 */
+	FQuat RetargetBoneRotationToTargetSkeleton(int32 TargetSkeletonBoneIndex, const FQuat& SourceRotation) const
+	{
+		// Compute the rotation part of FTransform(Q1) * Source * FTransform(Q0)
+		const TTuple<FQuat, FQuat>& QQ = RetargetingTable[TargetSkeletonBoneIndex];
+		return QQ.Get<0>() * SourceRotation * QQ.Get<1>();
+	}
+
+	/**
+	 * Get the specified additive transform retargeted from the source skeleton onto the target skeleton, corrected
+	 * for differences between source and target rest poses
+	 *
+	 * @param	TargetSkeletonBoneIndex		Skeleton bone index on the target skeleton
+	 * @param	SourceTransform				Bone transform from the corresponding bone on the source skeleton
+	 * @return								Transform mapped onto the target skeleton
+	 */
+	FTransform RetargetAdditiveTransformToTargetSkeleton(int32 TargetSkeletonBoneIndex, const FTransform& SourceTransform) const
+	{
+		return FTransform(
+			RetargetAdditiveRotationToTargetSkeleton(TargetSkeletonBoneIndex, SourceTransform.GetRotation()),
+			RetargetAdditiveTranslationToTargetSkeleton(TargetSkeletonBoneIndex, SourceTransform.GetTranslation()),
+			SourceTransform.GetScale3D());
+	}
+
+	/**
+	 * Get the specified additive translation retargeted from the source skeleton onto the target skeleton, corrected
+	 * for differences between source and target rest poses
+	 *
+	 * @param	TargetSkeletonBoneIndex		Skeleton bone index on the target skeleton
+	 * @param	SourceTranslation			Bone translation from the corresponding bone on the source skeleton
+	 * @return								Translation mapped onto the target skeleton
+	 */
+	FVector RetargetAdditiveTranslationToTargetSkeleton(int32 TargetSkeletonBoneIndex, const FVector& SourceTranslation) const
+	{
+		// Compute the translation part of FTransform(Q0.Inverse) * Source * FTransform(Q0)
+		const TTuple<FQuat, FQuat>& QQ = RetargetingTable[TargetSkeletonBoneIndex];
+		return QQ.Get<0>().RotateVector(SourceTranslation);
+	}
+
+	/**
+	 * Get the specified additive rotation retargeted from the source skeleton onto the target skeleton, corrected
+	 * for differences between source and target rest poses
+	 *
+	 * @param	TargetSkeletonBoneIndex		Skeleton bone index on the target skeleton
+	 * @param	SourceRotation				Bone rotation from the corresponding bone on the source skeleton
+	 * @return								Rotation mapped onto the target skeleton
+	 */
+	FQuat RetargetAdditiveRotationToTargetSkeleton(int32 TargetSkeletonBoneIndex, const FQuat& SourceRotation) const
+	{
+		// Compute the rotation part of FTransform(Q0.Inverse) * Source * FTransform(Q0)
+		const TTuple<FQuat, FQuat>& QQ = RetargetingTable[TargetSkeletonBoneIndex];
+		return QQ.Get<0>() * SourceRotation * QQ.Get<0>().Inverse();
+	}
+
+	UE_DEPRECATED(5.3, "This function is no longer used")
+	const TArray<SmartName::UID_Type>& GetSourceToTargetCurveMapping() const;
+
+	/** Refreshes the mapping. Empties mapping arrays and regenerates them against the stored skeletons. */
+	void RegenerateMapping();
+	
+	UE_DEPRECATED(5.3, "This function is no longer used")
+	void GenerateCurveMapping() {}
+
+	/** Check to see if a reference pose retarget is required between the source & target hierarchies */
+	bool RequiresReferencePoseRetarget() const
+	{
+		return RetargetingTable.Num() > 0;
+	}
+	
+	/** Check whether this mapping is valid - whether the source and target hierarchies still exist */
+	bool IsValid() const
+	{
+		return SourceSkeleton.IsValid() && TargetSkeleton.IsValid();
+	}
+
+private:
+	/** Internal helper function - performs mapping generation */
+	void GenerateMapping();
+	
+private:
+	TWeakObjectPtr<const USkeleton> SourceSkeleton;
+	TWeakObjectPtr<const USkeleton> TargetSkeleton;
+
+	// Table of target skeleton bone indexes (indexed by source skeleton bone index)
+	TArray<int32> SourceToTargetBoneIndexes;
+
+	// Table of source skeleton bone indexes (indexed by target skeleton bone index)
+	TArray<int32> TargetToSourceBoneIndexes;
+
+	// Table of precalculated constants for retargeting from source to target (indexed by target skeleton bone index)
+	TArray<TTuple<FQuat, FQuat>> RetargetingTable;
+};
\ No newline at end of file
diff --git a/Engine/Source/Runtime/Engine/Public/Animation/SkeletonRemappingRegistry.h b/Engine/Source/Runtime/Engine/Public/Animation/SkeletonRemappingRegistry.h
new file mode 100644
--- /dev/null
+++ b/Engine/Source/Runtime/Engine/Public/Animation/SkeletonRemappingRegistry.h
@@ -0,0 +1,55 @@
+// Copyright Epic Games, Inc. All Rights Reserved.
+
+#pragma once
+
+#include "CoreMinimal.h"
+#include "UObject/WeakObjectPtr.h"
+
+struct FSkeletonRemapping;
+class USkeleton;
+class FEngineModule;
+
+namespace UE::Anim
+{
+
+struct FSkeletonRemappingRegistryPrivate;
+
+// Global registry of skeleton remappings
+// Remappings are created on-demand when calling GetRemapping.
+// Calling public functions from multiple threads is expected. Data races are guarded by a FRWLock.
+class FSkeletonRemappingRegistry
+{
+public:
+	// Access the global registry
+	static FSkeletonRemappingRegistry& Get();
+
+	// Get a mapping between a source and target skeleton, built on demand
+	// Remapping is only valid locally to this call and could be free'd later (they are stored as shared ptrs)
+	// Do not assume ownership of the mapping after this call. 
+	const FSkeletonRemapping& GetRemapping(const USkeleton* InSourceSkeleton, const USkeleton* InTargetSkeleton);
+
+	// Refresh any existing mappings that the supplied skeleton is used with
+	void RefreshMappings(USkeleton* InSkeleton);
+
+private:
+	friend FSkeletonRemappingRegistryPrivate;
+	friend ::FEngineModule;
+	
+private:
+	// Initialize the global registry
+	static void Init();
+
+	// Shutdown the global registry
+	static void Destroy();
+	
+	// Mappings map is guarded with a RW lock
+	FRWLock MappingsLock;
+
+	using FWeakSkeletonPair = TPair<TWeakObjectPtr<const USkeleton>, TWeakObjectPtr<const USkeleton>>;
+	
+	// All mappings
+	TMap<FWeakSkeletonPair, TSharedPtr<FSkeletonRemapping>> Mappings;
+	TMultiMap<TWeakObjectPtr<const USkeleton>, TSharedPtr<FSkeletonRemapping>> PerSkeletonMappings;
+};
+
+}
\ No newline at end of file
