Fixed non-unity build
#rnx

diff --git a/Engine/Source/Runtime/Engine/Private/WorldPartition/HLOD/HLODLayer.cpp b/Engine/Source/Runtime/Engine/Private/WorldPartition/HLOD/HLODLayer.cpp
--- a/Engine/Source/Runtime/Engine/Private/WorldPartition/HLOD/HLODLayer.cpp
+++ b/Engine/Source/Runtime/Engine/Private/WorldPartition/HLOD/HLODLayer.cpp
@@ -5,8 +5,7 @@
 =============================================================================*/
 
 #include "WorldPartition/HLOD/HLODLayer.h"
-
-#include UE_INLINE_GENERATED_CPP_BY_NAME(HLODLayer)
+#include "WorldPartition/HLOD/HLODActor.h"
 
 #if WITH_EDITOR
 #include "Engine/World.h"
@@ -14,13 +13,14 @@
 #include "Misc/StringFormatArg.h"
 #include "Modules/ModuleManager.h"
 #include "UObject/UnrealType.h"
-#include "WorldPartition/HLOD/HLODActor.h"
 #include "WorldPartition/HLOD/IWorldPartitionHLODUtilities.h"
 #include "WorldPartition/HLOD/IWorldPartitionHLODUtilitiesModule.h"
 #include "WorldPartition/WorldPartition.h"
 #include "WorldPartition/WorldPartitionActorDescView.h"
 #endif
 
+#include UE_INLINE_GENERATED_CPP_BY_NAME(HLODLayer)
+
 DEFINE_LOG_CATEGORY_STATIC(LogHLODLayer, Log, All);
 
 UHLODLayer::UHLODLayer(const FObjectInitializer& ObjectInitializer)
