Adding missing includes.

diff --git a/Engine/Source/Runtime/Engine/Private/SceneTexturesConfig.cpp b/Engine/Source/Runtime/Engine/Private/SceneTexturesConfig.cpp
--- a/Engine/Source/Runtime/Engine/Private/SceneTexturesConfig.cpp
+++ b/Engine/Source/Runtime/Engine/Private/SceneTexturesConfig.cpp
@@ -3,6 +3,8 @@
 #include "SceneTexturesConfig.h"
 #include "ShaderCompiler.h"
 #include "Engine/Engine.h"
+#include "RenderGraphResources.h"
+#include "RenderUtils.h"
 #include "SceneInterface.h"
 #include "StereoRenderTargetManager.h"
 
