Fixed Client does not respect Allow Masked Material set in the editor

on behalf of Aleksander Netzel

#localization none
#tests performed by aleksander.netzel, lonnie.li

diff --git a/Engine/Source/Runtime/Engine/Classes/GameFramework/WorldSettings.h b/Engine/Source/Runtime/Engine/Classes/GameFramework/WorldSettings.h
--- a/Engine/Source/Runtime/Engine/Classes/GameFramework/WorldSettings.h
+++ b/Engine/Source/Runtime/Engine/Classes/GameFramework/WorldSettings.h
@@ -730,9 +730,14 @@ public:
 
 	/************************************/
 	/** NANITE SETTINGS **/
-	UPROPERTY(EditAnywhere, config, replicated, Category = Nanite)
+	UPROPERTY(EditAnywhere, config, replicated, ReplicatedUsing = OnRep_NaniteSettings, Category = Nanite)
 	FNaniteSettings NaniteSettings;
 
+	UFUNCTION()
+	virtual void OnRep_NaniteSettings();
+
+	virtual void SetAllowMaskedMaterials(bool bState);
+
 	/************************************/
 	/** AUDIO SETTINGS **/
 	/** Default reverb settings used by audio volumes.													*/
diff --git a/Engine/Source/Runtime/Engine/Private/WorldSettings.cpp b/Engine/Source/Runtime/Engine/Private/WorldSettings.cpp
--- a/Engine/Source/Runtime/Engine/Private/WorldSettings.cpp
+++ b/Engine/Source/Runtime/Engine/Private/WorldSettings.cpp
@@ -12,6 +12,7 @@
 #include "Engine/WorldComposition.h"
 #include "WorldPartition/WorldPartition.h"
 #include "Net/UnrealNetwork.h"
+#include "Net/Core/PushModel/PushModel.h"
 #include "GameFramework/GameNetworkManager.h"
 #include "AudioDevice.h"
 #include "Logging/MessageLog.h"
@@ -256,6 +257,21 @@ void AWorldSettings::OnRep_WorldGravityZ()
 	bWorldGravitySet = true;
 }
 
+void AWorldSettings::OnRep_NaniteSettings()
+{
+	// Need to recreate scene proxies when Nanite settings changes.
+	FGlobalComponentRecreateRenderStateContext Context;
+}
+
+void AWorldSettings::SetAllowMaskedMaterials(bool bState)
+{
+	if (GetLocalRole() == ROLE_Authority)
+	{
+		NaniteSettings.bAllowMaskedMaterials = bState;
+		MARK_PROPERTY_DIRTY_FROM_NAME(AWorldSettings, NaniteSettings, this);
+	}
+}
+
 float AWorldSettings::FixupDeltaSeconds(float DeltaSeconds, float RealDeltaSeconds)
 {
 	// DeltaSeconds is assumed to be fully dilated at this time, so we will dilate the clamp range as well
@@ -305,7 +321,10 @@ void AWorldSettings::GetLifetimeReplicatedProps( TArray< FLifetimeProperty > & O
 	DOREPLIFETIME( AWorldSettings, CinematicTimeDilation );
 	DOREPLIFETIME( AWorldSettings, WorldGravityZ );
 	DOREPLIFETIME( AWorldSettings, bHighPriorityLoading );
-	DOREPLIFETIME( AWorldSettings, NaniteSettings );
+
+	FDoRepLifetimeParams SharedParams;
+	SharedParams.bIsPushBased = true;
+	DOREPLIFETIME_WITH_PARAMS_FAST( AWorldSettings, NaniteSettings, SharedParams );
 }
 
 const FGuid FWorldSettingCustomVersion::GUID(0x1ED048F4, 0x2F2E4C68, 0x89D053A4, 0xF18F102D);
