fix incorrect Texture check( OutFormats.Num() >= 1 ) , can be zero for server only platforms

diff --git a/Engine/Source/Runtime/Engine/Private/Texture.cpp b/Engine/Source/Runtime/Engine/Private/Texture.cpp
--- a/Engine/Source/Runtime/Engine/Private/Texture.cpp
+++ b/Engine/Source/Runtime/Engine/Private/Texture.cpp
@@ -3529,7 +3529,7 @@ void UTexture::GetPlatformTextureFormatNamesWithPrefix(const class ITargetPlatfo
 	TargetPlatform->GetTextureFormats(this,OutFormats);
 
 	// almost always == 1, except for Android_Multi, which makes an array of layer formats per variant
-	check( OutFormats.Num() >= 1 );
+	// also OutFormats.Num() == 0 for server-only platforms
 	
 	// note : FinalizeVirtualTextureLayerFormat can do another set of remaps later; consider moving that here
 
diff --git a/Engine/Source/Runtime/Engine/Private/TextureDerivedData.cpp b/Engine/Source/Runtime/Engine/Private/TextureDerivedData.cpp
--- a/Engine/Source/Runtime/Engine/Private/TextureDerivedData.cpp
+++ b/Engine/Source/Runtime/Engine/Private/TextureDerivedData.cpp
@@ -1182,7 +1182,7 @@ static void GetBuildSettingsPerFormat(
 	Texture.GetPlatformTextureFormatNamesWithPrefix(TargetPlatform,PlatformFormats);
 
 	// almost always == 1, except for Android_Multi, which makes an array of layer formats per variant
-	check( PlatformFormats.Num() >= 1 );
+	// also OutFormats.Num() == 0 for server-only platforms
 
 	OutBuildSettingsPerFormat.Reserve(PlatformFormats.Num());
 	if (OutResultMetadataPerFormat)
@@ -3273,6 +3273,8 @@ void UTexture::BeginCachePlatformData()
 
 void UTexture::BeginCacheForCookedPlatformData( const ITargetPlatform *TargetPlatform )
 {
+	// @todo Oodle : if TargetPlatform->IsServerOnly() early exit?
+
 	TMap<FString, FTexturePlatformData*>* CookedPlatformDataPtr = GetCookedPlatformData();
 	if (CookedPlatformDataPtr && !GetOutermost()->HasAnyPackageFlags(PKG_FilterEditorOnly))
 	{
@@ -3427,7 +3429,9 @@ void UTexture::ClearAllCachedCookedPlatformData()
 }
 
 bool UTexture::IsCachedCookedPlatformDataLoaded(const ITargetPlatform* TargetPlatform)
-{ 
+{
+	// @todo Oodle : if TargetPlatform->IsServerOnly() early exit?
+
 	const TMap<FString, FTexturePlatformData*>* CookedPlatformDataPtr = GetCookedPlatformData();
 	if (!CookedPlatformDataPtr)
 	{
@@ -3454,9 +3458,8 @@ bool UTexture::IsCachedCookedPlatformDataLoaded(const ITargetPlatform* TargetPla
 		GetTextureBuildSettings(*this, TargetPlatform->GetTextureLODSettings(), *TargetPlatform, ETextureEncodeSpeed::Final, BuildSettings, nullptr);
 		GetBuildSettingsPerFormat(*this, BuildSettings, nullptr, TargetPlatform, ETextureEncodeSpeed::Final, BuildSettingsAllFormats, nullptr);
 	}
-
-	
 	
+	// on server-only platforms, BuildSettingsAllFormats is empty
 
 	for (const TArray<FTextureBuildSettings>& FormatBuildSettings : BuildSettingsAllFormats)
 	{
