Fix SCS_Node actor scenecomponents not getting reregistered correctly after Build events.

diff --git a/Engine/Source/Runtime/Engine/Classes/Engine/World.h b/Engine/Source/Runtime/Engine/Classes/Engine/World.h
--- a/Engine/Source/Runtime/Engine/Classes/Engine/World.h
+++ b/Engine/Source/Runtime/Engine/Classes/Engine/World.h
@@ -3062,6 +3062,9 @@ public:
 	/** Returns whether InitWorld has ever been called since this World was created.  */
 	bool HasEverBeenInitialized() const { return bHasEverBeenInitialized; }
 #endif
+	UE_DEPRECATED(5.2, "Not for public use. This function is a workaround for UE-170919 and will be removed in 5.3")
+	bool HasEverBeenInitialized_DONOTUSE() const { return bHasEverBeenInitialized; }
+
 	/** Returns whether InitWorld has been called without yet calling CleanupWorld.  */
 	bool IsInitialized() const { return bIsWorldInitialized; }
 
diff --git a/Engine/Source/Runtime/Engine/Private/SCS_Node.cpp b/Engine/Source/Runtime/Engine/Private/SCS_Node.cpp
--- a/Engine/Source/Runtime/Engine/Private/SCS_Node.cpp
+++ b/Engine/Source/Runtime/Engine/Private/SCS_Node.cpp
@@ -120,7 +120,9 @@ UActorComponent* USCS_Node::ExecuteNodeOnActor(AActor* Actor, USceneComponent* P
 		{
 			// Only register scene components if the world is initialized
 			UWorld* World = Actor->GetWorld();
-			bool bRegisterComponent = World && World->bIsWorldInitialized;
+			PRAGMA_DISABLE_DEPRECATION_WARNINGS
+			bool bRegisterComponent = World && World->HasEverBeenInitialized_DONOTUSE();
+			PRAGMA_ENABLE_DEPRECATION_WARNINGS
 
 			// If NULL is passed in, we are the root, so set transform and assign as RootComponent on Actor, similarly if the 
 			// NewSceneComp is the ParentComponent then we are the root component. This happens when the root component is recycled
