Added new macros to be able to have VLOG and UELOG still UELOG even when the Visual Logger is disabled.  These macros are:

*  UE_VLOG_ALWAYS_UELOG: Equivalent to UE_VLOG_UELOG, except it will still log if the Visual Logger is disabled.
* UE_CVLOG_ALWAYS_UELOG: Equivalent to UE_CVLOG_UELOG, except it will still log if the Visual Logger is disabled (as long as the condition is met).

#UE5

diff --git a/Engine/Source/Runtime/Engine/Public/VisualLogger/VisualLogger.h b/Engine/Source/Runtime/Engine/Public/VisualLogger/VisualLogger.h
--- a/Engine/Source/Runtime/Engine/Public/VisualLogger/VisualLogger.h
+++ b/Engine/Source/Runtime/Engine/Public/VisualLogger/VisualLogger.h
@@ -28,9 +28,11 @@
 // Text, regular log
 #define UE_VLOG(LogOwner, CategoryName, Verbosity, Format, ...) if( FVisualLogger::IsRecording() ) FVisualLogger::CategorizedLogf(LogOwner, CategoryName, ELogVerbosity::Verbosity, Format, ##__VA_ARGS__)
 #define UE_CVLOG(Condition, LogOwner, CategoryName, Verbosity, Format, ...)  if(FVisualLogger::IsRecording() && Condition) {UE_VLOG(LogOwner, CategoryName, Verbosity, Format, ##__VA_ARGS__);} 
-// Text, log with output to regular unreal logs too
+// Text, log with output to regular unreal logs too.  NOTE: UE_VLOG_UELOG will not UELOG if the Visual Logger is disabled (i.e. ENABLE_VISUAL_LOG is 0).  For the more common case where you want to always log, use UE_VLOG_ALWAYS_UELOG
 #define UE_VLOG_UELOG(LogOwner, CategoryName, Verbosity, Format, ...) { if(FVisualLogger::IsRecording()) FVisualLogger::CategorizedLogf(LogOwner, CategoryName, ELogVerbosity::Verbosity, Format, ##__VA_ARGS__); UE_LOG(CategoryName, Verbosity, Format, ##__VA_ARGS__); }
 #define UE_CVLOG_UELOG(Condition, LogOwner, CategoryName, Verbosity, Format, ...)  if(Condition) {UE_VLOG_UELOG(LogOwner, CategoryName, Verbosity, Format, ##__VA_ARGS__);} 
+// Text, log with output to regular unreal logs too.  Regular log will always happen even if the Visual Logger is disabled by via compiler switch.  See also UE_CVLOG_ALWAYS_UELOG below
+#define UE_VLOG_ALWAYS_UELOG(LogOwner, CategoryName, Verbosity, Format, ...) { UE_VLOG_UELOG(LogOwner, CategoryName, Verbosity, Format, ##__VA_ARGS__); }
 // Segment shape
 #define UE_VLOG_SEGMENT(LogOwner, CategoryName, Verbosity, SegmentStart, SegmentEnd, Color, Format, ...) if(FVisualLogger::IsRecording()) FVisualLogger::GeometryShapeLogf(LogOwner, CategoryName, ELogVerbosity::Verbosity, SegmentStart, SegmentEnd, Color, 0, Format, ##__VA_ARGS__)
 #define UE_CVLOG_SEGMENT(Condition, LogOwner, CategoryName, Verbosity, SegmentStart, SegmentEnd, Color, Format, ...) if(FVisualLogger::IsRecording() && Condition) {UE_VLOG_SEGMENT(LogOwner, CategoryName, Verbosity, SegmentStart, SegmentEnd, Color, Format, ##__VA_ARGS__);}
@@ -87,7 +89,7 @@
 
 #define UE_IFVLOG(__code_block__) if( FVisualLogger::IsRecording() ) { __code_block__; }
 
-#else
+#else // if !ENABLE_VISUAL_LOG
 #define REDIRECT_TO_VLOG(Dest)
 #define REDIRECT_OBJECT_TO_VLOG(Src, Dest)
 #define CONNECT_WITH_VLOG(Dest)
@@ -96,6 +98,8 @@
 #define UE_VLOG(Actor, CategoryName, Verbosity, Format, ...)
 #define UE_CVLOG(Condition, Actor, CategoryName, Verbosity, Format, ...)
 #define UE_VLOG_UELOG(LogOwner, CategoryName, Verbosity, Format, ...)
+// Always UELOG version will just become a UE_LOG when Visual Logging is disabled.
+#define UE_VLOG_ALWAYS_UELOG(LogOwner, CategoryName, Verbosity, Format, ...) { UE_LOG(CategoryName, Verbosity, Format, ##__VA_ARGS__); }
 #define UE_CVLOG_UELOG(Condition, Actor, CategoryName, Verbosity, Format, ...)
 #define UE_VLOG_SEGMENT(Actor, CategoryName, Verbosity, SegmentStart, SegmentEnd, Color, DescriptionFormat, ...)
 #define UE_CVLOG_SEGMENT(Condition, Actor, CategoryName, Verbosity, SegmentStart, SegmentEnd, Color, DescriptionFormat, ...)
@@ -139,6 +143,12 @@
 
 #endif //ENABLE_VISUAL_LOG
 
+// Definition is the same with/without ENABLE_VISUAL_LOG, because it's dependent on UE_VLOG_ALWAYS_UELOG which handles
+// the difference.
+// Text, log with output to regular logs if the condition is met... regular log will still happen if condition is true
+// even when the compiler switch (ENABLE_VISUAL_LOG) is off.
+#define UE_CVLOG_ALWAYS_UELOG(Condition, LogOwner, CategoryName, Verbosity, Format, ...) if (Condition) { UE_VLOG_ALWAYS_UELOG(LogOwner, CategoryName, Verbosity, Format, ##__VA_ARGS__); }
+
 // helper macros
 #define TEXT_EMPTY TEXT("")
 #define TEXT_NULL TEXT("NULL")
