This just adds comments to code associated with extracting transforms from simulation and blending them with the skeletal mesh component animation transforms.

diff --git a/Engine/Source/Runtime/Engine/Classes/Components/SkeletalMeshComponent.h b/Engine/Source/Runtime/Engine/Classes/Components/SkeletalMeshComponent.h
--- a/Engine/Source/Runtime/Engine/Classes/Components/SkeletalMeshComponent.h
+++ b/Engine/Source/Runtime/Engine/Classes/Components/SkeletalMeshComponent.h
@@ -1859,6 +1859,11 @@ public:
 	virtual bool ComponentIsTouchingSelectionFrustum(const FConvexVolume& InFrustum, const bool bConsiderOnlyBSP, const bool bMustEncompassEntireComponent) const override;
 #endif
 protected:
+	/**
+	 * If PhysicsTransformUpdateMode is set to SimulationUpatesComponentTransform, this returns the 
+	 * world-space transform of the body passed in (assumed to be the root body), after applying the 
+	 * TransformToRoot offset. Otherwise, it returns the component transform.
+	 */
 	virtual FTransform GetComponentTransformFromBodyInstance(FBodyInstance* UseBI) override;
 	//~ End UPrimitiveComponent Interface.
 
@@ -2433,9 +2438,11 @@ private:
 	
 	void BlendInPhysicsInternal(FTickFunction& ThisTickFunction);
 
-	//wrapper for parallel blend physics
+	// Blends the simulation output with the component and bone space transforms coming from
+	// animation (just calls PerformBlendPhysicsBones)
 	void ParallelBlendPhysics() { PerformBlendPhysicsBones(RequiredBones, AnimEvaluationContext.ComponentSpaceTransforms, AnimEvaluationContext.BoneSpaceTransforms); }
 
+	// Blends the simulation output with the component and bone space transforms coming from animation
 	void PerformBlendPhysicsBones(const TArray<FBoneIndexType>& InRequiredBones, TArray<FTransform>& InOutComponentSpaceTransforms, TArray<FTransform>& InOutBoneSpaceTransforms);
 
 	friend class FParallelClothTask;
diff --git a/Engine/Source/Runtime/Engine/Private/PhysicsEngine/PhysAnim.cpp b/Engine/Source/Runtime/Engine/Private/PhysicsEngine/PhysAnim.cpp
--- a/Engine/Source/Runtime/Engine/Private/PhysicsEngine/PhysAnim.cpp
+++ b/Engine/Source/Runtime/Engine/Private/PhysicsEngine/PhysAnim.cpp
@@ -141,7 +141,10 @@ void UpdateWorldBoneTM(TAssetWorldBoneTMArray& WorldBoneTMs, const TArray<FTrans
 }
 TAutoConsoleVariable<int32> CVarPhysicsAnimBlendUpdatesPhysX(TEXT("p.PhysicsAnimBlendUpdatesPhysX"), 1, TEXT("Whether to update the physx simulation with the results of physics animation blending"));
 
-void USkeletalMeshComponent::PerformBlendPhysicsBones(const TArray<FBoneIndexType>& InRequiredBones, TArray<FTransform>& InOutComponentSpaceTransforms, TArray<FTransform>& InOutBoneSpaceTransforms)
+void USkeletalMeshComponent::PerformBlendPhysicsBones(
+	const TArray<FBoneIndexType>& InRequiredBones, 
+	TArray<FTransform>& InOutComponentSpaceTransforms, 
+	TArray<FTransform>& InOutBoneSpaceTransforms)
 {
 	SCOPE_CYCLE_COUNTER(STAT_BlendInPhysics);
 	// Get drawscale from Owner (if there is one)
@@ -176,31 +179,38 @@ void USkeletalMeshComponent::PerformBlendPhysicsBones(const TArray<FBoneIndexTyp
 	FTransform LocalToWorldTM = GetComponentTransform();
 	LocalToWorldTM.RemoveScaling();
 
-	struct FBodyTMPair
-	{
-		FBodyInstance* BI;
-		FTransform TM;
-	};
-
+	// This will update both InOutComponentSpaceTransforms and InOutBoneSpaceTransforms for every
+	// required bone, leaving any others unchanged.
 	FPhysicsCommand::ExecuteRead(this, [&]()
 	{
 		bool bSetParentScale = false;
-		const bool bSimulatedRootBody = Bodies.IsValidIndex(RootBodyData.BodyIndex) && Bodies[RootBodyData.BodyIndex]->IsInstanceSimulatingPhysics();
-		const FTransform NewComponentToWorld = bSimulatedRootBody ? GetComponentTransformFromBodyInstance(Bodies[RootBodyData.BodyIndex]) : FTransform::Identity;
-
-		// For each bone - see if we need to provide some data for it.
+		// Note that IsInstanceSimulatingPhysics returns false for kinematic bodies, so
+		// bSimulatedRootBody means "is the root physics body dynamic" (not the same as the root bone)
+		const bool bSimulatedRootBody = Bodies.IsValidIndex(RootBodyData.BodyIndex) && 
+			Bodies[RootBodyData.BodyIndex]->IsInstanceSimulatingPhysics();
+
+		// Get the anticipated component transform - noting that if PhysicsTransformUpdateMode is
+		// set to SimulationUpatesComponentTransform then this will come from the simulation.
+		const FTransform NewComponentTransform = GetComponentTransformFromBodyInstance(Bodies[RootBodyData.BodyIndex]);
+
+		// For each bone:
+		// * Update the WorldBoneTMs entry
+		// * Update InOutBoneSpaceTransforms, using the physics blend weights
+		// * Update InOutComponentSpaceTransforms using transforms calculated from the bone space transforms
+		//
+		// This prevents intermediate physics blend weights from "breaking" joints (e.g. if the
+		// weight was just applied in world space).
 		for(int32 i=0; i<InRequiredBones.Num(); i++)
 		{
-			int32 BoneIndex = InRequiredBones[i];
-
-			// See if this is a physics bone..
-			int32 BodyIndex = PhysicsAsset->FindBodyIndex(GetSkeletalMeshAsset()->GetRefSkeleton().GetBoneName(BoneIndex));
-
 			// Gets set to true if we have a valid body and the skeletal mesh option has been set to
 			// be driven by kinematic body parts.
 			bool bDriveMeshWhenKinematic = false;
 
+			int32 BoneIndex = InRequiredBones[i];
+
+			// See if this is a physics bone - i.e. if there is a body registered to/associated with it.
 			// If so - get its world space matrix and its parents world space matrix and calc relative atom.
+			int32 BodyIndex = PhysicsAsset->FindBodyIndex(GetSkeletalMeshAsset()->GetRefSkeleton().GetBoneName(BoneIndex));
 			if(BodyIndex != INDEX_NONE )
 			{	
 #if !(UE_BUILD_SHIPPING || UE_BUILD_TEST)
@@ -220,7 +230,7 @@ void USkeletalMeshComponent::PerformBlendPhysicsBones(const TArray<FBoneIndexTyp
 					bUpdateMeshWhenKinematic &&
 					PhysicsAssetBodyInstance->IsValidBodyInstance();
 
-				//if simulated body copy back and blend with animation
+				// Only process this bone if it is simulated, or required due to the bDriveMeshWhenKinematic flag
 				if (PhysicsAssetBodyInstance->IsInstanceSimulatingPhysics() || bDriveMeshWhenKinematic)
 				{
 					FTransform PhysTM = PhysicsAssetBodyInstance->GetUnrealWorldTransform_AssumesLocked();
@@ -234,14 +244,29 @@ void USkeletalMeshComponent::PerformBlendPhysicsBones(const TArray<FBoneIndexTyp
 						continue;
 					}
 
-					// Note that when using bDriveMeshWhenKinematic you should still set the
-					// PhysicsBlendWeight. This allows you to use intermediate values for the blend
-					// weight when kinematic.
+					// Note that bBlendPhysics is a flag that is used to force use of the physics
+					// body, irrespective of whether they are simulated or not, or the value of the
+					// physics blend weight.
+					//
+					// Note that the existence of kinematic body parts towards/at the root can be a
+					// problem when they have a blend weight of zero in the asset, which is the
+					// default. This will result in InOutBoneSpaceTransforms not including the
+					// transform offset that goes to the real-world physics locations of body parts,
+					// so rag-dolls (etc) will get translated with the movement component. If this
+					// happens and is not the desired behavior, the solution to this is for the
+					// owner to explicitly set the physics blend weight of kinematic parts to one.
 					float PhysicsBlendWeight = bBlendPhysics ? 1.f : PhysicsAssetBodyInstance->PhysicsBlendWeight;
 
-					// if the body instance is disabled, then we want to use the animation transform
-					// and ignore the physics one, but still allow picking up transforms from
-					// kinematic movement.
+					// If the body instance is disabled, then we want to use the animation transform
+					// and ignore the physics one
+					if (PhysicsAssetBodyInstance->IsPhysicsDisabled())
+					{
+						PhysicsBlendWeight = 0.0f;
+					}
+
+					// Only do the calculations here if there is a PhysicsBlendWeight, since in the
+					// end we will use this weight to blend from the input/animation value to the
+					// physical one.
 					if (PhysicsBlendWeight > 0.f)
 					{
 						if (!(ensure(InOutBoneSpaceTransforms.Num())))
@@ -249,9 +274,9 @@ void USkeletalMeshComponent::PerformBlendPhysicsBones(const TArray<FBoneIndexTyp
 							continue;
 						}
 
-						// Find this bones parent matrix.
+						// Find the transform of the parent of this bone.
 						FTransform ParentWorldTM;
-						if(BoneIndex == 0)
+						if (BoneIndex == 0)
 						{
 							ParentWorldTM = LocalToWorldTM;
 						}
@@ -264,19 +289,21 @@ void USkeletalMeshComponent::PerformBlendPhysicsBones(const TArray<FBoneIndexTyp
 						}
 
 
-						// Then calc rel TM and convert to atom.
+						// Calculate the relative transform of the current and parent (physical) bones.
 						FTransform RelTM = PhysTM.GetRelativeTransform(ParentWorldTM);
 						RelTM.RemoveScaling();
 						FQuat RelRot(RelTM.GetRotation());
 						FVector RelPos =  RecipScale3D * RelTM.GetLocation();
-						FTransform PhysAtom = FTransform(RelRot, RelPos, InOutBoneSpaceTransforms[BoneIndex].GetScale3D());
+						FTransform PhysicalBoneSpaceTransform = FTransform(
+							RelRot, RelPos, InOutBoneSpaceTransforms[BoneIndex].GetScale3D());
 
 						// Now blend in this atom. See if we are forcing this bone to always be blended in
-						InOutBoneSpaceTransforms[BoneIndex].Blend( InOutBoneSpaceTransforms[BoneIndex], PhysAtom, PhysicsBlendWeight );
+						InOutBoneSpaceTransforms[BoneIndex].Blend(
+							InOutBoneSpaceTransforms[BoneIndex], PhysicalBoneSpaceTransform, PhysicsBlendWeight);
 
 						if (!bSetParentScale)
 						{
-							//We must update RecipScale3D based on the atom scale of the root
+							// We must update RecipScale3D based on the scale of the root
 							TotalScale3D *= InOutBoneSpaceTransforms[0].GetScale3D();
 							RecipScale3D = TotalScale3D.Reciprocal();
 							bSetParentScale = true;
@@ -291,7 +318,8 @@ void USkeletalMeshComponent::PerformBlendPhysicsBones(const TArray<FBoneIndexTyp
 				continue;
 			}
 
-			// Update SpaceBases entry for this bone now
+			// Update InOutComponentSpaceTransforms entry for this bone now - it will be the parent
+			// component-space transform, offset with the current bone-space transform.
 			if( BoneIndex == 0 )
 			{
 				if (!(ensure(InOutBoneSpaceTransforms.Num())))
@@ -322,7 +350,8 @@ void USkeletalMeshComponent::PerformBlendPhysicsBones(const TArray<FBoneIndexTyp
 				}
 				else if(bSimulatedRootBody)
 				{
-					InOutComponentSpaceTransforms[BoneIndex] = Bodies[BodyIndex]->GetUnrealWorldTransform_AssumesLocked().GetRelativeTransform(NewComponentToWorld);
+					InOutComponentSpaceTransforms[BoneIndex] = 
+						Bodies[BodyIndex]->GetUnrealWorldTransform_AssumesLocked().GetRelativeTransform(NewComponentTransform);
 				}
 			}
 		}
