Properly mark properties as NotReplicated if they're not serialized in the corresponding NetSerialize implementation. This avoids having to write custom serializers for Iris.

#rnx

diff --git a/Engine/Source/Runtime/Engine/Classes/GameFramework/RootMotionSource.h b/Engine/Source/Runtime/Engine/Classes/GameFramework/RootMotionSource.h
--- a/Engine/Source/Runtime/Engine/Classes/GameFramework/RootMotionSource.h
+++ b/Engine/Source/Runtime/Engine/Classes/GameFramework/RootMotionSource.h
@@ -244,7 +244,7 @@ struct ENGINE_API FRootMotionSource
 	 *  This is used to handle cases of inconsistent tick times (ServerMoves ticking for 1 second when root motion should
 	 *  have only applied for the last 0.1 seconds, or root motion source ending halfway through a 0.5 second tick)
 	 */
-	UPROPERTY()
+	UPROPERTY(NotReplicated)
 	float StartTime;
 
 	/** Time elapsed so far for this source */
@@ -252,7 +252,7 @@ struct ENGINE_API FRootMotionSource
 	float CurrentTime;
 
 	/** The last Time entry we had before the last SetTime() - used for simulated catchup */
-	UPROPERTY()
+	UPROPERTY(NotReplicated)
 	float PreviousTime;
 
 	/** The length of this root motion - < 0 for infinite (to be removed manually) */
@@ -264,7 +264,7 @@ struct ENGINE_API FRootMotionSource
 	FRootMotionSourceStatus Status;
 
 	/** Settings of this source */
-	UPROPERTY()
+	UPROPERTY(NotReplicated)
 	FRootMotionSourceSettings Settings;
 
 	/** True when this RootMotionSource is contributing local space accumulation (false for world space) */
@@ -287,11 +287,11 @@ struct ENGINE_API FRootMotionSource
 	bool bSimulatedNeedsSmoothing;
 
 	/** Root Motion generated by this Source */
-	UPROPERTY()
+	UPROPERTY(NotReplicated)
 	FRootMotionMovementParams RootMotionParams;
 
 	/** Finish Velocity Parameters */
-	UPROPERTY()
+	UPROPERTY(NotReplicated)
 	FRootMotionFinishVelocitySettings FinishVelocityParams;
 
 	FRootMotionSource();
