Fix for cached poses being invalid on subsequent evaluations of an anim graph

Tracks cached pose lifetime using a thread-local context and a RAII helper to ensure pose lifetime does not exceed lifetime of memstack scope

[Note duplicate of fix for Parrot (CL 24445752)]

diff --git a/Engine/Source/Runtime/Engine/Classes/Animation/AnimNode_SaveCachedPose.h b/Engine/Source/Runtime/Engine/Classes/Animation/AnimNode_SaveCachedPose.h
--- a/Engine/Source/Runtime/Engine/Classes/Animation/AnimNode_SaveCachedPose.h
+++ b/Engine/Source/Runtime/Engine/Classes/Animation/AnimNode_SaveCachedPose.h
@@ -33,6 +33,13 @@ private:
 	TUniqueFunction<void(TArrayView<const FMessageStack>)> Function;
 };
 
+/** RAII helper for cached pose lifetimes (as they are stored on the mem stack) */
+struct FCachedPoseScope
+{
+	FCachedPoseScope();
+	~FCachedPoseScope();
+};
+
 }}	// namespace UE::Anim
 
 USTRUCT(BlueprintInternalUseOnly)
@@ -50,10 +57,6 @@ struct ENGINE_API FAnimNode_SaveCachedPose : public FAnimNode_Base
 	float GlobalWeight;
 
 protected:
-	FCompactPose CachedPose;
-	FBlendedCurve CachedCurve;
-	UE::Anim::FStackAttributeContainer CachedAttributes;
-
 	struct FCachedUpdateContext
 	{
 		FAnimationUpdateContext Context;
diff --git a/Engine/Source/Runtime/Engine/Private/Animation/AnimInstance.cpp b/Engine/Source/Runtime/Engine/Private/Animation/AnimInstance.cpp
--- a/Engine/Source/Runtime/Engine/Private/Animation/AnimInstance.cpp
+++ b/Engine/Source/Runtime/Engine/Private/Animation/AnimInstance.cpp
@@ -25,6 +25,7 @@
 #include "Animation/AnimNode_LinkedAnimLayer.h"
 #include "Animation/AnimSubsystem_Tag.h"
 #include "Animation/ActiveMontageInstanceScope.h"
+#include "Animation/AnimNode_SaveCachedPose.h"
 #include "Animation/AnimSubsystem_SharedLinkedAnimLayers.h"
 
 #include UE_INLINE_GENERATED_CPP_BY_NAME(AnimInstance)
@@ -813,6 +814,9 @@ void UAnimInstance::ParallelEvaluateAnimation(bool bForceRefPose, const USkeleta
 	OutEvaluationData.OutPose.SetBoneContainer(&Proxy.GetRequiredBones());
 
 	FMemMark Mark(FMemStack::Get());
+	// Push cached pose scope to constrain cached pose lifetime to within this evaluate pass only (as cached poses are
+	// allocated with the above FMemMark)
+	UE::Anim::FCachedPoseScope CachedPoseScope;
 
 	if( !bForceRefPose )
 	{
diff --git a/Engine/Source/Runtime/Engine/Private/Animation/AnimNode_SaveCachedPose.cpp b/Engine/Source/Runtime/Engine/Private/Animation/AnimNode_SaveCachedPose.cpp
--- a/Engine/Source/Runtime/Engine/Private/Animation/AnimNode_SaveCachedPose.cpp
+++ b/Engine/Source/Runtime/Engine/Private/Animation/AnimNode_SaveCachedPose.cpp
@@ -7,6 +7,68 @@
 
 IMPLEMENT_ANIMGRAPH_MESSAGE(UE::Anim::FCachedPoseSkippedUpdateHandler);
 
+namespace UE::Anim
+{
+
+/** Holds references to scoped cached poses to ensure lifetime is correct */
+struct FCachedPoseThreadContext : TThreadSingleton<FCachedPoseThreadContext>
+{
+	/** Cached pose - note that all memory is allocated on the current mem stack allocator */
+	struct FCachedPose
+	{
+		FCompactPose Pose;
+		FBlendedCurve Curve;
+		UE::Anim::FStackAttributeContainer Attributes;
+	};
+
+	/** Cached pose scope */
+	struct FCachedPoseScopeInternal
+	{
+		// Map of nodes to cached poses
+		TMap<const FAnimNode_SaveCachedPose*, FCachedPose, TInlineSetAllocator<16>> CachedPoses;
+	};
+
+	/** Adds a cached pose */
+	const FCachedPose& AddCachedPose(const FAnimNode_SaveCachedPose* InNode, FCompactPose& InPose, FBlendedCurve& InCurve, UE::Anim::FStackAttributeContainer& InAttributes)
+	{
+		check(Scopes.Num() > 0);
+
+		FCachedPose& CachedPose = Scopes.Top().CachedPoses.Add(InNode);
+
+		CachedPose.Pose.MoveBonesFrom(InPose);
+		CachedPose.Curve.MoveFrom(InCurve);
+		CachedPose.Attributes.MoveFrom(InAttributes);
+
+		return CachedPose;
+	}
+
+	/** Finds a cached pose */
+	const FCachedPose* FindCachedPose(const FAnimNode_SaveCachedPose* InNode) const
+	{
+		if(Scopes.Num() > 0)
+		{
+			return Scopes.Top().CachedPoses.Find(InNode);
+		}
+
+		return nullptr;
+	}
+
+	/** All cached pose scopes for the current thread */
+	TArray<FCachedPoseScopeInternal, TInlineAllocator<1>> Scopes;
+};
+
+FCachedPoseScope::FCachedPoseScope()
+{
+	FCachedPoseThreadContext::Get().Scopes.Push(FCachedPoseThreadContext::FCachedPoseScopeInternal());
+}
+
+FCachedPoseScope::~FCachedPoseScope()
+{
+	FCachedPoseThreadContext::Get().Scopes.Pop();
+}
+
+}
+
 /////////////////////////////////////////////////////
 // FAnimNode_SaveCachedPose
 
@@ -64,24 +126,37 @@ void FAnimNode_SaveCachedPose::Update_AnyThread(const FAnimationUpdateContext& C
 
 void FAnimNode_SaveCachedPose::Evaluate_AnyThread(FPoseContext& Output)
 {
+	using namespace UE::Anim;
+
+	FCachedPoseThreadContext& CachedPoseThreadContext = FCachedPoseThreadContext::Get();
+	const FCachedPoseThreadContext::FCachedPose* CachedPose = nullptr;
+
 	// Note that we check here for IsSynchronized_All to deal with cases like Sequencer:
 	// In these cases the counter can stay zeroed between unbound/bound updates and this can cause issues
 	// with using out-of-date cached data (stack-allocated from a previous frame).
-	if (!EvaluationCounter.IsSynchronized_All(Output.AnimInstanceProxy->GetEvaluationCounter()))
+	const bool bSynchronized = EvaluationCounter.IsSynchronized_All(Output.AnimInstanceProxy->GetEvaluationCounter());
+	if(bSynchronized)
+	{
+		// Synchronized, so check whether we have a cached curve to use
+		CachedPose = CachedPoseThreadContext.FindCachedPose(this);
+	}
+
+	const bool bShouldCachePose = !bSynchronized || CachedPose == nullptr;
+
+	if (bShouldCachePose)
 	{
 		EvaluationCounter.SynchronizeWith(Output.AnimInstanceProxy->GetEvaluationCounter());
 
 		FPoseContext CachingContext(Output);
 		Pose.Evaluate(CachingContext);
-		CachedPose.MoveBonesFrom(CachingContext.Pose);
-		CachedCurve.MoveFrom(CachingContext.Curve);
-		CachedAttributes.MoveFrom(CachingContext.CustomAttributes);
+
+		CachedPose = &CachedPoseThreadContext.AddCachedPose(this, CachingContext.Pose, CachingContext.Curve, CachingContext.CustomAttributes);
 	}
 
 	// Return the cached result
-	Output.Pose.CopyBonesFrom(CachedPose);
-	Output.Curve.CopyFrom(CachedCurve);
-	Output.CustomAttributes.CopyFrom(CachedAttributes);
+	Output.Pose.CopyBonesFrom(CachedPose->Pose);
+	Output.Curve.CopyFrom(CachedPose->Curve);
+	Output.CustomAttributes.CopyFrom(CachedPose->Attributes);
 }
 
 void FAnimNode_SaveCachedPose::GatherDebugData(FNodeDebugData& DebugData)
