Improved behavior of ISM descriptor (comparison was not robust to hash not being computed)
Added soft-version of the ISM descriptor, and created a base class for common code/content

diff --git a/Engine/Source/Runtime/Engine/Private/ISMPartition/ISMComponentDescriptor.cpp b/Engine/Source/Runtime/Engine/Private/ISMPartition/ISMComponentDescriptor.cpp
--- a/Engine/Source/Runtime/Engine/Private/ISMPartition/ISMComponentDescriptor.cpp
+++ b/Engine/Source/Runtime/Engine/Private/ISMPartition/ISMComponentDescriptor.cpp
@@ -11,6 +11,8 @@
 #include "Serialization/ArchiveCrc32.h"
 #include "Components/InstancedStaticMeshComponent.h"
 #include "Components/HierarchicalInstancedStaticMeshComponent.h"
+#include "VT/RuntimeVirtualTexture.h"
+#include "Algo/Transform.h"
 
 FISMComponentDescriptor::FISMComponentDescriptor()
 {
@@ -18,6 +20,30 @@
 	InitFrom(UHierarchicalInstancedStaticMeshComponent::StaticClass()->GetDefaultObject<UHierarchicalInstancedStaticMeshComponent>());
 }
 
+FISMComponentDescriptor::FISMComponentDescriptor(const FSoftISMComponentDescriptor& Other)
+	: FISMComponentDescriptorBase(Other)
+{
+	StaticMesh = Other.StaticMesh.LoadSynchronous();
+	Algo::Transform(Other.OverrideMaterials, OverrideMaterials, [](TSoftObjectPtr<UMaterialInterface> Material) { return Material.LoadSynchronous(); });
+	Algo::Transform(Other.RuntimeVirtualTextures, RuntimeVirtualTextures, [](TSoftObjectPtr<URuntimeVirtualTexture> RVT) { return RVT.LoadSynchronous(); });
+	Hash = Other.Hash;
+}
+
+FSoftISMComponentDescriptor::FSoftISMComponentDescriptor()
+{
+	// Make sure we have proper defaults
+	InitFrom(UHierarchicalInstancedStaticMeshComponent::StaticClass()->GetDefaultObject<UHierarchicalInstancedStaticMeshComponent>());
+}
+
+FSoftISMComponentDescriptor::FSoftISMComponentDescriptor(const FISMComponentDescriptor& Other)
+	: FISMComponentDescriptorBase(Other)
+{
+	StaticMesh = Other.StaticMesh;
+	Algo::Transform(Other.OverrideMaterials, OverrideMaterials, [](TObjectPtr<UMaterialInterface> Material) { return Material; });
+	Algo::Transform(Other.RuntimeVirtualTextures, RuntimeVirtualTextures, [](TObjectPtr<URuntimeVirtualTexture> RVT) { return RVT; });
+	Hash = Other.Hash;
+}
+
 FISMComponentDescriptor FISMComponentDescriptor::CreateFrom(const TSubclassOf<UStaticMeshComponent>& From)
 {
 	FISMComponentDescriptor ComponentDescriptor;
@@ -28,14 +54,11 @@
 	return ComponentDescriptor;
 }
 
-void FISMComponentDescriptor::InitFrom(const UStaticMeshComponent* Template, bool bInitBodyInstance)
+void FISMComponentDescriptorBase::InitFrom(const UStaticMeshComponent* Template, bool bInitBodyInstance)
 {
 	bEnableDiscardOnLoad = false;
 	ComponentClass = Template->GetClass();
-	StaticMesh = Template->GetStaticMesh();
-	OverrideMaterials = Template->OverrideMaterials;
 	Mobility = Template->Mobility;
-	RuntimeVirtualTextures = Template->RuntimeVirtualTextures;
 	VirtualTextureRenderPassType = Template->VirtualTextureRenderPassType;
 	LightmapType = Template->LightmapType;
 	LightingChannels = Template->LightingChannels;
@@ -67,6 +90,7 @@
 	// Determine if this instance must render with reversed culling based on both scale and the component property
 	const bool bIsLocalToWorldDeterminantNegative = Template->GetRenderMatrix().Determinant() < 0;
 	bReverseCulling = Template->bReverseCulling != bIsLocalToWorldDeterminantNegative;
+	bUseDefaultCollision = Template->bUseDefaultCollision;
 
 #if WITH_EDITORONLY_DATA
 	HLODBatchingPolicy = Template->HLODBatchingPolicy;
@@ -92,19 +116,43 @@
 	}
 }
 
+void FISMComponentDescriptor::InitFrom(const UStaticMeshComponent* Template, bool bInitBodyInstance)
+{
+	StaticMesh = Template->GetStaticMesh();
+	OverrideMaterials = Template->OverrideMaterials;
+	RuntimeVirtualTextures = Template->RuntimeVirtualTextures;
+
+	Super::InitFrom(Template, bInitBodyInstance);
+}
+
+void FSoftISMComponentDescriptor::InitFrom(const UStaticMeshComponent* Template, bool bInitBodyInstance)
+{
+	StaticMesh = Template->GetStaticMesh();
+	Algo::Transform(Template->OverrideMaterials, OverrideMaterials, [](TObjectPtr<UMaterialInterface> Material) { return Material; });
+	Algo::Transform(Template->RuntimeVirtualTextures, RuntimeVirtualTextures, [](TObjectPtr<URuntimeVirtualTexture> RVT) { return RVT; });
+
+	Super::InitFrom(Template, bInitBodyInstance);
+}
+
+bool FISMComponentDescriptorBase::operator!=(const FISMComponentDescriptorBase& Other) const
+{
+	return !(*this == Other);
+}
+
 bool FISMComponentDescriptor::operator!=(const FISMComponentDescriptor& Other) const
 {
 	return !(*this == Other);
 }
 
-bool FISMComponentDescriptor::operator==(const FISMComponentDescriptor& Other) const
+bool FSoftISMComponentDescriptor::operator!=(const FSoftISMComponentDescriptor& Other) const
 {
-	return Hash == Other.Hash && // Check hash first, other checks are in case of Hash collision
-	ComponentClass == Other.ComponentClass &&
-	StaticMesh == Other.StaticMesh &&
-	OverrideMaterials == Other.OverrideMaterials &&
+	return !(*this == Other);
+}
+
+bool FISMComponentDescriptorBase::operator==(const FISMComponentDescriptorBase& Other) const
+{
+	return ComponentClass == Other.ComponentClass &&
 	Mobility == Other.Mobility &&
-	RuntimeVirtualTextures == Other.RuntimeVirtualTextures &&
 	VirtualTextureRenderPassType == Other.VirtualTextureRenderPassType &&
 	LightmapType == Other.LightmapType &&
 	GetLightingChannelMaskForStruct(LightingChannels) == GetLightingChannelMaskForStruct(Other.LightingChannels) &&
@@ -136,6 +184,7 @@
 	bVisibleInRayTracing == Other.bVisibleInRayTracing &&
 	bEvaluateWorldPositionOffset == Other.bEvaluateWorldPositionOffset &&
 	bReverseCulling == Other.bReverseCulling &&
+	bUseDefaultCollision == Other.bUseDefaultCollision &&
 	WorldPositionOffsetDisableDistance == Other.WorldPositionOffsetDisableDistance &&
 #if WITH_EDITORONLY_DATA
 	HLODBatchingPolicy == Other.HLODBatchingPolicy &&
@@ -148,6 +197,24 @@
 	(!BodyInstance.DoesUseCollisionProfile() || (BodyInstance.GetCollisionProfileName() == Other.BodyInstance.GetCollisionProfileName()));
 }
 
+bool FISMComponentDescriptor::operator==(const FISMComponentDescriptor& Other) const
+{
+	return (Hash == 0 || Other.Hash == 0 || Hash == Other.Hash) && // Check hash first, other checks are in case of Hash collision
+		StaticMesh == Other.StaticMesh &&
+		OverrideMaterials == Other.OverrideMaterials &&
+		RuntimeVirtualTextures == Other.RuntimeVirtualTextures &&
+		Super::operator==(Other);
+}
+
+bool FSoftISMComponentDescriptor::operator==(const FSoftISMComponentDescriptor& Other) const
+{
+	return (Hash == 0 || Other.Hash == 0 || Hash == Other.Hash) && // Check hash first, other checks are in case of Hash collision
+		StaticMesh == Other.StaticMesh &&
+		OverrideMaterials == Other.OverrideMaterials &&
+		RuntimeVirtualTextures == Other.RuntimeVirtualTextures &&
+		Super::operator==(Other);
+}
+
 uint32 FISMComponentDescriptor::ComputeHash() const
 {
 	FArchiveCrc32 CrcArchive;
@@ -159,6 +226,17 @@
 	return Hash;
 }
 
+uint32 FSoftISMComponentDescriptor::ComputeHash() const
+{
+	FArchiveCrc32 CrcArchive;
+
+	Hash = 0; // we don't want the hash to impact the calculation
+	CrcArchive << *this;
+	Hash = CrcArchive.GetCrc();
+
+	return Hash;
+}
+
 UInstancedStaticMeshComponent* FISMComponentDescriptor::CreateComponent(UObject* Outer, FName Name, EObjectFlags ObjectFlags) const
 {
 	UInstancedStaticMeshComponent* ISMComponent = NewObject<UInstancedStaticMeshComponent>(Outer, ComponentClass, Name, ObjectFlags);
@@ -168,41 +246,18 @@
 	return ISMComponent;
 }
 
-void FISMComponentDescriptor::InitComponent(UInstancedStaticMeshComponent* ISMComponent) const
+UInstancedStaticMeshComponent* FSoftISMComponentDescriptor::CreateComponent(UObject* Outer, FName Name, EObjectFlags ObjectFlags) const
 {
-	ISMComponent->SetStaticMesh(StaticMesh);
+	UInstancedStaticMeshComponent* ISMComponent = NewObject<UInstancedStaticMeshComponent>(Outer, ComponentClass, Name, ObjectFlags);
 
-	ISMComponent->OverrideMaterials.Empty(OverrideMaterials.Num());
-	for (UMaterialInterface* OverrideMaterial : OverrideMaterials)
-	{
-		if (OverrideMaterial && !OverrideMaterial->IsAsset())
-		{
-			// If the material is equivalent to its parent, just take a reference to its parent rather than making another redundant object 
-			if (UMaterialInstance* Instance = Cast<UMaterialInstance>(OverrideMaterial); Instance && Instance->IsRedundant())
-			{
-				OverrideMaterial = Instance->Parent;
-			}
-			else
-			{
-				// As override materials are normally outered to their owner component, we need to duplicate them here to make sure we don't create
-				// references to actors in other levels (for packed level instances or HLOD actors).
-				OverrideMaterial = DuplicateObject<UMaterialInterface>(OverrideMaterial, ISMComponent);
+	InitComponent(ISMComponent);
 
-				// If the MID we just duplicated has a nanite override that's also not an asset, duplicate that too
-				UMaterialInstanceDynamic* OverrideMID = Cast<UMaterialInstanceDynamic>(OverrideMaterial);
-				UMaterialInterface* NaniteOverride = OverrideMID ? OverrideMID->GetNaniteOverride() : nullptr; 
-				if (NaniteOverride && !NaniteOverride->IsAsset())
-				{
-					OverrideMID->SetNaniteOverride(DuplicateObject<UMaterialInterface>(NaniteOverride, ISMComponent));
-				}
-			}
-		}
-
-		ISMComponent->OverrideMaterials.Add(OverrideMaterial);
+	return ISMComponent;
 	}
 
+void FISMComponentDescriptorBase::InitComponent(UInstancedStaticMeshComponent* ISMComponent) const
+{
 	ISMComponent->Mobility = Mobility;
-	ISMComponent->RuntimeVirtualTextures = RuntimeVirtualTextures;
 	ISMComponent->VirtualTextureRenderPassType = VirtualTextureRenderPassType;
 	ISMComponent->LightmapType = LightmapType;
 	ISMComponent->LightingChannels = LightingChannels;
@@ -234,6 +289,7 @@
 	ISMComponent->bVisibleInRayTracing = bVisibleInRayTracing;
 	ISMComponent->bEvaluateWorldPositionOffset = bEvaluateWorldPositionOffset;
 	ISMComponent->bReverseCulling = bReverseCulling;
+	ISMComponent->bUseDefaultCollision = bUseDefaultCollision;
 	ISMComponent->WorldPositionOffsetDisableDistance = WorldPositionOffsetDisableDistance;
 	
 #if WITH_EDITORONLY_DATA
@@ -248,3 +304,87 @@
 		HISMComponent->bEnableDensityScaling = bEnableDensityScaling;
 	}
 }
+
+void FISMComponentDescriptor::InitComponent(UInstancedStaticMeshComponent* ISMComponent) const
+{
+	ISMComponent->SetStaticMesh(StaticMesh);
+
+	ISMComponent->OverrideMaterials.Empty(OverrideMaterials.Num());
+	for (UMaterialInterface* OverrideMaterial : OverrideMaterials)
+	{
+		if (OverrideMaterial && !OverrideMaterial->IsAsset())
+		{
+			// If the material is equivalent to its parent, just take a reference to its parent rather than making another redundant object 
+			if (UMaterialInstance* Instance = Cast<UMaterialInstance>(OverrideMaterial); Instance && Instance->IsRedundant())
+			{
+				OverrideMaterial = Instance->Parent;
+			}
+			else
+			{
+				// As override materials are normally outered to their owner component, we need to duplicate them here to make sure we don't create
+				// references to actors in other levels (for packed level instances or HLOD actors).
+				OverrideMaterial = DuplicateObject<UMaterialInterface>(OverrideMaterial, ISMComponent);
+
+				// If the MID we just duplicated has a nanite override that's also not an asset, duplicate that too
+				UMaterialInstanceDynamic* OverrideMID = Cast<UMaterialInstanceDynamic>(OverrideMaterial);
+				UMaterialInterface* NaniteOverride = OverrideMID ? OverrideMID->GetNaniteOverride() : nullptr; 
+				if (NaniteOverride && !NaniteOverride->IsAsset())
+				{
+					OverrideMID->SetNaniteOverride(DuplicateObject<UMaterialInterface>(NaniteOverride, ISMComponent));
+				}
+			}
+		}
+
+		ISMComponent->OverrideMaterials.Add(OverrideMaterial);
+	}
+
+	ISMComponent->RuntimeVirtualTextures = RuntimeVirtualTextures;
+
+	Super::InitComponent(ISMComponent);
+}
+
+void FSoftISMComponentDescriptor::InitComponent(UInstancedStaticMeshComponent* ISMComponent) const
+{
+	ISMComponent->SetStaticMesh(StaticMesh.LoadSynchronous());
+
+	ISMComponent->OverrideMaterials.Empty(OverrideMaterials.Num());
+	for (const TSoftObjectPtr<UMaterialInterface>& OverrideMaterialPtr : OverrideMaterials)
+	{
+		UMaterialInterface* OverrideMaterial = OverrideMaterialPtr.LoadSynchronous();
+		if (OverrideMaterial && !OverrideMaterial->IsAsset())
+		{
+			// If the material is equivalent to its parent, just take a reference to its parent rather than making another redundant object 
+			if (UMaterialInstance* Instance = Cast<UMaterialInstance>(OverrideMaterial); Instance && Instance->IsRedundant())
+			{
+				OverrideMaterial = Instance->Parent;
+			}
+			else
+			{
+				// As override materials are normally outered to their owner component, we need to duplicate them here to make sure we don't create
+				// references to actors in other levels (for packed level instances or HLOD actors).
+				OverrideMaterial = DuplicateObject<UMaterialInterface>(OverrideMaterial, ISMComponent);
+
+				// If the MID we just duplicated has a nanite override that's also not an asset, duplicate that too
+				UMaterialInstanceDynamic* OverrideMID = Cast<UMaterialInstanceDynamic>(OverrideMaterial);
+				UMaterialInterface* NaniteOverride = OverrideMID ? OverrideMID->GetNaniteOverride() : nullptr; 
+				if (NaniteOverride && !NaniteOverride->IsAsset())
+				{
+					OverrideMID->SetNaniteOverride(DuplicateObject<UMaterialInterface>(NaniteOverride, ISMComponent));
+				}
+			}
+		}
+
+		ISMComponent->OverrideMaterials.Add(OverrideMaterial);
+	}
+
+	ISMComponent->RuntimeVirtualTextures.Empty(RuntimeVirtualTextures.Num());
+	for (const TSoftObjectPtr<URuntimeVirtualTexture>& RuntimeVirtualTexturePtr : RuntimeVirtualTextures)
+	{
+		if (URuntimeVirtualTexture* RuntimeVirtualTexture = RuntimeVirtualTexturePtr.LoadSynchronous())
+		{
+			ISMComponent->RuntimeVirtualTextures.Add(RuntimeVirtualTexture);
+		}
+	}
+
+	Super::InitComponent(ISMComponent);
+}
diff --git a/Engine/Source/Runtime/Engine/Public/ISMPartition/ISMComponentDescriptor.h b/Engine/Source/Runtime/Engine/Public/ISMPartition/ISMComponentDescriptor.h
--- a/Engine/Source/Runtime/Engine/Public/ISMPartition/ISMComponentDescriptor.h
+++ b/Engine/Source/Runtime/Engine/Public/ISMPartition/ISMComponentDescriptor.h
@@ -15,34 +15,18 @@ enum class ERendererStencilMask : uint8;
 /** Struct that holds the relevant properties that can help decide if instances of different 
 	StaticMeshComponents/InstancedStaticMeshComponents/HISM can be merged into a single component. */
 USTRUCT()
-struct ENGINE_API FISMComponentDescriptor
+struct ENGINE_API FISMComponentDescriptorBase
 {
-	GENERATED_USTRUCT_BODY()
+	GENERATED_BODY()
 
-	FISMComponentDescriptor();
-	static FISMComponentDescriptor CreateFrom(const TSubclassOf<UStaticMeshComponent>& ComponentClass);
+	FISMComponentDescriptorBase() = default;
 	void InitFrom(const UStaticMeshComponent* Component, bool bInitBodyInstance = true);
 
 	uint32 ComputeHash() const;
-	UInstancedStaticMeshComponent* CreateComponent(UObject* Outer, FName Name = NAME_None, EObjectFlags ObjectFlags = EObjectFlags::RF_NoFlags) const;
 	void InitComponent(UInstancedStaticMeshComponent* ISMComponent) const;
 
-	friend inline uint32 GetTypeHash(const FISMComponentDescriptor& Key)
-	{
-		if (Key.Hash == 0)
-		{
-			Key.ComputeHash();
-		}
-		return Key.Hash;
-	}
-
-	bool operator!=(const FISMComponentDescriptor& Other) const;
-	bool operator==(const FISMComponentDescriptor& Other) const;
-
-	friend inline bool operator<(const FISMComponentDescriptor& Lhs, const FISMComponentDescriptor& Rhs)
-	{
-		return Lhs.Hash < Rhs.Hash;
-	}
+	bool operator!=(const FISMComponentDescriptorBase& Other) const;
+	bool operator==(const FISMComponentDescriptorBase& Other) const;
 
 public:
 	UPROPERTY()
@@ -51,15 +35,6 @@ public:
 	UPROPERTY(EditAnywhere, Category = "Component Settings")
 	TSubclassOf<UInstancedStaticMeshComponent> ComponentClass;
 	
-	UPROPERTY(EditAnywhere, Category = "Component Settings")
-	TObjectPtr<UStaticMesh> StaticMesh = nullptr;
-	
-	UPROPERTY(EditAnywhere, Category = "Component Settings")
-	TArray<TObjectPtr<UMaterialInterface>> OverrideMaterials;
-	
-	UPROPERTY(EditAnywhere, Category = "Component Settings")
-	TArray<TObjectPtr<URuntimeVirtualTexture>> RuntimeVirtualTextures;
-
 	UPROPERTY(EditAnywhere, Category = "Component Settings")
 	TEnumAsByte<EComponentMobility::Type> Mobility;
 		
@@ -175,6 +150,93 @@ public:
 	uint8 bConsiderForActorPlacementWhenHidden : 1;
 #endif
 
+	UPROPERTY(EditAnywhere, Category = "Component Settings", meta = (DisplayAfter = "CustomDepthStencilWriteMask"))
+	uint8 bUseDefaultCollision : 1;
+
 	UPROPERTY(EditAnywhere, Category = "Component Settings")
 	int32 WorldPositionOffsetDisableDistance;
 };
+
+USTRUCT()
+struct ENGINE_API FISMComponentDescriptor : public FISMComponentDescriptorBase
+{
+	GENERATED_USTRUCT_BODY()
+
+	FISMComponentDescriptor();
+	explicit FISMComponentDescriptor(const FSoftISMComponentDescriptor& Other);
+	static FISMComponentDescriptor CreateFrom(const TSubclassOf<UStaticMeshComponent>& ComponentClass);
+	void InitFrom(const UStaticMeshComponent* Component, bool bInitBodyInstance = true);
+
+	uint32 ComputeHash() const;
+	UInstancedStaticMeshComponent* CreateComponent(UObject* Outer, FName Name = NAME_None, EObjectFlags ObjectFlags = EObjectFlags::RF_NoFlags) const;
+	void InitComponent(UInstancedStaticMeshComponent* ISMComponent) const;
+
+	friend inline uint32 GetTypeHash(const FISMComponentDescriptor& Key)
+	{
+		if (Key.Hash == 0)
+		{
+			Key.ComputeHash();
+		}
+		return Key.Hash;
+	}
+
+	bool operator!=(const FISMComponentDescriptor& Other) const;
+	bool operator==(const FISMComponentDescriptor& Other) const;
+
+	friend inline bool operator<(const FISMComponentDescriptor& Lhs, const FISMComponentDescriptor& Rhs)
+	{
+		return Lhs.Hash < Rhs.Hash;
+	}
+
+public:
+	UPROPERTY(EditAnywhere, Category = "Component Settings", meta = (DisplayAfter = "ComponentClass"))
+	TObjectPtr<UStaticMesh> StaticMesh = nullptr;
+
+	UPROPERTY(EditAnywhere, Category = "Component Settings", meta = (DisplayAfter = "ComponentClass"))
+	TArray<TObjectPtr<UMaterialInterface>> OverrideMaterials;
+
+	UPROPERTY(EditAnywhere, Category = "Component Settings", meta = (DisplayAfter = "ComponentClass"))
+	TArray<TObjectPtr<URuntimeVirtualTexture>> RuntimeVirtualTextures;
+};
+
+USTRUCT()
+struct ENGINE_API FSoftISMComponentDescriptor : public FISMComponentDescriptorBase
+{
+	GENERATED_BODY()
+
+	FSoftISMComponentDescriptor();
+	explicit FSoftISMComponentDescriptor(const FISMComponentDescriptor& Other);
+	static FSoftISMComponentDescriptor CreateFrom(const TSubclassOf<UStaticMeshComponent>& ComponentClass);
+	void InitFrom(const UStaticMeshComponent* Component, bool bInitBodyInstance = true);
+
+	uint32 ComputeHash() const;
+	UInstancedStaticMeshComponent* CreateComponent(UObject* Outer, FName Name = NAME_None, EObjectFlags ObjectFlags = EObjectFlags::RF_NoFlags) const;
+	void InitComponent(UInstancedStaticMeshComponent* ISMComponent) const;
+
+	friend inline uint32 GetTypeHash(const FSoftISMComponentDescriptor& Key)
+	{
+		if (Key.Hash == 0)
+		{
+			Key.ComputeHash();
+		}
+		return Key.Hash;
+	}
+
+	bool operator!=(const FSoftISMComponentDescriptor& Other) const;
+	bool operator==(const FSoftISMComponentDescriptor& Other) const;
+
+	friend inline bool operator<(const FSoftISMComponentDescriptor& Lhs, const FSoftISMComponentDescriptor& Rhs)
+	{
+		return Lhs.Hash < Rhs.Hash;
+	}
+
+public:
+	UPROPERTY(EditAnywhere, Category = "Component Settings", meta = (DisplayAfter = "ComponentClass"))
+	TSoftObjectPtr<UStaticMesh> StaticMesh = nullptr;
+
+	UPROPERTY(EditAnywhere, Category = "Component Settings", meta = (DisplayAfter = "ComponentClass"))
+	TArray<TSoftObjectPtr<UMaterialInterface>> OverrideMaterials;
+
+	UPROPERTY(EditAnywhere, Category = "Component Settings", meta = (DisplayAfter = "ComponentClass"))
+	TArray<TSoftObjectPtr<URuntimeVirtualTexture>> RuntimeVirtualTextures;
+};
\ No newline at end of file
