World Partition: remove FlushAsyncLoading call that is no more required when initializing the world for the editor.

diff --git a/Engine/Source/Runtime/Engine/Private/WorldPartition/WorldPartition.cpp b/Engine/Source/Runtime/Engine/Private/WorldPartition/WorldPartition.cpp
--- a/Engine/Source/Runtime/Engine/Private/WorldPartition/WorldPartition.cpp
+++ b/Engine/Source/Runtime/Engine/Private/WorldPartition/WorldPartition.cpp
@@ -547,14 +547,6 @@ void UWorldPartition::Initialize(UWorld* InWorld, const FTransform& InTransform)
 			}
 		}
 	}
-
-	if (bIsEditor || bIsGame)
-	{
-		// Here we need to flush any async loading before starting any synchronous load (mixing synchronous and asynchronous load on the same package is not handled properly).
-		// Also, this will ensure that FindObject will work on an external actor that was loading asynchronously.
-		TRACE_CPUPROFILER_EVENT_SCOPE(FlushAsyncLoading);
-		FlushAsyncLoading();
-	}
 #endif
 
 	// Here's it's safe to initialize the DataLayerManager
