Remove logic duplicating the blueprint generated class from UWorld, rely instead on the UBlueprint::PostDuplicate behavior only, FBlueprintEditorUtils::PostDuplicateBlueprint now reconstruct nodes, to get any self or similarly derived pins up to date. avoid CreateReplicatedDestructionInfo when reinstancing, as this leaks REINST classes

diff --git a/Engine/Source/Editor/UnrealEd/Private/Kismet2/BlueprintEditorUtils.cpp b/Engine/Source/Editor/UnrealEd/Private/Kismet2/BlueprintEditorUtils.cpp
--- a/Engine/Source/Editor/UnrealEd/Private/Kismet2/BlueprintEditorUtils.cpp
+++ b/Engine/Source/Editor/UnrealEd/Private/Kismet2/BlueprintEditorUtils.cpp
@@ -2076,6 +2076,8 @@ void FBlueprintEditorUtils::PostDuplicateBlueprint(UBlueprint* Blueprint, bool b
 			check(NewCDO != nullptr);
 			UEditorEngine::CopyPropertiesForUnrelatedObjects(OldCDO, NewCDO);
 
+			FBlueprintEditorUtils::ReconstructAllNodes(Blueprint);
+
 			if (!FBlueprintDuplicationScopeFlags::HasAnyFlag(FBlueprintDuplicationScopeFlags::NoExtraCompilation))
 			{
 				// And compile again to make sure they go into the generated class, get cleaned up, etc...
diff --git a/Engine/Source/Runtime/Engine/Private/Level.cpp b/Engine/Source/Runtime/Engine/Private/Level.cpp
--- a/Engine/Source/Runtime/Engine/Private/Level.cpp
+++ b/Engine/Source/Runtime/Engine/Private/Level.cpp
@@ -752,7 +752,11 @@ void ULevel::CreateReplicatedDestructionInfo(AActor* const Actor)
 	// mimic the checks the package map will do before assigning a guid
 	const bool bIsActorStatic = Actor->IsFullNameStableForNetworking() && Actor->IsSupportedForNetworking();
 	const bool bActorHasRole = Actor->GetRemoteRole() != ROLE_None;
-	const bool bShouldCreateDestructionInfo = bIsActorStatic && bActorHasRole;
+	const bool bShouldCreateDestructionInfo = bIsActorStatic && bActorHasRole
+#if WITH_EDITOR
+		&& !GIsReinstancing
+#endif
+		;
 
 	if (bShouldCreateDestructionInfo)
 	{
diff --git a/Engine/Source/Runtime/Engine/Private/World.cpp b/Engine/Source/Runtime/Engine/Private/World.cpp
--- a/Engine/Source/Runtime/Engine/Private/World.cpp
+++ b/Engine/Source/Runtime/Engine/Private/World.cpp
@@ -876,55 +876,6 @@ void UWorld::PostDuplicate(bool bDuplicateForPIE)
 				}
 			}
 		}
-
-		// Duplicate the level script blueprint generated classes as well
-		const bool bDontCreate = true;
-		UBlueprint* LevelScriptBlueprint = PersistentLevel->GetLevelScriptBlueprint(bDontCreate);
-		if (LevelScriptBlueprint)
-		{
-			// Duplicating a UClass shallow copies the ClassGeneratedBy field and the CDO.
-			// This is problematic because we might accidentally reference these objects,
-			// which prevents saving of the level since these objects are external
-			// (eg: event dispatchers in a level blueprint can trigger this scenario).
-			auto FixupGeneratedClass = [LevelScriptBlueprint](UObject* NewClass)
-			{
-				if (UClass* Class = Cast<UClass>(NewClass))
-				{
-					// We can likely work around a deep-copied CDO, but it's also unexpected in this scenario.
-					ensureMsgf(!Class->ClassDefaultObject || Class->ClassDefaultObject->GetOuter() != Class, TEXT("Expecting shallow copy of CDO"));
-
-					Class->ClassGeneratedBy = LevelScriptBlueprint;
-					Class->StaticLink(true);
-					Class->ClassDefaultObject = nullptr;
-					Class->GetDefaultObject();
-				}
-			};
-
-			UObject* OldGeneratedClass = LevelScriptBlueprint->GeneratedClass;
-			if (OldGeneratedClass)
-			{
-				UObject* NewGeneratedClass = StaticDuplicateObject(OldGeneratedClass, MyPackage, OldGeneratedClass->GetFName());
-				FixupGeneratedClass(NewGeneratedClass);
-
-				ReplacementMap.Add(OldGeneratedClass, NewGeneratedClass);
-
-				// The class may have referenced a lightmap or landscape resource that is also being duplicated. Add it to the list of objects that need references fixed up.
-				ObjectsToFixReferences.Add(NewGeneratedClass);
-
-			}
-
-			UObject* OldSkeletonClass = LevelScriptBlueprint->SkeletonGeneratedClass;
-			if (OldSkeletonClass)
-			{
-				UObject* NewSkeletonClass = StaticDuplicateObject(OldSkeletonClass, MyPackage, OldSkeletonClass->GetFName());
-				FixupGeneratedClass(NewSkeletonClass);
-
-				ReplacementMap.Add(OldSkeletonClass, NewSkeletonClass);
-
-				// The class may have referenced a lightmap or landscape resource that is also being duplicated. Add it to the list of objects that need references fixed up.
-				ObjectsToFixReferences.Add(NewSkeletonClass);
-			}
-		}
 #endif // WITH_EDITOR
 	}
 
