Fixed material build errors

diff --git a/Engine/Source/Runtime/Engine/Private/Materials/HLSLMaterialTranslator.cpp b/Engine/Source/Runtime/Engine/Private/Materials/HLSLMaterialTranslator.cpp
--- a/Engine/Source/Runtime/Engine/Private/Materials/HLSLMaterialTranslator.cpp
+++ b/Engine/Source/Runtime/Engine/Private/Materials/HLSLMaterialTranslator.cpp
@@ -565,7 +565,7 @@ void FHLSLMaterialTranslator::CompileCustomOutputs(TArray<UMaterialExpressionCus
 
 			int32 MaxOutputs = CustomOutput->GetMaxOutputs();
 			int32 NumOutputs = CustomOutput->GetNumOutputs();
-			if (NumOutputs > MaxOutputs)
+			if ((MaxOutputs > 0) && (NumOutputs > MaxOutputs))
 			{
 				Errorf(TEXT("The material can only contain up to %i output %s nodes (current number: %i)"), MaxOutputs, *CustomOutput->GetDescription(), NumOutputs);
 			}
