Fix for manual USoundCue priming from game code not actually priming the assets in the cache.
#p4v-preflight-copy 24451320

diff --git a/Engine/Source/Runtime/Engine/Private/SoundNode.cpp b/Engine/Source/Runtime/Engine/Private/SoundNode.cpp
--- a/Engine/Source/Runtime/Engine/Private/SoundNode.cpp
+++ b/Engine/Source/Runtime/Engine/Private/SoundNode.cpp
@@ -101,7 +101,35 @@ UPTRINT USoundNode::GetNodeWaveInstanceHash(const UPTRINT ParentWaveInstanceHash
 
 void USoundNode::PrimeChildWavePlayers(bool bRecurse)
 {
-	OverrideLoadingBehaviorOnChildWaves(bRecurse, ESoundWaveLoadingBehavior::PrimeOnLoad);
+	if (!FPlatformCompressionUtilities::IsCurrentPlatformUsingStreamCaching())
+	{
+		return;
+	}
+
+	// Search child nodes for wave players, then prime each soundwave
+	IAudioStreamingManager &Mgr = IStreamingManager::Get().GetAudioStreamingManager();
+	for (USoundNode* ChildNode : ChildNodes)
+	{
+		if (ChildNode)
+		{
+			ChildNode->ConditionalPostLoad();
+			if (bRecurse)
+			{
+				ChildNode->PrimeChildWavePlayers(true);
+			}
+
+			if (USoundNodeWavePlayer* WavePlayer = Cast<USoundNodeWavePlayer>(ChildNode))
+			{
+				if (USoundWave* SoundWave = WavePlayer->GetSoundWave())
+				{
+					if (SoundWave->IsStreaming() && SoundWave->GetNumChunks() > 1)
+					{
+						Mgr.RequestChunk(SoundWave->CreateSoundWaveProxy(), 1, [](EAudioChunkLoadResult) {});
+					}
+				}
+			}
+		}
+	}
 }
 
 void USoundNode::RetainChildWavePlayers(bool bRecurse)
@@ -124,8 +152,7 @@ void USoundNode::OverrideLoadingBehaviorOnChildWaves(const bool bRecurse, const
 					ChildNode->OverrideLoadingBehaviorOnChildWaves(true, InLoadingBehavior);
 				}
 
-				USoundNodeWavePlayer* WavePlayer = Cast<USoundNodeWavePlayer>(ChildNode);
-				if (WavePlayer != nullptr)
+				if (USoundNodeWavePlayer* WavePlayer = Cast<USoundNodeWavePlayer>(ChildNode))
 				{
 					USoundWave* SoundWave = WavePlayer->GetSoundWave();
 					if (SoundWave)
