Fixed Non-unity build

#rnx

diff --git a/Engine/Source/Runtime/Engine/Private/ISMPartition/ISMComponentDescriptor.cpp b/Engine/Source/Runtime/Engine/Private/ISMPartition/ISMComponentDescriptor.cpp
--- a/Engine/Source/Runtime/Engine/Private/ISMPartition/ISMComponentDescriptor.cpp
+++ b/Engine/Source/Runtime/Engine/Private/ISMPartition/ISMComponentDescriptor.cpp
@@ -8,6 +8,7 @@
 #include UE_INLINE_GENERATED_CPP_BY_NAME(ISMComponentDescriptor)
 
 #include "Components/StaticMeshComponent.h"
+#include "Engine/StaticMesh.h"
 #include "Serialization/ArchiveCrc32.h"
 #include "Components/InstancedStaticMeshComponent.h"
 #include "Components/HierarchicalInstancedStaticMeshComponent.h"
