Exempt /Water/ and /Landmass/ folders from restrictive mode in material translator.

This is a workaround to allow custom expressions in built-in MFs which we so far only identify by its path name.
This won't interfere with the permissions list in the content sentry but unblocks our tooling for certain plugins.

diff --git a/Engine/Source/Runtime/Engine/Private/Materials/HLSLMaterialTranslator.cpp b/Engine/Source/Runtime/Engine/Private/Materials/HLSLMaterialTranslator.cpp
--- a/Engine/Source/Runtime/Engine/Private/Materials/HLSLMaterialTranslator.cpp
+++ b/Engine/Source/Runtime/Engine/Private/Materials/HLSLMaterialTranslator.cpp
@@ -11856,8 +11856,11 @@ static bool IsMaterialPathExemptFromRestrictiveMode(FStringBuilderBase& PathName
 		return false;
 	}
 
-	// Anything else in /Engine/ and /Game/ folders is exempt
-	if (PathNameView.StartsWith(TEXT("/Engine/")) || PathNameView.StartsWith(TEXT("/Game/")))
+	// Anything else in /Engine/ and /Game/ and other selected built-in folders is exempt
+	if (PathNameView.StartsWith(TEXT("/Engine/")) ||
+		PathNameView.StartsWith(TEXT("/Game/")) ||
+		PathNameView.StartsWith(TEXT("/Landmass/")) ||
+		PathNameView.StartsWith(TEXT("/Water/")))
 	{
 		return true;
 	}
