Improve the assert ELockState in FTextureSource to provide more information

#rnx

- The assert should now contain the state of the lock (readonly or readwrite), the type of the owning UObject and the path of the owning UObjects package.

diff --git a/Engine/Source/Runtime/Engine/Classes/Engine/Texture.h b/Engine/Source/Runtime/Engine/Classes/Engine/Texture.h
--- a/Engine/Source/Runtime/Engine/Classes/Engine/Texture.h
+++ b/Engine/Source/Runtime/Engine/Classes/Engine/Texture.h
@@ -550,6 +550,8 @@ private:
 	/** The state of any lock being held on the mip data */
 	ELockState LockState;
 
+	void CheckTextureIsUnlocked(const TCHAR* DebugMessage);
+
 #if WITH_EDITOR
 	/** Pointer to locked mip data, if any. */
 	FMipAllocation LockedMipData;
diff --git a/Engine/Source/Runtime/Engine/Private/Texture.cpp b/Engine/Source/Runtime/Engine/Private/Texture.cpp
--- a/Engine/Source/Runtime/Engine/Private/Texture.cpp
+++ b/Engine/Source/Runtime/Engine/Private/Texture.cpp
@@ -100,6 +100,23 @@ FName FTextureResource::TextureGroupStatFNames[TEXTUREGROUP_MAX] =
 	};
 #endif
 
+/** Convert FTextureSource::ELockState to TCHAR* */
+const TCHAR* LexToString(FTextureSource::ELockState State)
+{
+	switch (State)
+	{
+		case FTextureSource::ELockState::None:
+			return TEXT("None");
+		case FTextureSource::ELockState::ReadOnly:
+			return TEXT("ReadOnly");
+		case FTextureSource::ELockState::ReadWrite:
+			return TEXT("WriteOnly");
+		default:
+			checkNoEntry();
+			return TEXT("Unknown");
+	}
+}
+
 // This is used to prevent the PostEditChange to automatically update the material dependencies & material context, in some case we want to manually control this
 // to be more efficient.
 ENGINE_API bool GDisableAutomaticTextureMaterialUpdateDependencies = false;
@@ -2071,7 +2088,7 @@ FTextureSource FTextureSource::CopyTornOff() const
 
 void FTextureSource::Compress()
 {
-	checkf(LockState == ELockState::None, TEXT("Compress shouldn't be called in-between LockMip/UnlockMip"));
+	CheckTextureIsUnlocked(TEXT("Compress"));
 
 #if WITH_EDITOR
 	FWriteScopeLock BulkDataExclusiveScope(BulkDataLock.Get());
@@ -2175,6 +2192,15 @@ FSharedBuffer FTextureSource::Decompress(IImageWrapperModule* ) const
 	return Buffer;
 }
 
+void FTextureSource::CheckTextureIsUnlocked(const TCHAR* DebugMessage)
+{
+	// Asserts if a FTextureSource is locked for read or write access, along with additional debug data
+	checkf(LockState == ELockState::None, TEXT("%s cannot be called when FTextureSource is locked for %s access [%s]"), 
+		DebugMessage,
+		LexToString(LockState),
+		Owner ? *Owner->GetFullName() : TEXT("unowned"));	
+}
+
 // constructor locks the mip (can fail, pointer will be null)
 FTextureSource::FMipLock::FMipLock(ELockState InLockState,FTextureSource * InTextureSource,int32 InBlockIndex, int32 InLayerIndex, int32 InMipIndex)
 :
@@ -2371,7 +2397,7 @@ bool FTextureSource::GetMipData(TArray64<uint8>& OutMipData, int32 BlockIndex, i
 {
 	TRACE_CPUPROFILER_EVENT_SCOPE(FTextureSource::GetMipData (TArray64));
 
-	checkf(LockState == ELockState::None, TEXT("GetMipData (TArray64) shouldn't be called in-between LockMip/UnlockMip"));
+	CheckTextureIsUnlocked(TEXT("GetMipData (TArray64)"));
 
 	bool bSuccess = false;
 
@@ -2415,8 +2441,8 @@ FTextureSource::FMipData FTextureSource::GetMipData(IImageWrapperModule* )
 {
 	TRACE_CPUPROFILER_EVENT_SCOPE(FTextureSource::GetMipData (FMipData));
 
-	checkf(LockState == ELockState::None, TEXT("GetMipData (FMipData) shouldn't be called in-between LockMip/UnlockMip"));
-
+	CheckTextureIsUnlocked(TEXT("GetMipData (FMipData)"));
+	
 	check(LockedMipData.IsNull());
 	check(NumLockedMips == 0);
 
@@ -2664,7 +2690,7 @@ FSharedBuffer FTextureSource::TryDecompressData() const
 
 void FTextureSource::ExportCustomProperties(FOutputDevice& Out, uint32 Indent)
 {
-	check(LockState == ELockState::None);
+	CheckTextureIsUnlocked(TEXT("ExportCustomProperties"));
 
 	FSharedBuffer Payload = BulkData.GetPayload().Get();
 	uint64 PayloadSize = Payload.GetSize();
@@ -2681,7 +2707,7 @@ void FTextureSource::ExportCustomProperties(FOutputDevice& Out, uint32 Indent)
 
 void FTextureSource::ImportCustomProperties(const TCHAR* SourceText, FFeedbackContext* Warn)
 {
-	check(LockState == ELockState::None);
+	CheckTextureIsUnlocked(TEXT("ImportCustomProperties"));
 
 	if (FParse::Command(&SourceText, TEXT("TextureSourceData")))
 	{
@@ -2894,7 +2920,7 @@ void FTextureSource::UseHashAsGuid()
 {
 	if (HasPayloadData())
 	{
-		checkf(LockState == ELockState::None, TEXT("UseHashAsGuid shouldn't be called in-between LockMip/UnlockMip"));
+		CheckTextureIsUnlocked(TEXT("UseHashAsGuid"));
 
 		bGuidIsHash = true;
 		Id = UE::Serialization::IoHashToGuid(BulkData.GetPayloadId());
@@ -3862,7 +3888,7 @@ void FTextureSource::InitLayeredImpl(
 
 	BlockDataOffsets.Add(0);
 
-	checkf(LockState == ELockState::None, TEXT("InitLayered shouldn't be called in-between LockMip/UnlockMip"));
+	CheckTextureIsUnlocked(TEXT("InitLayered"));
 }
 
 void FTextureSource::InitBlockedImpl(const ETextureSourceFormat* InLayerFormats,
@@ -3903,7 +3929,7 @@ void FTextureSource::InitBlockedImpl(const ETextureSourceFormat* InLayerFormats,
 
 	EnsureBlocksAreSorted();
 
-	checkf(LockState == ELockState::None, TEXT("InitBlocked shouldn't be called in-between LockMip/UnlockMip"));
+	CheckTextureIsUnlocked(TEXT("InitBlockedImpl"));
 }
 
 namespace
