Added Schema Version to LoadEvent and CookEvent payloads

diff --git a/Engine/Source/Editor/UnrealEd/Private/Cooker/CookProfiling.cpp b/Engine/Source/Editor/UnrealEd/Private/Cooker/CookProfiling.cpp
--- a/Engine/Source/Editor/UnrealEd/Private/Cooker/CookProfiling.cpp
+++ b/Engine/Source/Editor/UnrealEd/Private/Cooker/CookProfiling.cpp
@@ -903,10 +903,13 @@ void LogCookStats(ECookMode::Type CookMode)
 
 	if (FStudioAnalytics::IsAvailable() && IsCookByTheBookMode(CookMode))
 	{
+		const int SchemaVersion = 1;
 
 		// convert filtered stats directly to an analytics event
 		TArray<FAnalyticsEventAttribute> Attributes;
 
+		Attributes.Emplace(TEXT("SchemaVersion"), SchemaVersion);
+
 		// Sends each cook stat to the studio analytics system.
 		auto SendCookStatsToAnalytics = [&Attributes](const FString& StatName, const TArray<FCookStatsManager::StringKeyValue>& StatAttributes)
 		{
diff --git a/Engine/Source/Runtime/Engine/Private/Analytics/StudioAnalytics.cpp b/Engine/Source/Runtime/Engine/Private/Analytics/StudioAnalytics.cpp
--- a/Engine/Source/Runtime/Engine/Private/Analytics/StudioAnalytics.cpp
+++ b/Engine/Source/Runtime/Engine/Private/Analytics/StudioAnalytics.cpp
@@ -135,6 +135,8 @@ void FStudioAnalytics::RecordEvent(const FString& EventName, const TArray<FAnaly
 
 void FStudioAnalytics::FireEvent_Loading(const FString& LoadingName, double SecondsSpentLoading, const TArray<FAnalyticsEventAttribute>& InAttributes)
 {
+	const int SchemaVersion = 1;
+
 	// Ignore anything less than a 1/4th a second.
 	if (SecondsSpentLoading < 0.250)
 	{
@@ -147,34 +149,38 @@ void FStudioAnalytics::FireEvent_Loading(const FString& LoadingName, double Seco
 		return;
 	}
 
-	if (FStudioAnalytics::IsAvailable())
-	{
-		TArray<FAnalyticsEventAttribute> Attributes;
+	TArray<FAnalyticsEventAttribute> Attributes;
 
-		Attributes.Emplace(TEXT("LoadingName"), LoadingName);
-		Attributes.Emplace(TEXT("LoadingSeconds"), SecondsSpentLoading);
-		Attributes.Append(InAttributes);
+	Attributes.Emplace(TEXT("SchemaVersion"), SchemaVersion);
+	Attributes.Emplace(TEXT("LoadingName"), LoadingName);
+	Attributes.Emplace(TEXT("LoadingSeconds"), SecondsSpentLoading);
+	Attributes.Append(InAttributes);
 
+	if (FStudioAnalytics::IsAvailable())
+	{
 		FStudioAnalytics::GetProvider().RecordEvent(TEXT("Performance.Loading"), Attributes);
+	}
 
 #if ENABLE_COOK_STATS
 
-		// Gather DDC analytics
-		GetDerivedDataCacheRef().GatherAnalytics(Attributes);
+	// Gather DDC analytics
+	GetDerivedDataCacheRef().GatherAnalytics(Attributes);
 
-		// Gather Virtualization analytics
-		UE::Virtualization::IVirtualizationSystem::Get().GatherAnalytics(Attributes);
+	// Gather Virtualization analytics
+	UE::Virtualization::IVirtualizationSystem::Get().GatherAnalytics(Attributes);
 
 #if UE_WITH_ZEN
-		// Gather Zen analytics
-		if (UE::Zen::IsDefaultServicePresent())
-		{
-			UE::Zen::GetDefaultServiceInstance().GatherAnalytics(Attributes);
-		}
+	// Gather Zen analytics
+	if (UE::Zen::IsDefaultServicePresent())
+	{
+		UE::Zen::GetDefaultServiceInstance().GatherAnalytics(Attributes);
+	}
 #endif
 
+	if (FStudioAnalytics::IsAvailable())
+	{
 		// Store it all in the loading event
 		FStudioAnalytics::GetProvider().RecordEvent(TEXT("Core.Loading"), Attributes);
-#endif
 	}
-}
\ No newline at end of file
+#endif
+}
