Adding export for RemoveClientConnection

PR #10197: Added exports for extra public method (Contributed by [REDACTED_PERSON])

diff --git a/Engine/Source/Runtime/Engine/Classes/Engine/NetDriver.h b/Engine/Source/Runtime/Engine/Classes/Engine/NetDriver.h
--- a/Engine/Source/Runtime/Engine/Classes/Engine/NetDriver.h
+++ b/Engine/Source/Runtime/Engine/Classes/Engine/NetDriver.h
@@ -1721,7 +1721,7 @@ public:
 	void UpdateGroupFilterStatusForLevel(const ULevel* Level, UE::Net::FNetObjectGroupHandle LevelGroupHandle);
 #endif // UE_WITH_IRIS
 
-	void RemoveClientConnection(UNetConnection* ClientConnectionToRemove);
+	ENGINE_API void RemoveClientConnection(UNetConnection* ClientConnectionToRemove);
 
 	/** Adds (fully initialized, ready to go) client connection to the ClientConnections list + any other game related setup */
 	ENGINE_API void	AddClientConnection(UNetConnection * NewConnection);
