Read `smpl` data when importing USoundWave.

- Audio.h/.cpp Add functionality to read the `smpl` chunk in a `.wav` asset and adds SampleLoop data to `FWaveModInfo`
- SoundWaveFactor.cpp, parse `WaveSampleLoops` from `FWaveModInfo` and add them to the`CuePoints` array instead of creating a new array
- add `bIsLoopRegion` to `FSoundWaveCuePoint` to identify if a cue point is being used as a "Sample Loop"
- add methods to `USoundWave` to get filtered "CuePoints" and "LoopRegions" from cuepoints array
- When converting `USoundWave` to proxy data, split cue points into "Loop Regions" and "Cue Points" array (less of a memory issue to add new arrays to these classes)
- Fail import if the sample loop start and end points are invalid

diff --git a/Engine/Source/Editor/AudioEditor/Private/Factories/SoundFactory.cpp b/Engine/Source/Editor/AudioEditor/Private/Factories/SoundFactory.cpp
--- a/Engine/Source/Editor/AudioEditor/Private/Factories/SoundFactory.cpp
+++ b/Engine/Source/Editor/AudioEditor/Private/Factories/SoundFactory.cpp
@@ -585,8 +585,10 @@ UObject* USoundFactory::CreateObject
 		Sound->AssetImportData->Update(CurrentFilename);
 
 		// Setup the cue points
-		Sound->CuePoints.Reset(WaveInfo.WaveCues.Num());
+		int TotalNum = WaveInfo.WaveCues.Num() + WaveInfo.WaveSampleLoops.Num();
+		Sound->CuePoints.Reset(TotalNum);
 
+		// Start with WaveCues
 		for (FWaveCue& WaveCue : WaveInfo.WaveCues)
 		{
 			FSoundWaveCuePoint NewCuePoint;
@@ -597,6 +599,36 @@ UObject* USoundFactory::CreateObject
 			Sound->CuePoints.Add(NewCuePoint);
 		}
 
+		// add Sample Loops to end
+		bool FoundInvalidSampleLoops = false;
+		for (FWaveSampleLoop& SampleLoop : WaveInfo.WaveSampleLoops)
+		{
+			FSoundWaveCuePoint NewCuePoint;
+			NewCuePoint.bIsLoopRegion = true;
+			NewCuePoint.CuePointID = (int32)SampleLoop.LoopID;
+			NewCuePoint.FramePosition = (int32)SampleLoop.StartFrame;
+			NewCuePoint.FrameLength = (int32)SampleLoop.EndFrame - (int32)SampleLoop.StartFrame;
+			if (SampleLoop.EndFrame <= SampleLoop.StartFrame)
+			{
+				Warn->Logf(ELogVerbosity::Error, 
+					TEXT("Found invalid start and end frames when creating Cue Point from Sample Loop Region! LoopID = %d, StartFrame = %d, EndFrame = %d"), 
+					SampleLoop.LoopID, SampleLoop.StartFrame, SampleLoop.EndFrame);
+
+				FoundInvalidSampleLoops = true;
+			}
+			Sound->CuePoints.Add(NewCuePoint);
+		}
+
+		// fail import if we found invalid sample loops
+		// each invalid sample loop is logged
+		if (FoundInvalidSampleLoops)
+		{
+			FText InvalidSamplesText = NSLOCTEXT("SoundFactory", "Invalid Sample Loops", "Found sample loops with invalid start and end frames. See logs for more info.");
+			FMessageDialog::Open(EAppMsgType::Ok, FText::Format(NSLOCTEXT("SoundFactory", "Import Failed", "Import failed for {0}: {1}"), FText::FromString(Name.ToString()), InvalidSamplesText));
+			GEditor->GetEditorSubsystem<UImportSubsystem>()->BroadcastAssetPostImport(this, nullptr);
+			return nullptr;
+		}
+
 		// If we've read some time-code.
 		if (WaveInfo.TimecodeInfo)
 		{
diff --git a/Engine/Source/Runtime/Engine/Classes/Sound/SoundWave.h b/Engine/Source/Runtime/Engine/Classes/Sound/SoundWave.h
--- a/Engine/Source/Runtime/Engine/Classes/Sound/SoundWave.h
+++ b/Engine/Source/Runtime/Engine/Classes/Sound/SoundWave.h
@@ -381,6 +381,14 @@ struct FSoundWaveCuePoint
 	// The frame length of the cue point (non-zero if it's a region)
 	UPROPERTY(Category = Info, VisibleAnywhere, BlueprintReadOnly)
 	int32 FrameLength = 0;
+
+	friend class USoundFactory;
+	friend class USoundWave;
+private:
+	// intentionally kept private.
+	// only USoundFactory should modify this value on import
+	UPROPERTY(Category = Info, VisibleAnywhere)
+	bool bIsLoopRegion = false;
 };
 
 struct ISoundWaveClient
@@ -467,6 +475,14 @@ public:
 	UFUNCTION(BlueprintCallable, Category = "Audio")
 	void SetSoundAssetCompressionType(ESoundAssetCompressionType InSoundAssetCompressionType, bool bMarkDirty = true);
 
+	/** Filters for the cue points that are _not_ loop regions and returns those as a new array*/
+	UFUNCTION(BlueprintCallable, Category = "Audio")
+	TArray<FSoundWaveCuePoint> GetCuePoints() const;
+
+	/** Filters for the cue points that _are_ loop regions and returns those as a new array */
+	UFUNCTION(BlueprintCallable, Category = "Audio")
+	TArray<FSoundWaveCuePoint> GetLoopRegions() const;
+
 	/** Returns the Runtime format of the wave */
 	FName GetRuntimeFormat() const;
 
@@ -710,7 +726,7 @@ public:
 	UPROPERTY(Category = Info, AssetRegistrySearchable, VisibleAnywhere)
 	int32 NumChannels;
 
-	/** Cue point data */
+	/** Cue point data parsed fro the .wav file. Contains "Loop Regions" as cue points as well! */
 	UPROPERTY(Category = Info, VisibleAnywhere, BlueprintReadOnly)
 	TArray<FSoundWaveCuePoint> CuePoints;
 
@@ -1275,6 +1291,7 @@ public:
 	float GetSampleRate() const { return SampleRate; }
 	uint32 GetNumChannels() const { return NumChannels; }
 	const TArray<FSoundWaveCuePoint>& GetCuePoints() const { return CuePoints; }
+	const TArray<FSoundWaveCuePoint>& GetLoopRegions() const { return LoopRegions; }
 
 	MaxChunkSizeResults GetMaxChunkSizeResults() const;
 
@@ -1363,6 +1380,7 @@ private:
 	FName RuntimeFormat{ "FSoundWaveProxy_InvalidFormat" };
 	FObjectKey SoundWaveKeyCached;
 	TArray<FSoundWaveCuePoint> CuePoints;
+	TArray<FSoundWaveCuePoint> LoopRegions;
 	ESoundAssetCompressionType SoundAssetCompressionType;
 	
 	float SampleRate = 0;
@@ -1422,6 +1440,7 @@ public:
 	uint32 GetNumChannels() const;
 	uint32 GetSizeOfChunk(uint32 ChunkIndex) const;
 	const TArray<FSoundWaveCuePoint>& GetCuePoints() const;
+	const TArray<FSoundWaveCuePoint>& GetLoopRegions() const;
 
 	FSoundWaveData::MaxChunkSizeResults GetMaxChunkSizeResults() const;
 
diff --git a/Engine/Source/Runtime/Engine/Private/Audio.cpp b/Engine/Source/Runtime/Engine/Private/Audio.cpp
--- a/Engine/Source/Runtime/Engine/Private/Audio.cpp
+++ b/Engine/Source/Runtime/Engine/Private/Audio.cpp
@@ -1167,6 +1167,34 @@ struct FRiffCuePointChunk
 	uint32 SampleOffset;	// Byte offset to sample byte of first channel
 };
 
+// ChunkID: 'smpl'
+// The sample chunk allows a MIDI sampler to use the Wave file as a collection of samples.
+// The 'smpl' chunk is optional and if included, a single 'smple' chunk should specify all Sample Loops
+struct FRiffSampleChunk
+{
+	uint32 ChunkID;				// 'smpl'
+	uint32 ChunkDataSize;		// Depends on the number of sample loops
+	uint32 ManufacturerCode;	// The MIDI Manufacturers Association manufacturer code
+	uint32 Product;				// The Product / Model ID of the target device, specific to the manufacturer
+	uint32 SamplePeriod;		// The period of one sample in nanoseconds.
+	uint32 MidiUnityNote;		// The MIDI note that will play when this sample is played at its current pitch
+	uint32 MidiPitchFraction;	// The fraction of a semitone up from the specified note. 
+	uint32 SmpteFormat;			// The SMPTE format. Possible values: 0, 24, 25, 29, 30
+	uint32 SmpteOffset;			// Specifies a time offset for the sample, if the sample should start at a later time and not immediately.
+	uint32 NumSampleLoops;		// Number of sample loops contained in this chunks data
+	uint32 NumSampleDataBytes;	// Number of bytes of optional sampler specific data that follows the sample loops. zero if there is no such data.
+};
+
+struct FRiffSampleLoopChunk
+{
+	uint32 LoopID;			// A unique ID of the loop, which could be a cue point
+	uint32 LoopType;		// The loop type. 0: Forward Looping, 1: Ping-Pong, 2: Backward, 3-31: future standard types. >=32: manufacturer specific types
+	uint32 StartFrame;		// Start point of the loop in samples
+	uint32 EndFrame;		// End point of the loop in samples. The end sample is also played.
+	uint32 Fraction;		// The resolution at which this loop should be fine tuned.
+	uint32 NumPlayTimes;	// The number of times to play the loop. A value of zero means inifity. In a Midi sampler, that may mean infinite sustain.
+};
+
 struct FRiffListChunk
 {
 	uint32 ChunkID;			// 'list'
@@ -1759,6 +1787,44 @@ bool FWaveModInfo::ReadWaveInfo( const uint8* WaveData, int32 WaveDataSize, FStr
 		}
 	}
 
+	// Look for smpl chunk
+	RiffChunk = FindRiffChunk(RiffChunkStart, WaveDataEnd, UE_mmioFOURCC('s', 'm', 'p', 'l'));
+
+	if (RiffChunk != nullptr)
+	{
+		const FRiffSampleChunk* SampleChunk = (const FRiffSampleChunk*)(RiffChunk);
+
+		// only swap the members that we care about
+#if !PLATFORM_LITTLE_ENDIAN
+		if (!AlreadySwapped)
+		{
+			SampleChunk->NumSampleLoops = INTEL_ORDER32(SampleChunk->NumSampleLoops);
+		}
+#endif
+		WaveSampleLoops.Reset(SampleChunk->NumSampleLoops);
+
+		// use the RiffSampleChunk size to offset into the chunk data and get the list of SampleLoops
+		FRiffSampleLoopChunk* SampleLoopChunks = (FRiffSampleLoopChunk*)((uint8*)RiffChunk + sizeof(FRiffSampleChunk));
+		for (uint32 LoopIdx = 0; LoopIdx < SampleChunk->NumSampleLoops; ++LoopIdx)
+		{
+			FRiffSampleLoopChunk& SampleLoopChunk = SampleLoopChunks[LoopIdx];
+
+#if !PLATFORM_LITTLE_ENDIAN
+			if (!AlreadySwapped)
+			{
+				SampleLoopChunk.LoopID = INTEL_ORDER32(SampleLoopChunk.LoopID);
+				SampleLoopChunk.StartFrame = INTEL_ORDER32(SampleLoopChunk.StartFrame);
+				SampleLoopChunk.EndFrame = INTEL_ORDER32(SampleLoopChunk.EndFrame);
+			}
+#endif
+			FWaveSampleLoop NewSampleLoop;
+			NewSampleLoop.LoopID = SampleLoopChunk.LoopID;
+			NewSampleLoop.StartFrame = SampleLoopChunk.StartFrame;
+			NewSampleLoop.EndFrame = SampleLoopChunk.EndFrame;
+			WaveSampleLoops.Add(NewSampleLoop);
+		}
+	}
+
 	return true;
 }
 
diff --git a/Engine/Source/Runtime/Engine/Private/SoundWave.cpp b/Engine/Source/Runtime/Engine/Private/SoundWave.cpp
--- a/Engine/Source/Runtime/Engine/Private/SoundWave.cpp
+++ b/Engine/Source/Runtime/Engine/Private/SoundWave.cpp
@@ -151,7 +151,8 @@ void FSoundWaveData::InitializeDataFromSoundWave(USoundWave& InWave)
 	}
 	
 	SoundWaveKeyCached = FObjectKey(&InWave);
-	CuePoints = InWave.CuePoints;
+	CuePoints = InWave.GetCuePoints();
+	LoopRegions = InWave.GetLoopRegions();
 	SampleRate = InWave.GetSampleRateForCurrentPlatform();
 	Duration = InWave.Duration;
 	NumChannels = InWave.NumChannels;
@@ -1150,6 +1151,32 @@ void USoundWave::SetSoundAssetCompressionType(ESoundAssetCompressionType InSound
 #endif // #if WITH_EDITOR
 }
 
+TArray<FSoundWaveCuePoint> USoundWave::GetCuePoints() const
+{
+	TArray<FSoundWaveCuePoint> OutCuePoints;
+	for (const FSoundWaveCuePoint& CuePoint : CuePoints)
+	{
+		if (!CuePoint.bIsLoopRegion)
+		{
+			OutCuePoints.Add(CuePoint);
+		}
+	}
+	return OutCuePoints;
+}
+
+TArray<FSoundWaveCuePoint> USoundWave::GetLoopRegions() const
+{
+	TArray<FSoundWaveCuePoint> OutLoopRegions;
+	for (const FSoundWaveCuePoint& CuePoint : CuePoints)
+	{
+		if (CuePoint.bIsLoopRegion)
+		{
+			OutLoopRegions.Add(CuePoint);
+		}
+	}
+	return OutLoopRegions;
+}
+
 FName USoundWave::GetRuntimeFormat() const
 {
 	return SoundWaveDataPtr->GetRuntimeFormat();
@@ -4215,6 +4242,12 @@ const TArray<FSoundWaveCuePoint>& FSoundWaveProxy::GetCuePoints() const
 	return SoundWaveDataPtr->GetCuePoints();
 }
 
+const TArray<FSoundWaveCuePoint>& FSoundWaveProxy::GetLoopRegions() const
+{
+	check(SoundWaveDataPtr);
+	return SoundWaveDataPtr->GetLoopRegions();
+}
+
 FSoundWaveData::MaxChunkSizeResults FSoundWaveProxy::GetMaxChunkSizeResults() const
 {
 	check(SoundWaveDataPtr);
diff --git a/Engine/Source/Runtime/Engine/Public/Audio.h b/Engine/Source/Runtime/Engine/Public/Audio.h
--- a/Engine/Source/Runtime/Engine/Public/Audio.h
+++ b/Engine/Source/Runtime/Engine/Public/Audio.h
@@ -860,6 +860,14 @@ struct FWaveCue
 	uint32 SampleLength = 0;
 };
 
+// data representing a sample loop in a wave file
+struct FWaveSampleLoop
+{
+	uint32 LoopID = 0;
+	uint32 StartFrame = 0;
+	uint32 EndFrame = 0;
+};
+
 //
 // Structure for in-memory interpretation and modification of WAVE sound structures.
 //
@@ -887,6 +895,9 @@ public:
 	// List of cues parsed from the wave file
 	TArray<FWaveCue> WaveCues;
 
+	// List of sample loops parsed from the wave file
+	TArray<FWaveSampleLoop> WaveSampleLoops;
+
 	// Timecode data if it was found on import.
 	TPimplPtr<FSoundWaveTimecodeInfo, EPimplPtrMode::DeepCopy> TimecodeInfo;
 
