TextureDerivedData make PVS Studio happy about OptionalMips

diff --git a/Engine/Source/Runtime/Engine/Private/TextureDerivedData.cpp b/Engine/Source/Runtime/Engine/Private/TextureDerivedData.cpp
--- a/Engine/Source/Runtime/Engine/Private/TextureDerivedData.cpp
+++ b/Engine/Source/Runtime/Engine/Private/TextureDerivedData.cpp
@@ -2854,8 +2854,9 @@ static void SerializePlatformData(
 			}
 
 			// Optional Mips (Streaming)
+			// OptionalMips == 0 when we don't have WITH_EDITORONLY
 			const uint32 OptionalBulkDataFlags = BULKDATA_Force_NOT_InlinePayload | BULKDATA_OptionalPayload;
-			for (int32 MipIndex = 0; MipIndex < OptionalMips; ++MipIndex)
+			for (int32 MipIndex = 0; MipIndex < OptionalMips; ++MipIndex) //-V654 //-V621
 			{
 				PlatformData->Mips[MipIndex + FirstMipToSerialize].BulkData.SetBulkDataFlags(OptionalBulkDataFlags);
 			}
