Fix ensure when using the OnAudioVirtualizationChanged delegate in editor

diff --git a/Engine/Source/Runtime/Engine/Private/AudioDevice.cpp b/Engine/Source/Runtime/Engine/Private/AudioDevice.cpp
--- a/Engine/Source/Runtime/Engine/Private/AudioDevice.cpp
+++ b/Engine/Source/Runtime/Engine/Private/AudioDevice.cpp
@@ -3266,6 +3266,10 @@ void FAudioDevice::VirtualizeInactiveLoops()
 	TRACE_CPUPROFILER_EVENT_SCOPE(FAudioDevice_VirtualizeLoops);
 
 	const bool bDoRangeCheck = true;
+	// Keep track of sounds to virtualize, then virtualize them after 
+	// to prevent ActiveSound array from changing while iterating 
+	// (ex. byOnAudioVirtualizationChanged delegate)
+	TArray<FActiveSound*> ActiveSoundsToVirtualize;
 	for (FActiveSound* ActiveSound : ActiveSounds)
 	{
 		// Don't virtualize if set to fade out
@@ -3280,14 +3284,19 @@ void FAudioDevice::VirtualizeInactiveLoops()
 			continue;
 		}
 
+		ActiveSoundsToVirtualize.Add(ActiveSound);
+	}
+
+	for (FActiveSound* ActiveSoundToVirtualize : ActiveSoundsToVirtualize)
+	{
 		FAudioVirtualLoop VirtualLoop;
-		if (FAudioVirtualLoop::Virtualize(*ActiveSound, bDoRangeCheck, VirtualLoop))
+		if (FAudioVirtualLoop::Virtualize(*ActiveSoundToVirtualize, bDoRangeCheck, VirtualLoop))
 		{
-			AddSoundToStop(ActiveSound);
+			AddSoundToStop(ActiveSoundToVirtualize);
 
 			// Clear must be called after AddSoundToStop to ensure AudioComponent is properly removed from AudioComponentIDToActiveSoundMap
-			ActiveSound->ClearAudioComponent();
-			if (USoundBase* Sound = ActiveSound->GetSound())
+			ActiveSoundToVirtualize->ClearAudioComponent();
+			if (USoundBase* Sound = ActiveSoundToVirtualize->GetSound())
 			{
 				UE_LOG(LogAudio, Verbose, TEXT("Playing ActiveSound %s Virtualizing: Out of audible range."), *Sound->GetName());
 			}
