AssetCompilingManager: Add OnPackageScopeEvent to let the cooker know for its reference tracking analytics which package is being processed during ProcessAsyncTasks.
#rnx

diff --git a/Engine/Source/Editor/UnrealEd/Private/Cooker/CookPackageData.cpp b/Engine/Source/Editor/UnrealEd/Private/Cooker/CookPackageData.cpp
--- a/Engine/Source/Editor/UnrealEd/Private/Cooker/CookPackageData.cpp
+++ b/Engine/Source/Editor/UnrealEd/Private/Cooker/CookPackageData.cpp
@@ -3010,7 +3010,21 @@ void FPackageDatas::PollPendingCookedPlatformDatas(bool bForce, double& LastCook
 
 	GShaderCompilingManager->ProcessAsyncResults(true /* bLimitExecutionTime */,
 		false /* bBlockOnGlobalShaderCompletion */);
+	FDelegateHandle EventHandle = FAssetCompilingManager::Get().OnPackageScopeEvent().AddLambda(
+		[this](UPackage* Package, bool bEntering)
+		{
+			if (bEntering)
+			{
+				CookOnTheFlyServer.SetActivePackage(Package->GetFName(), PackageAccessTrackingOps::NAME_CookerBuildObject);
+			}
+			else
+			{
+				CookOnTheFlyServer.ClearActivePackage();
+			}
+		});
 	FAssetCompilingManager::Get().ProcessAsyncTasks(true);
+	FAssetCompilingManager::Get().OnPackageScopeEvent().Remove(EventHandle);
+
 	if (LastCookableObjectTickTime + TickCookableObjectsFrameTime <= CurrentTime)
 	{
 		UE_SCOPED_COOKTIMER(TickCookableObjects);
diff --git a/Engine/Source/Runtime/Engine/Private/ActorDeferredScriptManager.cpp b/Engine/Source/Runtime/Engine/Private/ActorDeferredScriptManager.cpp
--- a/Engine/Source/Runtime/Engine/Private/ActorDeferredScriptManager.cpp
+++ b/Engine/Source/Runtime/Engine/Private/ActorDeferredScriptManager.cpp
@@ -189,6 +189,8 @@ void FActorDeferredScriptManager::ProcessAsyncTasks(bool bLimitExecutionTime)
 		TGuardValue<bool> IsEditorLoadingPackageGuard(GIsEditorLoadingPackage, true);
 
 		int32 CurrentIndex = NextIndexToProcess;
+		UPackage* EnteredPackage = nullptr;
+		FAssetCompilingManager::FPackageScopeEvent& OnPackageScopeEvent = FAssetCompilingManager::Get().OnPackageScopeEvent();
 
 		UE_LOG(LogActorDeferredScriptManager, VeryVerbose, TEXT("Starting to process at index %d"), CurrentIndex);
 
@@ -231,6 +233,20 @@ void FActorDeferredScriptManager::ProcessAsyncTasks(bool bLimitExecutionTime)
 
 				UE_LOG(LogActorDeferredScriptManager, VeryVerbose, TEXT("Rerunning construction script for index %d (%s)"), CircularIndex, *Actor->GetPathName());
 
+				UPackage* Package = Actor->GetPackage();
+				if (Package != EnteredPackage)
+				{
+					if (EnteredPackage)
+					{
+						OnPackageScopeEvent.Broadcast(EnteredPackage, false /* bEnter */);
+					}
+					EnteredPackage = Package;
+					if (EnteredPackage)
+					{
+						OnPackageScopeEvent.Broadcast(EnteredPackage, true /* bEnter */);
+					}
+				}
+
 				// Temporarily do not consider actor as initialized if they were when running deferred construction scripts
 				FGuardValue_Bitfield(Actor->bActorInitialized, false);
 				Actor->RerunConstructionScripts();
@@ -251,6 +267,10 @@ void FActorDeferredScriptManager::ProcessAsyncTasks(bool bLimitExecutionTime)
 				UE_LOG(LogActorDeferredScriptManager, VeryVerbose, TEXT("Finished processing construction scripts"));
 			}
 		}
+		if (EnteredPackage)
+		{
+			OnPackageScopeEvent.Broadcast(EnteredPackage, false /* bEnter */);
+		}
 	}
 
 	UpdateCompilationNotification();
diff --git a/Engine/Source/Runtime/Engine/Public/AssetCompilingManager.h b/Engine/Source/Runtime/Engine/Public/AssetCompilingManager.h
--- a/Engine/Source/Runtime/Engine/Public/AssetCompilingManager.h
+++ b/Engine/Source/Runtime/Engine/Public/AssetCompilingManager.h
@@ -2,6 +2,7 @@
 #pragma once
 
 #include "CoreMinimal.h"
+#include "Delegates/Delegate.h"
 #include "UObject/WeakObjectPtr.h"
 #include "UObject/Class.h"
 #include "UObject/ObjectMacros.h"
@@ -160,6 +161,16 @@ public:
 	 * Event called after an asset finishes compilation.
 	 */
 	FAssetPostCompileEvent& OnAssetPostCompileEvent() { return AssetPostCompileEvent; }
+
+	/**
+	 * Event called before and after FinishAllCompilation or ProcessAsyncTasks run operations on a specific package,
+	 * used for subscribers to associate low-level actions with that Package (e.g. TObjectPtr reads).
+	 * void OnPackageScopeEvent(UPackage* Package, bool bEntering)
+	 * Called with bEntering=true followed by actions for Package followed by call with bEntering=false.
+	 */
+	DECLARE_MULTICAST_DELEGATE_TwoParams(FPackageScopeEvent, UPackage*, bool);
+	FPackageScopeEvent& OnPackageScopeEvent() { return PackageScopeEvent; }
+
 private:
 	FAssetCompilingManager();
 	~FAssetCompilingManager();
@@ -176,4 +187,6 @@ private:
 
 	/** Event issued at the end of the compile process */
 	FAssetPostCompileEvent AssetPostCompileEvent;
+
+	FPackageScopeEvent PackageScopeEvent;
 };
