Fixed ensure condition when required bones are requested from a skeletal mesh component that hasn't yet allocated transforms.

diff --git a/Engine/Source/Runtime/Engine/Private/Components/SkeletalMeshComponent.cpp b/Engine/Source/Runtime/Engine/Private/Components/SkeletalMeshComponent.cpp
--- a/Engine/Source/Runtime/Engine/Private/Components/SkeletalMeshComponent.cpp
+++ b/Engine/Source/Runtime/Engine/Private/Components/SkeletalMeshComponent.cpp
@@ -1818,9 +1818,9 @@ void USkeletalMeshComponent::ComputeRequiredBones(TArray<FBoneIndexType>& OutReq
 	// Purge invisible bones and their children
 	// this has to be done before mirror table check/physics body checks
 	// mirror table/phys body ones has to be calculated
-	if (ShouldUpdateBoneVisibility())
+	const TArray<uint8>& EditableBoneVisibilityStates = GetEditableBoneVisibilityStates();
+	if (ShouldUpdateBoneVisibility() && EditableBoneVisibilityStates.Num() > 0)
 	{
-		const TArray<uint8>& EditableBoneVisibilityStates = GetEditableBoneVisibilityStates();
 		check(EditableBoneVisibilityStates.Num() == GetNumComponentSpaceTransforms());
 		
 		if (ensureMsgf(EditableBoneVisibilityStates.Num() >= OutRequiredBones.Num(), 
