[24.00] Fix "Audio+StreamedAudio" DDC Key to include AudioQuality
#localization none
#tests preflight

diff --git a/Engine/Source/Runtime/Engine/Private/AudioDerivedData.cpp b/Engine/Source/Runtime/Engine/Private/AudioDerivedData.cpp
--- a/Engine/Source/Runtime/Engine/Private/AudioDerivedData.cpp
+++ b/Engine/Source/Runtime/Engine/Private/AudioDerivedData.cpp
@@ -180,6 +180,20 @@ const TCHAR* LexToString(const ESoundwaveSampleRateSettings Enum)
 	return TEXT("<Unknown ESoundwaveSampleRateSettings>");
 }
 
+static FString GetSoundWaveHash(const USoundWave& InWave)
+{
+	// Hash the parts of the SoundWave that can affect the compressed data. It doesn't hurt to do this. 
+	// Typically the GUID will change if compressed data changes, but some settings can affect 
+	// the SoundWave. i.e. DefaultSoundWaveQuality which won't change the GUID. So is better
+	// it's reflected here.
+	using FPCU = FPlatformCompressionUtilities;
+	FString SoundWaveHash;
+	FPCU::AppendHash(SoundWaveHash, TEXT("QLT"), InWave.GetCompressionQuality());
+	FPCU::AppendHash(SoundWaveHash, TEXT("CHN"), InWave.NumChannels);
+	FPCU::AppendHash(SoundWaveHash, TEXT("SRQ"), InWave.SampleRateQuality);
+	return SoundWaveHash;
+}
+
 /**
  * Computes the derived data key suffix for a SoundWave's Streamed Audio.
  * @param SoundWave - The SoundWave for which to compute the derived data key.
@@ -214,23 +228,15 @@ static void GetStreamedAudioDerivedDataKeySuffix(
 		FPlatformAudioCookOverrides::GetHashSuffix(CompressionOverrides, AudioFormatNameString);
 	}
 
-	// Hash the parts of the SoundWave that can affect the compressed data. It doesn't hurt to do this. 
-	// Typically the GUID will change if compressed data changes, but some settings can affect 
-	// the SoundWave. i.e. DefaultSoundWaveQuality which won't change the GUID. So is better
-	// it's reflected here.
-	using FPCU = FPlatformCompressionUtilities;
-	FString SoundWaveHash;
-	FPCU::AppendHash(SoundWaveHash, TEXT("QLT"), SoundWave.GetCompressionQuality());
-	FPCU::AppendHash(SoundWaveHash, TEXT("CHN"), SoundWave.NumChannels);
-	FPCU::AppendHash(SoundWaveHash, TEXT("SRQ"), SoundWave.SampleRateQuality);
+	FString SoundWaveHash = GetSoundWaveHash(SoundWave);
 		
 	// build the key
-	OutKeySuffix = FString::Printf(TEXT("%s_%s_%d_%s"),
+	OutKeySuffix = FString::Printf(TEXT("%s_%d_%s_%s"),
 		*AudioFormatNameString,
-		*SoundWaveHash,
 		Version,
+		*SoundWaveHash,
 		*SoundWave.CompressedDataGuid.ToString()
-		);
+	);
 
 #if PLATFORM_CPU_ARM_FAMILY
 	// Separate out arm keys as x64 and arm64 clang do not generate the same data for a given
@@ -1289,7 +1295,8 @@ struct FAudioCookInputs
 	float                       SampleRateOverride = -1.0f;
 	bool                        bIsStreaming;
 	float                       CompressionQualityModifier;
-	
+	FString						SoundWaveHash;
+		
 	TArray<Audio::FTransformationPtr> WaveTransformations;
 
 	// Those are the only refs we keep on the actual USoundWave until
@@ -1313,6 +1320,7 @@ struct FAudioCookInputs
 		, bIsSoundWaveProcedural(InSoundWave->IsA<USoundWaveProcedural>())
 		, CompressionQuality(InSoundWave->GetCompressionQuality())
 		, CompressionQualityModifier(InCookOverrides ? InCookOverrides->CompressionQualityModifier : 1.0f)
+		, SoundWaveHash(GetSoundWaveHash(*InSoundWave))
 		, WaveTransformations(InSoundWave->CreateTransformations())
 		, BulkDataCriticalSection(InSoundWave->RawDataCriticalSection)
 		, BulkData(InSoundWave->RawData)
@@ -1764,7 +1772,13 @@ FString FDerivedAudioDataCompressor::GetPluginSpecificCacheKeySuffix() const
 
 	check(CookInputs->CompressedDataGuid.IsValid());
 	FString FormatHash = CookInputs->HashedFormat.ToString().ToUpper();
-	return FString::Printf(TEXT("%s_%04X_%s"), *FormatHash, FormatVersion, *CookInputs->CompressedDataGuid.ToString());
+	FString SoundWaveHash;
+
+#if WITH_EDITORONLY_DATA
+	SoundWaveHash = CookInputs->SoundWaveHash;
+#endif //WITH_EDITORONLY_DATA
+	
+	return FString::Printf(TEXT("%s_%04X_%s%s"), *FormatHash, FormatVersion, *SoundWaveHash, *CookInputs->CompressedDataGuid.ToString());
 }
 
 bool FDerivedAudioDataCompressor::IsBuildThreadsafe() const
