Added generate overlap option in the ISM descriptor.
Fixed class constructor warnings due to base class not initializing.

diff --git a/Engine/Source/Runtime/Engine/Private/ISMPartition/ISMComponentDescriptor.cpp b/Engine/Source/Runtime/Engine/Private/ISMPartition/ISMComponentDescriptor.cpp
--- a/Engine/Source/Runtime/Engine/Private/ISMPartition/ISMComponentDescriptor.cpp
+++ b/Engine/Source/Runtime/Engine/Private/ISMPartition/ISMComponentDescriptor.cpp
@@ -15,7 +15,14 @@
 #include "VT/RuntimeVirtualTexture.h"
 #include "Algo/Transform.h"
 
+FISMComponentDescriptorBase::FISMComponentDescriptorBase()
+{
+	// Note: should not really be used - prefer using FISMComponentDescriptor & FSoftISMDescriptor instead
+	InitFrom(UHierarchicalInstancedStaticMeshComponent::StaticClass()->GetDefaultObject<UHierarchicalInstancedStaticMeshComponent>());
+}
+
 FISMComponentDescriptor::FISMComponentDescriptor()
+	: FISMComponentDescriptorBase(NoInit)
 {
 	// Make sure we have proper defaults
 	InitFrom(UHierarchicalInstancedStaticMeshComponent::StaticClass()->GetDefaultObject<UHierarchicalInstancedStaticMeshComponent>());
@@ -31,6 +38,7 @@ FISMComponentDescriptor::FISMComponentDescriptor(const FSoftISMComponentDescript
 }
 
 FSoftISMComponentDescriptor::FSoftISMComponentDescriptor()
+	: FISMComponentDescriptorBase(NoInit)
 {
 	// Make sure we have proper defaults
 	InitFrom(UHierarchicalInstancedStaticMeshComponent::StaticClass()->GetDefaultObject<UHierarchicalInstancedStaticMeshComponent>());
@@ -92,6 +100,7 @@ void FISMComponentDescriptorBase::InitFrom(const UStaticMeshComponent* Template,
 	const bool bIsLocalToWorldDeterminantNegative = Template->GetRenderMatrix().Determinant() < 0;
 	bReverseCulling = Template->bReverseCulling != bIsLocalToWorldDeterminantNegative;
 	bUseDefaultCollision = Template->bUseDefaultCollision;
+	bGenerateOverlapEvents = Template->GetGenerateOverlapEvents();
 
 #if WITH_EDITORONLY_DATA
 	HLODBatchingPolicy = Template->HLODBatchingPolicy;
@@ -186,6 +195,7 @@ bool FISMComponentDescriptorBase::operator==(const FISMComponentDescriptorBase&
 	bEvaluateWorldPositionOffset == Other.bEvaluateWorldPositionOffset &&
 	bReverseCulling == Other.bReverseCulling &&
 	bUseDefaultCollision == Other.bUseDefaultCollision &&
+	bGenerateOverlapEvents == Other.bGenerateOverlapEvents &&
 	WorldPositionOffsetDisableDistance == Other.WorldPositionOffsetDisableDistance &&
 #if WITH_EDITORONLY_DATA
 	HLODBatchingPolicy == Other.HLODBatchingPolicy &&
@@ -291,6 +301,7 @@ void FISMComponentDescriptorBase::InitComponent(UInstancedStaticMeshComponent* I
 	ISMComponent->bEvaluateWorldPositionOffset = bEvaluateWorldPositionOffset;
 	ISMComponent->bReverseCulling = bReverseCulling;
 	ISMComponent->bUseDefaultCollision = bUseDefaultCollision;
+	ISMComponent->SetGenerateOverlapEvents(bGenerateOverlapEvents);
 	ISMComponent->WorldPositionOffsetDisableDistance = WorldPositionOffsetDisableDistance;
 	
 #if WITH_EDITORONLY_DATA
diff --git a/Engine/Source/Runtime/Engine/Public/ISMPartition/ISMComponentDescriptor.h b/Engine/Source/Runtime/Engine/Public/ISMPartition/ISMComponentDescriptor.h
--- a/Engine/Source/Runtime/Engine/Public/ISMPartition/ISMComponentDescriptor.h
+++ b/Engine/Source/Runtime/Engine/Public/ISMPartition/ISMComponentDescriptor.h
@@ -19,7 +19,8 @@ struct ENGINE_API FISMComponentDescriptorBase
 {
 	GENERATED_BODY()
 
-	FISMComponentDescriptorBase() = default;
+	FISMComponentDescriptorBase();
+	explicit FISMComponentDescriptorBase(ENoInit) {}
 	void InitFrom(const UStaticMeshComponent* Component, bool bInitBodyInstance = true);
 
 	uint32 ComputeHash() const;
@@ -150,9 +151,12 @@ public:
 	uint8 bConsiderForActorPlacementWhenHidden : 1;
 #endif
 
-	UPROPERTY(EditAnywhere, Category = "Component Settings", meta = (DisplayAfter = "CustomDepthStencilWriteMask"))
+	UPROPERTY(EditAnywhere, Category = "Component Settings", meta = (DisplayAfter = "bGenerateOverlapEvents"))
 	uint8 bUseDefaultCollision : 1;
 
+	UPROPERTY(EditAnywhere, Category = "Component Settings", meta = (DisplayAfter = "CustomDepthStencilWriteMask"))
+	uint8 bGenerateOverlapEvents : 1;
+
 	UPROPERTY(EditAnywhere, Category = "Component Settings")
 	int32 WorldPositionOffsetDisableDistance;
 };
