All SparseVolumeTexture mips are now stored in the same resource/proxy, making binding and sampling different mip levels much easier. This version is still missing streaming at mip level granularity and has a minimum mip level size of 16x16x16 (no mip tails). Since the added logic started bloating SparseVolumeTexture.h/.cpp, FSparseVolumeRawSource (now renamed to FSparseVolumeTextureData), FSparseVolumeTextureRuntime and FSparseVolumeTextureProxy were moved into their own files. FSparseVolumeTextureData stores a potentially mipped (after cooking) 3D page table indexing into linear arrays of physical tile data. At runtime the physical tiles are rearranged/packed into 3D physical tile textures in FSparseVolumeTextureRuntime and the page table payload adjusted from indices to 3D tile coordinates. The SVTSample node now also takes an optional MipLevel parameter which enables sampling across mips (lerp).

#rnx

diff --git a/Engine/Shaders/Private/SparseVolumeTexture/SparseVolumeTextureCommon.ush b/Engine/Shaders/Private/SparseVolumeTexture/SparseVolumeTextureCommon.ush
--- a/Engine/Shaders/Private/SparseVolumeTexture/SparseVolumeTextureCommon.ush
+++ b/Engine/Shaders/Private/SparseVolumeTexture/SparseVolumeTextureCommon.ush
@@ -4,39 +4,26 @@
 
 struct FSparseVolumeTextureUniforms
 {
-	float3			PageTableSize;
-	float3			UVScale;
-	float3			UVBias;
+	float3			VolumePageResolution;
+	float3			PageTableOffset;
+	int             MinMipLevel;
+	int             MaxMipLevel;
 	float			TileSize;
 };
 
-FSparseVolumeTextureUniforms SparseVolumeTextureCreateUniforms(float TileSize, float3 PageTableSize, float3 UVScale, float3 UVBias)
+FSparseVolumeTextureUniforms SparseVolumeTextureUnpackUniforms(const uint4 Packed0, const uint4 Packed1)
 {
 	FSparseVolumeTextureUniforms Result;
-	Result.PageTableSize = PageTableSize;
-	Result.UVScale = UVScale;
-	Result.UVBias = UVBias;
-	Result.TileSize = TileSize;
-	return Result;
-}
-
-FSparseVolumeTextureUniforms SparseVolumeTextureUnpackUniforms(uint4 Packed0, uint4 Packed1)
-{
-	FSparseVolumeTextureUniforms Result;
-	Result.UVScale = asfloat(Packed0.xyz);
-	Result.UVBias = asfloat(Packed1.xyz);
+	Result.VolumePageResolution = asfloat(Packed0.xyz);
+	Result.PageTableOffset = asfloat(Packed1.xyz);
+	Result.MinMipLevel = int(Packed0.w & 0xFFu);
+	Result.MaxMipLevel = int((Packed0.w >> 8u) & 0xFFu);
 	Result.TileSize = asfloat(Packed1.w);
 
-	// Packed as 11|11|10. Since this limits us to 2047x2047x1023, we add an implicit +1 during unpacking.
-	Result.PageTableSize = float3(
-		(Packed0.w & 0x7FF) + 1,
-		((Packed0.w >> 11) & 0x7FF) + 1,
-		((Packed0.w >> 22) & 0x3FF) + 1);
-
 	return Result;
 }
 
-float3 SparseVolumeTextureGetVoxelCoord(uint PackedPhysicalTileCoord, float TileSize, float3 PageTableCoord)
+float3 SparseVolumeTextureGetVoxelCoord(const uint PackedPhysicalTileCoord, const float TileSize, const float3 PageTableCoord)
 {
 	const int3 PhysicalTileCoord = int3(
 		PackedPhysicalTileCoord & 0x7FF,
@@ -47,18 +34,21 @@ float3 SparseVolumeTextureGetVoxelCoord(uint PackedPhysicalTileCoord, float Tile
 	return VoxelCoord;
 }
 
-float3 SparseVolumeTextureSamplePageTable(Texture3D<uint> PageTable, FSparseVolumeTextureUniforms Uniforms, float3 UVW)
+float3 SparseVolumeTextureSamplePageTable(Texture3D<uint> PageTable, const FSparseVolumeTextureUniforms Uniforms, const float3 UVW, int MipLevel = 0)
 {
-	const float3 UVWScaledBiased = UVW * Uniforms.UVScale + Uniforms.UVBias;
-	const float3 PageTableCoord = (UVWScaledBiased * Uniforms.PageTableSize);
-	
-	const uint PackedPhysicalTileCoord = PageTable.Load(int4(int3(PageTableCoord), 0)).x;
-	
+	MipLevel = clamp(MipLevel, Uniforms.MinMipLevel, Uniforms.MaxMipLevel);
+	const float RcpMipLevelFactor = rcp(float(1u << (uint)MipLevel));
+	const float3 VolumePageCoord = UVW * Uniforms.VolumePageResolution;
+	const float3 MipPageTableOffset = floor(Uniforms.PageTableOffset * RcpMipLevelFactor);
+	const float3 PageTableCoord = VolumePageCoord * RcpMipLevelFactor - MipPageTableOffset;
+
+	const uint PackedPhysicalTileCoord = PageTable.Load(int4(int3(PageTableCoord), (MipLevel - Uniforms.MinMipLevel))).x;
+
 	const float3 VoxelCoord = SparseVolumeTextureGetVoxelCoord(PackedPhysicalTileCoord, Uniforms.TileSize, PageTableCoord);
 	return VoxelCoord;
 }
 
-float4 SparseVolumeTextureSamplePhysicalTileData(Texture3D PhysicalTileDataA, Texture3D PhysicalTileDataB, float3 VoxelCoord, int PhysicalTileDataIndex)
+float4 SparseVolumeTextureSamplePhysicalTileData(Texture3D PhysicalTileDataA, Texture3D PhysicalTileDataB, const float3 VoxelCoord, const int PhysicalTileDataIndex)
 {
 	switch (PhysicalTileDataIndex)
 	{
@@ -66,4 +56,4 @@ float4 SparseVolumeTextureSamplePhysicalTileData(Texture3D PhysicalTileDataA, Te
 	case 1: return PhysicalTileDataB.Load(int4(int3(VoxelCoord), 0));
 	default: return 0.0f;
 	}
-} 
+}
diff --git a/Engine/Shaders/Private/SparseVolumeTexture/VisualizeSparseVolumeTexture.usf b/Engine/Shaders/Private/SparseVolumeTexture/VisualizeSparseVolumeTexture.usf
--- a/Engine/Shaders/Private/SparseVolumeTexture/VisualizeSparseVolumeTexture.usf
+++ b/Engine/Shaders/Private/SparseVolumeTexture/VisualizeSparseVolumeTexture.usf
@@ -5,6 +5,7 @@
 =============================================================================*/
 
 #include "/Engine/Private/Common.ush"
+#include "SparseVolumeTextureCommon.ush"
 
 
 #ifdef VERTEX_SHADER
@@ -40,6 +41,13 @@ bool Slabs(float3 p0, float3 p1, float3 rayOrigin, float3 invRaydir, out float o
 	return maxtmin <= mintmax;
 }
 
+Texture3D<uint> SparseVolumeTexturePageTable;
+Texture3D<float4> SparseVolumeTextureA;
+Texture3D<float4> SparseVolumeTextureB;
+uint4 PackedSVTUniforms0;
+uint4 PackedSVTUniforms1;
+float3 SparseVolumeTextureResolution;
+int MipLevel;
 float4 WorldToLocal0;
 float4 WorldToLocal1;
 float4 WorldToLocal2;
@@ -47,19 +55,8 @@ float3 WorldToLocalNoScale0;
 float3 WorldToLocalNoScale1;
 float3 WorldToLocalNoScale2;
 uint ComponentToVisualize;
-uint MipLevel;
 float Extinction;
 
-
-Texture3D<float4> SparseVolumeTextureA;
-Texture3D<float4> SparseVolumeTextureB;
-float3 SparseVolumeTextureResolution;
-
-Texture3D<uint> SparseVolumeTexturePageTable;
-float3 SparseVolumeTexturePageTableResolution;
-
-#define SPARSE_VOLUME_TILE_RES	16
-
 void VisualizeSparseVolumeTexturePS(
 	in float4 SVPos : SV_POSITION,
 	out float4 OutLuminanceAlpha : SV_Target0
@@ -101,7 +98,8 @@ void VisualizeSparseVolumeTexturePS(
 
 		float3 StartPos = RayOrig + RayDir * max(0.0, TMin);
 
-		float3 VolumeResolution = SparseVolumeTexturePageTableResolution * SPARSE_VOLUME_TILE_RES;
+		float3 VolumeResolution = SparseVolumeTextureResolution;
+		float3 RcpVolumeResolution = 1.0f / VolumeResolution;
 
 		const float3 StartUVs = StartPos * 0.5f + 0.5f;
 		float3 P = StartUVs * VolumeResolution;
@@ -132,31 +130,27 @@ void VisualizeSparseVolumeTexturePS(
 		tMax.y *= tDelta.y;
 		tMax.z *= tDelta.z;
 
-		float StepLength = float(1 << MipLevel);
+		const FSparseVolumeTextureUniforms SVTUniforms = SparseVolumeTextureUnpackUniforms(PackedSVTUniforms0, PackedSVTUniforms1);
+
+		float StepLength = 1.0f;
 
 		LOOP
 		while (P.x >= 0 && P.y >= 0 && P.z >= 0 && P.x <= VolumeResolution.x && P.y <= VolumeResolution.y && P.z <= VolumeResolution.z)
 		{
-			if (all(P >= 0.0f) && all(P < SparseVolumeTextureResolution))
+			const float3 UVW = P * RcpVolumeResolution;
+			const float3 VoxelCoord = SparseVolumeTextureSamplePageTable(SparseVolumeTexturePageTable, SVTUniforms, UVW, MipLevel);
+			if (any(VoxelCoord >= SVTUniforms.TileSize)) // skip on null tile
 			{
-				const float3 PageCoordF = floor(P * (1.0 / SPARSE_VOLUME_TILE_RES));
-				const uint EncodedPage = SparseVolumeTexturePageTable.Load(int4(int3(PageCoordF), 0));
-				const int3 TileCoord = int3(
-					EncodedPage & 0x7FF,
-					(EncodedPage >> 11) & 0x7FF,
-					(EncodedPage >> 22) & 0x3FF);
-
-				const int3 VoxelCoord = TileCoord * SPARSE_VOLUME_TILE_RES + int3(P - PageCoordF * SPARSE_VOLUME_TILE_RES);
+				const int PhysicalTileDataIndex = ComponentToVisualize < 4 ? 0 : 1;
+				const float4 VoxelData = SparseVolumeTextureSamplePhysicalTileData(SparseVolumeTextureA, SparseVolumeTextureB, VoxelCoord, PhysicalTileDataIndex);
 
 				float Density = 0.0f;
 				if (ComponentToVisualize < 4)
 				{
-					float4 VoxelData = SparseVolumeTextureA.Load(int4(VoxelCoord, 0));
 					Density = VoxelData[ComponentToVisualize];
 				}
 				else
 				{
-					float4 VoxelData = SparseVolumeTextureB.Load(int4(VoxelCoord, 0));
 					Density = VoxelData[min(ComponentToVisualize - 4, 3u)];
 				}
 
diff --git a/Engine/Source/Editor/SparseVolumeTexture/Private/SparseVolumeTextureFactory.cpp b/Engine/Source/Editor/SparseVolumeTexture/Private/SparseVolumeTextureFactory.cpp
--- a/Engine/Source/Editor/SparseVolumeTexture/Private/SparseVolumeTextureFactory.cpp
+++ b/Engine/Source/Editor/SparseVolumeTexture/Private/SparseVolumeTextureFactory.cpp
@@ -3,6 +3,7 @@
 #include "SparseVolumeTextureFactory.h"
 
 #include "SparseVolumeTexture/SparseVolumeTexture.h"
+#include "SparseVolumeTexture/SparseVolumeTextureData.h"
 
 #if WITH_EDITOR
 
@@ -478,7 +479,7 @@ UObject* USparseVolumeTextureFactory::FactoryCreateFile(UClass* InClass, UObject
 	{
 		int32 Levels = 1;
 		FIntVector3 Resolution = VolumeBoundsMax - VolumeBoundsMin;
-		while (Resolution.X > 1 || Resolution.Y > 1 || Resolution.Z > 1)
+		while (Resolution.X > SPARSE_VOLUME_TILE_RES || Resolution.Y > SPARSE_VOLUME_TILE_RES || Resolution.Z > SPARSE_VOLUME_TILE_RES)
 		{
 			Resolution /= 2;
 			++Levels;
@@ -503,11 +504,10 @@ UObject* USparseVolumeTextureFactory::FactoryCreateFile(UClass* InClass, UObject
 		ExpandVolumeBounds(ImportOptions, PreviewData.GridInfo, VolumeBoundsMin, VolumeBoundsMax);
 		StaticSVTexture->VolumeResolution = VolumeBoundsMax - VolumeBoundsMin;
 		StaticSVTexture->Frames.SetNum(1);
+		StaticSVTexture->NumMipLevels = ComputeNumMipLevels(VolumeBoundsMin, VolumeBoundsMax);
 
-		const int32 NumMipLevels = ComputeNumMipLevels(VolumeBoundsMin, VolumeBoundsMax);
-
-		FSparseVolumeRawSource SparseVolumeRawSource{};
-		const bool bConversionSuccess = ConvertOpenVDBToSparseVolumeTexture(PreviewData.LoadedFile, ImportOptions, VolumeBoundsMin, SparseVolumeRawSource);
+		FSparseVolumeTextureData TextureData{};
+		const bool bConversionSuccess = ConvertOpenVDBToSparseVolumeTexture(PreviewData.LoadedFile, ImportOptions, VolumeBoundsMin, TextureData);
 
 		if (!bConversionSuccess)
 		{
@@ -516,18 +516,8 @@ UObject* USparseVolumeTextureFactory::FactoryCreateFile(UClass* InClass, UObject
 		}
 
 		// Serialize the raw source data into the asset object.
-		// After serializing the data, we generate the next mip level by downsampling the current one (with a call to GenerateMipMap())
-		// and then repeating the loop until we processed all levels.
-		StaticSVTexture->Frames[0].SetNum(NumMipLevels);
-		for (int32 MipLevel = 0; MipLevel < NumMipLevels; ++MipLevel)
-		{
-			UE::Serialization::FEditorBulkDataWriter RawDataArchiveWriter(StaticSVTexture->Frames[0][MipLevel].RawData);
-			SparseVolumeRawSource.Serialize(RawDataArchiveWriter);
-			if ((MipLevel + 1) < NumMipLevels)
-			{
-				SparseVolumeRawSource = SparseVolumeRawSource.GenerateMipMap();
-			}
-		}
+		UE::Serialization::FEditorBulkDataWriter RawDataArchiveWriter(StaticSVTexture->Frames[0].RawData);
+		TextureData.Serialize(RawDataArchiveWriter);
 
 		if (ImportTask.ShouldCancel())
 		{
@@ -626,7 +616,7 @@ UObject* USparseVolumeTextureFactory::FactoryCreateFile(UClass* InClass, UObject
 
 		ImportTask.EnterProgressFrame(1.0f, LOCTEXT("ConvertingVDBAnim", "Converting OpenVDB animation"));
 
-		const int32 NumMipLevels = ComputeNumMipLevels(VolumeBoundsMin, VolumeBoundsMax);
+		AnimatedSVTexture->NumMipLevels = ComputeNumMipLevels(VolumeBoundsMin, VolumeBoundsMax);
 
 		FEvent* AllTasksFinishedEvent = FPlatformProcess::GetSynchEventFromPool();
 		std::atomic_int FinishedTasksCounter = 0; // Will be incremented even if frame processing failed
@@ -654,7 +644,7 @@ UObject* USparseVolumeTextureFactory::FactoryCreateFile(UClass* InClass, UObject
 
 			AsyncTask(ENamedThreads::AnyNormalThreadNormalTask, 
 				[FrameIdx, NumFrames, &PreviewData, &ImportOptions, &AnimatedSVTexture,
-				AllTasksFinishedEvent, &bErrored, &bCanceled, &FinishedTasksCounter, &ProcessedFramesCounter, &VolumeBoundsMin, &NumMipLevels]()
+				AllTasksFinishedEvent, &bErrored, &bCanceled, &FinishedTasksCounter, &ProcessedFramesCounter, &VolumeBoundsMin]()
 				{
 					// Ensure the FinishedTasksCounter will be incremented in all cases
 					FScopedIncrementer Incremeter(FinishedTasksCounter, NumFrames, AllTasksFinishedEvent);
@@ -677,8 +667,8 @@ UObject* USparseVolumeTextureFactory::FactoryCreateFile(UClass* InClass, UObject
 						return;
 					}
 
-					FSparseVolumeRawSource SparseVolumeRawSource{};
-					const bool bConversionSuccess = ConvertOpenVDBToSparseVolumeTexture(LoadedFrameFile, ImportOptions, VolumeBoundsMin, SparseVolumeRawSource);
+					FSparseVolumeTextureData TextureData{};
+					const bool bConversionSuccess = ConvertOpenVDBToSparseVolumeTexture(LoadedFrameFile, ImportOptions, VolumeBoundsMin, TextureData);
 
 					if (!bConversionSuccess)
 					{
@@ -688,18 +678,8 @@ UObject* USparseVolumeTextureFactory::FactoryCreateFile(UClass* InClass, UObject
 					}
 
 					// Serialize the raw source data from this frame into the asset object.
-					// After serializing the data, we generate the next mip level by downsampling the current one (with a call to GenerateMipMap())
-					// and then repeating the loop until we processed all levels.
-					AnimatedSVTexture->Frames[FrameIdx].SetNum(NumMipLevels);
-					for (int32 MipLevel = 0; MipLevel < NumMipLevels; ++MipLevel)
-					{
-						UE::Serialization::FEditorBulkDataWriter RawDataArchiveWriter(AnimatedSVTexture->Frames[FrameIdx][MipLevel].RawData);
-						SparseVolumeRawSource.Serialize(RawDataArchiveWriter);
-						if ((MipLevel + 1) < NumMipLevels)
-						{
-							SparseVolumeRawSource = SparseVolumeRawSource.GenerateMipMap();
-						}
-					}
+					UE::Serialization::FEditorBulkDataWriter RawDataArchiveWriter(AnimatedSVTexture->Frames[FrameIdx].RawData);
+					TextureData.Serialize(RawDataArchiveWriter);
 
 					// Increment ProcessedFramesCounter
 					ProcessedFramesCounter.fetch_add(1);
diff --git a/Engine/Source/Editor/SparseVolumeTexture/Private/SparseVolumeTextureOpenVDBUtility.cpp b/Engine/Source/Editor/SparseVolumeTexture/Private/SparseVolumeTextureOpenVDBUtility.cpp
--- a/Engine/Source/Editor/SparseVolumeTexture/Private/SparseVolumeTextureOpenVDBUtility.cpp
+++ b/Engine/Source/Editor/SparseVolumeTexture/Private/SparseVolumeTextureOpenVDBUtility.cpp
@@ -5,6 +5,7 @@
 #include "SparseVolumeTextureOpenVDBUtility.h"
 #include "SparseVolumeTextureOpenVDB.h"
 #include "SparseVolumeTexture/SparseVolumeTexture.h"
+#include "SparseVolumeTexture/SparseVolumeTextureData.h"
 #include "OpenVDBGridAdapter.h"
 #include "OpenVDBImportOptions.h"
 
@@ -175,13 +176,6 @@ static FOpenVDBGridInfo GetOpenVDBGridInfo(openvdb::GridBase::Ptr Grid, uint32 G
 
 bool IsOpenVDBGridValid(const FOpenVDBGridInfo& GridInfo, const FString& Filename)
 {
-	if (GridInfo.VolumeActiveDim.X * GridInfo.VolumeActiveDim.Y * GridInfo.VolumeActiveDim.Z == 0)
-	{
-		// SVT_TODO we should gently handle that case
-		UE_LOG(LogSparseVolumeTextureOpenVDBUtility, Warning, TEXT("OpenVDB grid is empty due to volume size being 0: %s"), *Filename);
-		return false;
-	}
-
 	if (!GridInfo.bHasUniformVoxels)
 	{
 		UE_LOG(LogSparseVolumeTextureOpenVDBUtility, Warning, TEXT("OpenVDB importer cannot handle non uniform voxels: %s"), *Filename);
@@ -258,7 +252,7 @@ static EPixelFormat GetMultiComponentFormat(ESparseVolumeAttributesFormat Format
 
 #if OPENVDB_AVAILABLE
 
-class FSparseVolumeRawSourceConstructionOpenVDBAdapter : public ISparseVolumeRawSourceConstructionAdapter
+class FSparseVolumeTextureDataConstructionOpenVDBAdapter : public ISparseVolumeTextureDataConstructionAdapter
 {
 public:
 
@@ -468,10 +462,10 @@ private:
 
 #endif // OPENVDB_AVAILABLE
 
-bool ConvertOpenVDBToSparseVolumeTexture(TArray64<uint8>& SourceFile, const FOpenVDBImportOptions& ImportOptions, const FIntVector3& VolumeBoundsMin, FSparseVolumeRawSource& OutResult)
+bool ConvertOpenVDBToSparseVolumeTexture(TArray64<uint8>& SourceFile, const FOpenVDBImportOptions& ImportOptions, const FIntVector3& VolumeBoundsMin, FSparseVolumeTextureData& OutResult)
 {
 #if OPENVDB_AVAILABLE
-	FSparseVolumeRawSourceConstructionOpenVDBAdapter Adapter;
+	FSparseVolumeTextureDataConstructionOpenVDBAdapter Adapter;
 	if (!Adapter.Initialize(SourceFile, ImportOptions, VolumeBoundsMin))
 	{
 		return false;
diff --git a/Engine/Source/Editor/SparseVolumeTexture/Private/SparseVolumeTextureOpenVDBUtility.h b/Engine/Source/Editor/SparseVolumeTexture/Private/SparseVolumeTextureOpenVDBUtility.h
--- a/Engine/Source/Editor/SparseVolumeTexture/Private/SparseVolumeTextureOpenVDBUtility.h
+++ b/Engine/Source/Editor/SparseVolumeTexture/Private/SparseVolumeTextureOpenVDBUtility.h
@@ -47,7 +47,7 @@ bool IsOpenVDBGridValid(const FOpenVDBGridInfo& GridInfo, const FString& Filenam
 
 bool GetOpenVDBGridInfo(TArray64<uint8>& SourceFile, bool bCreateStrings, TArray<FOpenVDBGridInfo>* OutGridInfo);
 
-bool ConvertOpenVDBToSparseVolumeTexture(TArray64<uint8>& SourceFile, const struct FOpenVDBImportOptions& ImportOptions, const FIntVector3& VolumeBoundsMin, struct FSparseVolumeRawSource& OutResult);
+bool ConvertOpenVDBToSparseVolumeTexture(TArray64<uint8>& SourceFile, const struct FOpenVDBImportOptions& ImportOptions, const FIntVector3& VolumeBoundsMin, struct FSparseVolumeTextureData& OutResult);
 
 const TCHAR* OpenVDBGridTypeToString(EOpenVDBGridType Type);
 
diff --git a/Engine/Source/Runtime/Engine/Classes/Materials/MaterialExpressionSparseVolumeTextureSample.h b/Engine/Source/Runtime/Engine/Classes/Materials/MaterialExpressionSparseVolumeTextureSample.h
--- a/Engine/Source/Runtime/Engine/Classes/Materials/MaterialExpressionSparseVolumeTextureSample.h
+++ b/Engine/Source/Runtime/Engine/Classes/Materials/MaterialExpressionSparseVolumeTextureSample.h
@@ -26,6 +26,9 @@ class ENGINE_API UMaterialExpressionSparseVolumeTextureSample : public UMaterial
 	UPROPERTY(meta = (RequiredInput = "false", ToolTip = "Defaults to 'SparseVolumeTexture' if not specified"))
 	FExpressionInput TextureObject;
 
+	UPROPERTY(meta = (RequiredInput = "false", ToolTip = "Defaults to 0 if not specified"))
+	FExpressionInput MipLevel;
+
 protected:
 
 #if WITH_EDITOR
diff --git a/Engine/Source/Runtime/Engine/Classes/SparseVolumeTexture/SparseVolumeTexture.h b/Engine/Source/Runtime/Engine/Classes/SparseVolumeTexture/SparseVolumeTexture.h
--- a/Engine/Source/Runtime/Engine/Classes/SparseVolumeTexture/SparseVolumeTexture.h
+++ b/Engine/Source/Runtime/Engine/Classes/SparseVolumeTexture/SparseVolumeTexture.h
@@ -18,159 +18,26 @@ namespace UE { namespace DerivedData { class FRequestOwner; } }
 
 #define SPARSE_VOLUME_TILE_RES	16
 
-uint32 SparseVolumeTexturePackPageTableEntry(const FIntVector3& Coord);
-FIntVector3 SparseVolumeTextureUnpackPageTableEntry(uint32 Packed);
+struct FSparseVolumeTextureData;
+class FSparseVolumeTextureSceneProxy;
 
-struct ENGINE_API FSparseVolumeAssetHeader
+struct ENGINE_API FSparseVolumeTextureHeader
 {
+	static const uint32 kVersion = 0; // The current data format version for the header.
+	uint32 Version = kVersion; // This version can be used to convert existing header to new version later.
+
 	FIntVector3 VirtualVolumeResolution = FIntVector3(0, 0, 0);
 	FIntVector3 VirtualVolumeAABBMin = FIntVector3(INT32_MAX, INT32_MAX, INT32_MAX);
 	FIntVector3 VirtualVolumeAABBMax = FIntVector3(INT32_MIN, INT32_MIN, INT32_MIN);
 	FIntVector3 PageTableVolumeResolution = FIntVector3(0, 0, 0);
 	FIntVector3 PageTableVolumeAABBMin = FIntVector3(INT32_MAX, INT32_MAX, INT32_MAX);
 	FIntVector3 PageTableVolumeAABBMax = FIntVector3(INT32_MIN, INT32_MIN, INT32_MIN);
-	FIntVector3 TileDataVolumeResolution = FIntVector3(0, 0, 0);
 	TStaticArray<EPixelFormat, 2> AttributesFormats = TStaticArray<EPixelFormat, 2>(InPlace, PF_Unknown);
-	int32 MipLevel = 0;
-	bool bHasNullTile = false;
-
-	// The current data format version for the header.
-	static const uint32 kVersion = 0;
-
-	// This version can be used to convert existing header to new version later.
-	uint32 Version = kVersion;
-
-	void Serialize(FArchive& Ar);
-};
-
-class ISparseVolumeRawSourceConstructionAdapter
-{
-public:
-	struct FAttributesInfo
-	{
-		EPixelFormat Format;
-		FVector4f FallbackValue;
-		FVector4f NormalizeScale;
-		FVector4f NormalizeBias;
-		bool bNormalized;
-	};
-
-	virtual void GetAttributesInfo(FAttributesInfo& OutInfoA, FAttributesInfo& OutInfoB) const = 0;
-	virtual FIntVector3 GetAABBMin() const = 0;
-	virtual FIntVector3 GetAABBMax() const = 0;
-	virtual FIntVector3 GetResolution() const = 0;
-	virtual void IteratePhysicalSource(TFunctionRef<void(const FIntVector3& Coord, int32 AttributesIdx, int32 ComponentIdx, float VoxelValue)> OnVisit) const = 0;
-	virtual ~ISparseVolumeRawSourceConstructionAdapter() = default;
-};
-
-// The structure represent the source asset in high quality. It is used to cook the runtime data
-struct ENGINE_API FSparseVolumeRawSource
-{
-	FSparseVolumeAssetHeader Header;
-	TArray<uint32> PageTable;
-	TArray<uint8> PhysicalTileDataA;
-	TArray<uint8> PhysicalTileDataB;
-
-	// The current data format version for the raw source data.
-	static const uint32 kVersion = 0;
-
-	// This version can be used to convert existing source data to new version later.
-	uint32 Version;
-
-	void Serialize(FArchive& Ar);
-
-	bool Construct(const ISparseVolumeRawSourceConstructionAdapter& Adapter);
-	uint32 ReadPageTablePacked(const FIntVector3& PageTableCoord) const;
-	FIntVector3 ReadPageTable(const FIntVector3& PageTableCoord) const;
-	FVector4f ReadTileDataVoxel(const FIntVector3& TileDataCoord, int32 AttributesIdx) const;
-	FVector4f Sample(const FIntVector3& VolumeCoord, int32 AttributesIdx) const;
-	void Sample(const FIntVector3& VolumeCoord, FVector4f& OutAttributesA, FVector4f& OutAttributesB) const;
-	void WriteTileDataVoxel(const FIntVector3& TileDataCoord, int32 AttributesIdx, const FVector4f& Value, int32 DstComponent = -1);
-	void FillNullTile(const FVector4f& FallbackValueA, const FVector4f& FallbackValueB);
-	FSparseVolumeRawSource GenerateMipMap() const;
-
-	FSparseVolumeRawSource()
-		: Version(kVersion)
-	{
-	}
-};
-
-// The structure represent the runtime data cooked runtime data.
-struct ENGINE_API FSparseVolumeTextureRuntime
-{
-	FSparseVolumeAssetHeader	Header;
-	TArray<uint32>				PageTable;
-	TArray<uint8>				PhysicalTileDataA;
-	TArray<uint8>				PhysicalTileDataB;
-
-	void SetAsDefaultTexture();
-
-	// The current data format version for the raw source data.
-	static const uint32 kVersion = 0;
-
-	// This version can be used to convert existing source data to new version later.
-	uint32 Version;
+	TStaticArray<FVector4f, 2> NullTileValues = TStaticArray<FVector4f, 2>(InPlace, FVector4f());
 
 	void Serialize(FArchive& Ar);
-
-	FSparseVolumeTextureRuntime()
-		: Header()
-		, Version(kVersion)
-	{
-	}
-};
-
-
-class FSparseVolumeTextureSceneProxy : public FRenderResource
-{
-public:
-
-	FSparseVolumeTextureSceneProxy();
-	virtual ~FSparseVolumeTextureSceneProxy() override;
-
-	FSparseVolumeTextureRuntime& GetRuntimeData()
-	{
-		return SparseVolumeTextureRuntime;
-	}
-
-	const FSparseVolumeTextureRuntime& GetRuntimeData() const
-	{
-		return SparseVolumeTextureRuntime;
-	}
-
-	const FSparseVolumeAssetHeader& GetHeader() const 
-	{
-		return SparseVolumeTextureRuntime.Header;
-	}
-
-	FTextureRHIRef GetPhysicalTileDataATextureRHI() const
-	{
-		return PhysicalTileDataATextureRHI;
-	}
-	FTextureRHIRef GetPhysicalTileDataBTextureRHI() const
-	{
-		return PhysicalTileDataBTextureRHI;
-	}
-	FTextureRHIRef GetPageTableTextureRHI() const
-	{
-		return PageTableTextureRHI;
-	}
-
-	void GetMemorySize(SIZE_T* SizeCPU, SIZE_T* SizeGPU) const;
-
-	virtual void InitRHI() override;
-	virtual void ReleaseRHI() override;
-
-private:
-
-	FSparseVolumeTextureRuntime			SparseVolumeTextureRuntime;
-
-	FTextureRHIRef						PageTableTextureRHI;
-	FTextureRHIRef						PhysicalTileDataATextureRHI;
-	FTextureRHIRef						PhysicalTileDataBTextureRHI;
 };
 
-
 struct ENGINE_API FSparseVolumeTextureFrame
 {
 	// The frame data that can be streamed in when in game.
@@ -186,12 +53,10 @@ struct ENGINE_API FSparseVolumeTextureFrame
 
 	FSparseVolumeTextureFrame();
 	virtual ~FSparseVolumeTextureFrame();
-	bool BuildRuntimeData(FSparseVolumeTextureRuntime* OutRuntimeData);
-	void Serialize(FArchive& Ar, UStreamableSparseVolumeTexture* Owner, int32 FrameIndex);
+	bool BuildDerivedData(FSparseVolumeTextureData* OutMippedTextureData);
+	void Serialize(FArchive& Ar, class UStreamableSparseVolumeTexture* Owner, int32 FrameIndex);
 };
 
-using FSparseVolumeTextureFrameMips = TArray<FSparseVolumeTextureFrame, TInlineAllocator<1u>>;
-
 
 enum ESparseVolumeTextureShaderUniform
 {
@@ -254,6 +119,9 @@ public:
 	UPROPERTY(VisibleAnywhere, Category = "Rendering")
 	FIntVector VolumeResolution;
 
+	UPROPERTY(VisibleAnywhere, Category = "Rendering")
+	int32 NumMipLevels;
+
 	UStreamableSparseVolumeTexture();
 	virtual ~UStreamableSparseVolumeTexture() = default;
 
@@ -270,16 +138,16 @@ public:
 
 	//~ Begin USparseVolumeTexture Interface.
 	int32 GetNumFrames() const override { return Frames.Num(); }
-	int32 GetNumMipLevels() const override { return Frames.IsEmpty() ? 0 : Frames[0].Num(); }
+	int32 GetNumMipLevels() const override { return NumMipLevels; }
 	FIntVector GetVolumeResolution() const override { return VolumeResolution; };
 	const FSparseVolumeTextureSceneProxy* GetSparseVolumeTextureSceneProxy() const override { return GetStreamedFrameProxyOrFallback(0 /*FrameIndex*/, 0 /*MipLevel*/); };
 	//~ End USparseVolumeTexture Interface.
 
 	const FSparseVolumeTextureSceneProxy* GetStreamedFrameProxyOrFallback(int32 FrameIndex, int32 MipLevel) const;
-	TArrayView<const FSparseVolumeTextureFrameMips> GetFrames() const;
+	TArrayView<const FSparseVolumeTextureFrame> GetFrames() const;
 
 protected:
-	TArray<FSparseVolumeTextureFrameMips> Frames;
+	TArray<FSparseVolumeTextureFrame> Frames;
 
 	void GenerateOrLoadDDCRuntimeDataAndCreateSceneProxy();
 	void GenerateOrLoadDDCRuntimeDataForFrame(FSparseVolumeTextureFrame& Frame, UE::DerivedData::FRequestOwner& DDCRequestOwner);
@@ -321,10 +189,6 @@ public:
 	const FSparseVolumeTextureSceneProxy* GetSparseVolumeTextureSceneProxy() const override;
 	//~ End USparseVolumeTexture Interface.
 
-	// Used for debugging a specific frame of an animated sequence.
-	const FSparseVolumeTextureSceneProxy* GetSparseVolumeTextureFrameSceneProxy(int32 FrameIndex, int32 MipLevel) const;
-	const FSparseVolumeAssetHeader* GetSparseVolumeTextureFrameHeader(int32 FrameIndex, int32 MipLevel) const;
-
 private:
 
 #if WITH_EDITOR
diff --git a/Engine/Source/Runtime/Engine/Private/Materials/HLSLMaterialTranslator.cpp b/Engine/Source/Runtime/Engine/Private/Materials/HLSLMaterialTranslator.cpp
--- a/Engine/Source/Runtime/Engine/Private/Materials/HLSLMaterialTranslator.cpp
+++ b/Engine/Source/Runtime/Engine/Private/Materials/HLSLMaterialTranslator.cpp
@@ -12621,9 +12621,9 @@ int32 FHLSLMaterialTranslator::SparseVolumeTextureUniformParameter(FName Paramet
 	return AddUniformExpression(new FMaterialUniformExpressionSparseVolumeTextureUniform(ParameterInfo, TextureIndex, VectorIndex), GetMaterialValueType(Type), TEXT(""));
 }
 
-int32 FHLSLMaterialTranslator::SparseVolumeTextureSamplePageTable(int32 SparseVolumeTextureIndex, int32 UVWIndex)
+int32 FHLSLMaterialTranslator::SparseVolumeTextureSamplePageTable(int32 SparseVolumeTextureIndex, int32 UVWIndex, int32 MipLevelIndex)
 {
-	if (SparseVolumeTextureIndex == INDEX_NONE || UVWIndex == INDEX_NONE)
+	if (SparseVolumeTextureIndex == INDEX_NONE || UVWIndex == INDEX_NONE || MipLevelIndex == INDEX_NONE)
 	{
 		return INDEX_NONE;
 	}
@@ -12634,12 +12634,6 @@ int32 FHLSLMaterialTranslator::SparseVolumeTextureSamplePageTable(int32 SparseVo
 		return Errorf(TEXT("FHLSLMaterialTranslator::SparseVolumeTextureSamplePageTable expects MCT_SparseVolumeTexture but was passed %s ERROR."), DescribeType(TextureType));
 	}
 
-	EMaterialValueType UVWType = GetParameterType(UVWIndex);
-	if ((UVWType & MCT_Float3) == 0)
-	{
-		return Errorf(TEXT("FHLSLMaterialTranslator::SparseVolumeTextureSamplePageTable expects MCT_Float3 as UVW input but was passed %s ERROR."), DescribeType(UVWType));
-	}
-
 	FMaterialUniformExpression* UniformExpression = GetParameterUniformExpression(SparseVolumeTextureIndex);
 	if (UniformExpression == nullptr)
 	{
@@ -12651,6 +12645,9 @@ int32 FHLSLMaterialTranslator::SparseVolumeTextureSamplePageTable(int32 SparseVo
 		return Errorf(TEXT("The provided uniform expression is not a texture"));
 	}
 
+	int32 UVWAsFloat3Index = ForceCast(UVWIndex, MCT_Float3);
+	int32 MipLevelAsFloatIndex = ForceCast(MipLevelIndex, MCT_Float1);
+
 	// Make sure the SVT has been added to UniformTextureExpressions
 	AccessUniformExpression(SparseVolumeTextureIndex);
 
@@ -12658,8 +12655,8 @@ int32 FHLSLMaterialTranslator::SparseVolumeTextureSamplePageTable(int32 SparseVo
 	check(UniformTextureExpressions[(uint32)EMaterialTextureParameterType::SparseVolume].IsValidIndex(SVTReferenceIndex));
 
 	AddEstimatedTextureSample();
-	FString SampleCode = FString::Printf(TEXT("SparseVolumeTextureSamplePageTable(Material.SparseVolumeTexturePageTable_%d, %s, %s)"),
-		SVTReferenceIndex , *GetParameterCode(SparseVolumeTextureIndex), *GetParameterCode(UVWIndex));
+	FString SampleCode = FString::Printf(TEXT("SparseVolumeTextureSamplePageTable(Material.SparseVolumeTexturePageTable_%d, %s, %s, %s)"),
+		SVTReferenceIndex , *GetParameterCode(SparseVolumeTextureIndex), *GetParameterCode(UVWAsFloat3Index), *GetParameterCode(MipLevelAsFloatIndex));
 	return AddCodeChunk(MCT_Float3, *SampleCode);
 }
 
@@ -12676,18 +12673,6 @@ int32 FHLSLMaterialTranslator::SparseVolumeTextureSamplePhysicalTileData(int32 S
 		return Errorf(TEXT("FHLSLMaterialTranslator::SparseVolumeTextureSamplePhysicalTileData expects MCT_SparseVolumeTexture but was passed %s ERROR."), DescribeType(TextureType));
 	}
 
-	EMaterialValueType VoxelCoordType = GetParameterType(VoxelCoordIndex);
-	if ((VoxelCoordType & MCT_Float3) == 0)
-	{
-		return Errorf(TEXT("FHLSLMaterialTranslator::SparseVolumeTextureSamplePhysicalTileData expects MCT_Float3 as VoxelCoord input but was passed %s ERROR."), DescribeType(VoxelCoordType));
-	}
-
-	EMaterialValueType IndexType = GetParameterType(PhysicalTileDataIdxIndex);
-	if ((IndexType & MCT_Float1) == 0 && (IndexType & MCT_UInt1) == 0)
-	{
-		return Errorf(TEXT("FHLSLMaterialTranslator::SparseVolumeTextureSamplePhysicalTileData expects MCT_Float1 or MCT_UInt1 as PhysicalTileDataIndex input but was passed %s ERROR."), DescribeType(IndexType));
-	}
-
 	FMaterialUniformExpression* UniformExpression = GetParameterUniformExpression(SparseVolumeTextureIndex);
 	if (UniformExpression == nullptr)
 	{
@@ -12699,6 +12684,9 @@ int32 FHLSLMaterialTranslator::SparseVolumeTextureSamplePhysicalTileData(int32 S
 		return Errorf(TEXT("The provided uniform expression is not a texture"));
 	}
 
+	int32 VoxelCoordAsFloat3Index = ForceCast(VoxelCoordIndex, MCT_Float3);
+	int32 IndexAsFloatIndex = ForceCast(PhysicalTileDataIdxIndex, MCT_Float1);
+
 	// Make sure the SVT has been added to UniformTextureExpressions
 	AccessUniformExpression(SparseVolumeTextureIndex);
 
@@ -12707,7 +12695,7 @@ int32 FHLSLMaterialTranslator::SparseVolumeTextureSamplePhysicalTileData(int32 S
 
 	AddEstimatedTextureSample();
 	FString SampleCode = FString::Printf(TEXT("SparseVolumeTextureSamplePhysicalTileData(Material.SparseVolumeTexturePhysicalA_%d, Material.SparseVolumeTexturePhysicalB_%d, %s, %s)"),
-		SVTReferenceIndex, SVTReferenceIndex, *GetParameterCode(VoxelCoordIndex), *GetParameterCode(PhysicalTileDataIdxIndex));
+		SVTReferenceIndex, SVTReferenceIndex, *GetParameterCode(VoxelCoordAsFloat3Index), *GetParameterCode(IndexAsFloatIndex));
 	return AddCodeChunk(MCT_Float4, *SampleCode);
 }
 
diff --git a/Engine/Source/Runtime/Engine/Private/Materials/HLSLMaterialTranslator.h b/Engine/Source/Runtime/Engine/Private/Materials/HLSLMaterialTranslator.h
--- a/Engine/Source/Runtime/Engine/Private/Materials/HLSLMaterialTranslator.h
+++ b/Engine/Source/Runtime/Engine/Private/Materials/HLSLMaterialTranslator.h
@@ -928,7 +928,7 @@ protected:
 	virtual int32 SparseVolumeTextureParameter(FName ParameterName, USparseVolumeTexture* InDefaultTexture, int32& TextureReferenceIndex, EMaterialSamplerType SamplerType) override;
 	virtual int32 SparseVolumeTextureUniform(int32 TextureIndex, int32 VectorIndex, UE::Shader::EValueType Type) override;
 	virtual int32 SparseVolumeTextureUniformParameter(FName ParameterName, int32 TextureIndex, int32 VectorIndex, UE::Shader::EValueType Type) override;
-	virtual int32 SparseVolumeTextureSamplePageTable(int32 SparseVolumeTextureIndex, int32 UVWIndex) override;
+	virtual int32 SparseVolumeTextureSamplePageTable(int32 SparseVolumeTextureIndex, int32 UVWIndex, int32 MipLevelIndex) override;
 	virtual int32 SparseVolumeTextureSamplePhysicalTileData(int32 SparseVolumeTextureIndex, int32 VoxelCoordIndex, int32 PhysicalTileDataIdxIndex) override;
 
 	virtual UObject* GetReferencedTexture(int32 Index);
diff --git a/Engine/Source/Runtime/Engine/Private/Materials/MaterialExpressions.cpp b/Engine/Source/Runtime/Engine/Private/Materials/MaterialExpressions.cpp
--- a/Engine/Source/Runtime/Engine/Private/Materials/MaterialExpressions.cpp
+++ b/Engine/Source/Runtime/Engine/Private/Materials/MaterialExpressions.cpp
@@ -26179,6 +26179,9 @@ uint32 UMaterialExpressionSparseVolumeTextureSample::GetInputType(int32 InputInd
 		return MCT_Float3;
 	case 1:
 		return MCT_SparseVolumeTexture;
+	case 2:
+		return MCT_Float1;
+	default:
 		break;
 	}
 
@@ -26218,11 +26221,35 @@ int32 UMaterialExpressionSparseVolumeTextureSample::Compile(class FMaterialCompi
 			return INDEX_NONE;
 		}
 
+		UMaterialExpression* MipLevelExpression = MipLevel.GetTracedInput().Expression;
+
+		// Shared inputs for both potential samples
 		int32 UVWIndex = CompileWithDefaultFloat3(Compiler, Coordinates, 0.0f, 0.0f, 0.0f);
-		int32 VoxelCoordIndex = Compiler->SparseVolumeTextureSamplePageTable(SparseVolumeTextureIndex, UVWIndex);
 		int32 PhysicalTileDataIdxIndex = Compiler->Constant(OutputIndex);
+		int32 MipLevelInputIndex = MipLevelExpression ? MipLevel.Compile(Compiler) : INDEX_NONE;
+		
+		// Sample the first mip
+		int32 MipLevel0Index = MipLevelExpression ? Compiler->Floor(MipLevelInputIndex) : Compiler->Constant(0.0f);
+		int32 VoxelCoordMip0Index = Compiler->SparseVolumeTextureSamplePageTable(SparseVolumeTextureIndex, UVWIndex, MipLevel0Index);
+		int32 Mip0SampleIndex = Compiler->SparseVolumeTextureSamplePhysicalTileData(SparseVolumeTextureIndex, VoxelCoordMip0Index, PhysicalTileDataIdxIndex);
 
-		return Compiler->SparseVolumeTextureSamplePhysicalTileData(SparseVolumeTextureIndex, VoxelCoordIndex, PhysicalTileDataIdxIndex);
+		if (MipLevelExpression)
+		{
+			// Sample the second mip
+			// SVT_TODO: Try to optimize out this second sample if LerpAlpha == 0. Might need to do that in HLSL.
+			int32 MipLevel1Index = Compiler->Ceil(MipLevelInputIndex);
+			int32 VoxelCoordMip1Index = Compiler->SparseVolumeTextureSamplePageTable(SparseVolumeTextureIndex, UVWIndex, MipLevel1Index);
+			int32 Mip1SampleIndex = Compiler->SparseVolumeTextureSamplePhysicalTileData(SparseVolumeTextureIndex, VoxelCoordMip1Index, PhysicalTileDataIdxIndex);
+
+			// Lerp
+			int32 LerpAlphaIndex = Compiler->Frac(MipLevelInputIndex);
+			int32 LerpedResultIndex = Compiler->Lerp(Mip0SampleIndex, Mip1SampleIndex, LerpAlphaIndex);
+			return LerpedResultIndex;
+		}
+		else
+		{
+			return Mip0SampleIndex;
+		}
 	}
 
 	return INDEX_NONE;
diff --git a/Engine/Source/Runtime/Engine/Private/Materials/MaterialUniformExpressions.cpp b/Engine/Source/Runtime/Engine/Private/Materials/MaterialUniformExpressions.cpp
--- a/Engine/Source/Runtime/Engine/Private/Materials/MaterialUniformExpressions.cpp
+++ b/Engine/Source/Runtime/Engine/Private/Materials/MaterialUniformExpressions.cpp
@@ -17,6 +17,7 @@
 #include "GlobalRenderResources.h"
 #include "Shader/PreshaderTypes.h"
 #include "SparseVolumeTexture/SparseVolumeTexture.h"
+#include "SparseVolumeTexture/SparseVolumeTextureSceneProxy.h"
 #include "VT/RuntimeVirtualTexture.h"
 #include "TextureResource.h"
 
diff --git a/Engine/Source/Runtime/Engine/Private/SparseVolumeTexture/SparseVolumeTexture.cpp b/Engine/Source/Runtime/Engine/Private/SparseVolumeTexture/SparseVolumeTexture.cpp
--- a/Engine/Source/Runtime/Engine/Private/SparseVolumeTexture/SparseVolumeTexture.cpp
+++ b/Engine/Source/Runtime/Engine/Private/SparseVolumeTexture/SparseVolumeTexture.cpp
@@ -10,6 +10,9 @@
 #include "EngineUtils.h"
 #include "Shader/ShaderTypes.h"
 #include "RenderingThread.h"
+#include "SparseVolumeTexture/SparseVolumeTextureSceneProxy.h"
+#include "SparseVolumeTexture/SparseVolumeTextureData.h"
+#include "SparseVolumeTexture/SparseVolumeTextureUtility.h"
 
 #if WITH_EDITORONLY_DATA
 #include "DerivedDataCache.h"
@@ -30,579 +33,9 @@
 
 DEFINE_LOG_CATEGORY_STATIC(LogSparseVolumeTexture, Log, All);
 
-// We are using this instead of GMaxVolumeTextureDimensions to be independent of the platform that
-// the asset is imported on. 2048 should be a safe value that should be supported by all our platforms.
-static constexpr int32 SVTMaxVolumeTextureDim = 2048;
-
-namespace UE
-{
-namespace SVT
-{
-namespace Private
-{
-	// SVT_TODO: This really should be a shared function.
-	template <typename Y, typename T>
-	void SerializeEnumAs(FArchive& Ar, T& Target)
-	{
-		Y Buffer = static_cast<Y>(Target);
-		Ar << Buffer;
-		if (Ar.IsLoading())
-		{
-			Target = static_cast<T>(Buffer);
-		}
-	}
-
-	static float F16ToF32(uint16 Packed)
-	{
-		FFloat16 F16{};
-		F16.Encoded = Packed;
-		return F16.GetFloat();
-	};
-
-	static FIntVector3 AdvanceTileCoord(const FIntVector3& TileCoord, const FIntVector3& TileCoordResolution)
-	{
-		FIntVector3 Result = TileCoord;
-		Result.X++;
-		if (Result.X >= TileCoordResolution.X)
-		{
-			Result.X = 0;
-			Result.Y++;
-		}
-		if (Result.Y >= TileCoordResolution.Y)
-		{
-			Result.Y = 0;
-			Result.Z++;
-		}
-		return Result;
-	};
-
-	static FIntVector3 ComputeTileDataVolumeResolution(int32 NumAllocatedPages, bool bAnyEmptyPageExists)
-	{
-		const int32 EffectivelyAllocatedPageEntries = NumAllocatedPages + (bAnyEmptyPageExists ? 1 : 0);
-		int32 TileVolumeResolutionCube = 1;
-		while (TileVolumeResolutionCube * TileVolumeResolutionCube * TileVolumeResolutionCube < EffectivelyAllocatedPageEntries)
-		{
-			TileVolumeResolutionCube++;				// We use a simple loop to compute the minimum resolution of a cube to store all the tile data
-		}
-		FIntVector3 TileDataVolumeResolution = FIntVector3(TileVolumeResolutionCube, TileVolumeResolutionCube, TileVolumeResolutionCube);
-		while (TileDataVolumeResolution.X * TileDataVolumeResolution.Y * (TileDataVolumeResolution.Z - 1) > EffectivelyAllocatedPageEntries)
-		{
-			TileDataVolumeResolution.Z--;	// We then trim an edge to get back space.
-		}
-
-		return TileDataVolumeResolution * SPARSE_VOLUME_TILE_RES;
-	}
-} // Private
-} // SVT
-} // UE
-
-uint32 SparseVolumeTexturePackPageTableEntry(const FIntVector3& Coord)
-{
-	// A page encodes the physical tile coord as unsigned int of 11 11 10 bits
-	// This means a page coord cannot be larger than 2047 for x and y and 1023 for z
-	// which mean we cannot have more than 2048*2048*1024 = 4 Giga tiles of 16^3 tiles.
-	uint32 Result = (Coord.X & 0x7FF) | ((Coord.Y & 0x7FF) << 11) | ((Coord.Z & 0x3FF) << 22);
-	return Result;
-}
-
-FIntVector3 SparseVolumeTextureUnpackPageTableEntry(uint32 Packed)
-{
-	FIntVector3 Result;
-	Result.X = Packed & 0x7FF;
-	Result.Y = (Packed >> 11) & 0x7FF;
-	Result.Z = (Packed >> 22) & 0x3FF;
-	return Result;
-}
-
-////////////////////////////////////////////////////////////////////////////////////////////////
-
-void FSparseVolumeRawSource::Serialize(FArchive& Ar)
-{
-	Header.Serialize(Ar);
-
-	Ar << Version;
-
-	if (Version == 0)
-	{
-		Ar << PageTable;
-		Ar << PhysicalTileDataA;
-		Ar << PhysicalTileDataB;
-	}
-	else
-	{
-		// FSparseVolumeRawSource needs to account for new version
-		check(false);
-	}
-}
-
-bool FSparseVolumeRawSource::Construct(const ISparseVolumeRawSourceConstructionAdapter& Adapter)
-{
-	ISparseVolumeRawSourceConstructionAdapter::FAttributesInfo AttributesInfo[2];
-	Adapter.GetAttributesInfo(AttributesInfo[0], AttributesInfo[1]);
-
-	Header.VirtualVolumeResolution = Adapter.GetResolution();
-	Header.VirtualVolumeAABBMin = Adapter.GetAABBMin();
-	Header.VirtualVolumeAABBMax = Adapter.GetAABBMax();
-	Header.AttributesFormats[0] = AttributesInfo[0].Format;
-	Header.AttributesFormats[1] = AttributesInfo[1].Format;
-
-	Header.PageTableVolumeAABBMin = Header.VirtualVolumeAABBMin / SPARSE_VOLUME_TILE_RES;
-	Header.PageTableVolumeAABBMax = (Header.VirtualVolumeAABBMax + FIntVector3(SPARSE_VOLUME_TILE_RES - 1)) / SPARSE_VOLUME_TILE_RES;
-	Header.bHasNullTile = false;
-
-	const FIntVector3 PageTableVolumeResolution = Header.PageTableVolumeAABBMax - Header.PageTableVolumeAABBMin;
-
-	if (PageTableVolumeResolution.X > SVTMaxVolumeTextureDim
-		|| PageTableVolumeResolution.Y > SVTMaxVolumeTextureDim
-		|| PageTableVolumeResolution.Z > SVTMaxVolumeTextureDim)
-	{
-		UE_LOG(LogSparseVolumeTexture, Warning, TEXT("SparseVolumeTexture page table texture dimensions exceed limit (%ix%ix%i): %ix%ix%i"), 
-			SVTMaxVolumeTextureDim, SVTMaxVolumeTextureDim, SVTMaxVolumeTextureDim, 
-			PageTableVolumeResolution.X, PageTableVolumeResolution.Y, PageTableVolumeResolution.Z);
-		return false;
-	}
-
-	Header.PageTableVolumeResolution = PageTableVolumeResolution;
-	Header.TileDataVolumeResolution = FIntVector::ZeroValue;	// unknown for now
-
-	// Tag all pages with valid data
-	TBitArray PagesWithData(false, Header.PageTableVolumeResolution.Z * Header.PageTableVolumeResolution.Y * Header.PageTableVolumeResolution.X);
-	Adapter.IteratePhysicalSource([&](const FIntVector3& Coord, int32 AttributesIdx, int32 ComponentIdx, float VoxelValue)
-		{
-			const bool bIsFallbackValue = (VoxelValue == AttributesInfo[AttributesIdx].FallbackValue[ComponentIdx]);
-			if (!bIsFallbackValue)
-			{
-				const FIntVector3 GridCoord = Coord;
-				check(GridCoord.X >= 0 && GridCoord.Y >= 0 && GridCoord.Z >= 0);
-				check(GridCoord.X < Header.VirtualVolumeAABBMax.X && GridCoord.Y < Header.VirtualVolumeAABBMax.Y && GridCoord.Z < Header.VirtualVolumeAABBMax.Z);
-				const FIntVector3 PageCoord = (GridCoord / SPARSE_VOLUME_TILE_RES) - Header.PageTableVolumeAABBMin;
-				check(PageCoord.X >= 0 && PageCoord.Y >= 0 && PageCoord.Z >= 0);
-				check(PageCoord.X < Header.PageTableVolumeResolution.X && PageCoord.Y < Header.PageTableVolumeResolution.Y && PageCoord.Z < Header.PageTableVolumeResolution.Z);
-
-				const int32 PageIndex = PageCoord.Z * (Header.PageTableVolumeResolution.Y * Header.PageTableVolumeResolution.X) + PageCoord.Y * Header.PageTableVolumeResolution.X + PageCoord.X;
-				PagesWithData[PageIndex] = true;
-			}
-		});
-
-	// Allocate some memory for temp data (worst case)
-	TArray<FIntVector3> LinearAllocatedPages;
-	LinearAllocatedPages.SetNum(Header.PageTableVolumeResolution.X * Header.PageTableVolumeResolution.Y * Header.PageTableVolumeResolution.Z);
-
-	// Go over each potential page from the source data and push allocate it if it has any data.
-	// Otherwise point to the default empty page.
-	int32 NumAllocatedPages = 0;
-	for (int32_t PageZ = 0; PageZ < Header.PageTableVolumeResolution.Z; ++PageZ)
-	{
-		for (int32_t PageY = 0; PageY < Header.PageTableVolumeResolution.Y; ++PageY)
-		{
-			for (int32_t PageX = 0; PageX < Header.PageTableVolumeResolution.X; ++PageX)
-			{
-				const int32 PageIndex = PageZ * (Header.PageTableVolumeResolution.Y * Header.PageTableVolumeResolution.X) + PageY * Header.PageTableVolumeResolution.X + PageX;
-				const bool bHasAnyData = PagesWithData[PageIndex];
-				if (bHasAnyData)
-				{
-					LinearAllocatedPages[NumAllocatedPages] = FIntVector3(PageX, PageY, PageZ);
-					NumAllocatedPages++;
-				}
-				Header.bHasNullTile |= !bHasAnyData;
-			}
-		}
-	}
-
-	// Compute Page and Tile VolumeResolution from allocated pages
-	Header.TileDataVolumeResolution = UE::SVT::Private::ComputeTileDataVolumeResolution(NumAllocatedPages, Header.bHasNullTile);
-	const FIntVector3 TileCoordResolution = Header.TileDataVolumeResolution / SPARSE_VOLUME_TILE_RES;
-
-	if (Header.TileDataVolumeResolution.X > SVTMaxVolumeTextureDim
-		|| Header.TileDataVolumeResolution.Y > SVTMaxVolumeTextureDim
-		|| Header.TileDataVolumeResolution.Z > SVTMaxVolumeTextureDim)
-	{
-		UE_LOG(LogSparseVolumeTexture, Warning, TEXT("SparseVolumeTexture tile data texture dimensions exceed limit (%ix%ix%i): %ix%ix%i"), SVTMaxVolumeTextureDim, SVTMaxVolumeTextureDim, SVTMaxVolumeTextureDim, Header.TileDataVolumeResolution.X, Header.TileDataVolumeResolution.Y, Header.TileDataVolumeResolution.Z);
-		return false;
-	}
-
-	// Initialize the SparseVolumeTexture page and tile.
-	const int32 FormatSize[] = { GPixelFormats[(SIZE_T)Header.AttributesFormats[0]].BlockBytes, GPixelFormats[(SIZE_T)Header.AttributesFormats[1]].BlockBytes };
-	PageTable.SetNumZeroed(Header.PageTableVolumeResolution.X * Header.PageTableVolumeResolution.Y * Header.PageTableVolumeResolution.Z);
-	PhysicalTileDataA.SetNumZeroed(Header.TileDataVolumeResolution.X * Header.TileDataVolumeResolution.Y * Header.TileDataVolumeResolution.Z * FormatSize[0]);
-	PhysicalTileDataB.SetNumZeroed(Header.TileDataVolumeResolution.X * Header.TileDataVolumeResolution.Y * Header.TileDataVolumeResolution.Z * FormatSize[1]);
-
-	// Generate page table and tile volume data by splatting the data
-	{
-		FIntVector3 DstTileCoord = FIntVector3::ZeroValue;
-
-		// Add an empty tile is needed, reserve slot at coord 0
-		if (Header.bHasNullTile)
-		{
-			// Fill null tile with fallback/background value
-			FVector4f FallbackValues[] = { AttributesInfo[0].FallbackValue, AttributesInfo[1].FallbackValue };
-			for (int32 AttributesIdx = 0; AttributesIdx < 2; ++AttributesIdx)
-			{
-				if (AttributesInfo[AttributesIdx].bNormalized)
-				{
-					FallbackValues[AttributesIdx] = FallbackValues[AttributesIdx] * AttributesInfo[AttributesIdx].NormalizeScale + AttributesInfo[AttributesIdx].NormalizeBias;
-				}
-			}
-			FillNullTile(FallbackValues[0], FallbackValues[1]);
-
-			// PageTable is all cleared to zero, simply skip a tile
-			DstTileCoord = UE::SVT::Private::AdvanceTileCoord(DstTileCoord, TileCoordResolution);
-		}
-
-		for (int32 i = 0; i < NumAllocatedPages; ++i)
-		{
-			const FIntVector3 PageCoordToSplat = LinearAllocatedPages[i];
-			const uint32 DestinationTileCoord32bit = SparseVolumeTexturePackPageTableEntry(DstTileCoord);
-
-			// Setup the page table entry
-			PageTable
-				[
-					PageCoordToSplat.Z * Header.PageTableVolumeResolution.X * Header.PageTableVolumeResolution.Y +
-					PageCoordToSplat.Y * Header.PageTableVolumeResolution.X +
-					PageCoordToSplat.X
-				] = DestinationTileCoord32bit;
-
-			// Set the next tile to be written to
-			DstTileCoord = UE::SVT::Private::AdvanceTileCoord(DstTileCoord, TileCoordResolution);
-		}
-	}
-
-	// Write physical tile data
-	Adapter.IteratePhysicalSource([&](const FIntVector3& Coord, int32 AttributesIdx, int32 ComponentIdx, float VoxelValue)
-		{
-			const FIntVector3 GridCoord = Coord;
-			check(GridCoord.X >= 0 && GridCoord.Y >= 0 && GridCoord.Z >= 0);
-			check(GridCoord.X < Header.VirtualVolumeAABBMax.X && GridCoord.Y < Header.VirtualVolumeAABBMax.Y && GridCoord.Z < Header.VirtualVolumeAABBMax.Z);
-			const FIntVector3 PageCoord = (GridCoord / SPARSE_VOLUME_TILE_RES) - Header.PageTableVolumeAABBMin;
-			check(PageCoord.X >= 0 && PageCoord.Y >= 0 && PageCoord.Z >= 0);
-			check(PageCoord.X < Header.PageTableVolumeResolution.X && PageCoord.Y < Header.PageTableVolumeResolution.Y && PageCoord.Z < Header.PageTableVolumeResolution.Z);
-
-			const int32 PageIndex = PageCoord.Z * (Header.PageTableVolumeResolution.Y * Header.PageTableVolumeResolution.X) + PageCoord.Y * Header.PageTableVolumeResolution.X + PageCoord.X;
-
-			if (!PagesWithData[PageIndex])
-			{
-				return;
-			}
-
-			float WriteValue = VoxelValue;
-			if (AttributesInfo[AttributesIdx].bNormalized)
-			{
-				WriteValue = WriteValue * AttributesInfo[AttributesIdx].NormalizeScale[ComponentIdx] + AttributesInfo[AttributesIdx].NormalizeBias[ComponentIdx];
-			}
-
-			const FIntVector3 WriteTileCoord = SparseVolumeTextureUnpackPageTableEntry(PageTable[PageIndex]);
-			const FIntVector3 TileLocalCoord = GridCoord % SPARSE_VOLUME_TILE_RES;
-			const FIntVector3 WriteCoord = WriteTileCoord * SPARSE_VOLUME_TILE_RES + TileLocalCoord;
-			WriteTileDataVoxel(WriteCoord, AttributesIdx, FVector4f(WriteValue, WriteValue, WriteValue, WriteValue), ComponentIdx);
-		});
-
-	return true;
-}
-
-uint32 FSparseVolumeRawSource::ReadPageTablePacked(const FIntVector3& PageTableCoord) const
-{
-	if (PageTableCoord.X < 0 || PageTableCoord.Y < 0 || PageTableCoord.Z < 0
-		|| PageTableCoord.X >= Header.PageTableVolumeResolution.X 
-		|| PageTableCoord.Y >= Header.PageTableVolumeResolution.Y 
-		|| PageTableCoord.Z >= Header.PageTableVolumeResolution.Z)
-	{
-		return 0;
-	}
-
-	const int32 PageIndex = PageTableCoord.Z * (Header.PageTableVolumeResolution.Y * Header.PageTableVolumeResolution.X) + PageTableCoord.Y * Header.PageTableVolumeResolution.X + PageTableCoord.X;
-	if (PageTable.IsValidIndex(PageIndex))
-	{
-		return PageTable[PageIndex];
-	}
-	return 0;
-}
-
-FIntVector3 FSparseVolumeRawSource::ReadPageTable(const FIntVector3& PageTableCoord) const
-{
-	const uint32 Packed = ReadPageTablePacked(PageTableCoord);
-	const FIntVector3 Unpacked = SparseVolumeTextureUnpackPageTableEntry(Packed);
-	return Unpacked;
-}
-
-FVector4f FSparseVolumeRawSource::ReadTileDataVoxel(const FIntVector3& TileDataCoord, int32 AttributesIdx) const
-{
-	using namespace UE::SVT::Private;
-	check(AttributesIdx >= 0 && AttributesIdx <= 1);
-	if ((AttributesIdx == 0 && PhysicalTileDataA.IsEmpty()) || (AttributesIdx == 1 && PhysicalTileDataB.IsEmpty()))
-	{
-		return FVector4f();
-	}
-
-	if (TileDataCoord.X < 0 || TileDataCoord.Y < 0 || TileDataCoord.Z < 0
-		|| TileDataCoord.X >= Header.TileDataVolumeResolution.X 
-		|| TileDataCoord.Y >= Header.TileDataVolumeResolution.Y 
-		|| TileDataCoord.Z >= Header.TileDataVolumeResolution.Z)
-	{
-		return FVector4f();
-	}
-
-	const int32 VoxelIndex = TileDataCoord.Z * (Header.TileDataVolumeResolution.Y * Header.TileDataVolumeResolution.X) + TileDataCoord.Y * Header.TileDataVolumeResolution.X + TileDataCoord.X;
-	const uint8* TileData = AttributesIdx == 0 ? PhysicalTileDataA.GetData() : PhysicalTileDataB.GetData();
-	const EPixelFormat Format = Header.AttributesFormats[AttributesIdx];
-
-	switch (Format)
-	{
-	case PF_R8:
-		return FVector4f(TileData[VoxelIndex] / 255.0f, 0.0f, 0.0f, 0.0f);
-	case PF_R8G8:
-		return FVector4f(TileData[VoxelIndex * 2 + 0] / 255.0f, TileData[VoxelIndex * 2 + 1] / 255.0f, 0.0f, 0.0f);
-	case PF_R8G8B8A8:
-		return FVector4f(TileData[VoxelIndex * 4 + 0] / 255.0f, TileData[VoxelIndex * 4 + 1] / 255.0f, TileData[VoxelIndex * 4 + 2] / 255.0f, TileData[VoxelIndex * 4 + 3] / 255.0f);
-	case PF_R16F:
-		return FVector4f(F16ToF32(((const uint16*)TileData)[VoxelIndex]), 0.0f, 0.0f, 0.0f);
-	case PF_G16R16F:
-		return FVector4f(F16ToF32(((const uint16*)TileData)[VoxelIndex * 2 + 0]), F16ToF32(((const uint16*)TileData)[VoxelIndex * 2 + 1]), 0.0f, 0.0f);
-	case PF_FloatRGBA:
-		return FVector4f(F16ToF32(((const uint16*)TileData)[VoxelIndex * 4 + 0]), F16ToF32(((const uint16*)TileData)[VoxelIndex * 4 + 1]), F16ToF32(((const uint16*)TileData)[VoxelIndex * 4 + 2]), F16ToF32(((const uint16*)TileData)[VoxelIndex * 4 + 3]));
-	case PF_R32_FLOAT:
-		return FVector4f(((const float*)TileData)[VoxelIndex], 0.0f, 0.0f, 0.0f);
-	case PF_G32R32F:
-		return FVector4f(((const float*)TileData)[VoxelIndex * 2 + 0], ((const float*)TileData)[VoxelIndex * 2 + 1], 0.0f, 0.0f);
-	case PF_A32B32G32R32F:
-		return FVector4f(((const float*)TileData)[VoxelIndex * 4 + 0], ((const float*)TileData)[VoxelIndex * 4 + 1], ((const float*)TileData)[VoxelIndex * 4 + 2], ((const float*)TileData)[VoxelIndex * 4 + 3]);
-	default:
-		checkNoEntry();
-		return FVector4f();
-	}
-}
-
-FVector4f FSparseVolumeRawSource::Sample(const FIntVector3& VolumeCoord, int32 AttributesIdx) const
-{
-	const FIntVector3 PageTableCoord = (VolumeCoord / SPARSE_VOLUME_TILE_RES) - Header.PageTableVolumeAABBMin;
-	const FIntVector3 TileCoord = ReadPageTable(PageTableCoord);
-	const FIntVector3 VoxelCoord = (TileCoord * SPARSE_VOLUME_TILE_RES) + (VolumeCoord % SPARSE_VOLUME_TILE_RES);
-	const FVector4f Sample = ReadTileDataVoxel(VoxelCoord, AttributesIdx);
-	return Sample;
-}
-
-void FSparseVolumeRawSource::Sample(const FIntVector3& VolumeCoord, FVector4f& OutAttributesA, FVector4f& OutAttributesB) const
-{
-	const FIntVector3 PageTableCoord = (VolumeCoord / SPARSE_VOLUME_TILE_RES) - Header.PageTableVolumeAABBMin;
-	const FIntVector3 TileCoord = ReadPageTable(PageTableCoord);
-	const FIntVector3 VoxelCoord = (TileCoord * SPARSE_VOLUME_TILE_RES) + (VolumeCoord % SPARSE_VOLUME_TILE_RES);
-	OutAttributesA = ReadTileDataVoxel(VoxelCoord, 0);
-	OutAttributesB = ReadTileDataVoxel(VoxelCoord, 1);
-}
-
-void FSparseVolumeRawSource::WriteTileDataVoxel(const FIntVector3& TileDataCoord, int32 AttributesIdx, const FVector4f& Value, int32 DstComponent)
-{
-	check(AttributesIdx >= 0 && AttributesIdx <= 1);
-	if ((AttributesIdx == 0 && PhysicalTileDataA.IsEmpty()) || (AttributesIdx == 1 && PhysicalTileDataB.IsEmpty()))
-	{
-		return;
-	}
-
-	const int32 VoxelIndex = TileDataCoord.Z * (Header.TileDataVolumeResolution.Y * Header.TileDataVolumeResolution.X) + TileDataCoord.Y * Header.TileDataVolumeResolution.X + TileDataCoord.X;
-	uint8* TileData = AttributesIdx == 0 ? PhysicalTileDataA.GetData() : PhysicalTileDataB.GetData();
-	const EPixelFormat Format = Header.AttributesFormats[AttributesIdx];
-
-	switch (Format)
-	{
-	case PF_R8:
-		if (DstComponent == -1 || DstComponent == 0) TileData[VoxelIndex] = uint8(FMath::Clamp(Value.X, 0.0f, 1.0f) * 255.0f);
-		break;
-	case PF_R8G8:
-		if (DstComponent == -1 || DstComponent == 0) TileData[VoxelIndex * 2 + 0] = uint8(FMath::Clamp(Value.X, 0.0f, 1.0f) * 255.0f);
-		if (DstComponent == -1 || DstComponent == 1) TileData[VoxelIndex * 2 + 1] = uint8(FMath::Clamp(Value.Y, 0.0f, 1.0f) * 255.0f);
-		break;
-	case PF_R8G8B8A8:
-		if (DstComponent == -1 || DstComponent == 0) TileData[VoxelIndex * 4 + 0] = uint8(FMath::Clamp(Value.X, 0.0f, 1.0f) * 255.0f);
-		if (DstComponent == -1 || DstComponent == 1) TileData[VoxelIndex * 4 + 1] = uint8(FMath::Clamp(Value.Y, 0.0f, 1.0f) * 255.0f);
-		if (DstComponent == -1 || DstComponent == 2) TileData[VoxelIndex * 4 + 2] = uint8(FMath::Clamp(Value.Z, 0.0f, 1.0f) * 255.0f);
-		if (DstComponent == -1 || DstComponent == 3) TileData[VoxelIndex * 4 + 3] = uint8(FMath::Clamp(Value.W, 0.0f, 1.0f) * 255.0f);
-		break;
-	case PF_R16F:
-		if (DstComponent == -1 || DstComponent == 0) ((uint16*)TileData)[VoxelIndex] = FFloat16(Value.X).Encoded;
-		break;
-	case PF_G16R16F:
-		if (DstComponent == -1 || DstComponent == 0) ((uint16*)TileData)[VoxelIndex * 2 + 0] = FFloat16(Value.X).Encoded;
-		if (DstComponent == -1 || DstComponent == 1) ((uint16*)TileData)[VoxelIndex * 2 + 1] = FFloat16(Value.Y).Encoded;
-		break;
-	case PF_FloatRGBA:
-		if (DstComponent == -1 || DstComponent == 0) ((uint16*)TileData)[VoxelIndex * 4 + 0] = FFloat16(Value.X).Encoded;
-		if (DstComponent == -1 || DstComponent == 1) ((uint16*)TileData)[VoxelIndex * 4 + 1] = FFloat16(Value.Y).Encoded;
-		if (DstComponent == -1 || DstComponent == 2) ((uint16*)TileData)[VoxelIndex * 4 + 2] = FFloat16(Value.Z).Encoded;
-		if (DstComponent == -1 || DstComponent == 3) ((uint16*)TileData)[VoxelIndex * 4 + 3] = FFloat16(Value.W).Encoded;
-		break;
-	case PF_R32_FLOAT:
-		if (DstComponent == -1 || DstComponent == 0) ((float*)TileData)[VoxelIndex] = Value.X;
-		break;
-	case PF_G32R32F:
-		if (DstComponent == -1 || DstComponent == 0) ((float*)TileData)[VoxelIndex * 2 + 0] = Value.X;
-		if (DstComponent == -1 || DstComponent == 1) ((float*)TileData)[VoxelIndex * 2 + 1] = Value.Y;
-		break;
-	case PF_A32B32G32R32F:
-		if (DstComponent == -1 || DstComponent == 0) ((float*)TileData)[VoxelIndex * 4 + 0] = Value.X;
-		if (DstComponent == -1 || DstComponent == 1) ((float*)TileData)[VoxelIndex * 4 + 1] = Value.Y;
-		if (DstComponent == -1 || DstComponent == 2) ((float*)TileData)[VoxelIndex * 4 + 2] = Value.Z;
-		if (DstComponent == -1 || DstComponent == 3) ((float*)TileData)[VoxelIndex * 4 + 3] = Value.W;
-		break;
-	default:
-		checkNoEntry();
-	}
-}
-
-void FSparseVolumeRawSource::FillNullTile(const FVector4f& FallbackValueA, const FVector4f& FallbackValueB)
-{
-	for (int32 AttributesIdx = 0; AttributesIdx < 2; ++AttributesIdx)
-	{
-		if ((AttributesIdx == 0 && PhysicalTileDataA.IsEmpty()) || (AttributesIdx == 1 && PhysicalTileDataB.IsEmpty()))
-		{
-			continue;
-		}
-
-		const FVector4f& FallbackValue = AttributesIdx == 0 ? FallbackValueA : FallbackValueB;
-
-		for (int32 Z = 0; Z < SPARSE_VOLUME_TILE_RES; ++Z)
-		{
-			for (int32 Y = 0; Y < SPARSE_VOLUME_TILE_RES; ++Y)
-			{
-				for (int32 X = 0; X < SPARSE_VOLUME_TILE_RES; ++X)
-				{
-					const FIntVector3 WriteCoord = FIntVector3(X, Y, Z);
-					WriteTileDataVoxel(WriteCoord, AttributesIdx, FallbackValue);
-				}
-			}
-		}
-	}
-}
-
-FSparseVolumeRawSource FSparseVolumeRawSource::GenerateMipMap() const
-{
-	FSparseVolumeRawSource Result{};
-	Result.Header = Header; // Overwrite values later
-	Result.Header.VirtualVolumeAABBMin = Header.VirtualVolumeAABBMin / 2;
-	Result.Header.VirtualVolumeAABBMax = FIntVector3(FMath::Max(1, Header.VirtualVolumeAABBMax.X / 2), FMath::Max(1, Header.VirtualVolumeAABBMax.Y / 2), FMath::Max(1, Header.VirtualVolumeAABBMax.Z / 2));
-	Result.Header.VirtualVolumeResolution = Result.Header.VirtualVolumeAABBMax - Result.Header.VirtualVolumeAABBMin;
-	Result.Header.PageTableVolumeAABBMin = Result.Header.VirtualVolumeAABBMin / SPARSE_VOLUME_TILE_RES;
-	Result.Header.PageTableVolumeAABBMax = (Result.Header.VirtualVolumeAABBMax + FIntVector3(SPARSE_VOLUME_TILE_RES - 1)) / SPARSE_VOLUME_TILE_RES;
-	Result.Header.PageTableVolumeResolution = Result.Header.PageTableVolumeAABBMax - Result.Header.PageTableVolumeAABBMin;
-	Result.Header.MipLevel = Header.MipLevel + 1;
-	Result.Header.bHasNullTile = false;
-
-	// Allocate some memory for temp data (worst case)
-	TArray<FIntVector3> LinearAllocatedPages;
-	LinearAllocatedPages.SetNum(Result.Header.PageTableVolumeResolution.X * Result.Header.PageTableVolumeResolution.Y * Result.Header.PageTableVolumeResolution.Z);
-
-	// Go over each potential page from the source data and push allocate it if it has any data.
-	// Otherwise point to the default empty page.
-	int32 NumAllocatedPages = 0;
-	for (int32_t PageZ = 0; PageZ < Result.Header.PageTableVolumeResolution.Z; ++PageZ)
-	{
-		for (int32_t PageY = 0; PageY < Result.Header.PageTableVolumeResolution.Y; ++PageY)
-		{
-			for (int32_t PageX = 0; PageX < Result.Header.PageTableVolumeResolution.X; ++PageX)
-			{
-				const FIntVector3 PageCoord = FIntVector3(PageX, PageY, PageZ);
-				bool bHasAnyData = false;
-				for (int32_t OffsetIdx = 0; OffsetIdx < 8; ++OffsetIdx)
-				{
-					const FIntVector3 Offset = FIntVector3(OffsetIdx, OffsetIdx >> 1, OffsetIdx >> 2) & 1;
-					const FIntVector3 ParentPageTableCoord = (Result.Header.PageTableVolumeAABBMin + PageCoord) * 2 + Offset - Header.PageTableVolumeAABBMin;
-					const uint32 PageSample = ReadPageTablePacked(ParentPageTableCoord);
-					
-					if (PageSample != 0 || !Header.bHasNullTile)
-					{
-						bHasAnyData = true;
-					}
-				}
-				if (bHasAnyData)
-				{
-					LinearAllocatedPages[NumAllocatedPages] = PageCoord;
-					NumAllocatedPages++;
-				}
-				Result.Header.bHasNullTile |= !bHasAnyData;
-			}
-		}
-	}
-
-	// Compute Page and Tile VolumeResolution from allocated pages
-	Result.Header.TileDataVolumeResolution = UE::SVT::Private::ComputeTileDataVolumeResolution(NumAllocatedPages, Result.Header.bHasNullTile);
-	const FIntVector3 TileCoordResolution = Result.Header.TileDataVolumeResolution / SPARSE_VOLUME_TILE_RES;
-
-	// Compute byte sizes of formats as NumTotalBytes/NumVoxels. Note that we do this on the data from this object, not the new downsampled one.
-	const int32 FormatSizeA = PhysicalTileDataA.Num() / (Header.TileDataVolumeResolution.X * Header.TileDataVolumeResolution.Y * Header.TileDataVolumeResolution.Z);
-	const int32 FormatSizeB = PhysicalTileDataB.Num() / (Header.TileDataVolumeResolution.X * Header.TileDataVolumeResolution.Y * Header.TileDataVolumeResolution.Z);
-
-	// Initialize the SparseVolumeTexture page and tile.
-	Result.PageTable.SetNumZeroed(Result.Header.PageTableVolumeResolution.X * Result.Header.PageTableVolumeResolution.Y * Result.Header.PageTableVolumeResolution.Z);
-	Result.PhysicalTileDataA.SetNumZeroed(Result.Header.TileDataVolumeResolution.X * Result.Header.TileDataVolumeResolution.Y * Result.Header.TileDataVolumeResolution.Z * FormatSizeA);
-	Result.PhysicalTileDataB.SetNumZeroed(Result.Header.TileDataVolumeResolution.X * Result.Header.TileDataVolumeResolution.Y * Result.Header.TileDataVolumeResolution.Z * FormatSizeB);
-
-	// Generate page table and tile volume data by splatting the data
-	{
-		FIntVector3 DstTileCoord = FIntVector3::ZeroValue;
-
-		// Add an empty tile is needed, reserve slot at coord 0
-		if (Result.Header.bHasNullTile)
-		{
-			Result.FillNullTile(ReadTileDataVoxel(FIntVector3(), 0), ReadTileDataVoxel(FIntVector3(), 1));
-			
-			// PageTable is all cleared to zero, simply skip a tile
-			DstTileCoord = UE::SVT::Private::AdvanceTileCoord(DstTileCoord, TileCoordResolution);
-		}
-
-		for (int32 i = 0; i < NumAllocatedPages; ++i)
-		{
-			const FIntVector3 PageCoordToSplat = LinearAllocatedPages[i];
-			const uint32 DestinationTileCoord32bit = SparseVolumeTexturePackPageTableEntry(DstTileCoord);
-
-			// Setup the page table entry
-			Result.PageTable
-				[
-					PageCoordToSplat.Z * Result.Header.PageTableVolumeResolution.X * Result.Header.PageTableVolumeResolution.Y +
-					PageCoordToSplat.Y * Result.Header.PageTableVolumeResolution.X +
-					PageCoordToSplat.X
-				] = DestinationTileCoord32bit;
-
-			// Write tile data
-			const FIntVector3 ParentVolumeCoordBase = (Result.Header.PageTableVolumeAABBMin + PageCoordToSplat) * SPARSE_VOLUME_TILE_RES * 2;
-			for (int32 AttributesIdx = 0; AttributesIdx < 2; ++AttributesIdx)
-			{
-				if ((AttributesIdx == 0 && PhysicalTileDataA.IsEmpty()) || (AttributesIdx == 1 && PhysicalTileDataB.IsEmpty()))
-				{
-					continue;
-				}
-
-				for (int32 Z = 0; Z < SPARSE_VOLUME_TILE_RES; ++Z)
-				{
-					for (int32 Y = 0; Y < SPARSE_VOLUME_TILE_RES; ++Y)
-					{
-						for (int32 X = 0; X < SPARSE_VOLUME_TILE_RES; ++X)
-						{
-							FVector4f DownsampledValue = FVector4f(0.0f, 0.0f, 0.0f, 0.0f);
-							for (int32_t OffsetIdx = 0; OffsetIdx < 8; ++OffsetIdx)
-							{
-								const FIntVector3 Offset = FIntVector3(OffsetIdx, OffsetIdx >> 1, OffsetIdx >> 2) & 1;
-								const FIntVector3 SourceCoord = ParentVolumeCoordBase + FIntVector3(X, Y, Z) * 2 + Offset;
-								const FVector4f SampleValue = Sample(SourceCoord, AttributesIdx);
-								DownsampledValue += SampleValue;
-							}
-							DownsampledValue /= 8.0f;
-
-							const FIntVector3 WriteCoord = DstTileCoord * SPARSE_VOLUME_TILE_RES + FIntVector3(X, Y, Z);
-							Result.WriteTileDataVoxel(WriteCoord, AttributesIdx, DownsampledValue);
-						}
-					}
-				}
-			}
-
-			// Set the next tile to be written to
-			DstTileCoord = UE::SVT::Private::AdvanceTileCoord(DstTileCoord, TileCoordResolution);
-		}
-	}
-
-	return Result;
-}
-
 ////////////////////////////////////////////////////////////////////////////////////////////////
 
-void FSparseVolumeAssetHeader::Serialize(FArchive& Ar)
+void FSparseVolumeTextureHeader::Serialize(FArchive& Ar)
 {
 	Ar << Version;
 
@@ -614,148 +47,20 @@ void FSparseVolumeAssetHeader::Serialize(FArchive& Ar)
 		Ar << PageTableVolumeResolution;
 		Ar << PageTableVolumeAABBMin;
 		Ar << PageTableVolumeAABBMax;
-		Ar << TileDataVolumeResolution;
 		UE::SVT::Private::SerializeEnumAs<uint8>(Ar, AttributesFormats[0]);
 		UE::SVT::Private::SerializeEnumAs<uint8>(Ar, AttributesFormats[1]);
-		Ar << MipLevel;
-		Ar << bHasNullTile;
+		Ar << NullTileValues[0];
+		Ar << NullTileValues[1];
 	}
 	else
 	{
-		// FSparseVolumeAssetHeader needs to account for new version
+		// FSparseVolumeTextureHeader needs to account for new version
 		check(false);
 	}
 }
 
 ////////////////////////////////////////////////////////////////////////////////////////////////
 
-void FSparseVolumeTextureRuntime::Serialize(FArchive& Ar)
-{
-	Header.Serialize(Ar);
-
-	Ar << Version;
-
-	if (Version == 0)
-	{
-		Ar << PageTable;
-		Ar << PhysicalTileDataA;
-		Ar << PhysicalTileDataB;
-	}
-	else
-	{
-		// FSparseVolumeRawSource needs to account for new version
-		check(false);
-	}
-}
-
-void FSparseVolumeTextureRuntime::SetAsDefaultTexture()
-{
-	const uint32 VolumeSize = 1;
-	PageTable.SetNumZeroed(VolumeSize * VolumeSize * VolumeSize);
-	PhysicalTileDataA.SetNumZeroed(VolumeSize * VolumeSize * VolumeSize);
-	PhysicalTileDataB.SetNumZeroed(VolumeSize * VolumeSize * VolumeSize);
-}
-
-////////////////////////////////////////////////////////////////////////////////////////////////
-
-FSparseVolumeTextureSceneProxy::FSparseVolumeTextureSceneProxy()
-	: FRenderResource()
-	, SparseVolumeTextureRuntime()
-	, PageTableTextureRHI(nullptr)
-	, PhysicalTileDataATextureRHI(nullptr)
-	, PhysicalTileDataBTextureRHI(nullptr)
-{
-}
-
-FSparseVolumeTextureSceneProxy::~FSparseVolumeTextureSceneProxy()
-{
-}
-
-void FSparseVolumeTextureSceneProxy::GetMemorySize(SIZE_T* SizeCPU, SIZE_T* SizeGPU) const
-{
-	*SizeCPU += sizeof(FSparseVolumeTextureSceneProxy);
-	*SizeCPU += SparseVolumeTextureRuntime.PageTable.GetAllocatedSize();
-	*SizeCPU += SparseVolumeTextureRuntime.PhysicalTileDataA.GetAllocatedSize();
-	*SizeCPU += SparseVolumeTextureRuntime.PhysicalTileDataB.GetAllocatedSize();
-
-#if RHI_ENABLE_RESOURCE_INFO
-	FRHIResourceInfo ResourceInfo;
-	if (PageTableTextureRHI && PageTableTextureRHI->GetResourceInfo(ResourceInfo))
-	{
-		*SizeGPU += ResourceInfo.VRamAllocation.AllocationSize;
-	}
-	if (PhysicalTileDataATextureRHI && PhysicalTileDataATextureRHI->GetResourceInfo(ResourceInfo))
-	{
-		*SizeGPU += ResourceInfo.VRamAllocation.AllocationSize;
-	}
-	if (PhysicalTileDataBTextureRHI && PhysicalTileDataBTextureRHI->GetResourceInfo(ResourceInfo))
-	{
-		*SizeGPU += ResourceInfo.VRamAllocation.AllocationSize;
-	}
-#endif
-}
-
-void FSparseVolumeTextureSceneProxy::InitRHI()
-{
-	// Page table
-	{
-		const EPixelFormat PageEntryFormat = PF_R32_UINT;
-		const FIntVector3 PageTableVolumeResolution = SparseVolumeTextureRuntime.Header.PageTableVolumeResolution;
-		const FRHITextureCreateDesc Desc =
-			FRHITextureCreateDesc::Create3D(TEXT("SparseVolumeTexture.PageTable.RHITexture"),
-				PageTableVolumeResolution.X, PageTableVolumeResolution.Y, PageTableVolumeResolution.Z, PageEntryFormat)
-			.SetFlags(ETextureCreateFlags::ShaderResource);
-
-		PageTableTextureRHI = RHICreateTexture(Desc);
-
-		const int32 FormatSize = GPixelFormats[PageEntryFormat].BlockBytes;
-		const FUpdateTextureRegion3D UpdateRegion(0, 0, 0, 0, 0, 0, PageTableVolumeResolution.X, PageTableVolumeResolution.Y, PageTableVolumeResolution.Z);
-		RHIUpdateTexture3D(PageTableTextureRHI, 0, UpdateRegion, UpdateRegion.Width * FormatSize, UpdateRegion.Width * UpdateRegion.Height * FormatSize, (const uint8*)SparseVolumeTextureRuntime.PageTable.GetData());
-	}
-
-	// Tile data
-	{
-		const FIntVector3 TileDataVolumeResolution = SparseVolumeTextureRuntime.Header.TileDataVolumeResolution;
-		const EPixelFormat VoxelFormatA = SparseVolumeTextureRuntime.Header.AttributesFormats[0];
-		const EPixelFormat VoxelFormatB = SparseVolumeTextureRuntime.Header.AttributesFormats[1];
-		const FUpdateTextureRegion3D UpdateRegion(0, 0, 0, 0, 0, 0, TileDataVolumeResolution.X, TileDataVolumeResolution.Y, TileDataVolumeResolution.Z);
-
-		// A
-		if (VoxelFormatA != PF_Unknown)
-		{
-			const FRHITextureCreateDesc Desc =
-				FRHITextureCreateDesc::Create3D(TEXT("SparseVolumeTexture.PhysicalTileDataA.RHITexture"),
-					TileDataVolumeResolution.X, TileDataVolumeResolution.Y, TileDataVolumeResolution.Z, VoxelFormatA)
-				.SetFlags(ETextureCreateFlags::ShaderResource);
-
-			PhysicalTileDataATextureRHI = RHICreateTexture(Desc);
-
-			const int32 FormatSize = GPixelFormats[VoxelFormatA].BlockBytes;
-			RHIUpdateTexture3D(PhysicalTileDataATextureRHI, 0, UpdateRegion, UpdateRegion.Width * FormatSize, UpdateRegion.Width * UpdateRegion.Height * FormatSize, (const uint8*)SparseVolumeTextureRuntime.PhysicalTileDataA.GetData());
-		}
-
-		// B
-		if (VoxelFormatB != PF_Unknown)
-		{
-			const FRHITextureCreateDesc Desc =
-				FRHITextureCreateDesc::Create3D(TEXT("SparseVolumeTexture.PhysicalTileDataB.RHITexture"),
-					TileDataVolumeResolution.X, TileDataVolumeResolution.Y, TileDataVolumeResolution.Z, VoxelFormatB)
-				.SetFlags(ETextureCreateFlags::ShaderResource);
-
-			PhysicalTileDataBTextureRHI = RHICreateTexture(Desc);
-
-			const int32 FormatSize = GPixelFormats[VoxelFormatB].BlockBytes;
-			RHIUpdateTexture3D(PhysicalTileDataBTextureRHI, 0, UpdateRegion, UpdateRegion.Width * FormatSize, UpdateRegion.Width * UpdateRegion.Height * FormatSize, (const uint8*)SparseVolumeTextureRuntime.PhysicalTileDataB.GetData());
-		}
-	}
-}
-
-void FSparseVolumeTextureSceneProxy::ReleaseRHI()
-{
-}
-
-////////////////////////////////////////////////////////////////////////////////////////////////
-
 FSparseVolumeTextureFrame::FSparseVolumeTextureFrame()
 	: RuntimeStreamedInData()
 	, SparseVolumeTextureSceneProxy()
@@ -769,7 +74,7 @@ FSparseVolumeTextureFrame::~FSparseVolumeTextureFrame()
 {
 }
 
-bool FSparseVolumeTextureFrame::BuildRuntimeData(FSparseVolumeTextureRuntime* OutRuntimeData)
+bool FSparseVolumeTextureFrame::BuildDerivedData(FSparseVolumeTextureData* OutMippedTextureData)
 {
 #if WITH_EDITORONLY_DATA
 	// Check if the virtualized bulk data payload is available now
@@ -777,14 +82,12 @@ bool FSparseVolumeTextureFrame::BuildRuntimeData(FSparseVolumeTextureRuntime* Ou
 	{
 		// First, read the source data in from the raw data stored as bulk data
 		UE::Serialization::FEditorBulkDataReader RawDataArchiveReader(RawData);
-		FSparseVolumeRawSource SparseVolumeRawSource;
-		SparseVolumeRawSource.Serialize(RawDataArchiveReader);
+		FSparseVolumeTextureData TextureData;
+		TextureData.Serialize(RawDataArchiveReader);
 
-		// Then, convert the raw source data to SVT
-		OutRuntimeData->Header = SparseVolumeRawSource.Header;
-		OutRuntimeData->PageTable = MoveTemp(SparseVolumeRawSource.PageTable);
-		OutRuntimeData->PhysicalTileDataA = MoveTemp(SparseVolumeRawSource.PhysicalTileDataA);
-		OutRuntimeData->PhysicalTileDataB = MoveTemp(SparseVolumeRawSource.PhysicalTileDataB);
+		TextureData.GenerateMipMaps();
+
+		*OutMippedTextureData = MoveTemp(TextureData);
 
 		// Now unload the raw data
 		RawData.UnloadData();
@@ -816,7 +119,10 @@ void FSparseVolumeTextureFrame::Serialize(FArchive& Ar, UStreamableSparseVolumeT
 			// Create runtime data from cooked bulk data
 			{
 				FBulkDataReader BulkDataReader(RuntimeStreamedInData);
-				SparseVolumeTextureSceneProxy->GetRuntimeData().Serialize(BulkDataReader);
+				FSparseVolumeTextureData TextureData;
+				TextureData.Serialize(BulkDataReader);
+				bool bSuccess = SparseVolumeTextureSceneProxy->GetRuntimeData().Create(TextureData);
+				check(bSuccess); // SVT_TODO
 			}
 
 			// The bulk data is no longer needed
@@ -830,14 +136,14 @@ void FSparseVolumeTextureFrame::Serialize(FArchive& Ar, UStreamableSparseVolumeT
 	{
 		// We are cooking the game, serialize the asset out.
 
-		FSparseVolumeTextureRuntime RuntimeData;
-		const bool bBuiltRuntimeData = BuildRuntimeData(&RuntimeData);
-		check(bBuiltRuntimeData); // SVT_TODO: actual error handling
+		FSparseVolumeTextureData DerivedData;
+		const bool bBuiltDerivedData = BuildDerivedData(&DerivedData);
+		check(bBuiltDerivedData); // SVT_TODO: actual error handling
 
-		// Write runtime data into RuntimeStreamedInData
+		// Write derived data into RuntimeStreamedInData
 		{
 			FBulkDataWriter BulkDataWriter(RuntimeStreamedInData);
-			RuntimeData.Serialize(BulkDataWriter);
+			DerivedData.Serialize(BulkDataWriter);
 		}
 
 		// And now write the cooked bulk data to the archive
@@ -867,7 +173,7 @@ FVector4 USparseVolumeTexture::GetUniformParameter(int32 Index) const
 	const FSparseVolumeTextureSceneProxy* Proxy = GetSparseVolumeTextureSceneProxy();
 	if (Proxy)
 	{
-		const FSparseVolumeAssetHeader& Header = Proxy->GetHeader();
+		const FSparseVolumeTextureRuntimeHeader& Header = Proxy->GetHeader();
 		switch (Index)
 		{
 		case ESparseVolumeTexture_TileSize:
@@ -901,11 +207,20 @@ FVector4 USparseVolumeTexture::GetUniformParameter(int32 Index) const
 
 void USparseVolumeTexture::GetPackedUniforms(FUintVector4& OutPacked0, FUintVector4& OutPacked1) const
 {
+	FIntVector3 PageTableOffset = FIntVector3::ZeroValue;
+	int32 MinMipLevel = 0;
+	int32 MaxMipLevel = 0;
+	const FSparseVolumeTextureSceneProxy* Proxy = GetSparseVolumeTextureSceneProxy();
+	if (Proxy)
+	{
+		const FSparseVolumeTextureRuntimeHeader& Header = Proxy->GetHeader();
+		PageTableOffset = Header.PageTableVolumeAABBMin;
+		MinMipLevel = Header.LowestResidentLevel;
+		MaxMipLevel = Header.HighestResidentLevel;
+	}
+	const FIntVector3 VolumeResolution = GetVolumeResolution();
+	const FVector3f VolumePageResolution = FVector3f(VolumeResolution) / SPARSE_VOLUME_TILE_RES;
 	const FVector4 TileSize = GetUniformParameter(ESparseVolumeTexture_TileSize);
-	const FVector4 PageTableSize = GetUniformParameter(ESparseVolumeTexture_PageTableSize);
-	const FVector4 UVScale = GetUniformParameter(ESparseVolumeTexture_UVScale);
-	const FVector4 UVBias = GetUniformParameter(ESparseVolumeTexture_UVBias);
-	const FUintVector4 PageTableSizeUIMinusOne = FUintVector4(PageTableSize.X - 1, PageTableSize.Y - 1, PageTableSize.Z - 1, 0);
 
 	auto AsUint = [](float X)
 	{
@@ -913,13 +228,13 @@ void USparseVolumeTexture::GetPackedUniforms(FUintVector4& OutPacked0, FUintVect
 		return FU.U;
 	};
 
-	OutPacked0.X = AsUint((float)UVScale.X);
-	OutPacked0.Y = AsUint((float)UVScale.Y);
-	OutPacked0.Z = AsUint((float)UVScale.Z);
-	OutPacked0.W = (PageTableSizeUIMinusOne.X & 0x7FF) | ((PageTableSizeUIMinusOne.Y & 0x7FF) << 11) | ((PageTableSizeUIMinusOne.Z & 0x3FF) << 22);
-	OutPacked1.X = AsUint((float)UVBias.X);
-	OutPacked1.Y = AsUint((float)UVBias.Y);
-	OutPacked1.Z = AsUint((float)UVBias.Z);
+	OutPacked0.X = AsUint(VolumePageResolution.X);
+	OutPacked0.Y = AsUint(VolumePageResolution.Y);
+	OutPacked0.Z = AsUint(VolumePageResolution.Z);
+	OutPacked0.W = (uint32)((MinMipLevel & 0xFF) | ((MaxMipLevel & 0xFF) << 8));
+	OutPacked1.X = AsUint((float)PageTableOffset.X);
+	OutPacked1.Y = AsUint((float)PageTableOffset.Y);
+	OutPacked1.Z = AsUint((float)PageTableOffset.Z);
 	OutPacked1.W = AsUint((float)TileSize.X);
 }
 
@@ -930,9 +245,8 @@ void USparseVolumeTexture::GetFrameUVScaleBias(FVector* OutScale, FVector* OutBi
 	const FSparseVolumeTextureSceneProxy* Proxy = GetSparseVolumeTextureSceneProxy();
 	if (Proxy)
 	{
-		const FSparseVolumeAssetHeader& Header = Proxy->GetHeader();
-		const int32 MipFactor = 1 << Header.MipLevel;
-		const FVector GlobalVolumeRes = FVector(GetVolumeResolution()) / MipFactor;
+		const FSparseVolumeTextureRuntimeHeader& Header = Proxy->GetHeader();
+		const FVector GlobalVolumeRes = FVector(GetVolumeResolution());
 		check(GlobalVolumeRes.X > 0.0 && GlobalVolumeRes.Y > 0.0 && GlobalVolumeRes.Z > 0.0);
 		const FVector FrameBoundsPaddedMin = FVector(Header.PageTableVolumeAABBMin * SPARSE_VOLUME_TILE_RES); // padded to multiple of page size
 		const FVector FrameBoundsPaddedMax = FVector(Header.PageTableVolumeAABBMax * SPARSE_VOLUME_TILE_RES);
@@ -986,20 +300,17 @@ void UStreamableSparseVolumeTexture::FinishDestroy()
 void UStreamableSparseVolumeTexture::BeginDestroy()
 {
 	Super::BeginDestroy();
-	for (FSparseVolumeTextureFrameMips& FrameMips : Frames)
+	for (FSparseVolumeTextureFrame& Frame : Frames)
 	{
-		for (FSparseVolumeTextureFrame& Frame : FrameMips)
+		if (Frame.SparseVolumeTextureSceneProxy)
 		{
-			if (Frame.SparseVolumeTextureSceneProxy)
-			{
-				ENQUEUE_RENDER_COMMAND(UStreamableSparseVolumeTexture_DeleteSVTProxy)(
-					[Proxy = Frame.SparseVolumeTextureSceneProxy](FRHICommandListImmediate& RHICmdList)
-					{
-						Proxy->ReleaseResource();
-						delete Proxy;
-					});
-				Frame.SparseVolumeTextureSceneProxy = nullptr;
-			}
+			ENQUEUE_RENDER_COMMAND(UStreamableSparseVolumeTexture_DeleteSVTProxy)(
+				[Proxy = Frame.SparseVolumeTextureSceneProxy](FRHICommandListImmediate& RHICmdList)
+				{
+					Proxy->ReleaseResource();
+					delete Proxy;
+				});
+			Frame.SparseVolumeTextureSceneProxy = nullptr;
 		}
 	}
 }
@@ -1010,8 +321,6 @@ void UStreamableSparseVolumeTexture::Serialize(FArchive& Ar)
 
 	int32 NumFrames = Frames.Num();
 	Ar << NumFrames;
-	int32 NumMipLevels = Frames.IsEmpty() ? 0 : Frames[0].Num();
-	Ar << NumMipLevels;
 
 	if (Ar.IsLoading())
 	{
@@ -1020,11 +329,7 @@ void UStreamableSparseVolumeTexture::Serialize(FArchive& Ar)
 	}
 	for (int32 FrameIndex = 0; FrameIndex < NumFrames; ++FrameIndex)
 	{
-		Frames[FrameIndex].SetNum(NumMipLevels);
-		for (int32 MipLevel = 0; MipLevel < NumMipLevels; ++MipLevel)
-		{
-			Frames[FrameIndex][MipLevel].Serialize(Ar, this, FrameIndex);
-		}
+		Frames[FrameIndex].Serialize(Ar, this, FrameIndex);
 	}
 }
 
@@ -1043,15 +348,11 @@ void UStreamableSparseVolumeTexture::GetResourceSizeEx(FResourceSizeEx& Cumulati
 	SIZE_T SizeCPU = sizeof(Frames);
 	SIZE_T SizeGPU = 0;
 	SizeCPU += Frames.GetAllocatedSize();
-	for (const FSparseVolumeTextureFrameMips& FrameMips : Frames)
+	for (const FSparseVolumeTextureFrame& Frame : Frames)
 	{
-		SizeCPU += FrameMips.GetAllocatedSize();
-		for (const FSparseVolumeTextureFrame& Frame : FrameMips)
+		if (Frame.SparseVolumeTextureSceneProxy)
 		{
-			if (Frame.SparseVolumeTextureSceneProxy)
-			{
-				Frame.SparseVolumeTextureSceneProxy->GetMemorySize(&SizeCPU, &SizeGPU);
-			}
+			Frame.SparseVolumeTextureSceneProxy->GetMemorySize(&SizeCPU, &SizeGPU);
 		}
 	}
 	ISparseVolumeTextureStreamingManager& StreamingManager = IStreamingManager::Get().GetSparseVolumeTextureStreamingManager();
@@ -1080,7 +381,7 @@ const FSparseVolumeTextureSceneProxy* UStreamableSparseVolumeTexture::GetStreame
 	return Proxy;
 }
 
-TArrayView<const FSparseVolumeTextureFrameMips> UStreamableSparseVolumeTexture::GetFrames() const
+TArrayView<const FSparseVolumeTextureFrame> UStreamableSparseVolumeTexture::GetFrames() const
 {
 	return Frames;
 }
@@ -1091,35 +392,29 @@ void UStreamableSparseVolumeTexture::GenerateOrLoadDDCRuntimeDataAndCreateSceneP
 	UE::DerivedData::FRequestOwner DDCRequestOwner(UE::DerivedData::EPriority::Normal);
 	{
 		UE::DerivedData::FRequestBarrier DDCRequestBarrier(DDCRequestOwner);
-		for (FSparseVolumeTextureFrameMips& FrameMips : Frames)
+		for (FSparseVolumeTextureFrame& Frame : Frames)
 		{
-			for (FSparseVolumeTextureFrame& Frame : FrameMips)
+			// Release any previously allocated render thread proxy
+			if (Frame.SparseVolumeTextureSceneProxy)
 			{
-				// Release any previously allocated render thread proxy
-				if (Frame.SparseVolumeTextureSceneProxy)
-				{
-					BeginReleaseResource(Frame.SparseVolumeTextureSceneProxy);
-				}
-				else
-				{
-					Frame.SparseVolumeTextureSceneProxy = new FSparseVolumeTextureSceneProxy();
-				}
-
-				GenerateOrLoadDDCRuntimeDataForFrame(Frame, DDCRequestOwner);
+				BeginReleaseResource(Frame.SparseVolumeTextureSceneProxy);
+			}
+			else
+			{
+				Frame.SparseVolumeTextureSceneProxy = new FSparseVolumeTextureSceneProxy();
 			}
+
+			GenerateOrLoadDDCRuntimeDataForFrame(Frame, DDCRequestOwner);
 		}
 	}
 
 	// Wait for all DDC requests to complete before creating the proxies
 	DDCRequestOwner.Wait();
 
-	for (FSparseVolumeTextureFrameMips& FrameMips : Frames)
+	for (FSparseVolumeTextureFrame& Frame : Frames)
 	{
-		for (FSparseVolumeTextureFrame& Frame : FrameMips)
-		{
-			// Runtime data is now valid, initialize the render thread proxy
-			BeginInitResource(Frame.SparseVolumeTextureSceneProxy);
-		}
+		// Runtime data is now valid, initialize the render thread proxy
+		BeginInitResource(Frame.SparseVolumeTextureSceneProxy);
 	}
 
 	IStreamingManager::Get().GetSparseVolumeTextureStreamingManager().AddSparseVolumeTexture(this);
@@ -1147,7 +442,10 @@ void UStreamableSparseVolumeTexture::GenerateOrLoadDDCRuntimeDataForFrame(FSpars
 					{
 						FSharedBuffer Data = Value.GetData().Decompress();
 						FMemoryReaderView Ar(Data, true /*bIsPersistent*/);
-						Frame.SparseVolumeTextureSceneProxy->GetRuntimeData().Serialize(Ar);
+						FSparseVolumeTextureData TextureData;
+						TextureData.Serialize(Ar);
+						bool bSuccess = Frame.SparseVolumeTextureSceneProxy->GetRuntimeData().Create(TextureData);
+						check(bSuccess); // SVT_TODO
 					});
 			}
 			else if (Response.Status == EStatus::Error)
@@ -1160,32 +458,36 @@ void UStreamableSparseVolumeTexture::GenerateOrLoadDDCRuntimeDataForFrame(FSpars
 						// Check if the virtualized bulk data payload is available now
 						if (Frame.RawData.HasPayloadData())
 						{
-							const bool bSuccess = Frame.BuildRuntimeData(&RuntimeData);
+							FSparseVolumeTextureData TextureData;
+							bool bSuccess = Frame.BuildDerivedData(&TextureData);
 							ensure(bSuccess);
-						}
-						else
-						{
-							UE_LOG(LogSparseVolumeTexture, Error, TEXT("SparseVolumeTexture - Raw source data is not available for %s. Using default data."), *GetName());
-							RuntimeData.SetAsDefaultTexture();
-						}
 
-						// Using a LargeMemoryWriter for serialization since the data can be bigger than 2 GB
-						FLargeMemoryWriter LargeMemWriter(0, /*bIsPersistent=*/ true);
-						RuntimeData.Serialize(LargeMemWriter);
+							bSuccess = RuntimeData.Create(TextureData);
+							ensure(bSuccess);
 
-						int64 UncompressedSize = LargeMemWriter.TotalSize();
+							// Using a LargeMemoryWriter for serialization since the data can be bigger than 2 GB
+							FLargeMemoryWriter LargeMemWriter(0, /*bIsPersistent=*/ true);
+							TextureData.Serialize(LargeMemWriter);
 
-						// Since the DDC doesn't support data bigger than 2 GB, we only cache for such uncompressed size.
-						constexpr int64 SizeThreshold = 2147483648LL;	// 2GB
-						const bool bIsCacheable = UncompressedSize < SizeThreshold;
-						if (bIsCacheable)
-						{
-							FValue Value = FValue::Compress(FSharedBuffer::MakeView(LargeMemWriter.GetView()));
-							UE::DerivedData::GetCache().PutValue({ {Name, Key, Value} }, DDCRequestOwner);
+							const int64 UncompressedSize = LargeMemWriter.TotalSize();
+
+							// Since the DDC doesn't support data bigger than 2 GB, we only cache for such uncompressed size.
+							constexpr int64 SizeThreshold = 2147483648LL;	// 2GB
+							const bool bIsCacheable = UncompressedSize < SizeThreshold;
+							if (bIsCacheable)
+							{
+								FValue Value = FValue::Compress(FSharedBuffer::MakeView(LargeMemWriter.GetView()));
+								UE::DerivedData::GetCache().PutValue({ {Name, Key, Value} }, DDCRequestOwner);
+							}
+							else
+							{
+								UE_LOG(LogSparseVolumeTexture, Error, TEXT("SparseVolumeTexture - the asset is too large to fit in Derived Data Cache %s"), *GetName());
+							}
 						}
 						else
 						{
-							UE_LOG(LogSparseVolumeTexture, Error, TEXT("SparseVolumeTexture - the asset is too large to fit in Derived Data Cache %s"), *GetName());
+							UE_LOG(LogSparseVolumeTexture, Error, TEXT("SparseVolumeTexture - Raw source data is not available for %s. Using default data."), *GetName());
+							RuntimeData.SetAsDefaultTexture();
 						}
 					});
 			}
@@ -1216,23 +518,6 @@ const FSparseVolumeTextureSceneProxy* UAnimatedSparseVolumeTexture::GetSparseVol
 	return GetStreamedFrameProxyOrFallback(FrameIndex, MipLevel);
 }
 
-const FSparseVolumeTextureSceneProxy* UAnimatedSparseVolumeTexture::GetSparseVolumeTextureFrameSceneProxy(int32 FrameIndex, int32 MipLevel) const
-{
-	check(!Frames.IsEmpty());
-	FrameIndex = FrameIndex % Frames.Num();
-	MipLevel = FMath::Clamp(MipLevel, 0, GetNumMipLevels() - 1);
-	return GetStreamedFrameProxyOrFallback(FrameIndex, MipLevel);
-}
-
-const FSparseVolumeAssetHeader* UAnimatedSparseVolumeTexture::GetSparseVolumeTextureFrameHeader(int32 FrameIndex, int32 MipLevel) const
-{
-	check(!Frames.IsEmpty());
-	FrameIndex = FrameIndex % Frames.Num();
-	MipLevel = FMath::Clamp(MipLevel, 0, GetNumMipLevels() - 1);
-	const FSparseVolumeTextureSceneProxy* Proxy = GetStreamedFrameProxyOrFallback(FrameIndex, MipLevel);
-	return Proxy ? &Proxy->GetHeader() : nullptr;
-}
-
 ////////////////////////////////////////////////////////////////////////////////////////////////
 
 USparseVolumeTextureFrame::USparseVolumeTextureFrame(const FObjectInitializer& ObjectInitializer)
diff --git a/Engine/Source/Runtime/Engine/Private/SparseVolumeTexture/SparseVolumeTextureData.cpp b/Engine/Source/Runtime/Engine/Private/SparseVolumeTexture/SparseVolumeTextureData.cpp
new file mode 100644
--- /dev/null
+++ b/Engine/Source/Runtime/Engine/Private/SparseVolumeTexture/SparseVolumeTextureData.cpp
@@ -0,0 +1,457 @@
+// Copyright Epic Games, Inc. All Rights Reserved.
+
+#include "SparseVolumeTexture/SparseVolumeTextureData.h"
+#include "SparseVolumeTexture/SparseVolumeTextureUtility.h"
+
+DEFINE_LOG_CATEGORY_STATIC(LogSparseVolumeTextureData, Log, All);
+
+////////////////////////////////////////////////////////////////////////////////////////////////
+
+void FSparseVolumeTextureDataHeader::Serialize(FArchive& Ar)
+{
+	FSparseVolumeTextureHeader::Serialize(Ar);
+
+	Ar << Version;
+
+	if (Version == 0)
+	{
+		int32 NumMipLevels = MipInfo.Num();
+		Ar << NumMipLevels;
+		if (Ar.IsLoading())
+		{
+			MipInfo.Reset(NumMipLevels);
+			MipInfo.SetNum(NumMipLevels);
+		}
+		for (int32 MipLevel = 0; MipLevel < NumMipLevels; ++MipLevel)
+		{
+			Ar << MipInfo[MipLevel].TileOffset;
+			Ar << MipInfo[MipLevel].TileCount;
+		}
+	}
+	else
+	{
+		// FSparseVolumeTextureDataHeader needs to account for new version
+		check(false);
+	}
+}
+
+////////////////////////////////////////////////////////////////////////////////////////////////
+
+void FSparseVolumeTextureData::Serialize(FArchive& Ar)
+{
+	Header.Serialize(Ar);
+
+	Ar << Version;
+
+	if (Version == 0)
+	{
+		Ar << PageTable;
+		Ar << PhysicalTileDataA;
+		Ar << PhysicalTileDataB;
+	}
+	else
+	{
+		// FSparseVolumeTextureData needs to account for new version
+		check(false);
+	}
+}
+
+bool FSparseVolumeTextureData::Construct(const ISparseVolumeTextureDataConstructionAdapter& Adapter)
+{
+	using namespace UE::SVT;
+	using namespace UE::SVT::Private;
+
+	ISparseVolumeTextureDataConstructionAdapter::FAttributesInfo AttributesInfo[2];
+	Adapter.GetAttributesInfo(AttributesInfo[0], AttributesInfo[1]);
+
+	Header.MipInfo.SetNum(1);
+	FSparseVolumeTextureMipInfo& MipInfo = Header.MipInfo[0];
+	Header.AttributesFormats[0] = AttributesInfo[0].Format;
+	Header.AttributesFormats[1] = AttributesInfo[1].Format;
+
+	Header.VirtualVolumeResolution = Adapter.GetResolution();
+	Header.VirtualVolumeAABBMin = Adapter.GetAABBMin();
+	Header.VirtualVolumeAABBMax = Adapter.GetAABBMax();
+	Header.PageTableVolumeAABBMin = Header.VirtualVolumeAABBMin / SPARSE_VOLUME_TILE_RES;
+	Header.PageTableVolumeAABBMax = (Header.VirtualVolumeAABBMax + FIntVector3(SPARSE_VOLUME_TILE_RES - 1)) / SPARSE_VOLUME_TILE_RES;
+
+	FIntVector3 PageTableVolumeResolution = Header.PageTableVolumeAABBMax - Header.PageTableVolumeAABBMin;
+	
+	// We need to ensure a power of two resolution for the page table in order to fit all mips of the page table into the physical mips of the texture resource.
+	PageTableVolumeResolution.X = FMath::RoundUpToPowerOfTwo(PageTableVolumeResolution.X);
+	PageTableVolumeResolution.Y = FMath::RoundUpToPowerOfTwo(PageTableVolumeResolution.Y);
+	PageTableVolumeResolution.Z = FMath::RoundUpToPowerOfTwo(PageTableVolumeResolution.Z);
+
+	Header.PageTableVolumeAABBMax = Header.PageTableVolumeAABBMin + PageTableVolumeResolution;
+
+	if (PageTableVolumeResolution.X > SVTMaxVolumeTextureDim
+		|| PageTableVolumeResolution.Y > SVTMaxVolumeTextureDim
+		|| PageTableVolumeResolution.Z > SVTMaxVolumeTextureDim)
+	{
+		UE_LOG(LogSparseVolumeTextureData, Warning, TEXT("SparseVolumeTexture page table texture dimensions exceed limit (%ix%ix%i): %ix%ix%i"),
+			SVTMaxVolumeTextureDim, SVTMaxVolumeTextureDim, SVTMaxVolumeTextureDim,
+			PageTableVolumeResolution.X, PageTableVolumeResolution.Y, PageTableVolumeResolution.Z);
+		return false;
+	}
+
+	Header.PageTableVolumeResolution = PageTableVolumeResolution;
+
+	// Tag all pages with valid data
+	TBitArray PagesWithData(false, Header.PageTableVolumeResolution.Z * Header.PageTableVolumeResolution.Y * Header.PageTableVolumeResolution.X);
+	Adapter.IteratePhysicalSource([&](const FIntVector3& Coord, int32 AttributesIdx, int32 ComponentIdx, float VoxelValue)
+		{
+			const bool bIsFallbackValue = (VoxelValue == AttributesInfo[AttributesIdx].FallbackValue[ComponentIdx]);
+			if (!bIsFallbackValue)
+			{
+				const FIntVector3 GridCoord = Coord;
+				check(GridCoord.X >= 0 && GridCoord.Y >= 0 && GridCoord.Z >= 0);
+				check(GridCoord.X < Header.VirtualVolumeAABBMax.X&& GridCoord.Y < Header.VirtualVolumeAABBMax.Y&& GridCoord.Z < Header.VirtualVolumeAABBMax.Z);
+				const FIntVector3 PageCoord = (GridCoord / SPARSE_VOLUME_TILE_RES) - Header.PageTableVolumeAABBMin;
+				check(PageCoord.X >= 0 && PageCoord.Y >= 0 && PageCoord.Z >= 0);
+				check(PageCoord.X < Header.PageTableVolumeResolution.X&& PageCoord.Y < Header.PageTableVolumeResolution.Y&& PageCoord.Z < Header.PageTableVolumeResolution.Z);
+
+				const int32 PageIndex = PageCoord.Z * (Header.PageTableVolumeResolution.Y * Header.PageTableVolumeResolution.X) + PageCoord.Y * Header.PageTableVolumeResolution.X + PageCoord.X;
+				PagesWithData[PageIndex] = true;
+			}
+		});
+
+	// Allocate some memory for temp data (worst case)
+	TArray<FIntVector3> LinearAllocatedPages;
+	LinearAllocatedPages.SetNum(Header.PageTableVolumeResolution.X * Header.PageTableVolumeResolution.Y * Header.PageTableVolumeResolution.Z);
+
+	// Go over each potential page from the source data and push allocate it if it has any data.
+	// Otherwise point to the default empty page.
+	int32 NumAllocatedPages = 0;
+	for (int32_t PageZ = 0; PageZ < Header.PageTableVolumeResolution.Z; ++PageZ)
+	{
+		for (int32_t PageY = 0; PageY < Header.PageTableVolumeResolution.Y; ++PageY)
+		{
+			for (int32_t PageX = 0; PageX < Header.PageTableVolumeResolution.X; ++PageX)
+			{
+				const int32 PageIndex = PageZ * (Header.PageTableVolumeResolution.Y * Header.PageTableVolumeResolution.X) + PageY * Header.PageTableVolumeResolution.X + PageX;
+				const bool bHasAnyData = PagesWithData[PageIndex];
+				if (bHasAnyData)
+				{
+					LinearAllocatedPages[NumAllocatedPages] = FIntVector3(PageX, PageY, PageZ);
+					NumAllocatedPages++;
+				}
+			}
+		}
+	}
+
+	const int32 EffectivelyAllocatedTiles = NumAllocatedPages + 1 /*null tile*/;
+
+	// Initialize the SparseVolumeTexture page and tile.
+	const int32 FormatSize[] = { GPixelFormats[Header.AttributesFormats[0]].BlockBytes, GPixelFormats[Header.AttributesFormats[1]].BlockBytes };
+	PageTable.SetNum(1);
+	PageTable[0].SetNumZeroed(Header.PageTableVolumeResolution.X * Header.PageTableVolumeResolution.Y * Header.PageTableVolumeResolution.Z);
+	PhysicalTileDataA.SetNumZeroed(EffectivelyAllocatedTiles * SVTNumVoxelsPerTile * FormatSize[0]);
+	PhysicalTileDataB.SetNumZeroed(EffectivelyAllocatedTiles * SVTNumVoxelsPerTile * FormatSize[1]);
+
+	// Generate page table and tile volume data by splatting the data
+	{
+		uint32 DstTileIndex = 0;
+
+		// Add an empty tile, reserve slot at coord 0
+		{
+			// Fill null tile with fallback/background value
+			FVector4f FallbackValues[] = { AttributesInfo[0].FallbackValue, AttributesInfo[1].FallbackValue };
+			for (int32 AttributesIdx = 0; AttributesIdx < 2; ++AttributesIdx)
+			{
+				if (AttributesInfo[AttributesIdx].bNormalized)
+				{
+					FallbackValues[AttributesIdx] = FallbackValues[AttributesIdx] * AttributesInfo[AttributesIdx].NormalizeScale + AttributesInfo[AttributesIdx].NormalizeBias;
+				}
+			}
+			for (int32 Z = 0; Z < SPARSE_VOLUME_TILE_RES; ++Z)
+			{
+				for (int32 Y = 0; Y < SPARSE_VOLUME_TILE_RES; ++Y)
+				{
+					for (int32 X = 0; X < SPARSE_VOLUME_TILE_RES; ++X)
+					{
+						const FIntVector3 WriteCoord = FIntVector3(X, Y, Z);
+						WriteTileDataVoxel(0 /*TileIndex*/, WriteCoord, 0 /*AttributesIndex*/, FallbackValues[0]);
+						WriteTileDataVoxel(0 /*TileIndex*/, WriteCoord, 1 /*AttributesIndex*/, FallbackValues[1]);
+					}
+				}
+			}
+
+			Header.NullTileValues[0] = FallbackValues[0];
+			Header.NullTileValues[1] = FallbackValues[1];
+
+			// PageTable is all cleared to zero, simply skip a tile
+			++DstTileIndex;
+		}
+
+		for (int32 i = 0; i < NumAllocatedPages; ++i)
+		{
+			const FIntVector3 PageCoordToSplat = LinearAllocatedPages[i];
+
+			// Setup the page table entry
+			PageTable[0]
+				[
+					PageCoordToSplat.Z * Header.PageTableVolumeResolution.X * Header.PageTableVolumeResolution.Y +
+					PageCoordToSplat.Y * Header.PageTableVolumeResolution.X +
+					PageCoordToSplat.X
+				] = DstTileIndex;
+
+			// Set the next tile to be written to
+			++DstTileIndex;
+		}
+		MipInfo.TileOffset = 0;
+		MipInfo.TileCount = EffectivelyAllocatedTiles;
+	}
+
+	// Write physical tile data
+	Adapter.IteratePhysicalSource([&](const FIntVector3& Coord, int32 AttributesIdx, int32 ComponentIdx, float VoxelValue)
+		{
+			const FIntVector3 GridCoord = Coord;
+			check(GridCoord.X >= 0 && GridCoord.Y >= 0 && GridCoord.Z >= 0);
+			check(GridCoord.X < Header.VirtualVolumeAABBMax.X&& GridCoord.Y < Header.VirtualVolumeAABBMax.Y&& GridCoord.Z < Header.VirtualVolumeAABBMax.Z);
+			const FIntVector3 PageCoord = (GridCoord / SPARSE_VOLUME_TILE_RES) - Header.PageTableVolumeAABBMin;
+			check(PageCoord.X >= 0 && PageCoord.Y >= 0 && PageCoord.Z >= 0);
+			check(PageCoord.X < Header.PageTableVolumeResolution.X&& PageCoord.Y < Header.PageTableVolumeResolution.Y&& PageCoord.Z < Header.PageTableVolumeResolution.Z);
+
+			const int32 PageIndex = PageCoord.Z * (Header.PageTableVolumeResolution.Y * Header.PageTableVolumeResolution.X) + PageCoord.Y * Header.PageTableVolumeResolution.X + PageCoord.X;
+
+			if (!PagesWithData[PageIndex])
+			{
+				return;
+			}
+
+			float WriteValue = VoxelValue;
+			if (AttributesInfo[AttributesIdx].bNormalized)
+			{
+				WriteValue = WriteValue * AttributesInfo[AttributesIdx].NormalizeScale[ComponentIdx] + AttributesInfo[AttributesIdx].NormalizeBias[ComponentIdx];
+			}
+
+			const int32 TileIndex = (int32)PageTable[0][PageIndex];
+			const FIntVector3 TileLocalCoord = GridCoord % SPARSE_VOLUME_TILE_RES;
+			WriteTileDataVoxel(TileIndex, TileLocalCoord, AttributesIdx, FVector4f(WriteValue, WriteValue, WriteValue, WriteValue), ComponentIdx);
+		});
+
+	return true;
+}
+
+uint32 FSparseVolumeTextureData::ReadPageTable(const FIntVector3& PageTableCoord, int32 MipLevel) const
+{
+	if (!Header.MipInfo.IsValidIndex(MipLevel))
+	{
+		return INDEX_NONE;
+	}
+	const int32 MipLevelFactor = 1 << MipLevel;
+	const FIntVector3 PageTableResolution = FIntVector3(
+		FMath::Max(1, Header.PageTableVolumeResolution.X / MipLevelFactor), 
+		FMath::Max(1, Header.PageTableVolumeResolution.Y / MipLevelFactor), 
+		FMath::Max(1, Header.PageTableVolumeResolution.Z / MipLevelFactor));
+	if (PageTableCoord.X < 0 || PageTableCoord.Y < 0 || PageTableCoord.Z < 0
+		|| PageTableCoord.X >= PageTableResolution.X
+		|| PageTableCoord.Y >= PageTableResolution.Y
+		|| PageTableCoord.Z >= PageTableResolution.Z)
+	{
+		return INDEX_NONE;
+	}
+
+	const int32 PageIndex = PageTableCoord.Z * (PageTableResolution.Y * PageTableResolution.X) + PageTableCoord.Y * PageTableResolution.X + PageTableCoord.X;
+	if (PageTable[MipLevel].IsValidIndex(PageIndex))
+	{
+		return PageTable[MipLevel][PageIndex];
+	}
+	return INDEX_NONE;
+}
+
+FVector4f FSparseVolumeTextureData::ReadTileDataVoxel(int32 TileIndex, const FIntVector3& TileDataCoord, int32 AttributesIdx) const
+{
+	check(AttributesIdx >= 0 && AttributesIdx <= 1);
+	if ((AttributesIdx == 0 && PhysicalTileDataA.IsEmpty()) || (AttributesIdx == 1 && PhysicalTileDataB.IsEmpty()))
+	{
+		return FVector4f();
+	}
+
+	// SV_TODO check for TileIndex oob
+
+	if (TileDataCoord.X < 0 || TileDataCoord.Y < 0 || TileDataCoord.Z < 0
+		|| TileDataCoord.X >= SPARSE_VOLUME_TILE_RES
+		|| TileDataCoord.Y >= SPARSE_VOLUME_TILE_RES
+		|| TileDataCoord.Z >= SPARSE_VOLUME_TILE_RES)
+	{
+		return FVector4f();
+	}
+
+	const int32 VoxelIndex = TileIndex * UE::SVT::SVTNumVoxelsPerTile + TileDataCoord.Z * (SPARSE_VOLUME_TILE_RES * SPARSE_VOLUME_TILE_RES) + TileDataCoord.Y * SPARSE_VOLUME_TILE_RES + TileDataCoord.X;
+	const uint8* TileData = AttributesIdx == 0 ? PhysicalTileDataA.GetData() : PhysicalTileDataB.GetData();
+	const EPixelFormat Format = Header.AttributesFormats[AttributesIdx];
+	return UE::SVT::ReadVoxel(VoxelIndex, TileData, Format);
+}
+
+FVector4f FSparseVolumeTextureData::Load(const FIntVector3& VolumeCoord, int32 MipLevel, int32 AttributesIdx) const
+{
+	if (!Header.MipInfo.IsValidIndex(MipLevel))
+	{
+		return FVector4f();
+	}
+	const FIntVector3 PageTableCoord = (VolumeCoord / SPARSE_VOLUME_TILE_RES) - (Header.PageTableVolumeAABBMin / (1 << MipLevel));
+	const uint32 TileIndex = ReadPageTable(PageTableCoord, MipLevel);
+	const FIntVector3 VoxelCoord = VolumeCoord % SPARSE_VOLUME_TILE_RES;
+	const FVector4f Sample = ReadTileDataVoxel(TileIndex, VoxelCoord, AttributesIdx);
+	return Sample;
+}
+
+void FSparseVolumeTextureData::WriteTileDataVoxel(int32 TileIndex, const FIntVector3& TileDataCoord, int32 AttributesIdx, const FVector4f& Value, int32 DstComponent)
+{
+	check(AttributesIdx >= 0 && AttributesIdx <= 1);
+	if ((AttributesIdx == 0 && PhysicalTileDataA.IsEmpty()) || (AttributesIdx == 1 && PhysicalTileDataB.IsEmpty()))
+	{
+		return;
+	}
+
+	// SV_TODO check for TileIndex oob
+
+	if (TileDataCoord.X < 0 || TileDataCoord.Y < 0 || TileDataCoord.Z < 0
+		|| TileDataCoord.X >= SPARSE_VOLUME_TILE_RES
+		|| TileDataCoord.Y >= SPARSE_VOLUME_TILE_RES
+		|| TileDataCoord.Z >= SPARSE_VOLUME_TILE_RES)
+	{
+		return;
+	}
+
+	const int32 VoxelIndex = TileIndex * UE::SVT::SVTNumVoxelsPerTile + TileDataCoord.Z * (SPARSE_VOLUME_TILE_RES * SPARSE_VOLUME_TILE_RES) + TileDataCoord.Y * SPARSE_VOLUME_TILE_RES + TileDataCoord.X;
+	uint8* TileData = AttributesIdx == 0 ? PhysicalTileDataA.GetData() : PhysicalTileDataB.GetData();
+	const EPixelFormat Format = Header.AttributesFormats[AttributesIdx];
+	UE::SVT::WriteVoxel(VoxelIndex, TileData, Format, Value, DstComponent);
+}
+
+void FSparseVolumeTextureData::GenerateMipMaps(int32 NumMipLevels)
+{
+	check(!Header.MipInfo.IsEmpty());
+	check(FMath::IsPowerOfTwo(Header.PageTableVolumeResolution.X) 
+		&& FMath::IsPowerOfTwo(Header.PageTableVolumeResolution.Y) 
+		&& FMath::IsPowerOfTwo(Header.PageTableVolumeResolution.Z));
+	if (NumMipLevels == -1)
+	{
+		NumMipLevels = 1;
+		FIntVector3 Resolution = Header.VirtualVolumeResolution;
+		while (Resolution.X > SPARSE_VOLUME_TILE_RES || Resolution.Y > SPARSE_VOLUME_TILE_RES || Resolution.Z > SPARSE_VOLUME_TILE_RES)
+		{
+			Resolution /= 2;
+			++NumMipLevels;
+		}
+	}
+	Header.MipInfo.SetNum(1);
+	PageTable.SetNum(1);
+
+	for (int32 MipLevel = 1; MipLevel < NumMipLevels; ++MipLevel)
+	{
+		Header.MipInfo.AddDefaulted();
+		FSparseVolumeTextureMipInfo& MipInfo = Header.MipInfo[MipLevel];
+		MipInfo.TileOffset = Header.MipInfo[MipLevel - 1].TileOffset + Header.MipInfo[MipLevel - 1].TileCount;
+
+		const int32 MipLevelFactor = 1 << MipLevel;
+		const int32 ParentMipLevelFactor = 1 << (MipLevel - 1);
+		const FIntVector3 PageTableVolumeAABBMin = (Header.VirtualVolumeAABBMin / SPARSE_VOLUME_TILE_RES) / MipLevelFactor;
+		const FIntVector3 ParentPageTableVolumeAABBMin = (Header.VirtualVolumeAABBMin / SPARSE_VOLUME_TILE_RES) / ParentMipLevelFactor;
+		const FIntVector3 PageTableVolumeResolution = FIntVector3(FMath::Max(1, Header.PageTableVolumeResolution.X / MipLevelFactor), FMath::Max(1, Header.PageTableVolumeResolution.Y / MipLevelFactor), FMath::Max(1, Header.PageTableVolumeResolution.Z / MipLevelFactor));
+		
+		check(FMath::IsPowerOfTwo(PageTableVolumeResolution.X)
+			&& FMath::IsPowerOfTwo(PageTableVolumeResolution.Y)
+			&& FMath::IsPowerOfTwo(PageTableVolumeResolution.Z));
+
+		// Allocate some memory for temp data (worst case)
+		TArray<FIntVector3> LinearAllocatedPages;
+		LinearAllocatedPages.SetNum(PageTableVolumeResolution.X * PageTableVolumeResolution.Y * PageTableVolumeResolution.Z);
+
+		// Go over each potential page from the source data and push allocate it if it has any data.
+		// Otherwise point to the default empty page.
+		int32 NumAllocatedPages = 0;
+		for (int32_t PageZ = 0; PageZ < PageTableVolumeResolution.Z; ++PageZ)
+		{
+			for (int32_t PageY = 0; PageY < PageTableVolumeResolution.Y; ++PageY)
+			{
+				for (int32_t PageX = 0; PageX < PageTableVolumeResolution.X; ++PageX)
+				{
+					const FIntVector3 PageCoord = FIntVector3(PageX, PageY, PageZ);
+					bool bHasAnyData = false;
+					for (int32_t OffsetIdx = 0; OffsetIdx < 8; ++OffsetIdx)
+					{
+						const FIntVector3 Offset = FIntVector3(OffsetIdx, OffsetIdx >> 1, OffsetIdx >> 2) & 1;
+						const FIntVector3 ParentPageTableCoord = (PageTableVolumeAABBMin + PageCoord) * 2 + Offset - ParentPageTableVolumeAABBMin;
+						const uint32 PageSample = ReadPageTable(ParentPageTableCoord, MipLevel - 1);
+
+						if (PageSample != INDEX_NONE && PageSample != 0)
+						{
+							bHasAnyData = true;
+						}
+					}
+					if (bHasAnyData)
+					{
+						LinearAllocatedPages[NumAllocatedPages] = PageCoord;
+						NumAllocatedPages++;
+					}
+				}
+			}
+		}
+
+		// Initialize the SparseVolumeTexture page and tile.
+		const int32 FormatSize[] = { GPixelFormats[Header.AttributesFormats[0]].BlockBytes, GPixelFormats[Header.AttributesFormats[1]].BlockBytes };
+		PageTable.AddDefaulted();
+		PageTable[MipLevel].SetNumZeroed(PageTableVolumeResolution.X * PageTableVolumeResolution.Y * PageTableVolumeResolution.Z);
+		PhysicalTileDataA.AddZeroed(NumAllocatedPages * UE::SVT::SVTNumVoxelsPerTile * FormatSize[0]);
+		PhysicalTileDataB.AddZeroed(NumAllocatedPages * UE::SVT::SVTNumVoxelsPerTile * FormatSize[1]);
+
+		// Generate page table and tile volume data by splatting the data
+		{
+			uint32 DstTileIndex = (uint32)MipInfo.TileOffset;
+
+			for (int32 i = 0; i < NumAllocatedPages; ++i)
+			{
+				const FIntVector3 PageCoordToSplat = LinearAllocatedPages[i];
+
+				// Setup the page table entry
+				PageTable[MipLevel]
+					[
+						PageCoordToSplat.Z * PageTableVolumeResolution.X * PageTableVolumeResolution.Y +
+						PageCoordToSplat.Y * PageTableVolumeResolution.X +
+						PageCoordToSplat.X
+					] = DstTileIndex;
+
+				// Write tile data
+				const FIntVector3 ParentVolumeCoordBase = (PageTableVolumeAABBMin + PageCoordToSplat) * SPARSE_VOLUME_TILE_RES * 2;
+				for (int32 AttributesIdx = 0; AttributesIdx < 2; ++AttributesIdx)
+				{
+					if ((AttributesIdx == 0 && PhysicalTileDataA.IsEmpty()) || (AttributesIdx == 1 && PhysicalTileDataB.IsEmpty()))
+					{
+						continue;
+					}
+
+					for (int32 Z = 0; Z < SPARSE_VOLUME_TILE_RES; ++Z)
+					{
+						for (int32 Y = 0; Y < SPARSE_VOLUME_TILE_RES; ++Y)
+						{
+							for (int32 X = 0; X < SPARSE_VOLUME_TILE_RES; ++X)
+							{
+								FVector4f DownsampledValue = FVector4f(0.0f, 0.0f, 0.0f, 0.0f);
+								for (int32_t OffsetIdx = 0; OffsetIdx < 8; ++OffsetIdx)
+								{
+									const FIntVector3 Offset = FIntVector3(OffsetIdx, OffsetIdx >> 1, OffsetIdx >> 2) & 1;
+									const FIntVector3 SourceCoord = ParentVolumeCoordBase + FIntVector3(X, Y, Z) * 2 + Offset;
+									const FVector4f SampleValue = Load(SourceCoord, MipLevel - 1, AttributesIdx);
+									DownsampledValue += SampleValue;
+								}
+								DownsampledValue /= 8.0f;
+
+								const FIntVector3 WriteCoord = FIntVector3(X, Y, Z);
+								WriteTileDataVoxel(DstTileIndex, WriteCoord, AttributesIdx, DownsampledValue);
+							}
+						}
+					}
+				}
+
+				// Set the next tile to be written to
+				++DstTileIndex;
+			}
+			MipInfo.TileCount = NumAllocatedPages;
+		}
+	}
+}
+
+////////////////////////////////////////////////////////////////////////////////////////////////
diff --git a/Engine/Source/Runtime/Engine/Private/SparseVolumeTexture/SparseVolumeTextureSceneProxy.cpp b/Engine/Source/Runtime/Engine/Private/SparseVolumeTexture/SparseVolumeTextureSceneProxy.cpp
new file mode 100644
--- /dev/null
+++ b/Engine/Source/Runtime/Engine/Private/SparseVolumeTexture/SparseVolumeTextureSceneProxy.cpp
@@ -0,0 +1,424 @@
+// Copyright Epic Games, Inc. All Rights Reserved.
+
+#include "SparseVolumeTexture/SparseVolumeTextureSceneProxy.h"
+#include "SparseVolumeTexture/SparseVolumeTextureUtility.h"
+#include "SparseVolumeTexture/SparseVolumeTextureData.h"
+
+DEFINE_LOG_CATEGORY_STATIC(LogSparseVolumeTextureProxy, Log, All);
+
+namespace UE
+{
+namespace SVT
+{
+namespace Private
+{
+	static FIntVector3 AdvanceTileCoord(const FIntVector3& TileCoord, const FIntVector3& TileCoordResolution)
+	{
+		FIntVector3 Result = TileCoord;
+		Result.X++;
+		if (Result.X >= TileCoordResolution.X)
+		{
+			Result.X = 0;
+			Result.Y++;
+		}
+		if (Result.Y >= TileCoordResolution.Y)
+		{
+			Result.Y = 0;
+			Result.Z++;
+		}
+		return Result;
+	}
+
+	static FIntVector3 ComputeTileDataVolumeResolution(int32 NumAllocatedPages)
+	{
+		int32 TileVolumeResolutionCube = 1;
+		while (TileVolumeResolutionCube * TileVolumeResolutionCube * TileVolumeResolutionCube < NumAllocatedPages)
+		{
+			TileVolumeResolutionCube++;				// We use a simple loop to compute the minimum resolution of a cube to store all the tile data
+		}
+		FIntVector3 TileDataVolumeResolution = FIntVector3(TileVolumeResolutionCube, TileVolumeResolutionCube, TileVolumeResolutionCube);
+		while (TileDataVolumeResolution.X * TileDataVolumeResolution.Y * (TileDataVolumeResolution.Z - 1) > NumAllocatedPages)
+		{
+			TileDataVolumeResolution.Z--;	// We then trim an edge to get back space.
+		}
+
+		return TileDataVolumeResolution * SPARSE_VOLUME_TILE_RES;
+	}
+} // Private
+} // SVT
+} // UE
+
+////////////////////////////////////////////////////////////////////////////////////////////////
+
+bool FSparseVolumeTextureRuntime::Create(const FSparseVolumeTextureData& TextureData)
+{
+	bool bSuccess = Create(TextureData.Header, TextureData.Header.MipInfo.Num());
+
+	if (!bSuccess)
+	{
+		return false;
+	}
+
+	const int32 NumMipLevels = TextureData.Header.MipInfo.Num();
+	const int32 TileSizeBytesA = GPixelFormats[Header.AttributesFormats[0]].BlockBytes * UE::SVT::SVTNumVoxelsPerTile;
+	const int32 TileSizeBytesB = GPixelFormats[Header.AttributesFormats[1]].BlockBytes * UE::SVT::SVTNumVoxelsPerTile;
+
+	TArray<FSparseVolumeTextureTileMapping> Mappings;
+	Mappings.Empty(NumMipLevels);
+	for (int32 MipLevel = 0; MipLevel < NumMipLevels; ++MipLevel)
+	{
+		FSparseVolumeTextureTileMapping Mapping;
+		Mapping.TileIndices = TextureData.PageTable[MipLevel].GetData();
+		Mapping.TileDataA = TextureData.PhysicalTileDataA.GetData() + TileSizeBytesA * TextureData.Header.MipInfo[MipLevel].TileOffset;
+		Mapping.TileDataB = TextureData.PhysicalTileDataB.GetData() + TileSizeBytesB * TextureData.Header.MipInfo[MipLevel].TileOffset;
+		Mapping.NumPhysicalTiles = TextureData.Header.MipInfo[MipLevel].TileCount;
+		Mapping.TileIndicesOffset = -TextureData.Header.MipInfo[MipLevel].TileOffset;
+
+		Mappings.Add(Mapping);
+	}
+
+	return SetTileMappings(Mappings);
+}
+
+bool FSparseVolumeTextureRuntime::Create(const FSparseVolumeTextureHeader& SVTHeader, int32 NumMipLevels)
+{
+	using namespace UE::SVT;
+
+	// Check if the requested resolution exceeds hardware limits
+	if (SVTHeader.PageTableVolumeResolution.X > SVTMaxVolumeTextureDim
+		|| SVTHeader.PageTableVolumeResolution.Y > SVTMaxVolumeTextureDim
+		|| SVTHeader.PageTableVolumeResolution.Z > SVTMaxVolumeTextureDim)
+	{
+		UE_LOG(LogSparseVolumeTextureProxy, Warning, TEXT("FSparseVolumeTextureRuntime page table texture dimensions exceed limit (%ix%ix%i): %ix%ix%i"),
+			SVTMaxVolumeTextureDim, SVTMaxVolumeTextureDim, SVTMaxVolumeTextureDim,
+			SVTHeader.PageTableVolumeResolution.X, SVTHeader.PageTableVolumeResolution.Y, SVTHeader.PageTableVolumeResolution.Z);
+		return false;
+	}
+
+	// Check if number of requested mip levels are possible
+	{
+		int32 PossibleLevels = 1;
+		FIntVector3 Resolution = SVTHeader.VirtualVolumeResolution;
+		while (Resolution.X > SPARSE_VOLUME_TILE_RES || Resolution.Y > SPARSE_VOLUME_TILE_RES || Resolution.Z > SPARSE_VOLUME_TILE_RES)
+		{
+			Resolution /= 2;
+			++PossibleLevels;
+		}
+		if (NumMipLevels <= 0 || NumMipLevels > PossibleLevels)
+		{
+			UE_LOG(LogSparseVolumeTextureProxy, Warning, TEXT("Tried to create FSparseVolumeTextureRuntime with more mip levels than possible. Requested %i Possible %i"), NumMipLevels, PossibleLevels);
+			return false;
+		}
+	}
+
+	// Check if formats are supported
+	if (SVTHeader.AttributesFormats[0] == PF_Unknown && SVTHeader.AttributesFormats[1] == PF_Unknown)
+	{
+		UE_LOG(LogSparseVolumeTextureProxy, Warning, TEXT("Tried to create FSparseVolumeTextureRuntime with pixel format == PF_Unknown for both attribute textures."));
+		return false;
+	}
+	for (int32 i = 0; i < 2; ++i)
+	{
+		EPixelFormat Format = SVTHeader.AttributesFormats[i];
+		switch (Format)
+		{
+		case PF_R8:
+		case PF_R8G8:
+		case PF_R8G8B8A8:
+		case PF_R16F:
+		case PF_G16R16F:
+		case PF_FloatRGBA:
+		case PF_R32_FLOAT:
+		case PF_G32R32F:
+		case PF_A32B32G32R32F:
+		case PF_Unknown: // Allow PF_Unknown because the above check guarantees that at least one format is not PF_Unknown.
+			break;
+		default:
+			UE_LOG(LogSparseVolumeTextureProxy, Warning, TEXT("Tried to create FSparseVolumeTextureRuntime with unsupported pixel format. Requested %i"), Format);
+			return false;
+		}
+	}
+
+	*static_cast<FSparseVolumeTextureHeader*>(&Header) = SVTHeader;
+	Header.TileDataVolumeResolution = FIntVector3::ZeroValue;
+	Header.NumMipLevels = NumMipLevels;
+	Header.HighestResidentLevel = INT32_MIN;
+	Header.LowestResidentLevel = INT32_MAX;
+	PageTable.SetNum(NumMipLevels);
+	PhysicalTileDataA.Reset();
+	PhysicalTileDataB.Reset();
+
+	for (int32 MipLevel = 0; MipLevel < NumMipLevels; ++MipLevel)
+	{
+		PageTable[MipLevel].Empty();
+	}
+
+	return true;
+}
+
+bool FSparseVolumeTextureRuntime::SetTileMappings(const TArrayView<const FSparseVolumeTextureTileMapping>& Mappings)
+{
+	using namespace UE::SVT;
+	using namespace UE::SVT::Private;
+
+	const int32 NumMipLevels = Header.NumMipLevels;
+	check(Mappings.Num() == NumMipLevels)
+
+		// Compute number of required tiles
+		int32 NumTiles = 1; // always need a null tile
+	for (const FSparseVolumeTextureTileMapping& Mapping : Mappings)
+	{
+		NumTiles += Mapping.NumPhysicalTiles;
+	}
+
+	FIntVector3 TileDataVolumeRes = ComputeTileDataVolumeResolution(NumTiles);
+	if (TileDataVolumeRes.X > SVTMaxVolumeTextureDim
+		|| TileDataVolumeRes.Y > SVTMaxVolumeTextureDim
+		|| TileDataVolumeRes.Z > SVTMaxVolumeTextureDim)
+	{
+		UE_LOG(LogSparseVolumeTextureProxy, Warning, TEXT("SparseVolumeTexture tile data texture dimensions exceed limit (%ix%ix%i): %ix%ix%i"),
+			SVTMaxVolumeTextureDim, SVTMaxVolumeTextureDim, SVTMaxVolumeTextureDim,
+			TileDataVolumeRes.X, TileDataVolumeRes.Y, TileDataVolumeRes.Z);
+		return false;
+	}
+
+	Header.TileDataVolumeResolution = TileDataVolumeRes;
+	const FIntVector3 TileCoordSpace = TileDataVolumeRes / SPARSE_VOLUME_TILE_RES;
+	check((TileCoordSpace.X * SPARSE_VOLUME_TILE_RES) == TileDataVolumeRes.X
+		&& (TileCoordSpace.Y * SPARSE_VOLUME_TILE_RES) == TileDataVolumeRes.Y
+		&& (TileCoordSpace.Z * SPARSE_VOLUME_TILE_RES) == TileDataVolumeRes.Z);
+
+	// Clear page table
+	for (TArray<uint32>& PageTableMip : PageTable)
+	{
+		PageTableMip.Empty();
+	}
+
+	// Allocate memory for tile data
+	const int32 FormatSize[] = { GPixelFormats[Header.AttributesFormats[0]].BlockBytes, GPixelFormats[Header.AttributesFormats[1]].BlockBytes };
+	PhysicalTileDataA.SetNumZeroed(Header.TileDataVolumeResolution.X * Header.TileDataVolumeResolution.Y * Header.TileDataVolumeResolution.Z * FormatSize[0]);
+	PhysicalTileDataB.SetNumZeroed(Header.TileDataVolumeResolution.X * Header.TileDataVolumeResolution.Y * Header.TileDataVolumeResolution.Z * FormatSize[1]);
+	uint8* PhysicalTileDataPtrs[] = { PhysicalTileDataA.GetData(), PhysicalTileDataB.GetData() };
+
+	int32 NumWrittenTiles = 0;
+	FIntVector3 DstTileCoord = FIntVector3::ZeroValue;
+
+	// Write null tile
+	for (int32 Z = 0; Z < SPARSE_VOLUME_TILE_RES; ++Z)
+	{
+		for (int32 Y = 0; Y < SPARSE_VOLUME_TILE_RES; ++Y)
+		{
+			for (int32 X = 0; X < SPARSE_VOLUME_TILE_RES; ++X)
+			{
+				const int32 VoxelIndex = Z * (SPARSE_VOLUME_TILE_RES * SPARSE_VOLUME_TILE_RES) + Y * SPARSE_VOLUME_TILE_RES + X;
+				WriteVoxel(VoxelIndex, PhysicalTileDataPtrs[0], Header.AttributesFormats[0], Header.NullTileValues[0]);
+				WriteVoxel(VoxelIndex, PhysicalTileDataPtrs[1], Header.AttributesFormats[1], Header.NullTileValues[1]);
+			}
+		}
+	}
+	++NumWrittenTiles;
+	DstTileCoord = AdvanceTileCoord(DstTileCoord, TileCoordSpace);
+
+	// Write page table and physical tiles
+	Header.HighestResidentLevel = INT32_MIN;
+	Header.LowestResidentLevel = INT32_MAX;
+	FIntVector3 PageTableRes = Header.PageTableVolumeResolution * 2;
+	for (int32 MipLevel = 0; MipLevel < NumMipLevels; ++MipLevel)
+	{
+		PageTableRes = FIntVector3(FMath::Max(1, PageTableRes.X / 2), FMath::Max(1, PageTableRes.Y / 2), FMath::Max(1, PageTableRes.Z / 2));
+
+		const FSparseVolumeTextureTileMapping& Mapping = Mappings[MipLevel];
+		if (Mapping.NumPhysicalTiles <= 0)
+		{
+			continue;
+		}
+
+		Header.HighestResidentLevel = FMath::Max(Header.HighestResidentLevel, MipLevel);
+		Header.LowestResidentLevel = FMath::Min(Header.LowestResidentLevel, MipLevel);
+
+		// Write page table
+		PageTable[MipLevel].SetNum(PageTableRes.X * PageTableRes.Y * PageTableRes.Z);
+		uint32* PageTablePtr = PageTable[MipLevel].GetData();
+		const int32 MipTileOffset = NumWrittenTiles;
+		for (int32 PageZ = 0; PageZ < PageTableRes.Z; ++PageZ)
+		{
+			for (int32 PageY = 0; PageY < PageTableRes.Y; ++PageY)
+			{
+				for (int32 PageX = 0; PageX < PageTableRes.X; ++PageX)
+				{
+					const int32 PageIndex = PageZ * (PageTableRes.Y * PageTableRes.X) + PageY * PageTableRes.X + PageX;
+					uint32 TileIndex = Mapping.TileIndices[PageIndex];
+					if (TileIndex == 0) // points to null tile
+					{
+						PageTablePtr[PageIndex] = 0;
+					}
+					else // points to actual physical tile
+					{
+						// Make index relative to the start index of this mip level
+						TileIndex += Mapping.TileIndicesOffset;
+						TileIndex += MipTileOffset;
+
+						FIntVector3 TileCoord;
+						TileCoord.X = TileIndex % TileCoordSpace.X;
+						TileCoord.Z = TileIndex / (TileCoordSpace.Y * TileCoordSpace.X);
+						TileCoord.Y = (TileIndex - (TileCoord.Z * (TileCoordSpace.Y * TileCoordSpace.X))) / TileCoordSpace.X;
+
+						// Write to page table
+						PageTablePtr[PageIndex] = PackPageTableEntry(TileCoord);
+					}
+				}
+			}
+		}
+
+		// Write to tile data
+		for (int32 PhysicalTileIndex = 0; PhysicalTileIndex < Mapping.NumPhysicalTiles; ++PhysicalTileIndex)
+		{
+			for (int32 Z = 0; Z < SPARSE_VOLUME_TILE_RES; ++Z)
+			{
+				for (int32 Y = 0; Y < SPARSE_VOLUME_TILE_RES; ++Y)
+				{
+					for (int32 X = 0; X < SPARSE_VOLUME_TILE_RES; ++X)
+					{
+						const int32 SrcVoxelIndex = PhysicalTileIndex * SVTNumVoxelsPerTile + Z * (SPARSE_VOLUME_TILE_RES * SPARSE_VOLUME_TILE_RES) + Y * SPARSE_VOLUME_TILE_RES + X;
+						const FVector4f ValueA = ReadVoxel(SrcVoxelIndex, Mapping.TileDataA, Header.AttributesFormats[0]);
+						const FVector4f ValueB = ReadVoxel(SrcVoxelIndex, Mapping.TileDataB, Header.AttributesFormats[1]);
+
+						const FIntVector3 VoxelCoord = DstTileCoord * SPARSE_VOLUME_TILE_RES + FIntVector3(X, Y, Z);
+						const int32 DstVoxelIndex = VoxelCoord.Z * (Header.TileDataVolumeResolution.Y * Header.TileDataVolumeResolution.X) + VoxelCoord.Y * Header.TileDataVolumeResolution.X + VoxelCoord.X;
+						WriteVoxel(DstVoxelIndex, PhysicalTileDataPtrs[0], Header.AttributesFormats[0], ValueA);
+						WriteVoxel(DstVoxelIndex, PhysicalTileDataPtrs[1], Header.AttributesFormats[1], ValueB);
+					}
+				}
+			}
+			DstTileCoord = AdvanceTileCoord(DstTileCoord, TileCoordSpace);
+		}
+		NumWrittenTiles += Mapping.NumPhysicalTiles;
+	}
+	return true;
+}
+
+void FSparseVolumeTextureRuntime::SetAsDefaultTexture()
+{
+	Header.PageTableVolumeResolution = FIntVector3(1, 1, 1);
+	Header.TileDataVolumeResolution = FIntVector3(1, 1, 1);
+	Header.NumMipLevels = 1;
+	Header.HighestResidentLevel = 0;
+	Header.LowestResidentLevel = 0;
+	PageTable.SetNum(1);
+	PageTable[0].SetNumZeroed(1);
+	const int32 FormatSize[] = { GPixelFormats[Header.AttributesFormats[0]].BlockBytes, GPixelFormats[Header.AttributesFormats[1]].BlockBytes };
+	PhysicalTileDataA.SetNumZeroed(FormatSize[0]);
+	PhysicalTileDataB.SetNumZeroed(FormatSize[1]);
+	UE::SVT::WriteVoxel(0, PhysicalTileDataA.GetData(), Header.AttributesFormats[0], Header.NullTileValues[0]);
+	UE::SVT::WriteVoxel(0, PhysicalTileDataB.GetData(), Header.AttributesFormats[1], Header.NullTileValues[1]);
+}
+
+////////////////////////////////////////////////////////////////////////////////////////////////
+
+FSparseVolumeTextureSceneProxy::FSparseVolumeTextureSceneProxy()
+	: FRenderResource()
+	, SparseVolumeTextureRuntime()
+	, PageTableTextureRHI(nullptr)
+	, PhysicalTileDataATextureRHI(nullptr)
+	, PhysicalTileDataBTextureRHI(nullptr)
+{
+}
+
+FSparseVolumeTextureSceneProxy::~FSparseVolumeTextureSceneProxy()
+{
+}
+
+void FSparseVolumeTextureSceneProxy::GetMemorySize(SIZE_T* SizeCPU, SIZE_T* SizeGPU) const
+{
+	*SizeCPU += sizeof(FSparseVolumeTextureSceneProxy);
+	*SizeCPU += SparseVolumeTextureRuntime.PageTable.GetAllocatedSize();
+	*SizeCPU += SparseVolumeTextureRuntime.PhysicalTileDataA.GetAllocatedSize();
+	*SizeCPU += SparseVolumeTextureRuntime.PhysicalTileDataB.GetAllocatedSize();
+
+#if RHI_ENABLE_RESOURCE_INFO
+	FRHIResourceInfo ResourceInfo;
+	if (PageTableTextureRHI && PageTableTextureRHI->GetResourceInfo(ResourceInfo))
+	{
+		*SizeGPU += ResourceInfo.VRamAllocation.AllocationSize;
+	}
+	if (PhysicalTileDataATextureRHI && PhysicalTileDataATextureRHI->GetResourceInfo(ResourceInfo))
+	{
+		*SizeGPU += ResourceInfo.VRamAllocation.AllocationSize;
+	}
+	if (PhysicalTileDataBTextureRHI && PhysicalTileDataBTextureRHI->GetResourceInfo(ResourceInfo))
+	{
+		*SizeGPU += ResourceInfo.VRamAllocation.AllocationSize;
+	}
+#endif
+}
+
+void FSparseVolumeTextureSceneProxy::InitRHI()
+{
+	// Can't create a proxy without any resident mips
+	check(SparseVolumeTextureRuntime.Header.HighestResidentLevel >= SparseVolumeTextureRuntime.Header.LowestResidentLevel);
+
+	// Page table
+	{
+		const int32 NumResidentMipLevels = SparseVolumeTextureRuntime.Header.HighestResidentLevel - SparseVolumeTextureRuntime.Header.LowestResidentLevel + 1;
+		FIntVector3 PageTableResolution = SparseVolumeTextureRuntime.Header.PageTableVolumeResolution / (1 << SparseVolumeTextureRuntime.Header.LowestResidentLevel);
+		PageTableResolution = FIntVector3(FMath::Max(1, PageTableResolution.X), FMath::Max(1, PageTableResolution.Y), FMath::Max(1, PageTableResolution.Z));
+
+		const EPixelFormat PageEntryFormat = PF_R32_UINT;
+		const FRHITextureCreateDesc Desc =
+			FRHITextureCreateDesc::Create3D(TEXT("SparseVolumeTexture.PageTable.RHITexture"), PageTableResolution.X, PageTableResolution.Y, PageTableResolution.Z, PageEntryFormat)
+			.SetFlags(ETextureCreateFlags::ShaderResource)
+			.SetNumMips((uint8)NumResidentMipLevels);
+
+		PageTableTextureRHI = RHICreateTexture(Desc);
+
+		const int32 FormatSize = GPixelFormats[PageEntryFormat].BlockBytes;
+		FIntVector3 PageTableMipRes = PageTableResolution;
+		for (int32 RelativeMipLevel = 0; RelativeMipLevel < NumResidentMipLevels; ++RelativeMipLevel)
+		{
+			const FUpdateTextureRegion3D UpdateRegion(0, 0, 0, 0, 0, 0, PageTableMipRes.X, PageTableMipRes.Y, PageTableMipRes.Z);
+			const uint8* PageTableData = (const uint8*)SparseVolumeTextureRuntime.PageTable[RelativeMipLevel + SparseVolumeTextureRuntime.Header.LowestResidentLevel].GetData();
+			RHIUpdateTexture3D(PageTableTextureRHI, (uint32)RelativeMipLevel, UpdateRegion, UpdateRegion.Width * FormatSize, UpdateRegion.Width * UpdateRegion.Height * FormatSize, PageTableData);
+			PageTableMipRes = FIntVector3(FMath::Max(1, PageTableMipRes.X / 2), FMath::Max(1, PageTableMipRes.Y / 2), FMath::Max(1, PageTableMipRes.Z / 2));
+		}
+	}
+
+	// Tile data
+	{
+		const FIntVector3 TileDataVolumeResolution = SparseVolumeTextureRuntime.Header.TileDataVolumeResolution;
+		const EPixelFormat VoxelFormatA = SparseVolumeTextureRuntime.Header.AttributesFormats[0];
+		const EPixelFormat VoxelFormatB = SparseVolumeTextureRuntime.Header.AttributesFormats[1];
+		const FUpdateTextureRegion3D UpdateRegion(0, 0, 0, 0, 0, 0, TileDataVolumeResolution.X, TileDataVolumeResolution.Y, TileDataVolumeResolution.Z);
+
+		// A
+		if (VoxelFormatA != PF_Unknown)
+		{
+			const FRHITextureCreateDesc Desc =
+				FRHITextureCreateDesc::Create3D(TEXT("SparseVolumeTexture.PhysicalTileDataA.RHITexture"), TileDataVolumeResolution.X, TileDataVolumeResolution.Y, TileDataVolumeResolution.Z, VoxelFormatA)
+				.SetFlags(ETextureCreateFlags::ShaderResource);
+
+			PhysicalTileDataATextureRHI = RHICreateTexture(Desc);
+
+			const int32 FormatSize = GPixelFormats[VoxelFormatA].BlockBytes;
+			RHIUpdateTexture3D(PhysicalTileDataATextureRHI, 0, UpdateRegion, UpdateRegion.Width * FormatSize, UpdateRegion.Width * UpdateRegion.Height * FormatSize, (const uint8*)SparseVolumeTextureRuntime.PhysicalTileDataA.GetData());
+		}
+
+		// B
+		if (VoxelFormatB != PF_Unknown)
+		{
+			const FRHITextureCreateDesc Desc =
+				FRHITextureCreateDesc::Create3D(TEXT("SparseVolumeTexture.PhysicalTileDataB.RHITexture"), TileDataVolumeResolution.X, TileDataVolumeResolution.Y, TileDataVolumeResolution.Z, VoxelFormatB)
+				.SetFlags(ETextureCreateFlags::ShaderResource);
+
+			PhysicalTileDataBTextureRHI = RHICreateTexture(Desc);
+
+			const int32 FormatSize = GPixelFormats[VoxelFormatB].BlockBytes;
+			RHIUpdateTexture3D(PhysicalTileDataBTextureRHI, 0, UpdateRegion, UpdateRegion.Width * FormatSize, UpdateRegion.Width * UpdateRegion.Height * FormatSize, (const uint8*)SparseVolumeTextureRuntime.PhysicalTileDataB.GetData());
+		}
+	}
+}
+
+void FSparseVolumeTextureSceneProxy::ReleaseRHI()
+{
+}
+
+////////////////////////////////////////////////////////////////////////////////////////////////
\ No newline at end of file
diff --git a/Engine/Source/Runtime/Engine/Private/SparseVolumeTexture/SparseVolumeTextureStreaming.cpp b/Engine/Source/Runtime/Engine/Private/SparseVolumeTexture/SparseVolumeTextureStreaming.cpp
--- a/Engine/Source/Runtime/Engine/Private/SparseVolumeTexture/SparseVolumeTextureStreaming.cpp
+++ b/Engine/Source/Runtime/Engine/Private/SparseVolumeTexture/SparseVolumeTextureStreaming.cpp
@@ -2,6 +2,8 @@
 
 #include "SparseVolumeTexture/SparseVolumeTextureStreaming.h"
 #include "SparseVolumeTexture/SparseVolumeTexture.h"
+#include "SparseVolumeTexture/SparseVolumeTextureSceneProxy.h"
+#include "SparseVolumeTexture/SparseVolumeTextureData.h"
 #include "RenderingThread.h"
 
 #include "Serialization/LargeMemoryReader.h"
@@ -77,9 +79,8 @@ bool FStreamingSparseVolumeTextureData::Initialize(UStreamableSparseVolumeTextur
 	check(LoadedChunks.IsEmpty());
 	check(LoadedChunkIndices.IsEmpty());
 	check(!InSparseVolumeTexture->GetFrames().IsEmpty()); // Must hold at least the first frame
-	check(!InSparseVolumeTexture->GetFrames()[0].IsEmpty()); // Must hold at least the first frame
 
-	AddNewLoadedChunk(0, InSparseVolumeTexture->GetFrames()[0][0].SparseVolumeTextureSceneProxy);
+	AddNewLoadedChunk(0, InSparseVolumeTexture->GetFrames()[0].SparseVolumeTextureSceneProxy);
 	LoadedChunkIndices.Add(0);
 
 	return true;
@@ -239,15 +240,13 @@ void FStreamingSparseVolumeTextureData::BeginPendingRequests(const TArray<int32>
 
 	// Set off all IO Requests
 	const int32 NumFrames = SparseVolumeTexture->GetNumFrames();
-	const int32 NumMipLevels = SparseVolumeTexture->GetNumMipLevels();
-	check(NumFrames > 0 && NumMipLevels > 0);
+	check(NumFrames > 0);
 	const EAsyncIOPriorityAndFlags AsyncIOPriority = AIOP_CriticalPath; //Set to Crit temporarily as emergency speculative fix for streaming issue
-	TArrayView<const FSparseVolumeTextureFrameMips> SVTFrames = SparseVolumeTexture->GetFrames();
+	TArrayView<const FSparseVolumeTextureFrame> SVTFrames = SparseVolumeTexture->GetFrames();
 	for (int32 IndexToLoad : IndicesToLoad)
 	{
 		const int32 FrameToLoad = SVTChunkIndexToFrame(IndexToLoad, NumFrames) % NumFrames;
-		const int32 MipLevelToLoad = FMath::Clamp(SVTChunkIndexToMipLevel(IndexToLoad, NumFrames), 0, NumMipLevels - 1);
-		const FSparseVolumeTextureFrame& Frame = SVTFrames[FrameToLoad][MipLevelToLoad];
+		const FSparseVolumeTextureFrame& Frame = SVTFrames[FrameToLoad];
 		FSparseVolumeTextureSceneProxy* ExistingProxy = Frame.SparseVolumeTextureSceneProxy;
 		FLoadedSparseVolumeTextureChunk& ChunkStorage = AddNewLoadedChunk(IndexToLoad, ExistingProxy);
 
@@ -459,6 +458,7 @@ void FSparseVolumeTextureStreamingManager::GetMemorySizeForSparseVolumeTexture(c
 
 const FSparseVolumeTextureSceneProxy* FSparseVolumeTextureStreamingManager::GetSparseVolumeTextureSceneProxy(const UStreamableSparseVolumeTexture* SparseVolumeTexture, int32 FrameIndex, int32 MipLevel, bool bTrackAsRequested)
 {
+	MipLevel = 0; // Currently streaming all mips at once
 	FScopeLock Lock(&CriticalSection);
 
 	FStreamingSparseVolumeTextureData* StreamingData = StreamingSparseVolumeTextures.FindRef(SparseVolumeTexture);
@@ -529,7 +529,9 @@ void FSparseVolumeTextureStreamingManager::OnAsyncFileCallback(FStreamingSparseV
 		FMemoryReaderView Reader(MemView);
 
 		FSparseVolumeTextureSceneProxy* NewProxy = new FSparseVolumeTextureSceneProxy();
-		NewProxy->GetRuntimeData().Serialize(Reader);
+		FSparseVolumeTextureData TextureData;
+		TextureData.Serialize(Reader);
+		NewProxy->GetRuntimeData().Create(TextureData);
 		BeginInitResource(NewProxy);
 
 		ChunkStorage.Proxy = NewProxy;
diff --git a/Engine/Source/Runtime/Engine/Private/SparseVolumeTexture/SparseVolumeTextureUtility.cpp b/Engine/Source/Runtime/Engine/Private/SparseVolumeTexture/SparseVolumeTextureUtility.cpp
new file mode 100644
--- /dev/null
+++ b/Engine/Source/Runtime/Engine/Private/SparseVolumeTexture/SparseVolumeTextureUtility.cpp
@@ -0,0 +1,122 @@
+// Copyright Epic Games, Inc. All Rights Reserved.
+
+#include "SparseVolumeTextureUtility.h"
+
+namespace UE
+{
+namespace SVT
+{
+namespace Private
+{
+	float F16ToF32(uint16 Packed)
+	{
+		FFloat16 F16{};
+		F16.Encoded = Packed;
+		return F16.GetFloat();
+	};
+} // Private
+} // SVT
+} // UE
+
+uint32 UE::SVT::PackPageTableEntry(const FIntVector3& Coord)
+{
+	// A page encodes the physical tile coord as unsigned int of 11 11 10 bits
+	// This means a page coord cannot be larger than 2047 for x and y and 1023 for z
+	// which mean we cannot have more than 2048*2048*1024 = 4 Giga tiles of 16^3 tiles.
+	uint32 Result = (Coord.X & 0x7FF) | ((Coord.Y & 0x7FF) << 11) | ((Coord.Z & 0x3FF) << 22);
+	return Result;
+}
+
+FIntVector3 UE::SVT::UnpackPageTableEntry(uint32 Packed)
+{
+	FIntVector3 Result;
+	Result.X = Packed & 0x7FF;
+	Result.Y = (Packed >> 11) & 0x7FF;
+	Result.Z = (Packed >> 22) & 0x3FF;
+	return Result;
+}
+
+FVector4f UE::SVT::ReadVoxel(int32 VoxelIndex, const uint8* TileData, EPixelFormat Format)
+{
+	using namespace UE::SVT::Private;
+	if (Format == PF_Unknown)
+	{
+		return FVector4f();
+	}
+	switch (Format)
+	{
+	case PF_R8:
+		return FVector4f(TileData[VoxelIndex] / 255.0f, 0.0f, 0.0f, 0.0f);
+	case PF_R8G8:
+		return FVector4f(TileData[VoxelIndex * 2 + 0] / 255.0f, TileData[VoxelIndex * 2 + 1] / 255.0f, 0.0f, 0.0f);
+	case PF_R8G8B8A8:
+		return FVector4f(TileData[VoxelIndex * 4 + 0] / 255.0f, TileData[VoxelIndex * 4 + 1] / 255.0f, TileData[VoxelIndex * 4 + 2] / 255.0f, TileData[VoxelIndex * 4 + 3] / 255.0f);
+	case PF_R16F:
+		return FVector4f(F16ToF32(((const uint16*)TileData)[VoxelIndex]), 0.0f, 0.0f, 0.0f);
+	case PF_G16R16F:
+		return FVector4f(F16ToF32(((const uint16*)TileData)[VoxelIndex * 2 + 0]), F16ToF32(((const uint16*)TileData)[VoxelIndex * 2 + 1]), 0.0f, 0.0f);
+	case PF_FloatRGBA:
+		return FVector4f(F16ToF32(((const uint16*)TileData)[VoxelIndex * 4 + 0]), F16ToF32(((const uint16*)TileData)[VoxelIndex * 4 + 1]), F16ToF32(((const uint16*)TileData)[VoxelIndex * 4 + 2]), F16ToF32(((const uint16*)TileData)[VoxelIndex * 4 + 3]));
+	case PF_R32_FLOAT:
+		return FVector4f(((const float*)TileData)[VoxelIndex], 0.0f, 0.0f, 0.0f);
+	case PF_G32R32F:
+		return FVector4f(((const float*)TileData)[VoxelIndex * 2 + 0], ((const float*)TileData)[VoxelIndex * 2 + 1], 0.0f, 0.0f);
+	case PF_A32B32G32R32F:
+		return FVector4f(((const float*)TileData)[VoxelIndex * 4 + 0], ((const float*)TileData)[VoxelIndex * 4 + 1], ((const float*)TileData)[VoxelIndex * 4 + 2], ((const float*)TileData)[VoxelIndex * 4 + 3]);
+	default:
+		checkNoEntry();
+		return FVector4f();
+	}
+}
+
+void UE::SVT::WriteVoxel(int32 VoxelIndex, uint8* TileData, EPixelFormat Format, const FVector4f& Value, int32 DstComponent)
+{
+	if (Format == PF_Unknown)
+	{
+		return;
+	}
+	switch (Format)
+	{
+	case PF_R8:
+		if (DstComponent == -1 || DstComponent == 0) TileData[VoxelIndex] = uint8(FMath::Clamp(Value.X, 0.0f, 1.0f) * 255.0f);
+		break;
+	case PF_R8G8:
+		if (DstComponent == -1 || DstComponent == 0) TileData[VoxelIndex * 2 + 0] = uint8(FMath::Clamp(Value.X, 0.0f, 1.0f) * 255.0f);
+		if (DstComponent == -1 || DstComponent == 1) TileData[VoxelIndex * 2 + 1] = uint8(FMath::Clamp(Value.Y, 0.0f, 1.0f) * 255.0f);
+		break;
+	case PF_R8G8B8A8:
+		if (DstComponent == -1 || DstComponent == 0) TileData[VoxelIndex * 4 + 0] = uint8(FMath::Clamp(Value.X, 0.0f, 1.0f) * 255.0f);
+		if (DstComponent == -1 || DstComponent == 1) TileData[VoxelIndex * 4 + 1] = uint8(FMath::Clamp(Value.Y, 0.0f, 1.0f) * 255.0f);
+		if (DstComponent == -1 || DstComponent == 2) TileData[VoxelIndex * 4 + 2] = uint8(FMath::Clamp(Value.Z, 0.0f, 1.0f) * 255.0f);
+		if (DstComponent == -1 || DstComponent == 3) TileData[VoxelIndex * 4 + 3] = uint8(FMath::Clamp(Value.W, 0.0f, 1.0f) * 255.0f);
+		break;
+	case PF_R16F:
+		if (DstComponent == -1 || DstComponent == 0) ((uint16*)TileData)[VoxelIndex] = FFloat16(Value.X).Encoded;
+		break;
+	case PF_G16R16F:
+		if (DstComponent == -1 || DstComponent == 0) ((uint16*)TileData)[VoxelIndex * 2 + 0] = FFloat16(Value.X).Encoded;
+		if (DstComponent == -1 || DstComponent == 1) ((uint16*)TileData)[VoxelIndex * 2 + 1] = FFloat16(Value.Y).Encoded;
+		break;
+	case PF_FloatRGBA:
+		if (DstComponent == -1 || DstComponent == 0) ((uint16*)TileData)[VoxelIndex * 4 + 0] = FFloat16(Value.X).Encoded;
+		if (DstComponent == -1 || DstComponent == 1) ((uint16*)TileData)[VoxelIndex * 4 + 1] = FFloat16(Value.Y).Encoded;
+		if (DstComponent == -1 || DstComponent == 2) ((uint16*)TileData)[VoxelIndex * 4 + 2] = FFloat16(Value.Z).Encoded;
+		if (DstComponent == -1 || DstComponent == 3) ((uint16*)TileData)[VoxelIndex * 4 + 3] = FFloat16(Value.W).Encoded;
+		break;
+	case PF_R32_FLOAT:
+		if (DstComponent == -1 || DstComponent == 0) ((float*)TileData)[VoxelIndex] = Value.X;
+		break;
+	case PF_G32R32F:
+		if (DstComponent == -1 || DstComponent == 0) ((float*)TileData)[VoxelIndex * 2 + 0] = Value.X;
+		if (DstComponent == -1 || DstComponent == 1) ((float*)TileData)[VoxelIndex * 2 + 1] = Value.Y;
+		break;
+	case PF_A32B32G32R32F:
+		if (DstComponent == -1 || DstComponent == 0) ((float*)TileData)[VoxelIndex * 4 + 0] = Value.X;
+		if (DstComponent == -1 || DstComponent == 1) ((float*)TileData)[VoxelIndex * 4 + 1] = Value.Y;
+		if (DstComponent == -1 || DstComponent == 2) ((float*)TileData)[VoxelIndex * 4 + 2] = Value.Z;
+		if (DstComponent == -1 || DstComponent == 3) ((float*)TileData)[VoxelIndex * 4 + 3] = Value.W;
+		break;
+	default:
+		checkNoEntry();
+	}
+}
\ No newline at end of file
diff --git a/Engine/Source/Runtime/Engine/Private/SparseVolumeTexture/SparseVolumeTextureUtility.h b/Engine/Source/Runtime/Engine/Private/SparseVolumeTexture/SparseVolumeTextureUtility.h
new file mode 100644
--- /dev/null
+++ b/Engine/Source/Runtime/Engine/Private/SparseVolumeTexture/SparseVolumeTextureUtility.h
@@ -0,0 +1,38 @@
+// Copyright Epic Games, Inc. All Rights Reserved.
+
+#pragma once
+
+#include "CoreMinimal.h"
+#include "SparseVolumeTexture/SparseVolumeTexture.h"
+
+namespace UE
+{
+namespace SVT
+{
+	// We are using this instead of GMaxVolumeTextureDimensions to be independent of the platform that
+	// the asset is imported on. 2048 should be a safe value that should be supported by all our platforms.
+	static constexpr int32 SVTMaxVolumeTextureDim = 2048;
+	static constexpr int32 SVTNumVoxelsPerTile = SPARSE_VOLUME_TILE_RES * SPARSE_VOLUME_TILE_RES * SPARSE_VOLUME_TILE_RES;
+
+	uint32 PackPageTableEntry(const FIntVector3& Coord);
+	FIntVector3 UnpackPageTableEntry(uint32 Packed);
+	FVector4f ReadVoxel(int32 VoxelIndex, const uint8* TileData, EPixelFormat Format);
+	void WriteVoxel(int32 VoxelIndex, uint8* TileData, EPixelFormat Format, const FVector4f& Value, int32 DstComponent = -1);
+
+namespace Private
+{
+	// SVT_TODO: This really should be a shared function.
+	template <typename Y, typename T>
+	void SerializeEnumAs(FArchive& Ar, T& Target)
+	{
+		Y Buffer = static_cast<Y>(Target);
+		Ar << Buffer;
+		if (Ar.IsLoading())
+		{
+			Target = static_cast<T>(Buffer);
+		}
+	}
+} // Private
+
+} // SVT
+} // UE
diff --git a/Engine/Source/Runtime/Engine/Public/MaterialCompiler.h b/Engine/Source/Runtime/Engine/Public/MaterialCompiler.h
--- a/Engine/Source/Runtime/Engine/Public/MaterialCompiler.h
+++ b/Engine/Source/Runtime/Engine/Public/MaterialCompiler.h
@@ -428,7 +428,7 @@ public:
 	 * @param UVWIndex The UVW coordinate at which to sample the texture.
 	 * @return The code chunk index of the result of the texture sample.
 	 */
-	virtual int32 SparseVolumeTextureSamplePageTable(int32 SparseVolumeTextureIndex, int32 UVWIndex) = 0;
+	virtual int32 SparseVolumeTextureSamplePageTable(int32 SparseVolumeTextureIndex, int32 UVWIndex, int32 MipLevelIndex) = 0;
 	/**
 	 * Sample a sparse volume texture physical tile data texture.
 	 * @param SparseVolumeTextureIndex The code chunk index of the texture.
@@ -923,7 +923,7 @@ public:
 	virtual int32 SparseVolumeTextureParameter(FName ParameterName, USparseVolumeTexture* InDefaultTexture, int32& TextureReferenceIndex, EMaterialSamplerType SamplerType)									{ return Compiler->SparseVolumeTextureParameter(ParameterName, InDefaultTexture, TextureReferenceIndex, SamplerType); }
 	virtual int32 SparseVolumeTextureUniform(int32 TextureIndex, int32 VectorIndex, UE::Shader::EValueType Type) override																					{ return Compiler->SparseVolumeTextureUniform(TextureIndex, VectorIndex, Type); }
 	virtual int32 SparseVolumeTextureUniformParameter(FName ParameterName, int32 TextureIndex, int32 VectorIndex, UE::Shader::EValueType Type) override														{ return Compiler->SparseVolumeTextureUniformParameter(ParameterName, TextureIndex, VectorIndex, Type); }
-	virtual int32 SparseVolumeTextureSamplePageTable(int32 SparseVolumeTextureIndex, int32 UVWIndex) override																								{ return Compiler->SparseVolumeTextureSamplePageTable(SparseVolumeTextureIndex, UVWIndex); }
+	virtual int32 SparseVolumeTextureSamplePageTable(int32 SparseVolumeTextureIndex, int32 UVWIndex, int32 MipLevelIndex) override																			{ return Compiler->SparseVolumeTextureSamplePageTable(SparseVolumeTextureIndex, UVWIndex, MipLevelIndex); }
 	virtual int32 SparseVolumeTextureSamplePhysicalTileData(int32 SparseVolumeTextureIndex, int32 VoxelCoordIndex, int32 PhysicalTileDataIdxIndex) override													{ return Compiler->SparseVolumeTextureSamplePhysicalTileData(SparseVolumeTextureIndex, VoxelCoordIndex, PhysicalTileDataIdxIndex); }
 
 	virtual UObject* GetReferencedTexture(int32 Index) override { return Compiler->GetReferencedTexture(Index); }
diff --git a/Engine/Source/Runtime/Engine/Public/SparseVolumeTexture/SparseVolumeTextureData.h b/Engine/Source/Runtime/Engine/Public/SparseVolumeTexture/SparseVolumeTextureData.h
new file mode 100644
--- /dev/null
+++ b/Engine/Source/Runtime/Engine/Public/SparseVolumeTexture/SparseVolumeTextureData.h
@@ -0,0 +1,63 @@
+// Copyright Epic Games, Inc. All Rights Reserved.
+
+#pragma once
+
+#include "CoreMinimal.h"
+#include "SparseVolumeTexture/SparseVolumeTexture.h"
+
+class ENGINE_API ISparseVolumeTextureDataConstructionAdapter
+{
+public:
+	struct ENGINE_API FAttributesInfo
+	{
+		EPixelFormat Format;
+		FVector4f FallbackValue;
+		FVector4f NormalizeScale;
+		FVector4f NormalizeBias;
+		bool bNormalized;
+	};
+
+	virtual void GetAttributesInfo(FAttributesInfo& OutInfoA, FAttributesInfo& OutInfoB) const = 0;
+	virtual FIntVector3 GetAABBMin() const = 0;
+	virtual FIntVector3 GetAABBMax() const = 0;
+	virtual FIntVector3 GetResolution() const = 0;
+	virtual void IteratePhysicalSource(TFunctionRef<void(const FIntVector3& Coord, int32 AttributesIdx, int32 ComponentIdx, float VoxelValue)> OnVisit) const = 0;
+	virtual ~ISparseVolumeTextureDataConstructionAdapter() = default;
+};
+
+struct ENGINE_API FSparseVolumeTextureMipInfo
+{
+	int32 TileOffset = 0;
+	int32 TileCount = 0;
+};
+
+struct ENGINE_API FSparseVolumeTextureDataHeader : public FSparseVolumeTextureHeader
+{
+	static const uint32 kVersion = 0; // The current data format version for the header.
+	uint32 Version = kVersion; // This version can be used to convert existing header to new version later.
+	TArray<FSparseVolumeTextureMipInfo> MipInfo;
+
+	void Serialize(FArchive& Ar);
+};
+
+// Holds the data for a SparseVolumeTexture that is stored on disk. It only has a single mip after importing a source asset. The mip chain is built during cooking.
+// Tiles are addressed by a flat index; unlike the runtime representation, this one stores all tiles in a 1D array and doesn't have the concept of a 3D physical tile texture.
+// The page table itself is 3D though.
+struct ENGINE_API FSparseVolumeTextureData
+{
+	static const uint32 kVersion = 0; // The current data format version for the raw source data.
+	uint32 Version = kVersion; // This version can be used to convert existing source data to new version later.
+
+	FSparseVolumeTextureDataHeader Header = {};
+	TArray<TArray<uint32>> PageTable; // PageTable[MipLevel][PageCoord]
+	TArray<uint8> PhysicalTileDataA;
+	TArray<uint8> PhysicalTileDataB;
+
+	void Serialize(FArchive& Ar);
+	bool Construct(const ISparseVolumeTextureDataConstructionAdapter& Adapter);
+	uint32 ReadPageTable(const FIntVector3& PageTableCoord, int32 MipLevel) const;
+	FVector4f ReadTileDataVoxel(int32 TileIndex, const FIntVector3& TileDataCoord, int32 AttributesIdx) const;
+	FVector4f Load(const FIntVector3& VolumeCoord, int32 MipLevel, int32 AttributesIdx) const;
+	void WriteTileDataVoxel(int32 TileIndex, const FIntVector3& TileDataCoord, int32 AttributesIdx, const FVector4f& Value, int32 DstComponent = -1);
+	void GenerateMipMaps(int32 NumMipLevels = -1);
+};
diff --git a/Engine/Source/Runtime/Engine/Public/SparseVolumeTexture/SparseVolumeTextureSceneProxy.h b/Engine/Source/Runtime/Engine/Public/SparseVolumeTexture/SparseVolumeTextureSceneProxy.h
new file mode 100644
--- /dev/null
+++ b/Engine/Source/Runtime/Engine/Public/SparseVolumeTexture/SparseVolumeTextureSceneProxy.h
@@ -0,0 +1,67 @@
+// Copyright Epic Games, Inc. All Rights Reserved.
+
+#pragma once
+
+#include "CoreMinimal.h"
+#include "SparseVolumeTexture/SparseVolumeTexture.h"
+
+struct FSparseVolumeTextureData;
+
+struct FSparseVolumeTextureTileMapping
+{
+	const uint32* TileIndices; // Indices to tiles stored in TileDataA and TileDataB
+	const uint8* TileDataA;
+	const uint8* TileDataB;
+	int32 NumPhysicalTiles; // Number of tiles stored in TileDataA and TileDataB
+	int32 TileIndicesOffset; // Offset to apply to each element in TileIndices
+};
+
+struct FSparseVolumeTextureRuntimeHeader : public FSparseVolumeTextureHeader
+{
+	FIntVector3 TileDataVolumeResolution = FIntVector3::ZeroValue;
+	int32 NumMipLevels = 0;
+	int32 HighestResidentLevel = INT32_MIN;
+	int32 LowestResidentLevel = INT32_MAX;
+};
+
+// The structure represent the runtime data cooked runtime data.
+struct FSparseVolumeTextureRuntime
+{
+	FSparseVolumeTextureRuntimeHeader Header;
+	TArray<TArray<uint32>> PageTable;
+	TArray<uint8> PhysicalTileDataA;
+	TArray<uint8> PhysicalTileDataB;
+
+	bool Create(const FSparseVolumeTextureData& TextureData);
+	bool Create(const FSparseVolumeTextureHeader& SVTHeader, int32 NumMipLevels);
+	/* Sets the mapped physical tile data for this SVT. The array must have one entry per mip level. */
+	bool SetTileMappings(const TArrayView<const FSparseVolumeTextureTileMapping>& Mappings);
+	void SetAsDefaultTexture();
+};
+
+class FSparseVolumeTextureSceneProxy : public FRenderResource
+{
+public:
+
+	FSparseVolumeTextureSceneProxy();
+	virtual ~FSparseVolumeTextureSceneProxy() override;
+
+	FSparseVolumeTextureRuntime& GetRuntimeData() { return SparseVolumeTextureRuntime; }
+	const FSparseVolumeTextureRuntime& GetRuntimeData() const { return SparseVolumeTextureRuntime; }
+	const FSparseVolumeTextureRuntimeHeader& GetHeader() const { return SparseVolumeTextureRuntime.Header; }
+	FTextureRHIRef GetPhysicalTileDataATextureRHI() const { return PhysicalTileDataATextureRHI; }
+	FTextureRHIRef GetPhysicalTileDataBTextureRHI() const { return PhysicalTileDataBTextureRHI; }
+	FTextureRHIRef GetPageTableTextureRHI() const { return PageTableTextureRHI; }
+	void GetMemorySize(SIZE_T* SizeCPU, SIZE_T* SizeGPU) const;
+
+	virtual void InitRHI() override;
+	virtual void ReleaseRHI() override;
+
+private:
+
+	FSparseVolumeTextureRuntime			SparseVolumeTextureRuntime;
+
+	FTextureRHIRef						PageTableTextureRHI;
+	FTextureRHIRef						PhysicalTileDataATextureRHI;
+	FTextureRHIRef						PhysicalTileDataBTextureRHI;
+};
diff --git a/Engine/Source/Runtime/Renderer/Private/SparseVolumeTexture/SparseVolumeTextureViewerComponent.cpp b/Engine/Source/Runtime/Renderer/Private/SparseVolumeTexture/SparseVolumeTextureViewerComponent.cpp
--- a/Engine/Source/Runtime/Renderer/Private/SparseVolumeTexture/SparseVolumeTextureViewerComponent.cpp
+++ b/Engine/Source/Runtime/Renderer/Private/SparseVolumeTexture/SparseVolumeTextureViewerComponent.cpp
@@ -154,18 +154,15 @@ void USparseVolumeTextureViewerComponent::SendRenderTransformCommand()
 	if (SparseVolumeTextureViewerSceneProxy)
 	{
 		FVector VolumeExtent = FVector(SVTViewerDefaultVolumeExtent);
-		// In order to keep the contents of an animated SVT sequence stable in world space, we need to account for the fact that
-		// different frames of the sequence have different AABBs.
-		FVector UVScale = FVector::One();
-		FVector UVBias = FVector::Zero();
-		
+		FVector VolumeResolution = FVector(SVTViewerDefaultVolumeExtent * 2.0);
+		FUintVector4 PackedSVTUniforms0 = FUintVector4();
+		FUintVector4 PackedSVTUniforms1 = FUintVector4();
 		if (SparseVolumeTextureFrame)
 		{
-			const FVector VolumeResolution = FVector(SparseVolumeTextureFrame->GetVolumeResolution());
+			VolumeResolution = FVector(SparseVolumeTextureFrame->GetVolumeResolution());
 			const double MaxBoundsDim = FMath::Max(FMath::Max(VolumeResolution.X, VolumeResolution.Y), VolumeResolution.Z);
 			VolumeExtent = VolumeResolution / MaxBoundsDim * SVTViewerDefaultVolumeExtent;
-
-			SparseVolumeTextureFrame->GetFrameUVScaleBias(&UVScale, &UVBias);
+			SparseVolumeTextureFrame->GetPackedUniforms(PackedSVTUniforms0, PackedSVTUniforms1);
 		}
 
 		const FTransform ToWorldTransform = GetComponentTransform();
@@ -178,20 +175,25 @@ void USparseVolumeTextureViewerComponent::SendRenderTransformCommand()
 		const FRotationMatrix WorldToLocalRotation = FRotationMatrix(FRotator(ToWorldTransform.GetRotation().Inverse()));
 		const FMatrix44f ToLocalMatNoScale = FMatrix44f(WorldToLocalRotation);
 		const FMatrix44f ToLocalMat = FMatrix44f(
-			FTranslationMatrix(-ToWorldTransform.GetTranslation()) * WorldToLocalRotation * FScaleMatrix((VolumeExtent * MaxScaling).Reciprocal())
-			* FScaleMatrix(0.5) * FTranslationMatrix(FVector(0.5))
-			* FScaleMatrix(UVScale) * FTranslationMatrix(UVBias)
-			* FScaleMatrix(2.0) * FTranslationMatrix(FVector(-1.0)));
+			FTranslationMatrix(-ToWorldTransform.GetTranslation()) 
+			* WorldToLocalRotation 
+			* FScaleMatrix((VolumeExtent * MaxScaling).Reciprocal()));
+
+		const FVector3f VolumeRes3f = FVector3f(VolumeResolution.X, VolumeResolution.Y, VolumeResolution.Z);
 
-		const uint32 CompToVisualize = (uint32)ComponentToVisualize;
 
 		FSparseVolumeTextureViewerSceneProxy* SVTViewerSceneProxy = SparseVolumeTextureViewerSceneProxy;
 		ENQUEUE_RENDER_COMMAND(FUpdateSparseVolumeTextureViewerProxyTransformCommand)(
-			[SVTViewerSceneProxy, ToLocalMat, ToLocalMatNoScale, CompToVisualize, Ext = Extinction](FRHICommandList& RHICmdList)
+			[SVTViewerSceneProxy, ToLocalMat, ToLocalMatNoScale, CompIdx = (uint32)ComponentToVisualize, Ext = Extinction, Mip = MipLevel, PackedSVTUniforms0, PackedSVTUniforms1, VolumeRes3f]
+		(FRHICommandList& RHICmdList)
 			{
 				SVTViewerSceneProxy->WorldToLocal = ToLocalMat;
 				SVTViewerSceneProxy->WorldToLocalNoScale = ToLocalMatNoScale;
-				SVTViewerSceneProxy->ComponentToVisualize = CompToVisualize;
+				SVTViewerSceneProxy->PackedSVTUniforms0 = PackedSVTUniforms0;
+				SVTViewerSceneProxy->PackedSVTUniforms1 = PackedSVTUniforms1;
+				SVTViewerSceneProxy->VolumeResolution = VolumeRes3f;
+				SVTViewerSceneProxy->MipLevel = Mip;
+				SVTViewerSceneProxy->ComponentToVisualize = CompIdx;
 				SVTViewerSceneProxy->Extinction = Ext;
 			});
 	}
diff --git a/Engine/Source/Runtime/Renderer/Private/SparseVolumeTexture/SparseVolumeTextureViewerRendering.cpp b/Engine/Source/Runtime/Renderer/Private/SparseVolumeTexture/SparseVolumeTextureViewerRendering.cpp
--- a/Engine/Source/Runtime/Renderer/Private/SparseVolumeTexture/SparseVolumeTextureViewerRendering.cpp
+++ b/Engine/Source/Runtime/Renderer/Private/SparseVolumeTexture/SparseVolumeTextureViewerRendering.cpp
@@ -3,6 +3,7 @@
 #include "SparseVolumeTexture/SparseVolumeTextureViewerRendering.h"
 #include "SparseVolumeTexture/SparseVolumeTextureViewerSceneProxy.h"
 #include "SparseVolumeTexture/SparseVolumeTexture.h"
+#include "SparseVolumeTexture/SparseVolumeTextureSceneProxy.h"
 
 #include "DataDrivenShaderPlatformInfo.h"
 #include "SceneView.h"
@@ -56,20 +57,21 @@ class FVisualizeSparseVolumeTexturePS : public FGlobalShader
 	BEGIN_SHADER_PARAMETER_STRUCT(FParameters, )
 		SHADER_PARAMETER_STRUCT_REF(FViewUniformShaderParameters, ViewUniformBuffer)
 		RENDER_TARGET_BINDING_SLOTS()
+		SHADER_PARAMETER_TEXTURE(Texture3D<uint>, SparseVolumeTexturePageTable)
+		SHADER_PARAMETER_TEXTURE(Texture3D, SparseVolumeTextureA)
+		SHADER_PARAMETER_TEXTURE(Texture3D, SparseVolumeTextureB)
+		SHADER_PARAMETER(FUintVector4, PackedSVTUniforms0)
+		SHADER_PARAMETER(FUintVector4, PackedSVTUniforms1)
+		SHADER_PARAMETER(FVector3f, SparseVolumeTextureResolution)
+		SHADER_PARAMETER(int32, MipLevel)
 		SHADER_PARAMETER(FVector4f, WorldToLocal0)
 		SHADER_PARAMETER(FVector4f, WorldToLocal1)
 		SHADER_PARAMETER(FVector4f, WorldToLocal2)
 		SHADER_PARAMETER(FVector3f, WorldToLocalNoScale0)
 		SHADER_PARAMETER(FVector3f, WorldToLocalNoScale1)
 		SHADER_PARAMETER(FVector3f, WorldToLocalNoScale2)
-		SHADER_PARAMETER(FVector3f, SparseVolumeTextureResolution)
-		SHADER_PARAMETER(FVector3f, SparseVolumeTexturePageTableResolution)
 		SHADER_PARAMETER(uint32, ComponentToVisualize)
-		SHADER_PARAMETER(uint32, MipLevel)
 		SHADER_PARAMETER(float, Extinction)
-		SHADER_PARAMETER_TEXTURE(Texture3D, SparseVolumeTextureA)
-		SHADER_PARAMETER_TEXTURE(Texture3D, SparseVolumeTextureB)
-		SHADER_PARAMETER_TEXTURE(Texture3D<uint>, SparseVolumeTexturePageTable)
 	END_SHADER_PARAMETER_STRUCT()
 
 	static FPermutationDomain RemapPermutation(FPermutationDomain PermutationVector)
@@ -128,6 +130,13 @@ void AddSparseVolumeTextureViewerRenderPass(FRDGBuilder& GraphBuilder, FSceneRen
 			FVisualizeSparseVolumeTexturePS::FParameters* PsPassParameters = GraphBuilder.AllocParameters<FVisualizeSparseVolumeTexturePS::FParameters>();
 			PsPassParameters->ViewUniformBuffer = View.ViewUniformBuffer;
 			PsPassParameters->RenderTargets[0] = FRenderTargetBinding(SceneTextures.Color.Target, ERenderTargetLoadAction::ELoad);
+			PsPassParameters->SparseVolumeTexturePageTable = GBlackUintVolumeTexture->TextureRHI;
+			PsPassParameters->SparseVolumeTextureA = GBlackVolumeTexture->TextureRHI;
+			PsPassParameters->SparseVolumeTextureB = GBlackVolumeTexture->TextureRHI;
+			PsPassParameters->PackedSVTUniforms0 = SVTProxy->PackedSVTUniforms0;
+			PsPassParameters->PackedSVTUniforms1 = SVTProxy->PackedSVTUniforms1;
+			PsPassParameters->SparseVolumeTextureResolution = SVTProxy->VolumeResolution;
+			PsPassParameters->MipLevel = SVTProxy->MipLevel;
 			PsPassParameters->WorldToLocal0 = FVector4f(WorldToLocal.M[0][0], WorldToLocal.M[1][0], WorldToLocal.M[2][0], WorldToLocal.M[3][0]);
 			PsPassParameters->WorldToLocal1 = FVector4f(WorldToLocal.M[0][1], WorldToLocal.M[1][1], WorldToLocal.M[2][1], WorldToLocal.M[3][1]);
 			PsPassParameters->WorldToLocal2 = FVector4f(WorldToLocal.M[0][2], WorldToLocal.M[1][2], WorldToLocal.M[2][2], WorldToLocal.M[3][2]);
@@ -135,23 +144,15 @@ void AddSparseVolumeTextureViewerRenderPass(FRDGBuilder& GraphBuilder, FSceneRen
 			PsPassParameters->WorldToLocalNoScale1 = FVector3f(WorldToLocalNoScale.M[0][1], WorldToLocalNoScale.M[1][1], WorldToLocalNoScale.M[2][1]);
 			PsPassParameters->WorldToLocalNoScale2 = FVector3f(WorldToLocalNoScale.M[0][2], WorldToLocalNoScale.M[1][2], WorldToLocalNoScale.M[2][2]);
 			PsPassParameters->ComponentToVisualize = SVTProxy->ComponentToVisualize;
-			PsPassParameters->MipLevel = 0;
 			PsPassParameters->Extinction = SVTProxy->Extinction;
-			PsPassParameters->SparseVolumeTextureResolution = FVector3f::OneVector;
-			PsPassParameters->SparseVolumeTexturePageTableResolution = FVector3f::OneVector;
-			PsPassParameters->SparseVolumeTextureA = GBlackVolumeTexture->TextureRHI;
-			PsPassParameters->SparseVolumeTextureB = GBlackVolumeTexture->TextureRHI;
-			PsPassParameters->SparseVolumeTexturePageTable = GBlackUintVolumeTexture->TextureRHI;
-
-
+			
 			if (SVTProxy->SparseVolumeTextureSceneProxy)
 			{
-				const FSparseVolumeAssetHeader& Header = SVTProxy->SparseVolumeTextureSceneProxy->GetHeader();
-
-				PsPassParameters->SparseVolumeTextureResolution = FVector3f(Header.VirtualVolumeResolution);
-				PsPassParameters->SparseVolumeTexturePageTableResolution = FVector3f(Header.PageTableVolumeResolution);
-				PsPassParameters->MipLevel = (uint32)Header.MipLevel;
-
+				FRHITexture* PageTableTexture = SVTProxy->SparseVolumeTextureSceneProxy->GetPageTableTextureRHI();
+				if (PageTableTexture)
+				{
+					PsPassParameters->SparseVolumeTexturePageTable = PageTableTexture;
+				}
 				FRHITexture* TextureA = SVTProxy->SparseVolumeTextureSceneProxy->GetPhysicalTileDataATextureRHI();
 				if (TextureA)
 				{
@@ -162,11 +163,6 @@ void AddSparseVolumeTextureViewerRenderPass(FRDGBuilder& GraphBuilder, FSceneRen
 				{
 					PsPassParameters->SparseVolumeTextureB = TextureB;
 				}
-				FRHITexture* PageTableTexture = SVTProxy->SparseVolumeTextureSceneProxy->GetPageTableTextureRHI();
-				if (PageTableTexture)
-				{
-					PsPassParameters->SparseVolumeTexturePageTable = PageTableTexture;
-				}
 			}
 
 			ClearUnusedGraphResources(PixelShader, PsPassParameters);
diff --git a/Engine/Source/Runtime/Renderer/Private/SparseVolumeTexture/SparseVolumeTextureViewerSceneProxy.h b/Engine/Source/Runtime/Renderer/Private/SparseVolumeTexture/SparseVolumeTextureViewerSceneProxy.h
--- a/Engine/Source/Runtime/Renderer/Private/SparseVolumeTexture/SparseVolumeTextureViewerSceneProxy.h
+++ b/Engine/Source/Runtime/Renderer/Private/SparseVolumeTexture/SparseVolumeTextureViewerSceneProxy.h
@@ -20,6 +20,10 @@ public:
 	const FSparseVolumeTextureSceneProxy* SparseVolumeTextureSceneProxy = nullptr;
 	FMatrix44f WorldToLocal;
 	FMatrix44f WorldToLocalNoScale;
+	FUintVector4 PackedSVTUniforms0;
+	FUintVector4 PackedSVTUniforms1;
+	FVector3f VolumeResolution;
+	int32 MipLevel;
 	uint32 ComponentToVisualize;
 	float Extinction;
 
