geometry collections should not be added to the external SQ if they're a part of a cluster union

diff --git a/Engine/Source/Runtime/Engine/Private/PhysicsEngine/Experimental/PhysScene_Chaos.cpp b/Engine/Source/Runtime/Engine/Private/PhysicsEngine/Experimental/PhysScene_Chaos.cpp
--- a/Engine/Source/Runtime/Engine/Private/PhysicsEngine/Experimental/PhysScene_Chaos.cpp
+++ b/Engine/Source/Runtime/Engine/Private/PhysicsEngine/Experimental/PhysScene_Chaos.cpp
@@ -1904,8 +1904,16 @@ void FPhysScene_Chaos::OnSyncBodies(Chaos::FPhysicsSolverBase* Solver)
 
 			auto ShouldGCParticleBeInAccelerationStructure = [&Interface](Chaos::FPhysicsObjectHandle Handle)
 			{
+				IPhysicsProxyBase* HandleProxy = Chaos::FPhysicsObjectInterface::GetProxy({ &Handle, 1 });
+
 				// It's possible to be an enabled particle and not qualify for query collisions if the GC particle has been replication abandoned. 
-				return !Interface->AreAllDisabled({ &Handle, 1 }) && Interface->AreAllShapesQueryEnabled({ &Handle, 1 });
+				// Furthermore, it's possible for a particle to be enabled on the game thread but not actually enabled on the physics thread. A particle
+				// with an internal cluster parent (e.g. a cluster union) could get into this state where the geometry collection physics proxy will not
+				// transfer the disabled state from the PT to the GT because the particle has an internal cluster for a parent. We can detect this case by making
+				// sure that the proxy of the GC does not have a parent proxy (only happens with a cluster union currently).
+				return !Interface->AreAllDisabled({ &Handle, 1 })
+					&& Interface->AreAllShapesQueryEnabled({ &Handle, 1 })
+					&& (HandleProxy->GetParentProxy() == nullptr);
 			};
 
 			TArray<Chaos::FPhysicsObjectHandle> ActiveHandles = Proxy->GetAllPhysicsObjects().FilterByPredicate(
