Changing a lot more code to use batched shader parameters.

diff --git a/Engine/Plugins/Compositing/LensDistortion/Source/LensDistortion/Private/LensDistortionRendering.cpp b/Engine/Plugins/Compositing/LensDistortion/Source/LensDistortion/Private/LensDistortionRendering.cpp
--- a/Engine/Plugins/Compositing/LensDistortion/Source/LensDistortion/Private/LensDistortionRendering.cpp
+++ b/Engine/Plugins/Compositing/LensDistortion/Source/LensDistortion/Private/LensDistortionRendering.cpp
@@ -85,10 +85,8 @@ public:
 		OutputMultiplyAndAdd.Bind(Initializer.ParameterMap, TEXT("OutputMultiplyAndAdd"));
 	}
 
-	template<typename TShaderRHIParamRef>
 	void SetParameters(
-		FRHICommandListImmediate& RHICmdList,
-		const TShaderRHIParamRef ShaderRHI,
+		FRHIBatchedShaderParameters& BatchedParameters,
 		const FCompiledCameraModel& CompiledCameraModel,
 		const FIntPoint& DisplacementMapResolution)
 	{
@@ -102,12 +100,12 @@ public:
 			CompiledCameraModel.OriginalCameraModel.P1,
 			CompiledCameraModel.OriginalCameraModel.P2);
 
-		SetShaderValue(RHICmdList, ShaderRHI, PixelUVSize, PixelUVSizeValue);
-		SetShaderValue(RHICmdList, ShaderRHI, DistortedCameraMatrix, FVector4f(CompiledCameraModel.DistortedCameraMatrix));
-		SetShaderValue(RHICmdList, ShaderRHI, UndistortedCameraMatrix, FVector4f(CompiledCameraModel.UndistortedCameraMatrix));
-		SetShaderValue(RHICmdList, ShaderRHI, RadialDistortionCoefs, RadialDistortionCoefsValue);
-		SetShaderValue(RHICmdList, ShaderRHI, TangentialDistortionCoefs, TangentialDistortionCoefsValue);
-		SetShaderValue(RHICmdList, ShaderRHI, OutputMultiplyAndAdd, FVector2f(CompiledCameraModel.OutputMultiplyAndAdd));
+		SetShaderValue(BatchedParameters, PixelUVSize, PixelUVSizeValue);
+		SetShaderValue(BatchedParameters, DistortedCameraMatrix, FVector4f(CompiledCameraModel.DistortedCameraMatrix));
+		SetShaderValue(BatchedParameters, UndistortedCameraMatrix, FVector4f(CompiledCameraModel.UndistortedCameraMatrix));
+		SetShaderValue(BatchedParameters, RadialDistortionCoefs, RadialDistortionCoefsValue);
+		SetShaderValue(BatchedParameters, TangentialDistortionCoefs, TangentialDistortionCoefsValue);
+		SetShaderValue(BatchedParameters, OutputMultiplyAndAdd, FVector2f(CompiledCameraModel.OutputMultiplyAndAdd));
 	}
 
 private:
@@ -210,8 +208,9 @@ static void DrawUVDisplacementToRenderTarget_RenderThread(
 			OutTextureRenderTargetResource->GetSizeX(), OutTextureRenderTargetResource->GetSizeY(), 1.f);
 
 		// Update shader uniform parameters.
-		VertexShader->SetParameters(RHICmdList, VertexShader.GetVertexShader(), CompiledCameraModel, DisplacementMapResolution);
-		PixelShader->SetParameters(RHICmdList, PixelShader.GetPixelShader(), CompiledCameraModel, DisplacementMapResolution);
+
+		SetShaderParametersLegacyVS(RHICmdList, VertexShader, CompiledCameraModel, DisplacementMapResolution);
+		SetShaderParametersLegacyPS(RHICmdList, PixelShader, CompiledCameraModel, DisplacementMapResolution);
 
 		// Draw grid.
 		uint32 PrimitiveCount = kGridSubdivisionX * kGridSubdivisionY * 2;
diff --git a/Engine/Plugins/Compositing/OpenCVLensDistortion/Source/OpenCVLensDistortion/Private/LensDistortionDisplacementMapRendering.cpp b/Engine/Plugins/Compositing/OpenCVLensDistortion/Source/OpenCVLensDistortion/Private/LensDistortionDisplacementMapRendering.cpp
--- a/Engine/Plugins/Compositing/OpenCVLensDistortion/Source/OpenCVLensDistortion/Private/LensDistortionDisplacementMapRendering.cpp
+++ b/Engine/Plugins/Compositing/OpenCVLensDistortion/Source/OpenCVLensDistortion/Private/LensDistortionDisplacementMapRendering.cpp
@@ -48,13 +48,12 @@ public:
 		BilinearSampler.Bind(Initializer.ParameterMap, TEXT("BilinearClampedSampler"));
 	}
 
-	template<typename TShaderRHIParamRef>
-	void SetParameters(FRHICommandListImmediate& RHICmdList, const TShaderRHIParamRef ShaderRHI, const FTextureResource* PreComputedDisplacementMap, const FIntPoint& DisplacementMapResolution)
+	void SetParameters(FRHIBatchedShaderParameters& BatchedParameters, const FTextureResource* PreComputedDisplacementMap, const FIntPoint& DisplacementMapResolution)
 	{
 		FVector2f PixelUVSizeValue(1.f / float(DisplacementMapResolution.X), 1.f / float(DisplacementMapResolution.Y));
 
-		SetShaderValue(RHICmdList, ShaderRHI, PixelUVSize, PixelUVSizeValue);
-		SetTextureParameter(RHICmdList, ShaderRHI, UndistortDisplacementMap, BilinearSampler, TStaticSamplerState<SF_Bilinear, AM_Clamp, AM_Clamp, AM_Clamp>::GetRHI(), PreComputedDisplacementMap->TextureRHI);
+		SetShaderValue(BatchedParameters, PixelUVSize, PixelUVSizeValue);
+		SetTextureParameter(BatchedParameters, UndistortDisplacementMap, BilinearSampler, TStaticSamplerState<SF_Bilinear, AM_Clamp, AM_Clamp, AM_Clamp>::GetRHI(), PreComputedDisplacementMap->TextureRHI);
 	}
 
 private:
@@ -147,8 +146,8 @@ static void DrawUVDisplacementToRenderTargetFromPreComputedDisplacementMap_Rende
 		SetGraphicsPipelineState(RHICmdList, GraphicsPSOInit, 0);
 
 		// Update shader uniform parameters.
-		VertexShader->SetParameters(RHICmdList, VertexShader.GetVertexShader(), PreComputedDisplacementMap, DisplacementMapResolution);
-		PixelShader->SetParameters(RHICmdList, PixelShader.GetPixelShader(), PreComputedDisplacementMap, DisplacementMapResolution);
+		SetShaderParametersLegacyVS(RHICmdList, VertexShader, PreComputedDisplacementMap, DisplacementMapResolution);
+		SetShaderParametersLegacyPS(RHICmdList, PixelShader, PreComputedDisplacementMap, DisplacementMapResolution);
 
 		// Draw grid.
 		const uint32 PrimitiveCount = kGridSubdivisionX * kGridSubdivisionY * 2;
diff --git a/Engine/Plugins/Experimental/ColorCorrectRegions/Source/ColorCorrectRegions/Private/ColorCorrectRegionsSceneViewExtension.cpp b/Engine/Plugins/Experimental/ColorCorrectRegions/Source/ColorCorrectRegions/Private/ColorCorrectRegionsSceneViewExtension.cpp
--- a/Engine/Plugins/Experimental/ColorCorrectRegions/Source/ColorCorrectRegions/Private/ColorCorrectRegionsSceneViewExtension.cpp
+++ b/Engine/Plugins/Experimental/ColorCorrectRegions/Source/ColorCorrectRegions/Private/ColorCorrectRegionsSceneViewExtension.cpp
@@ -573,20 +573,28 @@ namespace
 					EScreenPassDrawFlags::None,
 					[&](FRHICommandList& RHICmdList)
 					{
-						SetUniformBufferParameterImmediate(RHICmdList, PixelShader.GetPixelShader(), PixelShader->GetUniformBufferParameter<FCCRRegionDataInputParameter>(), RegionData);
-						SetUniformBufferParameterImmediate(RHICmdList, PixelShader.GetPixelShader(), PixelShader->GetUniformBufferParameter<FCCRColorCorrectParameter>(), CCBase);
+						FRHIBatchedShaderParameters& BatchedParameters = RHICmdList.GetScratchShaderParameters();
+
+						SetUniformBufferParameterImmediate(BatchedParameters, PixelShader->GetUniformBufferParameter<FCCRRegionDataInputParameter>(), RegionData);
+						SetUniformBufferParameterImmediate(BatchedParameters, PixelShader->GetUniformBufferParameter<FCCRColorCorrectParameter>(), CCBase);
 						if (bIsAdvanced)
 						{
-							SetUniformBufferParameterImmediate(RHICmdList, PixelShader.GetPixelShader(), PixelShader->GetUniformBufferParameter<FCCRColorCorrectShadowsParameter>(), CCShadows);
-							SetUniformBufferParameterImmediate(RHICmdList, PixelShader.GetPixelShader(), PixelShader->GetUniformBufferParameter<FCCRColorCorrectMidtonesParameter>(), CCMidtones);
-							SetUniformBufferParameterImmediate(RHICmdList, PixelShader.GetPixelShader(), PixelShader->GetUniformBufferParameter<FCCRColorCorrectHighlightsParameter>(), CCHighlights);
+							SetUniformBufferParameterImmediate(BatchedParameters, PixelShader->GetUniformBufferParameter<FCCRColorCorrectShadowsParameter>(), CCShadows);
+							SetUniformBufferParameterImmediate(BatchedParameters, PixelShader->GetUniformBufferParameter<FCCRColorCorrectMidtonesParameter>(), CCMidtones);
+							SetUniformBufferParameterImmediate(BatchedParameters, PixelShader->GetUniformBufferParameter<FCCRColorCorrectHighlightsParameter>(), CCHighlights);
 						}
 
-						VertexShader->SetParameters(RHICmdList, View);
-						SetShaderParameters(RHICmdList, VertexShader, VertexShader.GetVertexShader(), *PostProcessMaterialParameters);
+						PixelShader->SetParameters(BatchedParameters, View);
+						SetShaderParameters(BatchedParameters, PixelShader, *PostProcessMaterialParameters);
+
+						RHICmdList.SetBatchedShaderParameters(PixelShader.GetPixelShader(), BatchedParameters);
+
+						BatchedParameters.Reset();
+
+						VertexShader->SetParameters(BatchedParameters, View);
+						SetShaderParameters(BatchedParameters, VertexShader, *PostProcessMaterialParameters);
 
-						PixelShader->SetParameters(RHICmdList, View);
-						SetShaderParameters(RHICmdList, PixelShader, PixelShader.GetPixelShader(), *PostProcessMaterialParameters);
+						RHICmdList.SetBatchedShaderParameters(VertexShader.GetVertexShader(), BatchedParameters);
 					});
 
 			});
diff --git a/Engine/Plugins/Experimental/ColorCorrectRegions/Source/ColorCorrectRegions/Public/ColorCorrectRegionsPostProcessMaterial.h b/Engine/Plugins/Experimental/ColorCorrectRegions/Source/ColorCorrectRegions/Public/ColorCorrectRegionsPostProcessMaterial.h
--- a/Engine/Plugins/Experimental/ColorCorrectRegions/Source/ColorCorrectRegions/Public/ColorCorrectRegionsPostProcessMaterial.h
+++ b/Engine/Plugins/Experimental/ColorCorrectRegions/Source/ColorCorrectRegions/Public/ColorCorrectRegionsPostProcessMaterial.h
@@ -118,11 +118,16 @@ public:
 		: FColorCorrectRegionsPostProcessMaterialShader(Initializer)
 	{}
 
+	void SetParameters(FRHIBatchedShaderParameters& BatchedParameters, const FSceneView& View)
+	{
+		FGlobalShader::SetParameters<FViewUniformShaderParameters>(BatchedParameters, View.ViewUniformBuffer);
+	}
+
+	UE_DEPRECATED(5.3, "SetParameters with FRHIBatchedShaderParameters should be used.")
 	void SetParameters(FRHICommandList& RHICmdList, const FSceneView& View)
 	{
 		FGlobalShader::SetParameters<FViewUniformShaderParameters>(RHICmdList, RHICmdList.GetBoundVertexShader(), View.ViewUniformBuffer);
 	}
-	
 };
 
 class FColorCorrectGenericPS : public FColorCorrectRegionsPostProcessMaterialShader
@@ -147,16 +152,16 @@ public:
 		: FColorCorrectRegionsPostProcessMaterialShader(Initializer)
 	{}
 
-	void SetParameters(FRHICommandList& RHICmdList, const FSceneView& View)
+	void SetParameters(FRHIBatchedShaderParameters& BatchedParameters, const FSceneView& View)
 	{
-		FGlobalShader::SetParameters<FViewUniformShaderParameters>(RHICmdList, RHICmdList.GetBoundPixelShader(), View.ViewUniformBuffer);
+		FGlobalShader::SetParameters<FViewUniformShaderParameters>(BatchedParameters, View.ViewUniformBuffer);
 	}
 
-	static void ModifyCompilationEnvironment(const FGlobalShaderPermutationParameters& Parameters, FShaderCompilerEnvironment& OutEnvironment)
+	UE_DEPRECATED(5.3, "SetParameters with FRHIBatchedShaderParameters should be used.")
+	void SetParameters(FRHICommandList& RHICmdList, const FSceneView& View)
 	{
-		FColorCorrectRegionsPostProcessMaterialShader::ModifyCompilationEnvironment(Parameters, OutEnvironment);
+		FGlobalShader::SetParameters<FViewUniformShaderParameters>(RHICmdList, RHICmdList.GetBoundPixelShader(), View.ViewUniformBuffer);
 	}
-
 };
 
 class FColorCorrectRegionMaterialPS : public FColorCorrectGenericPS
@@ -198,11 +203,6 @@ class FColorCorrectScreenPassVS : public FGlobalShader
 public:
 	DECLARE_GLOBAL_SHADER(FColorCorrectScreenPassVS);
 
-	static bool ShouldCompilePermutation(const FGlobalShaderPermutationParameters&)
-	{
-		return true;
-	}
-
 	FColorCorrectScreenPassVS() = default;
 	FColorCorrectScreenPassVS(const ShaderMetaType::CompiledShaderInitializerType& Initializer)
 		: FGlobalShader(Initializer)
@@ -215,11 +215,6 @@ class FClearRectPS : public FGlobalShader
 	DECLARE_GLOBAL_SHADER(FClearRectPS);
 	SHADER_USE_PARAMETER_STRUCT(FClearRectPS, FGlobalShader);
 
-	static bool ShouldCompilePermutation(const FGlobalShaderPermutationParameters& Parameters) 
-	{
-		return true;
-	}
-
 	static void ModifyCompilationEnvironment(const FGlobalShaderPermutationParameters& Parameters, FShaderCompilerEnvironment& OutEnvironment)
 	{
 		FColorCorrectRegionMaterialPS::ModifyCompilationEnvironment(Parameters, OutEnvironment);
@@ -238,11 +233,6 @@ class FCCRStencilMergerVS : public FGlobalShader
 public:
 	DECLARE_GLOBAL_SHADER(FCCRStencilMergerVS);
 
-	static bool ShouldCompilePermutation(const FGlobalShaderPermutationParameters&)
-	{
-		return true;
-	}
-
 	FCCRStencilMergerVS() = default;
 	FCCRStencilMergerVS(const ShaderMetaType::CompiledShaderInitializerType& Initializer)
 		: FGlobalShader(Initializer)
@@ -255,11 +245,6 @@ class FCCRStencilMergerPS : public FGlobalShader
 	DECLARE_GLOBAL_SHADER(FCCRStencilMergerPS);
 	SHADER_USE_PARAMETER_STRUCT(FCCRStencilMergerPS, FGlobalShader);
 
-	static bool ShouldCompilePermutation(const FGlobalShaderPermutationParameters& Parameters)
-	{
-		return true;
-	}
-
 	BEGIN_SHADER_PARAMETER_STRUCT(FParameters, )
 		SHADER_PARAMETER_STRUCT_REF(FViewUniformShaderParameters, View)
 		SHADER_PARAMETER_RDG_BUFFER_SRV(StructuredBuffer<uint>, StencilIds)
diff --git a/Engine/Plugins/Media/PixelCapture/Source/PixelCapture/Private/PixelCaptureUtils.h b/Engine/Plugins/Media/PixelCapture/Source/PixelCapture/Private/PixelCaptureUtils.h
--- a/Engine/Plugins/Media/PixelCapture/Source/PixelCapture/Private/PixelCaptureUtils.h
+++ b/Engine/Plugins/Media/PixelCapture/Source/PixelCapture/Private/PixelCaptureUtils.h
@@ -55,7 +55,7 @@ inline void CopyTexture(FRHICommandList& RHICmdList, FTextureRHIRef SourceTextur
             GraphicsPSOInit.PrimitiveType = PT_TriangleList;
             SetGraphicsPipelineState(RHICmdList, GraphicsPSOInit, 0);
 
-            PixelShader->SetParameters(RHICmdList, TStaticSamplerState<SF_Point>::GetRHI(), SourceTexture);
+			SetShaderParametersLegacyPS(RHICmdList, PixelShader, TStaticSamplerState<SF_Point>::GetRHI(), SourceTexture);
 
             FIntPoint TargetBufferSize(DestTexture->GetDesc().Extent.X, DestTexture->GetDesc().Extent.Y);
             RendererModule->DrawRectangle(RHICmdList, 0, 0, // Dest X, Y
diff --git a/Engine/Plugins/Media/PixelStreaming/Source/PixelStreaming/Private/UtilsRender.h b/Engine/Plugins/Media/PixelStreaming/Source/PixelStreaming/Private/UtilsRender.h
--- a/Engine/Plugins/Media/PixelStreaming/Source/PixelStreaming/Private/UtilsRender.h
+++ b/Engine/Plugins/Media/PixelStreaming/Source/PixelStreaming/Private/UtilsRender.h
@@ -84,7 +84,7 @@ namespace UE::PixelStreaming
 				GraphicsPSOInit.PrimitiveType = PT_TriangleList;
 				SetGraphicsPipelineState(RHICmdList, GraphicsPSOInit, 0);
 
-				PixelShader->SetParameters(RHICmdList, TStaticSamplerState<SF_Point>::GetRHI(), SourceTexture);
+				SetShaderParametersLegacyPS(RHICmdList, PixelShader, TStaticSamplerState<SF_Point>::GetRHI(), SourceTexture);
 
 				FIntPoint TargetBufferSize(DestTexture->GetDesc().Extent.X, DestTexture->GetDesc().Extent.Y);
 				RendererModule->DrawRectangle(RHICmdList, 0, 0, // Dest X, Y
diff --git a/Engine/Plugins/MovieScene/MovieRenderPipeline/Source/MovieRenderPipelineCore/Private/MoviePipelineSurfaceReader.cpp b/Engine/Plugins/MovieScene/MovieRenderPipeline/Source/MovieRenderPipelineCore/Private/MoviePipelineSurfaceReader.cpp
--- a/Engine/Plugins/MovieScene/MovieRenderPipeline/Source/MovieRenderPipelineCore/Private/MoviePipelineSurfaceReader.cpp
+++ b/Engine/Plugins/MovieScene/MovieRenderPipeline/Source/MovieRenderPipelineCore/Private/MoviePipelineSurfaceReader.cpp
@@ -129,7 +129,7 @@ void FMoviePipelineSurfaceReader::ResolveSampleToReadbackTexture_RenderThread(co
 			SetGraphicsPipelineState(RHICmdList, GraphicsPSOInit, 0);
 
 			// Bind the SourceSurfaceSample to our texture sampler with a point sample (since the resolutions match).
-			PixelShader->SetParameters(RHICmdList, TStaticSamplerState<SF_Point>::GetRHI(), SourceSurfaceSample);
+			SetShaderParametersLegacyPS(RHICmdList, PixelShader, TStaticSamplerState<SF_Point>::GetRHI(), SourceSurfaceSample);
 		}
 		else
 		{
@@ -139,7 +139,7 @@ void FMoviePipelineSurfaceReader::ResolveSampleToReadbackTexture_RenderThread(co
 			SetGraphicsPipelineState(RHICmdList, GraphicsPSOInit, 0);
 
 			// Bind the SourceSurfaceSample to our texture sampler with a point sample (since the resolutions match).
-			PixelShader->SetParameters(RHICmdList, TStaticSamplerState<SF_Point>::GetRHI(), SourceSurfaceSample);
+			SetShaderParametersLegacyPS(RHICmdList, PixelShader, TStaticSamplerState<SF_Point>::GetRHI(), SourceSurfaceSample);
 		}
 
 		RendererModule->DrawRectangle(
diff --git a/Engine/Plugins/Runtime/Oculus/OculusVR/Source/OculusHMD/Private/OculusHMD_CustomPresent.cpp b/Engine/Plugins/Runtime/Oculus/OculusVR/Source/OculusHMD/Private/OculusHMD_CustomPresent.cpp
--- a/Engine/Plugins/Runtime/Oculus/OculusVR/Source/OculusHMD/Private/OculusHMD_CustomPresent.cpp
+++ b/Engine/Plugins/Runtime/Oculus/OculusVR/Source/OculusHMD/Private/OculusHMD_CustomPresent.cpp
@@ -441,14 +441,14 @@ void FCustomPresent::CopyTexture_RenderThread(FRHICommandListImmediate& RHICmdLi
 					TShaderMapRef<FScreenPSMipLevel> PixelShader(ShaderMap);
 					GraphicsPSOInit.BoundShaderState.PixelShaderRHI = PixelShader.GetPixelShader();
 					SetGraphicsPipelineState(RHICmdList, GraphicsPSOInit, 0);
-					PixelShader->SetParameters(RHICmdList, SamplerState, SrcTextureRHI, MipIndex);
+					SetShaderParametersLegacyPS(RHICmdList, PixelShader, SamplerState, SrcTextureRHI, MipIndex);
 				}
 				else
 				{
 					TShaderMapRef<FScreenPSsRGBSourceMipLevel> PixelShader(ShaderMap);
 					GraphicsPSOInit.BoundShaderState.PixelShaderRHI = PixelShader.GetPixelShader();
 					SetGraphicsPipelineState(RHICmdList, GraphicsPSOInit, 0);
-					PixelShader->SetParameters(RHICmdList, SamplerState, SrcTextureRHI, MipIndex);
+					SetShaderParametersLegacyPS(RHICmdList, PixelShader, SamplerState, SrcTextureRHI, MipIndex);
 				}
 
 				RHICmdList.SetViewport(DstRect.Min.X, DstRect.Min.Y, 0.0f, DstRect.Min.X + MipViewportWidth, DstRect.Min.Y + MipViewportHeight, 1.0f);
diff --git a/Engine/Plugins/Runtime/OpenXR/Source/OpenXRHMD/Private/OpenXRHMD.cpp b/Engine/Plugins/Runtime/OpenXR/Source/OpenXRHMD/Private/OpenXRHMD.cpp
--- a/Engine/Plugins/Runtime/OpenXR/Source/OpenXRHMD/Private/OpenXRHMD.cpp
+++ b/Engine/Plugins/Runtime/OpenXR/Source/OpenXRHMD/Private/OpenXRHMD.cpp
@@ -3451,28 +3451,28 @@ public:
 
 	}
 
-	void SetParameters(FRHICommandList& RHICmdList, EDisplayOutputFormat DisplayOutputFormat, EDisplayColorGamut DisplayColorGamut, EDisplayColorGamut TextureColorGamut, FRHITexture* SceneTextureRHI, bool bSameSize)
+	void SetParameters(FRHIBatchedShaderParameters& BatchedParameters, EDisplayOutputFormat DisplayOutputFormat, EDisplayColorGamut DisplayColorGamut, EDisplayColorGamut TextureColorGamut, FRHITexture* SceneTextureRHI, bool bSameSize)
 	{
 		int32 OutputDeviceValue = (int32)DisplayOutputFormat;
 		int32 OutputGamutValue = (int32)DisplayColorGamut;
 
-		SetShaderValue(RHICmdList, RHICmdList.GetBoundPixelShader(), OutputDevice, OutputDeviceValue);
-		SetShaderValue(RHICmdList, RHICmdList.GetBoundPixelShader(), OutputGamut, OutputGamutValue);
+		SetShaderValue(BatchedParameters, OutputDevice, OutputDeviceValue);
+		SetShaderValue(BatchedParameters, OutputGamut, OutputGamutValue);
 
 		const FMatrix44f TextureGamutMatrixToXYZ = GamutToXYZMatrix(TextureColorGamut);
 		const FMatrix44f XYZToDisplayMatrix = XYZToGamutMatrix(DisplayColorGamut);
 		// note: we use mul(m,v) instead of mul(v,m) in the shaders for color conversions which is why matrix multiplication is reversed compared to what we usually do
 		const FMatrix44f CombinedMatrix = XYZToDisplayMatrix * TextureGamutMatrixToXYZ;
 
-		SetShaderValue(RHICmdList, RHICmdList.GetBoundPixelShader(), TextureToOutputGamutMatrix, CombinedMatrix);
+		SetShaderValue(BatchedParameters, TextureToOutputGamutMatrix, CombinedMatrix);
 
 		if (bSameSize)
 		{
-			SetTextureParameter(RHICmdList, RHICmdList.GetBoundPixelShader(), SceneTexture, SceneSampler, TStaticSamplerState<SF_Point>::GetRHI(), SceneTextureRHI);
+			SetTextureParameter(BatchedParameters, SceneTexture, SceneSampler, TStaticSamplerState<SF_Point>::GetRHI(), SceneTextureRHI);
 		}
 		else
 		{
-			SetTextureParameter(RHICmdList, RHICmdList.GetBoundPixelShader(), SceneTexture, SceneSampler, TStaticSamplerState<SF_Bilinear>::GetRHI(), SceneTextureRHI);
+			SetTextureParameter(BatchedParameters, SceneTexture, SceneSampler, TStaticSamplerState<SF_Bilinear>::GetRHI(), SceneTextureRHI);
 		}
 	}
 
@@ -3588,13 +3588,13 @@ void FOpenXRHMD::CopyTexture_RenderThread(FRHICommandListImmediate& RHICmdList,
 		GraphicsPSOInit.PrimitiveType = PT_TriangleList;
 
 		const auto FeatureLevel = GMaxRHIFeatureLevel;
-		auto ShaderMap = GetGlobalShaderMap(FeatureLevel);
+		FGlobalShaderMap* ShaderMap = GetGlobalShaderMap(FeatureLevel);
 
 		TShaderMapRef<FScreenVS> VertexShader(ShaderMap);
 
 		TShaderRef<FGlobalShader> PixelShader;
-		FDisplayMappingPS* DisplayMappingPS = nullptr;
-		FScreenPS* ScreenPS = nullptr;
+		TShaderRef<FDisplayMappingPS> DisplayMappingPS;
+		TShaderRef<FScreenPS> ScreenPS;
 
 		bool bNeedsDisplayMapping = false;
 		EDisplayOutputFormat TVDisplayOutputFormat = EDisplayOutputFormat::SDR_sRGB;
@@ -3631,7 +3631,7 @@ void FOpenXRHMD::CopyTexture_RenderThread(FRHICommandListImmediate& RHICmdList,
 
 			ensureMsgf(!bIsArraySource, TEXT("Stub implementation for array source exists, but this path is not fully supported because post-processing not working with MMV"));
 
-			DisplayMappingPS = DisplayMappingPSRef.GetShader();
+			DisplayMappingPS = DisplayMappingPSRef;
 			PixelShader = DisplayMappingPSRef;
 		}
 		else
@@ -3639,13 +3639,13 @@ void FOpenXRHMD::CopyTexture_RenderThread(FRHICommandListImmediate& RHICmdList,
 			if (LIKELY(!bIsArraySource))
 			{
 				TShaderMapRef<FScreenPS> ScreenPSRef(ShaderMap);
-				ScreenPS = ScreenPSRef.GetShader();
+				ScreenPS = ScreenPSRef;
 				PixelShader = ScreenPSRef;
 			}
 			else
 			{
 				TShaderMapRef<FScreenFromSlice0PS> ScreenPSRef(ShaderMap);
-				ScreenPS = ScreenPSRef.GetShader();
+				ScreenPS = ScreenPSRef;
 				PixelShader = ScreenPSRef;
 			}
 		}
@@ -3659,20 +3659,14 @@ void FOpenXRHMD::CopyTexture_RenderThread(FRHICommandListImmediate& RHICmdList,
 		RHICmdList.Transition(FRHITransitionInfo(SrcTexture, ERHIAccess::Unknown, ERHIAccess::SRVMask));
 
 		const bool bSameSize = DstRect.Size() == SrcRect.Size();
-		if (ScreenPS)
+		if (ScreenPS.IsValid())
 		{
-			if (bSameSize)
-			{
-				ScreenPS->SetParameters(RHICmdList, TStaticSamplerState<SF_Point>::GetRHI(), SrcTexture);
-			}
-			else
-			{
-				ScreenPS->SetParameters(RHICmdList, TStaticSamplerState<SF_Bilinear>::GetRHI(), SrcTexture);
-			}
+			FRHISamplerState* PixelSampler = bSameSize ? TStaticSamplerState<SF_Point>::GetRHI() : TStaticSamplerState<SF_Bilinear>::GetRHI();
+			SetShaderParametersLegacyPS(RHICmdList, ScreenPS, PixelSampler, SrcTexture);
 		}
-		else if (DisplayMappingPS)
+		else if (DisplayMappingPS.IsValid())
 		{
-			DisplayMappingPS->SetParameters(RHICmdList, TVDisplayOutputFormat, TVColorGamut, HMDColorGamut, SrcTexture, bSameSize);
+			SetShaderParametersLegacyPS(RHICmdList, DisplayMappingPS, TVDisplayOutputFormat, TVColorGamut, HMDColorGamut, SrcTexture, bSameSize);
 		}
 
 		RendererModule->DrawRectangle(
diff --git a/Engine/Plugins/Runtime/Steam/SteamVR/Source/SteamVR/Private/SteamVRSpectatorScreen.cpp b/Engine/Plugins/Runtime/Steam/SteamVR/Source/SteamVR/Private/SteamVRSpectatorScreen.cpp
--- a/Engine/Plugins/Runtime/Steam/SteamVR/Source/SteamVR/Private/SteamVRSpectatorScreen.cpp
+++ b/Engine/Plugins/Runtime/Steam/SteamVR/Source/SteamVR/Private/SteamVRSpectatorScreen.cpp
@@ -88,7 +88,7 @@ void FSteamVRHMD::CopyTexture_RenderThread(FRHICommandListImmediate& RHICmdList,
 			GraphicsPSOInit.BoundShaderState.PixelShaderRHI = PixelShader.GetPixelShader();
 
 			SetGraphicsPipelineState(RHICmdList, GraphicsPSOInit, 0);
-			PixelShader->SetParameters(RHICmdList, PixelSampler, SrcTexture);
+			SetShaderParametersLegacyPS(RHICmdList, PixelShader, PixelSampler, SrcTexture);
 		}
 		else
 		{
@@ -96,7 +96,7 @@ void FSteamVRHMD::CopyTexture_RenderThread(FRHICommandListImmediate& RHICmdList,
 			GraphicsPSOInit.BoundShaderState.PixelShaderRHI = PixelShader.GetPixelShader();
 
 			SetGraphicsPipelineState(RHICmdList, GraphicsPSOInit, 0);
-			PixelShader->SetParameters(RHICmdList, PixelSampler, SrcTexture);
+			SetShaderParametersLegacyPS(RHICmdList, PixelShader, PixelSampler, SrcTexture);
 		}
 
 		RendererModule->DrawRectangle(
diff --git a/Engine/Plugins/Runtime/WarpUtils/Source/PFMExporter/Private/PFMExporterShader.cpp b/Engine/Plugins/Runtime/WarpUtils/Source/PFMExporter/Private/PFMExporterShader.cpp
--- a/Engine/Plugins/Runtime/WarpUtils/Source/PFMExporter/Private/PFMExporterShader.cpp
+++ b/Engine/Plugins/Runtime/WarpUtils/Source/PFMExporter/Private/PFMExporterShader.cpp
@@ -47,10 +47,9 @@ public:
 		MeshToPFMMatrixParameter.Bind(Initializer.ParameterMap, TEXT("MeshToPFMMatrix"));
 	}	
 
-	template<typename TShaderRHIParamRef>
-	void SetMeshToPFMMatrix(FRHICommandListImmediate& RHICmdList, const TShaderRHIParamRef ShaderRHI, const FMatrix& MeshToPFMMatrix)
+	void SetParameters(FRHIBatchedShaderParameters& BatchedParameters, const FMatrix& MeshToPFMMatrix)
 	{
-		SetShaderValue(RHICmdList, ShaderRHI, MeshToPFMMatrixParameter, (FMatrix44f)MeshToPFMMatrix);
+		SetShaderValue(BatchedParameters, MeshToPFMMatrixParameter, (FMatrix44f)MeshToPFMMatrix);
 	}
 
 private:
@@ -162,7 +161,8 @@ bool FPFMExporterShader::ApplyPFMExporter_RenderThread(
 
 				SetGraphicsPipelineState(RHICmdList, GraphicsPSOInit, 0);
 			}
-			VertexShader->SetMeshToPFMMatrix(RHICmdList, VertexShader.GetVertexShader(), MeshToPFMMatrix);
+
+			SetShaderParametersLegacyVS(RHICmdList, VertexShader, MeshToPFMMatrix);
 
 			// Render mesh to PFM texture:
 			RHICmdList.SetStreamSource(0, VertexBufferRHI, 0);
diff --git a/Engine/Plugins/Runtime/nDisplay/Source/DisplayCluster/Private/Render/Viewport/DisplayClusterViewportProxy.cpp b/Engine/Plugins/Runtime/nDisplay/Source/DisplayCluster/Private/Render/Viewport/DisplayClusterViewportProxy.cpp
--- a/Engine/Plugins/Runtime/nDisplay/Source/DisplayCluster/Private/Render/Viewport/DisplayClusterViewportProxy.cpp
+++ b/Engine/Plugins/Runtime/nDisplay/Source/DisplayCluster/Private/Render/Viewport/DisplayClusterViewportProxy.cpp
@@ -813,14 +813,10 @@ void ResampleCopyTextureImpl_RenderThread(FRHICommandListImmediate& RHICmdList,
 
 		SetGraphicsPipelineState(RHICmdList, GraphicsPSOInit, 0);
 
-		if (SrcRect.Size() != DstRect.Size())
-		{
-			PixelShader->SetParameters(RHICmdList, TStaticSamplerState<SF_Bilinear>::GetRHI(), SrcTexture);
-		}
-		else
-		{
-			PixelShader->SetParameters(RHICmdList, TStaticSamplerState<SF_Point>::GetRHI(), SrcTexture);
-		}
+		const bool bSameSize = (SrcRect.Size() == DstRect.Size());
+		FRHISamplerState* PixelSampler = bSameSize ? TStaticSamplerState<SF_Point>::GetRHI() : TStaticSamplerState<SF_Bilinear>::GetRHI();
+
+		SetShaderParametersLegacyPS(RHICmdList, PixelShader, PixelSampler, SrcTexture);
 
 		// Set up vertex uniform parameters for scaling and biasing the rectangle.
 		// Note: Use DrawRectangle in the vertex shader to calculate the correct vertex position and uv.
diff --git a/Engine/Plugins/Runtime/nDisplay/Source/DisplayClusterMedia/Private/DisplayClusterMediaHelpers.cpp b/Engine/Plugins/Runtime/nDisplay/Source/DisplayClusterMedia/Private/DisplayClusterMediaHelpers.cpp
--- a/Engine/Plugins/Runtime/nDisplay/Source/DisplayClusterMedia/Private/DisplayClusterMediaHelpers.cpp
+++ b/Engine/Plugins/Runtime/nDisplay/Source/DisplayClusterMedia/Private/DisplayClusterMediaHelpers.cpp
@@ -51,14 +51,9 @@ namespace DisplayClusterMediaHelpers
 
 			SetGraphicsPipelineState(RHICmdList, GraphicsPSOInit, 0);
 
-			if (SrcRect.Size() != DstRect.Size())
-			{
-				PixelShader->SetParameters(RHICmdList, TStaticSamplerState<SF_Bilinear>::GetRHI(), SrcTexture);
-			}
-			else
-			{
-				PixelShader->SetParameters(RHICmdList, TStaticSamplerState<SF_Point>::GetRHI(), SrcTexture);
-			}
+			FRHISamplerState* SamplerState = (SrcRect.Size() != DstRect.Size()) ? TStaticSamplerState<SF_Bilinear>::GetRHI() : TStaticSamplerState<SF_Point>::GetRHI();
+
+			SetShaderParametersLegacyPS(RHICmdList, PixelShader, SamplerState, SrcTexture);
 
 			// Set up vertex uniform parameters for scaling and biasing the rectangle.
 			// Note: Use DrawRectangle in the vertex shader to calculate the correct vertex position and uv.
diff --git a/Engine/Plugins/VirtualProduction/DMX/DMXPixelMapping/Source/DMXPixelMappingRenderer/Private/DMXPixelMappingRenderer.cpp b/Engine/Plugins/VirtualProduction/DMX/DMXPixelMapping/Source/DMXPixelMappingRenderer/Private/DMXPixelMappingRenderer.cpp
--- a/Engine/Plugins/VirtualProduction/DMX/DMXPixelMapping/Source/DMXPixelMappingRenderer/Private/DMXPixelMappingRenderer.cpp
+++ b/Engine/Plugins/VirtualProduction/DMX/DMXPixelMapping/Source/DMXPixelMappingRenderer/Private/DMXPixelMappingRenderer.cpp
@@ -280,7 +280,7 @@ void FDMXPixelMappingRenderer::RenderPreview(const FTextureResource* TextureReso
 				GraphicsPSOInit.PrimitiveType = PT_TriangleList;
 				SetGraphicsPipelineState(RHICmdList, GraphicsPSOInit, 0);
 
-				PixelShader->SetParameters(RHICmdList, TStaticSamplerState<SF_Point>::GetRHI(), DownsampleTextureRef);
+				SetShaderParametersLegacyPS(RHICmdList, PixelShader, TStaticSamplerState<SF_Point>::GetRHI(), DownsampleTextureRef);
 
 				const float DownsampleSizeX = DownsampleResource->GetSizeX();
 				const float DownsampleSizeY = DownsampleResource->GetSizeY();
@@ -415,14 +415,14 @@ void FDMXPixelMappingRenderer::RenderTextureToRectangle(const FTextureResource*
 				TShaderMapRef<FScreenPSsRGBSource> PixelShader(ShaderMap);
 				GraphicsPSOInit.BoundShaderState.PixelShaderRHI = PixelShader.GetPixelShader();
 				SetGraphicsPipelineState(RHICmdList, GraphicsPSOInit, 0);
-				PixelShader->SetParameters(RHICmdList, TStaticSamplerState<SF_Point>::GetRHI(), RenderContext.TextureResource->TextureRHI);
+				SetShaderParametersLegacyPS(RHICmdList, PixelShader, TStaticSamplerState<SF_Point>::GetRHI(), RenderContext.TextureResource->TextureRHI);
 			}
 			else
 			{
 				TShaderMapRef<FScreenPS> PixelShader(ShaderMap);
 				GraphicsPSOInit.BoundShaderState.PixelShaderRHI = PixelShader.GetPixelShader();
 				SetGraphicsPipelineState(RHICmdList, GraphicsPSOInit, 0);
-				PixelShader->SetParameters(RHICmdList, TStaticSamplerState<SF_Point>::GetRHI(), RenderContext.TextureResource->TextureRHI);
+				SetShaderParametersLegacyPS(RHICmdList, PixelShader, TStaticSamplerState<SF_Point>::GetRHI(), RenderContext.TextureResource->TextureRHI);
 			}
 
 			RendererModule->DrawRectangle(
diff --git a/Engine/Plugins/VirtualProduction/TextureShare/Source/TextureShare/Private/Resources/TextureShareResourcesProxy.cpp b/Engine/Plugins/VirtualProduction/TextureShare/Source/TextureShare/Private/Resources/TextureShareResourcesProxy.cpp
--- a/Engine/Plugins/VirtualProduction/TextureShare/Source/TextureShare/Private/Resources/TextureShareResourcesProxy.cpp
+++ b/Engine/Plugins/VirtualProduction/TextureShare/Source/TextureShare/Private/Resources/TextureShareResourcesProxy.cpp
@@ -103,14 +103,10 @@ namespace TextureShareResourcesHelpers
 
 			SetGraphicsPipelineState(RHICmdList, GraphicsPSOInit, 0);
 
-			if (SrcRect.Size() != DstRect.Size())
-			{
-				PixelShader->SetParameters(RHICmdList, TStaticSamplerState<SF_Bilinear>::GetRHI(), SrcTexture);
-			}
-			else
-			{
-				PixelShader->SetParameters(RHICmdList, TStaticSamplerState<SF_Point>::GetRHI(), SrcTexture);
-			}
+			const bool bSameSize = (SrcRect.Size() == DstRect.Size());
+			FRHISamplerState* PixelSampler = bSameSize ? TStaticSamplerState<SF_Point>::GetRHI() : TStaticSamplerState<SF_Bilinear>::GetRHI();
+
+			SetShaderParametersLegacyPS(RHICmdList, PixelShader, PixelSampler, SrcTexture);
 
 			// Set up vertex uniform parameters for scaling and biasing the rectangle.
 			// Note: Use DrawRectangle in the vertex shader to calculate the correct vertex position and uv.
diff --git a/Engine/Source/Editor/UnrealEd/Private/MeshPaintRendering.cpp b/Engine/Source/Editor/UnrealEd/Private/MeshPaintRendering.cpp
--- a/Engine/Source/Editor/UnrealEd/Private/MeshPaintRendering.cpp
+++ b/Engine/Source/Editor/UnrealEd/Private/MeshPaintRendering.cpp
@@ -42,9 +42,9 @@ namespace MeshPaintRendering
 			TransformParameter.Bind( Initializer.ParameterMap, TEXT( "c_Transform" ) );
 		}
 
-		void SetParameters(FRHICommandList& RHICmdList, const FMatrix44f& InTransform )
+		void SetParameters(FRHIBatchedShaderParameters& BatchedParameters, const FMatrix44f& InTransform )
 		{
-			SetShaderValue(RHICmdList, RHICmdList.GetBoundVertexShader(), TransformParameter, InTransform );
+			SetShaderValue(BatchedParameters, TransformParameter, InTransform );
 		}
 
 	private:
@@ -85,44 +85,41 @@ namespace MeshPaintRendering
 			GammaParameter.Bind( Initializer.ParameterMap, TEXT( "c_Gamma" ) );
 		}
 
-		void SetParameters(FRHICommandList& RHICmdList, const float InGamma, const FMeshPaintShaderParameters& InShaderParams )
+		void SetParameters(FRHIBatchedShaderParameters& BatchedParameters, const float InGamma, const FMeshPaintShaderParameters& InShaderParams )
 		{
-			FRHIPixelShader* ShaderRHI = RHICmdList.GetBoundPixelShader();
-
 			SetTextureParameter(
-				RHICmdList, 
-				ShaderRHI,
+				BatchedParameters,
 				CloneTextureParameter,
 				CloneTextureParameterSampler,
 				TStaticSamplerState< SF_Point, AM_Clamp, AM_Clamp, AM_Clamp >::GetRHI(),
 				InShaderParams.CloneTexture->GetRenderTargetResource()->TextureRHI );
 
-			SetShaderValue(RHICmdList, ShaderRHI, WorldToBrushMatrixParameter, (FMatrix44f)InShaderParams.WorldToBrushMatrix );
+			SetShaderValue(BatchedParameters, WorldToBrushMatrixParameter, (FMatrix44f)InShaderParams.WorldToBrushMatrix );
 
 			FVector4f BrushMetrics;
 			BrushMetrics.X = InShaderParams.BrushRadius;
 			BrushMetrics.Y = InShaderParams.BrushRadialFalloffRange;
 			BrushMetrics.Z = InShaderParams.BrushDepth;
 			BrushMetrics.W = InShaderParams.BrushDepthFalloffRange;
-			SetShaderValue(RHICmdList, ShaderRHI, BrushMetricsParameter, BrushMetrics );
+			SetShaderValue(BatchedParameters, BrushMetricsParameter, BrushMetrics );
 
 			FVector4f BrushStrength4( InShaderParams.BrushStrength, 0.0f, 0.0f, 0.0f );
-			SetShaderValue(RHICmdList, ShaderRHI, BrushStrengthParameter, BrushStrength4 );
+			SetShaderValue(BatchedParameters, BrushStrengthParameter, BrushStrength4 );
 
-			SetShaderValue(RHICmdList, ShaderRHI, BrushColorParameter, InShaderParams.BrushColor );
+			SetShaderValue(BatchedParameters, BrushColorParameter, InShaderParams.BrushColor );
 
 			FVector4f ChannelFlags;
 			ChannelFlags.X = InShaderParams.RedChannelFlag;
 			ChannelFlags.Y = InShaderParams.GreenChannelFlag;
 			ChannelFlags.Z = InShaderParams.BlueChannelFlag;
 			ChannelFlags.W = InShaderParams.AlphaChannelFlag;
-			SetShaderValue(RHICmdList, ShaderRHI, ChannelFlagsParameter, ChannelFlags );
+			SetShaderValue(BatchedParameters, ChannelFlagsParameter, ChannelFlags );
 			
 			float MaskVal = InShaderParams.GenerateMaskFlag ? 1.0f : 0.0f;
-			SetShaderValue(RHICmdList, ShaderRHI, GenerateMaskFlagParameter, MaskVal );
+			SetShaderValue(BatchedParameters, GenerateMaskFlagParameter, MaskVal );
 
 			// @todo MeshPaint
-			SetShaderValue(RHICmdList, ShaderRHI, GammaParameter, InGamma );
+			SetShaderValue(BatchedParameters, GammaParameter, InGamma );
 		}
 
 	private:
@@ -179,9 +176,9 @@ namespace MeshPaintRendering
 			TransformParameter.Bind( Initializer.ParameterMap, TEXT( "c_Transform" ) );
 		}
 
-		void SetParameters(FRHICommandList& RHICmdList, const FMatrix44f& InTransform )
+		void SetParameters(FRHIBatchedShaderParameters& BatchedParameters, const FMatrix44f& InTransform )
 		{
-			SetShaderValue(RHICmdList, RHICmdList.GetBoundVertexShader(), TransformParameter, InTransform );
+			SetShaderValue(BatchedParameters, TransformParameter, InTransform );
 		}
 
 	private:
@@ -224,37 +221,32 @@ namespace MeshPaintRendering
 			GammaParameter.Bind( Initializer.ParameterMap, TEXT( "Gamma" ) );
 		}
 
-		void SetParameters(FRHICommandList& RHICmdList, const float InGamma, const FMeshPaintDilateShaderParameters& InShaderParams )
+		void SetParameters(FRHIBatchedShaderParameters& BatchedParameters, const float InGamma, const FMeshPaintDilateShaderParameters& InShaderParams )
 		{
-			FRHIPixelShader* ShaderRHI = RHICmdList.GetBoundPixelShader();
-
 			SetTextureParameter(
-				RHICmdList, 
-				ShaderRHI,
+				BatchedParameters,
 				Texture0Parameter,
 				Texture0ParameterSampler,
 				TStaticSamplerState< SF_Point, AM_Clamp, AM_Clamp, AM_Clamp >::GetRHI(),
 				InShaderParams.Texture0->GetRenderTargetResource()->TextureRHI );
 
 			SetTextureParameter(
-				RHICmdList, 
-				ShaderRHI,
+				BatchedParameters,
 				Texture1Parameter,
 				Texture1ParameterSampler,
 				TStaticSamplerState< SF_Point, AM_Clamp, AM_Clamp, AM_Clamp >::GetRHI(),
 				InShaderParams.Texture1->GetRenderTargetResource()->TextureRHI );
 
 			SetTextureParameter(
-				RHICmdList, 
-				ShaderRHI,
+				BatchedParameters,
 				Texture2Parameter,
 				Texture2ParameterSampler,
 				TStaticSamplerState< SF_Point, AM_Clamp, AM_Clamp, AM_Clamp >::GetRHI(),
 				InShaderParams.Texture2->GetRenderTargetResource()->TextureRHI );
 
-			SetShaderValue(RHICmdList, ShaderRHI, WidthPixelOffsetParameter, InShaderParams.WidthPixelOffset );
-			SetShaderValue(RHICmdList, ShaderRHI, HeightPixelOffsetParameter, InShaderParams.HeightPixelOffset );
-			SetShaderValue(RHICmdList, ShaderRHI, GammaParameter, InGamma );
+			SetShaderValue(BatchedParameters, WidthPixelOffsetParameter, InShaderParams.WidthPixelOffset );
+			SetShaderValue(BatchedParameters, HeightPixelOffsetParameter, InShaderParams.HeightPixelOffset );
+			SetShaderValue(BatchedParameters, GammaParameter, InGamma );
 		}
 
 
@@ -322,10 +314,10 @@ namespace MeshPaintRendering
 		SetGraphicsPipelineState(RHICmdList, GraphicsPSOInit, 0);
 
 		// Set vertex shader parameters
-		VertexShader->SetParameters(RHICmdList, FMatrix44f(InTransform) );	// LWC_TODO: Precision loss
+		SetShaderParametersLegacyVS(RHICmdList, VertexShader, FMatrix44f(InTransform) );	// LWC_TODO: Precision loss
 		
 		// Set pixel shader parameters
-		PixelShader->SetParameters(RHICmdList, InGamma, InShaderParams );
+		SetShaderParametersLegacyPS(RHICmdList, PixelShader, InGamma, InShaderParams );
 
 		// @todo MeshPaint: Make sure blending/color writes are setup so we can write to ALPHA if needed!
 	}
@@ -347,10 +339,10 @@ namespace MeshPaintRendering
 		SetGraphicsPipelineState(RHICmdList, GraphicsPSOInit, 0);
 
 		// Set vertex shader parameters
-		VertexShader->SetParameters(RHICmdList, FMatrix44f(InTransform) );	// LWC_TODO: Precision loss
+		SetShaderParametersLegacyVS(RHICmdList, VertexShader, FMatrix44f(InTransform) );	// LWC_TODO: Precision loss
 
 		// Set pixel shader parameters
-		PixelShader->SetParameters(RHICmdList, InGamma, InShaderParams );
+		SetShaderParametersLegacyPS(RHICmdList, PixelShader, InGamma, InShaderParams );
 	}
 
 }
diff --git a/Engine/Source/Runtime/Engine/Private/SkeletalRenderGPUSkin.cpp b/Engine/Source/Runtime/Engine/Private/SkeletalRenderGPUSkin.cpp
--- a/Engine/Source/Runtime/Engine/Private/SkeletalRenderGPUSkin.cpp
+++ b/Engine/Source/Runtime/Engine/Private/SkeletalRenderGPUSkin.cpp
@@ -1118,26 +1118,28 @@ FGPUMorphUpdateCS::FGPUMorphUpdateCS(const ShaderMetaType::CompiledShaderInitial
 	MorphDataBufferParameter.Bind(Initializer.ParameterMap, TEXT("MorphDataBuffer"));
 }
 
-void FGPUMorphUpdateCS::SetParameters(FRHICommandList& RHICmdList, const FVector4& LocalScale, const FMorphTargetVertexInfoBuffers& MorphTargetVertexInfoBuffers, FMorphVertexBuffer& MorphVertexBuffer, uint32 NumGroups)
+void FGPUMorphUpdateCS::SetParameters(
+	FRHIBatchedShaderParameters& BatchedParameters,
+	const FVector4& LocalScale,
+	const FMorphTargetVertexInfoBuffers& MorphTargetVertexInfoBuffers,
+	FMorphVertexBuffer& MorphVertexBuffer,
+	uint32 NumGroups,
+	uint32 BatchOffsets[MorphTargetDispatchBatchSize],
+	uint32 GroupOffsets[MorphTargetDispatchBatchSize],
+	float Weights[MorphTargetDispatchBatchSize])
 {
-	FRHIComputeShader* CS = RHICmdList.GetBoundComputeShader();
-
-	SetUAVParameter(RHICmdList, CS, MorphVertexBufferParameter, MorphVertexBuffer.GetUAV());
+	SetUAVParameter(BatchedParameters, MorphVertexBufferParameter, MorphVertexBuffer.GetUAV());
 
-	SetShaderValue(RHICmdList, CS, PositionScaleParameter, (FVector4f)LocalScale);
+	SetShaderValue(BatchedParameters, PositionScaleParameter, (FVector4f)LocalScale);
 	FVector2f Precision = { MorphTargetVertexInfoBuffers.GetPositionPrecision(), MorphTargetVertexInfoBuffers.GetTangentZPrecision() };
-	SetShaderValue(RHICmdList, CS, PrecisionParameter, Precision);
-	SetShaderValue(RHICmdList, CS, NumGroupsParameter, NumGroups);
+	SetShaderValue(BatchedParameters, PrecisionParameter, Precision);
+	SetShaderValue(BatchedParameters, NumGroupsParameter, NumGroups);
 
-	SetSRVParameter(RHICmdList, CS, MorphDataBufferParameter, MorphTargetVertexInfoBuffers.MorphDataSRV);
-}
+	SetSRVParameter(BatchedParameters, MorphDataBufferParameter, MorphTargetVertexInfoBuffers.MorphDataSRV);
 
-void FGPUMorphUpdateCS::SetMorphOffsetsAndWeights(FRHICommandList& RHICmdList, uint32 BatchOffsets[MorphTargetDispatchBatchSize], uint32 GroupOffsets[MorphTargetDispatchBatchSize], float Weights[MorphTargetDispatchBatchSize])
-{
-	FRHIComputeShader* CS = RHICmdList.GetBoundComputeShader();
-	SetShaderValue(RHICmdList, CS, MorphTargetBatchOffsetsParameter,*(uint32(*)[MorphTargetDispatchBatchSize]) BatchOffsets);
-	SetShaderValue(RHICmdList, CS, MorphTargetGroupOffsetsParameter,*(uint32(*)[MorphTargetDispatchBatchSize]) GroupOffsets);
-	SetShaderValue(RHICmdList, CS, MorphTargetWeightsParameter,		*(float(*)[MorphTargetDispatchBatchSize]) Weights);
+	SetShaderValue(BatchedParameters, MorphTargetBatchOffsetsParameter, *(uint32(*)[MorphTargetDispatchBatchSize]) BatchOffsets);
+	SetShaderValue(BatchedParameters, MorphTargetGroupOffsetsParameter, *(uint32(*)[MorphTargetDispatchBatchSize]) GroupOffsets);
+	SetShaderValue(BatchedParameters, MorphTargetWeightsParameter, *(float(*)[MorphTargetDispatchBatchSize]) Weights);
 }
 
 void FGPUMorphUpdateCS::Dispatch(FRHICommandList& RHICmdList, uint32 Size)
@@ -1146,10 +1148,9 @@ void FGPUMorphUpdateCS::Dispatch(FRHICommandList& RHICmdList, uint32 Size)
 	RHICmdList.DispatchComputeShader(DispatchSize.X, DispatchSize.Y, DispatchSize.Z);
 }
 
-void FGPUMorphUpdateCS::EndAllDispatches(FRHICommandList& RHICmdList)
+void FGPUMorphUpdateCS::UnsetParameters(FRHIBatchedShaderParameters& BatchedParameters)
 {
-	FRHIComputeShader* CS = RHICmdList.GetBoundComputeShader();
-	SetUAVParameter(RHICmdList, CS, MorphVertexBufferParameter, nullptr);
+	SetUAVParameter(BatchedParameters, MorphVertexBufferParameter, nullptr);
 }
 
 bool FGPUMorphUpdateCS::ShouldCompilePermutation(const FGlobalShaderPermutationParameters& Parameters)
@@ -1174,12 +1175,11 @@ bool FGPUMorphNormalizeCS::ShouldCompilePermutation(const FGlobalShaderPermutati
 	return IsFeatureLevelSupported(Parameters.Platform, ERHIFeatureLevel::SM5);
 }
 
-void FGPUMorphNormalizeCS::SetParameters(FRHICommandList& RHICmdList, const FVector4& InvLocalScale, const FMorphTargetVertexInfoBuffers& MorphTargetVertexInfoBuffers, FMorphVertexBuffer& MorphVertexBuffer, uint32 NumVertices)
+void FGPUMorphNormalizeCS::SetParameters(FRHIBatchedShaderParameters& BatchedParameters, const FVector4& InvLocalScale, const FMorphTargetVertexInfoBuffers& MorphTargetVertexInfoBuffers, FMorphVertexBuffer& MorphVertexBuffer, uint32 NumVertices)
 {
-	FRHIComputeShader* CS = RHICmdList.GetBoundComputeShader();
-	SetUAVParameter(RHICmdList, CS, MorphVertexBufferParameter, MorphVertexBuffer.GetUAV());
-	SetShaderValue(RHICmdList, CS, PositionScaleParameter, (FVector4f)InvLocalScale);
-	SetShaderValue(RHICmdList, CS, NumVerticesParameter, NumVertices);
+	SetUAVParameter(BatchedParameters, MorphVertexBufferParameter, MorphVertexBuffer.GetUAV());
+	SetShaderValue(BatchedParameters, PositionScaleParameter, (FVector4f)InvLocalScale);
+	SetShaderValue(BatchedParameters, NumVerticesParameter, NumVertices);
 }
 
 void FGPUMorphNormalizeCS::Dispatch(FRHICommandList& RHICmdList, uint32 NumVertices)
@@ -1188,10 +1188,9 @@ void FGPUMorphNormalizeCS::Dispatch(FRHICommandList& RHICmdList, uint32 NumVerti
 	RHICmdList.DispatchComputeShader(DispatchSize.X, DispatchSize.Y, DispatchSize.Z);
 }
 
-void FGPUMorphNormalizeCS::EndAllDispatches(FRHICommandList& RHICmdList)
+void FGPUMorphNormalizeCS::UnsetParameters(FRHIBatchedShaderParameters& BatchedParameters)
 {
-	FRHIComputeShader* CS = RHICmdList.GetBoundComputeShader();
-	SetUAVParameter(RHICmdList, CS, MorphVertexBufferParameter, nullptr);
+	SetUAVParameter(BatchedParameters, MorphVertexBufferParameter, nullptr);
 }
 
 IMPLEMENT_SHADER_TYPE(, FGPUMorphNormalizeCS, TEXT("/Engine/Private/MorphTargets.usf"), TEXT("GPUMorphNormalizeCS"), SF_Compute);
@@ -1279,12 +1278,22 @@ void FSkeletalMeshObjectGPUSkin::FSkeletalMeshObjectLOD::UpdateMorphVertexBuffer
 
 					SetComputePipelineState(RHICmdList, GPUMorphUpdateCS.GetComputeShader());
 
-					GPUMorphUpdateCS->SetParameters(RHICmdList, MorphScale, MorphTargetVertexInfoBuffers, MorphVertexBuffer, NumBatches);
-					GPUMorphUpdateCS->SetMorphOffsetsAndWeights(RHICmdList, BatchOffsets, GroupOffsets, Weights);
+					SetShaderParametersLegacyCS(
+						RHICmdList,
+						GPUMorphUpdateCS,
+						MorphScale,
+						MorphTargetVertexInfoBuffers,
+						MorphVertexBuffer,
+						NumBatches,
+						BatchOffsets,
+						GroupOffsets,
+						Weights);
+
 					GPUMorphUpdateCS->Dispatch(RHICmdList, NumBatches);
 				}
+
+				UnsetShaderParametersLegacyCS(RHICmdList, GPUMorphUpdateCS);
 				
-				GPUMorphUpdateCS->EndAllDispatches(RHICmdList);
 				RHICmdList.EndUAVOverlap(MorphVertexBuffer.GetUAV());
 			}
 
@@ -1302,10 +1311,10 @@ void FSkeletalMeshObjectGPUSkin::FSkeletalMeshObjectLOD::UpdateMorphVertexBuffer
 				TShaderMapRef<FGPUMorphNormalizeCS> GPUMorphNormalizeCS(GetGlobalShaderMap(FeatureLevel));
 
 				SetComputePipelineState(RHICmdList, GPUMorphNormalizeCS.GetComputeShader());
-
-				GPUMorphNormalizeCS->SetParameters(RHICmdList, InvMorphScale, MorphTargetVertexInfoBuffers, MorphVertexBuffer, MorphVertexBuffer.GetNumVerticies());
+				SetShaderParametersLegacyCS(RHICmdList, GPUMorphNormalizeCS, InvMorphScale, MorphTargetVertexInfoBuffers, MorphVertexBuffer, MorphVertexBuffer.GetNumVerticies());
 				GPUMorphNormalizeCS->Dispatch(RHICmdList, MorphVertexBuffer.GetNumVerticies());
-				GPUMorphNormalizeCS->EndAllDispatches(RHICmdList);
+				UnsetShaderParametersLegacyCS(RHICmdList, GPUMorphNormalizeCS);
+
 				RHICmdList.Transition(FRHITransitionInfo(MorphVertexBuffer.GetUAV(), ERHIAccess::UAVCompute, ERHIAccess::VertexOrIndexBuffer | ERHIAccess::SRVMask));
 			}
 		}
diff --git a/Engine/Source/Runtime/Engine/Private/SkeletalRenderGPUSkin.h b/Engine/Source/Runtime/Engine/Private/SkeletalRenderGPUSkin.h
--- a/Engine/Source/Runtime/Engine/Private/SkeletalRenderGPUSkin.h
+++ b/Engine/Source/Runtime/Engine/Private/SkeletalRenderGPUSkin.h
@@ -712,11 +712,18 @@ public:
 
 	static const uint32 MorphTargetDispatchBatchSize = 128;
 
-	void SetParameters(FRHICommandList& RHICmdList, const FVector4& LocalScale, const FMorphTargetVertexInfoBuffers& MorphTargetVertexInfoBuffers, FMorphVertexBuffer& MorphVertexBuffer, uint32 NumGroups);
-	void SetMorphOffsetsAndWeights(FRHICommandList& RHICmdList, uint32 BatchOffsets[MorphTargetDispatchBatchSize], uint32 GroupOffsets[MorphTargetDispatchBatchSize], float Weights[MorphTargetDispatchBatchSize]);
+	void SetParameters(
+		FRHIBatchedShaderParameters& BatchedParameters,
+		const FVector4& LocalScale,
+		const FMorphTargetVertexInfoBuffers& MorphTargetVertexInfoBuffers,
+		FMorphVertexBuffer& MorphVertexBuffer,
+		uint32 NumGroups,
+		uint32 BatchOffsets[MorphTargetDispatchBatchSize],
+		uint32 GroupOffsets[MorphTargetDispatchBatchSize],
+		float Weights[MorphTargetDispatchBatchSize]);
 
 	void Dispatch(FRHICommandList& RHICmdList, uint32 Size);
-	void EndAllDispatches(FRHICommandList& RHICmdList);
+	void UnsetParameters(FRHIBatchedShaderParameters& BatchedParameters);
 
 	static bool ShouldCompilePermutation(const FGlobalShaderPermutationParameters& Parameters);
 
@@ -744,10 +751,10 @@ public:
 
 	static bool ShouldCompilePermutation(const FGlobalShaderPermutationParameters& Parameters);
 
-	void SetParameters(FRHICommandList& RHICmdList, const FVector4& LocalScale, const FMorphTargetVertexInfoBuffers& MorphTargetVertexInfoBuffers, FMorphVertexBuffer& MorphVertexBuffer, uint32 NumVertices);
+	void SetParameters(FRHIBatchedShaderParameters& BatchedParameters, const FVector4& LocalScale, const FMorphTargetVertexInfoBuffers& MorphTargetVertexInfoBuffers, FMorphVertexBuffer& MorphVertexBuffer, uint32 NumVertices);
 
 	void Dispatch(FRHICommandList& RHICmdList, uint32 NumVertices);
-	void EndAllDispatches(FRHICommandList& RHICmdList);
+	void UnsetParameters(FRHIBatchedShaderParameters& BatchedParameters);
 
 protected:
 	LAYOUT_FIELD(FShaderResourceParameter, MorphVertexBufferParameter);
diff --git a/Engine/Source/Runtime/GameplayMediaEncoder/Private/GameplayMediaEncoder.cpp b/Engine/Source/Runtime/GameplayMediaEncoder/Private/GameplayMediaEncoder.cpp
--- a/Engine/Source/Runtime/GameplayMediaEncoder/Private/GameplayMediaEncoder.cpp
+++ b/Engine/Source/Runtime/GameplayMediaEncoder/Private/GameplayMediaEncoder.cpp
@@ -724,14 +724,10 @@ void FGameplayMediaEncoder::CopyTexture(const FTexture2DRHIRef& SourceTexture, F
 
 			SetGraphicsPipelineState(RHICmdList, GraphicsPSOInit, 0);
 
-			if(DestinationTexture->GetSizeX() != SourceTexture->GetSizeX() || DestinationTexture->GetSizeY() != SourceTexture->GetSizeY())
-			{
-				PixelShader->SetParameters(RHICmdList, TStaticSamplerState<SF_Bilinear>::GetRHI(), SourceTexture);
-			}
-			else
-			{
-				PixelShader->SetParameters(RHICmdList, TStaticSamplerState<SF_Point>::GetRHI(), SourceTexture);
-			}
+			const bool bSameSize = (DestinationTexture->GetDesc().Extent == SourceTexture->GetDesc().Extent);
+			FRHISamplerState* PixelSampler = bSameSize ? TStaticSamplerState<SF_Point>::GetRHI() : TStaticSamplerState<SF_Bilinear>::GetRHI();
+
+			SetShaderParametersLegacyPS(RHICmdList, PixelShader, PixelSampler, SourceTexture);
 
 			RendererModule->DrawRectangle(RHICmdList, 0, 0,                // Dest X, Y
 			                              (float)DestinationTexture->GetSizeX(),  // Dest Width
diff --git a/Engine/Source/Runtime/MovieSceneCapture/Private/FrameGrabber.cpp b/Engine/Source/Runtime/MovieSceneCapture/Private/FrameGrabber.cpp
--- a/Engine/Source/Runtime/MovieSceneCapture/Private/FrameGrabber.cpp
+++ b/Engine/Source/Runtime/MovieSceneCapture/Private/FrameGrabber.cpp
@@ -147,14 +147,9 @@ void FViewportSurfaceReader::ResolveRenderTarget(FViewportSurfaceReader* RenderT
 			const bool bIsSourceBackBufferSameAsWindowSize = SourceBackBuffer->GetSizeX() == WindowSize.X && SourceBackBuffer->GetSizeY() == WindowSize.Y;
 			const bool bIsSourceBackBufferSameAsTargetSize = TargetSize.X == SourceBackBuffer->GetSizeX() && TargetSize.Y == SourceBackBuffer->GetSizeY();
 
-			if (bIsSourceBackBufferSameAsWindowSize || bIsSourceBackBufferSameAsTargetSize)
-			{
-				PixelShader->SetParameters(RHICmdList, TStaticSamplerState<SF_Point>::GetRHI(), SourceBackBuffer);
-			}
-			else
-			{
-				PixelShader->SetParameters(RHICmdList, TStaticSamplerState<SF_Bilinear>::GetRHI(), SourceBackBuffer);
-			}
+			FRHISamplerState* SamplerState = (bIsSourceBackBufferSameAsWindowSize || bIsSourceBackBufferSameAsTargetSize) ? TStaticSamplerState<SF_Point>::GetRHI() : TStaticSamplerState<SF_Bilinear>::GetRHI();
+
+			SetShaderParametersLegacyPS(RHICmdList, PixelShader, SamplerState, SourceBackBuffer);
 
 			float U = float(CaptureRect.Min.X) / float(SourceBackBuffer->GetSizeX());
 			float V = float(CaptureRect.Min.Y) / float(SourceBackBuffer->GetSizeY());
diff --git a/Engine/Source/Runtime/RenderCore/Private/RayTracingValidationShaders.cpp b/Engine/Source/Runtime/RenderCore/Private/RayTracingValidationShaders.cpp
--- a/Engine/Source/Runtime/RenderCore/Private/RayTracingValidationShaders.cpp
+++ b/Engine/Source/Runtime/RenderCore/Private/RayTracingValidationShaders.cpp
@@ -50,18 +50,22 @@ void FRayTracingValidateGeometryBuildParamsCS::Dispatch(FRHICommandList& RHICmdL
 
 		const uint32 IndexBufferOffsetInBytes = Segment.FirstPrimitive * IndexStride * 3;
 
-		SetShaderValue(RHICmdList, ShaderRHI, ComputeShader->VertexBufferStrideParam, Segment.VertexBufferStride);
-		SetShaderValue(RHICmdList, ShaderRHI, ComputeShader->VertexBufferOffsetInBytesParam, Segment.VertexBufferOffset);
-		SetShaderValue(RHICmdList, ShaderRHI, ComputeShader->IndexBufferOffsetInBytesParam, IndexBufferOffsetInBytes);
-		SetShaderValue(RHICmdList, ShaderRHI, ComputeShader->IndexBufferStrideParam, IndexStride);
-		SetShaderValue(RHICmdList, ShaderRHI, ComputeShader->NumPrimitivesParam, Segment.NumPrimitives);
-		SetShaderValue(RHICmdList, ShaderRHI, ComputeShader->MaxVerticesParam, Segment.MaxVertices);
+		FRHIBatchedShaderParameters& BatchedParameters = RHICmdList.GetScratchShaderParameters();
+
+		SetShaderValue(BatchedParameters, ComputeShader->VertexBufferStrideParam, Segment.VertexBufferStride);
+		SetShaderValue(BatchedParameters, ComputeShader->VertexBufferOffsetInBytesParam, Segment.VertexBufferOffset);
+		SetShaderValue(BatchedParameters, ComputeShader->IndexBufferOffsetInBytesParam, IndexBufferOffsetInBytes);
+		SetShaderValue(BatchedParameters, ComputeShader->IndexBufferStrideParam, IndexStride);
+		SetShaderValue(BatchedParameters, ComputeShader->NumPrimitivesParam, Segment.NumPrimitives);
+		SetShaderValue(BatchedParameters, ComputeShader->MaxVerticesParam, Segment.MaxVertices);
 
 		const FRawBufferShaderResourceViewInitializer VBViewInitializer(Segment.VertexBuffer);
 		FShaderResourceViewRHIRef VertexBufferSRV = RHICreateShaderResourceView(VBViewInitializer);
 
-		SetSRVParameter(RHICmdList, ShaderRHI, ComputeShader->VertexBufferParam, VertexBufferSRV);
-		SetSRVParameter(RHICmdList, ShaderRHI, ComputeShader->IndexBufferParam, IndexBufferSRV);
+		SetSRVParameter(BatchedParameters, ComputeShader->VertexBufferParam, VertexBufferSRV);
+		SetSRVParameter(BatchedParameters, ComputeShader->IndexBufferParam, IndexBufferSRV);
+
+		RHICmdList.SetBatchedShaderParameters(ShaderRHI, BatchedParameters);
 
 		// TODO: handle arbitrary large meshes that may overrun the 1D dispatch limit
 		const uint32 MaxDispatchDimension = 65536; // D3D11_CS_DISPATCH_MAX_THREAD_GROUPS_PER_DIMENSION
@@ -70,8 +74,12 @@ void FRayTracingValidateGeometryBuildParamsCS::Dispatch(FRHICommandList& RHICmdL
 		RHICmdList.DispatchComputeShader(NumGroupsX, 1, 1);
 	}
 
-	SetSRVParameter(RHICmdList, ShaderRHI, ComputeShader->VertexBufferParam, nullptr);
-	SetSRVParameter(RHICmdList, ShaderRHI, ComputeShader->IndexBufferParam, nullptr);
+	FRHIBatchedShaderParameters& BatchedParameters = RHICmdList.GetScratchShaderParameters();
+
+	SetSRVParameter(BatchedParameters, ComputeShader->VertexBufferParam, nullptr);
+	SetSRVParameter(BatchedParameters, ComputeShader->IndexBufferParam, nullptr);
+
+	RHICmdList.SetBatchedShaderParameters(ShaderRHI, BatchedParameters);
 
 	RHICmdList.PopEvent();
 }
@@ -90,21 +98,29 @@ void FRayTracingValidateSceneBuildParamsCS::Dispatch(FRHICommandList& RHICmdList
 
 	RHICmdList.PushEvent(TEXT("RTSceneValidation"), FColor::Black);
 
-	const FRawBufferShaderResourceViewInitializer InstanceBufferViewInitializer(InstanceBuffer);
-	FShaderResourceViewRHIRef InstanceBufferSRV = RHICreateShaderResourceView(InstanceBufferViewInitializer);
+	{
+		FRHIBatchedShaderParameters& BatchedParameters = RHICmdList.GetScratchShaderParameters();
+
+		const FRawBufferShaderResourceViewInitializer InstanceBufferViewInitializer(InstanceBuffer);
+		FShaderResourceViewRHIRef InstanceBufferSRV = RHICreateShaderResourceView(InstanceBufferViewInitializer);
+
+		SetShaderValue(BatchedParameters, ComputeShader->NumInstancesParam, NumInstances);
+		SetShaderValue(BatchedParameters, ComputeShader->NumHitGroupsParam, NumHitGroups);
+		SetShaderValue(BatchedParameters, ComputeShader->InstanceBufferOffsetInBytesParam, InstanceBufferOffset);
+		SetShaderValue(BatchedParameters, ComputeShader->InstanceBufferStrideInBytesParam, InstanceBufferStride);
+		SetSRVParameter(BatchedParameters, ComputeShader->InstanceBufferParam, InstanceBufferSRV);
 
-	SetShaderValue(RHICmdList, ShaderRHI, ComputeShader->NumInstancesParam, NumInstances);
-	SetShaderValue(RHICmdList, ShaderRHI, ComputeShader->NumHitGroupsParam, NumHitGroups);
-	SetShaderValue(RHICmdList, ShaderRHI, ComputeShader->InstanceBufferOffsetInBytesParam, InstanceBufferOffset);
-	SetShaderValue(RHICmdList, ShaderRHI, ComputeShader->InstanceBufferStrideInBytesParam, InstanceBufferStride);
-	SetSRVParameter(RHICmdList, ShaderRHI, ComputeShader->InstanceBufferParam, InstanceBufferSRV);
+		const uint32 MaxDispatchDimension = 65536; // D3D11_CS_DISPATCH_MAX_THREAD_GROUPS_PER_DIMENSION
+		const uint32 NumGroupsX = FMath::Min((NumInstances + NumThreadsX - 1) / NumThreadsX, MaxDispatchDimension);
 
-	const uint32 MaxDispatchDimension = 65536; // D3D11_CS_DISPATCH_MAX_THREAD_GROUPS_PER_DIMENSION
-	const uint32 NumGroupsX = FMath::Min((NumInstances + NumThreadsX - 1) / NumThreadsX, MaxDispatchDimension);
+		RHICmdList.SetBatchedShaderParameters(ShaderRHI, BatchedParameters);
 
-	RHICmdList.DispatchComputeShader(NumGroupsX, 1, 1);
+		RHICmdList.DispatchComputeShader(NumGroupsX, 1, 1);
+	}
 
-	SetSRVParameter(RHICmdList, ShaderRHI, ComputeShader->InstanceBufferParam, nullptr);
+	FRHIBatchedShaderParameters& BatchedParameters = RHICmdList.GetScratchShaderParameters();
+	SetSRVParameter(BatchedParameters, ComputeShader->InstanceBufferParam, nullptr);
+	RHICmdList.SetBatchedShaderParameters(ShaderRHI, BatchedParameters);
 
 	RHICmdList.PopEvent();
 }
diff --git a/Engine/Source/Runtime/Renderer/Private/DecalRenderingShared.cpp b/Engine/Source/Runtime/Renderer/Private/DecalRenderingShared.cpp
--- a/Engine/Source/Runtime/Renderer/Private/DecalRenderingShared.cpp
+++ b/Engine/Source/Runtime/Renderer/Private/DecalRenderingShared.cpp
@@ -86,11 +86,12 @@ public:
 		MobileBasePassUniformBuffer.Bind(Initializer.ParameterMap, FMobileBasePassUniformParameters::FTypeInfo::GetStructMetadata()->GetShaderVariableName());
 	}
 
-	void SetParameters(FRHICommandList& RHICmdList, const FViewInfo& View, const FDeferredDecalProxy& DecalProxy, const FMaterialRenderProxy* MaterialProxy, const FMaterial* MaterialResource, const float FadeAlphaValue = 1.0f)
+	void SetParameters(FRHIBatchedShaderParameters& BatchedParameters, const FViewInfo& View, const FDeferredDecalProxy& DecalProxy, const FMaterialRenderProxy* MaterialProxy, const FMaterial* MaterialResource, const float FadeAlphaValue = 1.0f)
 	{
-		FRHIPixelShader* ShaderRHI = RHICmdList.GetBoundPixelShader();
+		auto& PrimitivePS = GetUniformBufferParameter<FPrimitiveUniformShaderParameters>();
+		SetUniformBufferParameter(BatchedParameters, PrimitivePS, GIdentityPrimitiveUniformBuffer);
 
-		FMaterialShader::SetParameters(RHICmdList, ShaderRHI, MaterialProxy, *MaterialResource, View);
+		FMaterialShader::SetParameters(BatchedParameters, MaterialProxy, *MaterialResource, View);
 
 		const FMatrix DecalToWorldMatrix = DecalProxy.ComponentTrans.ToMatrixWithScale();
 		const FMatrix WorldToDecalMatrix = DecalProxy.ComponentTrans.ToInverseMatrixWithScale();
@@ -101,7 +102,7 @@ public:
 
 		if (DecalTilePosition.IsBound())
 		{
-			SetShaderValue(RHICmdList, ShaderRHI, DecalTilePosition, TilePosition);
+			SetShaderValue(BatchedParameters, DecalTilePosition, TilePosition);
 		}
 		if(SvPositionToDecal.IsBound())
 		{
@@ -125,19 +126,19 @@ public:
 					FPlane(Ax, Ay,   0,  1)
 				) * View.ViewMatrices.GetInvViewProjectionMatrix() * WorldToDecalMatrix);
 
-			SetShaderValue(RHICmdList, ShaderRHI, SvPositionToDecal, SvPositionToDecalValue);
+			SetShaderValue(BatchedParameters, SvPositionToDecal, SvPositionToDecalValue);
 		}
 		if(DecalToWorld.IsBound())
 		{
-			SetShaderValue(RHICmdList, ShaderRHI, DecalToWorld, RelativeDecalToWorldMatrix);
+			SetShaderValue(BatchedParameters, DecalToWorld, RelativeDecalToWorldMatrix);
 		}
 		if (DecalToWorldInvScale.IsBound())
 		{
-			SetShaderValue(RHICmdList, ShaderRHI, DecalToWorldInvScale, RelativeDecalToWorldMatrix.GetScaleVector().Reciprocal());
+			SetShaderValue(BatchedParameters, DecalToWorldInvScale, RelativeDecalToWorldMatrix.GetScaleVector().Reciprocal());
 		}
 		if (DecalOrientation.IsBound())
 		{
-			SetShaderValue(RHICmdList, ShaderRHI, DecalOrientation, OrientationVector);
+			SetShaderValue(BatchedParameters, DecalOrientation, OrientationVector);
 		}
 		
 		float LifetimeAlpha = 1.0f;
@@ -148,8 +149,8 @@ public:
 			LifetimeAlpha = FMath::Clamp(FMath::Min(View.Family->Time.GetWorldTimeSeconds() * -DecalProxy.InvFadeDuration + DecalProxy.FadeStartDelayNormalized, View.Family->Time.GetWorldTimeSeconds() * DecalProxy.InvFadeInDuration + DecalProxy.FadeInStartDelayNormalized), 0.0f, 1.0f);
 		}
  
-		SetShaderValue(RHICmdList, ShaderRHI, DecalParams, FVector2f(FadeAlphaValue, LifetimeAlpha));
-		SetShaderValue(RHICmdList, ShaderRHI, DecalColorParam, DecalProxy.DecalColor);
+		SetShaderValue(BatchedParameters, DecalParams, FVector2f(FadeAlphaValue, LifetimeAlpha));
+		SetShaderValue(BatchedParameters, DecalColorParam, DecalProxy.DecalColor);
 	}
 
 private:
@@ -454,9 +455,7 @@ namespace DecalRendering
 
 		// Set pixel shader parameters.
 		{
-			PixelShader->SetParameters(RHICmdList, View, DecalData.Proxy, MaterialProxy, MaterialResource, DecalData.FadeAlpha);
-			auto& PrimitivePS = PixelShader->GetUniformBufferParameter<FPrimitiveUniformShaderParameters>();
-			SetUniformBufferParameter(RHICmdList, PixelShader.GetPixelShader(), PrimitivePS, GIdentityPrimitiveUniformBuffer);
+			SetShaderParametersLegacyPS(RHICmdList, PixelShader, View, DecalData.Proxy, MaterialProxy, MaterialResource, DecalData.FadeAlpha);
 		}
 
 		// Set stream source after updating cached strides
diff --git a/Engine/Source/Runtime/Renderer/Private/GPUBenchmark.cpp b/Engine/Source/Runtime/Renderer/Private/GPUBenchmark.cpp
--- a/Engine/Source/Runtime/Renderer/Private/GPUBenchmark.cpp
+++ b/Engine/Source/Runtime/Renderer/Private/GPUBenchmark.cpp
@@ -62,13 +62,11 @@ public:
 		InputTextureSampler.Bind(Initializer.ParameterMap,TEXT("InputTextureSampler"));
 	}
 
-	void SetParameters(FRHICommandList& RHICmdList, const FSceneView& View, TRefCountPtr<IPooledRenderTarget>& Src)
+	void SetParameters(FRHIBatchedShaderParameters& BatchedParameters, const FSceneView& View, TRefCountPtr<IPooledRenderTarget>& Src)
 	{
-		FRHIPixelShader* ShaderRHI = RHICmdList.GetBoundPixelShader();
+		FGlobalShader::SetParameters<FViewUniformShaderParameters>(BatchedParameters, View.ViewUniformBuffer);
 
-		FGlobalShader::SetParameters<FViewUniformShaderParameters>(RHICmdList, ShaderRHI, View.ViewUniformBuffer);
-
-		SetTextureParameter(RHICmdList, ShaderRHI, InputTexture, InputTextureSampler, TStaticSamplerState<>::GetRHI(), Src->GetRHI());
+		SetTextureParameter(BatchedParameters, InputTexture, InputTextureSampler, TStaticSamplerState<>::GetRHI(), Src->GetRHI());
 	}
 
 	static const TCHAR* GetSourceFilename()
@@ -117,11 +115,9 @@ public:
 	{
 	}
 
-	void SetParameters(FRHICommandList& RHICmdList, const FSceneView& View)
+	void SetParameters(FRHIBatchedShaderParameters& BatchedParameters, const FSceneView& View)
 	{
-		FRHIVertexShader* ShaderRHI = RHICmdList.GetBoundVertexShader();
-		
-		FGlobalShader::SetParameters<FViewUniformShaderParameters>(RHICmdList, ShaderRHI, View.ViewUniformBuffer);
+		FGlobalShader::SetParameters<FViewUniformShaderParameters>(BatchedParameters, View.ViewUniformBuffer);
 	}
 };
 
@@ -202,8 +198,8 @@ void RunBenchmarkShader(FRHICommandList& RHICmdList, FRHIBuffer* VertexThroughpu
 
 	SetGraphicsPipelineState(RHICmdList, GraphicsPSOInit, 0);
 
-	PixelShader->SetParameters(RHICmdList, View, Src);
-	VertexShader->SetParameters(RHICmdList, View);
+	SetShaderParametersLegacyPS(RHICmdList, PixelShader, View, Src);
+	SetShaderParametersLegacyVS(RHICmdList, VertexShader, View);
 
 	if (bVertexTest)
 	{
diff --git a/Engine/Source/Runtime/Renderer/Private/HeterogeneousVolumes/HeterogeneousVolumesLiveShadingPipeline.cpp b/Engine/Source/Runtime/Renderer/Private/HeterogeneousVolumes/HeterogeneousVolumesLiveShadingPipeline.cpp
--- a/Engine/Source/Runtime/Renderer/Private/HeterogeneousVolumes/HeterogeneousVolumesLiveShadingPipeline.cpp
+++ b/Engine/Source/Runtime/Renderer/Private/HeterogeneousVolumes/HeterogeneousVolumesLiveShadingPipeline.cpp
@@ -116,18 +116,6 @@ class FRenderLightingCacheWithLiveShadingCS : public FMeshMaterialShader
 		OutEnvironment.SetDefine(TEXT("GET_PRIMITIVE_DATA_OVERRIDE"), 1);
 	}
 
-	void SetParameters(
-		FRHIComputeCommandList & RHICmdList,
-		FRHIComputeShader * ShaderRHI,
-		const FViewInfo & View,
-		const FMaterialRenderProxy * MaterialProxy,
-		const FMaterial & Material
-	)
-	{
-		FMaterialShader::SetViewParameters(RHICmdList, ShaderRHI, View, View.ViewUniformBuffer);
-		FMaterialShader::SetParameters(RHICmdList, ShaderRHI, MaterialProxy, Material, View);
-	}
-
 	static int32 GetThreadGroupSize1D() { return GetThreadGroupSize2D() * GetThreadGroupSize2D(); }
 	static int32 GetThreadGroupSize2D() { return 8; }
 	static int32 GetThreadGroupSize3D() { return 4; }
@@ -244,18 +232,6 @@ class FRenderSingleScatteringWithLiveShadingCS : public FMeshMaterialShader
 		OutEnvironment.SetDefine(TEXT("GET_PRIMITIVE_DATA_OVERRIDE"), 1);
 	}
 
-	void SetParameters(
-		FRHIComputeCommandList& RHICmdList,
-		FRHIComputeShader* ShaderRHI,
-		const FViewInfo& View,
-		const FMaterialRenderProxy* MaterialProxy,
-		const FMaterial& Material
-	)
-	{
-		FMaterialShader::SetViewParameters(RHICmdList, ShaderRHI, View, View.ViewUniformBuffer);
-		FMaterialShader::SetParameters(RHICmdList, ShaderRHI, MaterialProxy, Material, View);
-	}
-
 	static int32 GetThreadGroupSize1D() { return GetThreadGroupSize2D() * GetThreadGroupSize2D(); }
 	static int32 GetThreadGroupSize2D() { return 8; }
 };
diff --git a/Engine/Source/Runtime/Renderer/Private/HeterogeneousVolumes/HeterogeneousVolumesMaterialBakingPipeline.cpp b/Engine/Source/Runtime/Renderer/Private/HeterogeneousVolumes/HeterogeneousVolumesMaterialBakingPipeline.cpp
--- a/Engine/Source/Runtime/Renderer/Private/HeterogeneousVolumes/HeterogeneousVolumesMaterialBakingPipeline.cpp
+++ b/Engine/Source/Runtime/Renderer/Private/HeterogeneousVolumes/HeterogeneousVolumesMaterialBakingPipeline.cpp
@@ -89,18 +89,6 @@ class FHeterogeneousVolumesBakeMaterialCS : public FMeshMaterialShader
 		OutEnvironment.SetDefine(TEXT("GET_PRIMITIVE_DATA_OVERRIDE"), 1);
 	}
 
-	void SetParameters(
-		FRHIComputeCommandList& RHICmdList,
-		FRHIComputeShader* ShaderRHI,
-		const FViewInfo& View,
-		const FMaterialRenderProxy* MaterialProxy,
-		const FMaterial& Material
-	)
-	{
-		FMaterialShader::SetViewParameters(RHICmdList, ShaderRHI, View, View.ViewUniformBuffer);
-		FMaterialShader::SetParameters(RHICmdList, ShaderRHI, MaterialProxy, Material, View);
-	}
-
 	static int32 GetThreadGroupSize1D() { return GetThreadGroupSize2D() * GetThreadGroupSize2D(); }
 	static int32 GetThreadGroupSize2D() { return 8; }
 	static int32 GetThreadGroupSize3D() { return 4; }
diff --git a/Engine/Source/Runtime/Renderer/Private/RayTracing/RaytracingSkylight.cpp b/Engine/Source/Runtime/Renderer/Private/RayTracing/RaytracingSkylight.cpp
--- a/Engine/Source/Runtime/Renderer/Private/RayTracing/RaytracingSkylight.cpp
+++ b/Engine/Source/Runtime/Renderer/Private/RayTracing/RaytracingSkylight.cpp
@@ -679,9 +679,8 @@ public:
 
 	FVisualizeSkyLightMipTreePS() {}
 
-	template<typename TRHICommandList>
 	void SetParameters(
-		TRHICommandList& RHICmdList,
+		FRHIBatchedShaderParameters& BatchedParameters,
 		const FViewInfo& View,
 		const FIntVector Dimensions,
 		const FRWBuffer& MipTreePosX,
@@ -691,16 +690,15 @@ public:
 		const FRWBuffer& MipTreePosZ,
 		const FRWBuffer& MipTreeNegZ)
 	{
-		FRHIPixelShader* ShaderRHI = RHICmdList.GetBoundPixelShader();
-		FGlobalShader::SetParameters<FViewUniformShaderParameters>(RHICmdList, ShaderRHI, View.ViewUniformBuffer);
-
-		SetShaderValue(RHICmdList, ShaderRHI, DimensionsParameter, Dimensions);
-		SetSRVParameter(RHICmdList, ShaderRHI, MipTreePosXParameter, MipTreePosX.SRV);
-		SetSRVParameter(RHICmdList, ShaderRHI, MipTreeNegXParameter, MipTreeNegX.SRV);
-		SetSRVParameter(RHICmdList, ShaderRHI, MipTreePosYParameter, MipTreePosY.SRV);
-		SetSRVParameter(RHICmdList, ShaderRHI, MipTreeNegYParameter, MipTreeNegY.SRV);
-		SetSRVParameter(RHICmdList, ShaderRHI, MipTreePosZParameter, MipTreePosZ.SRV);
-		SetSRVParameter(RHICmdList, ShaderRHI, MipTreeNegZParameter, MipTreeNegZ.SRV);
+		FGlobalShader::SetParameters<FViewUniformShaderParameters>(BatchedParameters, View.ViewUniformBuffer);
+
+		SetShaderValue(BatchedParameters, DimensionsParameter, Dimensions);
+		SetSRVParameter(BatchedParameters, MipTreePosXParameter, MipTreePosX.SRV);
+		SetSRVParameter(BatchedParameters, MipTreeNegXParameter, MipTreeNegX.SRV);
+		SetSRVParameter(BatchedParameters, MipTreePosYParameter, MipTreePosY.SRV);
+		SetSRVParameter(BatchedParameters, MipTreeNegYParameter, MipTreeNegY.SRV);
+		SetSRVParameter(BatchedParameters, MipTreePosZParameter, MipTreePosZ.SRV);
+		SetSRVParameter(BatchedParameters, MipTreeNegZParameter, MipTreeNegZ.SRV);
 	}
 
 private:
@@ -771,7 +769,7 @@ void FDeferredShadingSceneRenderer::VisualizeSkyLightMipTree(
 
 	// Draw
 	RHICmdList.SetViewport((float)View.ViewRect.Min.X, (float)View.ViewRect.Min.Y, 0.0f, (float)View.ViewRect.Max.X, (float)View.ViewRect.Max.Y, 1.0f);
-	PixelShader->SetParameters(RHICmdList, View, SkyLightMipDimensions, SkyLightMipTreePosX, SkyLightMipTreeNegX, SkyLightMipTreePosY, SkyLightMipTreeNegY, SkyLightMipTreePosZ, SkyLightMipTreeNegZ);
+	SetShaderParametersLegacyPS(RHICmdList, PixelShader, View, SkyLightMipDimensions, SkyLightMipTreePosX, SkyLightMipTreeNegX, SkyLightMipTreePosY, SkyLightMipTreeNegY, SkyLightMipTreePosZ, SkyLightMipTreeNegZ);
 	DrawRectangle(
 		RHICmdList,
 		0, 0,
diff --git a/Engine/Source/Runtime/Renderer/Private/VT/VirtualTextureSpace.cpp b/Engine/Source/Runtime/Renderer/Private/VT/VirtualTextureSpace.cpp
--- a/Engine/Source/Runtime/Renderer/Private/VT/VirtualTextureSpace.cpp
+++ b/Engine/Source/Runtime/Renderer/Private/VT/VirtualTextureSpace.cpp
@@ -197,6 +197,14 @@ public:
 		return true;
 	}
 
+	void SetParameters(FRHIBatchedShaderParameters& BatchedParameters, FIntPoint InPageTableSize, uint32 InFirstUpdate, uint32 InNumUpdates, FRHIShaderResourceView* InUpdateBuffer)
+	{
+		SetShaderValue(BatchedParameters, PageTableSize, InPageTableSize);
+		SetShaderValue(BatchedParameters, FirstUpdate, InFirstUpdate);
+		SetShaderValue(BatchedParameters, NumUpdates, InNumUpdates);
+		SetSRVParameter(BatchedParameters, UpdateBuffer, InUpdateBuffer);
+	}
+
 	LAYOUT_FIELD(FShaderParameter, PageTableSize);
 	LAYOUT_FIELD(FShaderParameter, FirstUpdate);
 	LAYOUT_FIELD(FShaderParameter, NumUpdates);
@@ -522,13 +530,7 @@ void FVirtualTextureSpace::ApplyUpdates(FVirtualTextureSystem* System, FRDGBuild
 
 					SetGraphicsPipelineState(RHICmdList, GraphicsPSOInit, 0);
 
-					{
-						FRHIVertexShader* ShaderRHI = VertexShader.GetVertexShader();
-						SetShaderValue(RHICmdList, ShaderRHI, VertexShader->PageTableSize, FIntPoint(CachedPageTableWidth, CachedPageTableHeight));
-						SetShaderValue(RHICmdList, ShaderRHI, VertexShader->FirstUpdate, FirstUpdate);
-						SetShaderValue(RHICmdList, ShaderRHI, VertexShader->NumUpdates, NumUpdates);
-						SetSRVParameter(RHICmdList, ShaderRHI, VertexShader->UpdateBuffer, UpdateBufferSRV);
-					}
+					SetShaderParametersLegacyVS(RHICmdList, VertexShader, FIntPoint(CachedPageTableWidth, CachedPageTableHeight), FirstUpdate, NumUpdates, UpdateBufferSRV);
 
 					// needs to be the same on shader side (faster on NVIDIA and AMD)
 					uint32 QuadsPerInstance = 8;
diff --git a/Engine/Source/Runtime/Renderer/Private/VolumetricFogLightFunction.cpp b/Engine/Source/Runtime/Renderer/Private/VolumetricFogLightFunction.cpp
--- a/Engine/Source/Runtime/Renderer/Private/VolumetricFogLightFunction.cpp
+++ b/Engine/Source/Runtime/Renderer/Private/VolumetricFogLightFunction.cpp
@@ -84,20 +84,20 @@ class FVolumetricFogLightFunctionPS : public FMaterialShader
 		return Parameters.MaterialParameters.MaterialDomain == MD_LightFunction && DoesPlatformSupportVolumetricFog(Parameters.Platform);
 	}
 
-	void SetParameters(
-		FRHICommandList& RHICmdList, 
-		FParameters& PS,
+	void SetParameters(FRHIBatchedShaderParameters& BatchedParameters, const FSceneView& View, const FMaterialRenderProxy* MaterialProxy)
+	{
+		const FMaterial& Material = MaterialProxy->GetMaterialWithFallback(View.GetFeatureLevel(), MaterialProxy);
+		FMaterialShader::SetViewParameters(BatchedParameters, View, View.ViewUniformBuffer);
+		FMaterialShader::SetParameters(BatchedParameters, MaterialProxy, Material, View);
+	}
+
+	FParameters GetParameters(
 		const FSceneView& View, 
 		const FLightSceneInfo* LightSceneInfo, 
-		const FMaterialRenderProxy* MaterialProxy,
-		FVector2D LightFunctionTexelSizeValue,
+		const FVector2D& LightFunctionTexelSizeValue,
 		const FMatrix44f& ShadowToTranslatedWorldValue)
 	{
-		FRHIPixelShader* ShaderRHI = RHICmdList.GetBoundPixelShader();
-
-		const FMaterial& Material = MaterialProxy->GetMaterialWithFallback(View.GetFeatureLevel(), MaterialProxy);
-		FMaterialShader::SetViewParameters(RHICmdList, ShaderRHI, View, View.ViewUniformBuffer);
-		FMaterialShader::SetParameters(RHICmdList, ShaderRHI, MaterialProxy, Material, View);
+		FParameters PS;
 
 		PS.LightFunctionParameters = FLightFunctionSharedParameters::GetLightFunctionSharedParameters(LightSceneInfo, 1.0f);
 		PS.LightFunctionParameters2 = FVector3f(
@@ -118,6 +118,8 @@ class FVolumetricFogLightFunctionPS : public FMaterialShader
 		PS.LightFunctionTexelSize = FVector2f(LightFunctionTexelSizeValue);
 		PS.ShadowToTranslatedWorld = ShadowToTranslatedWorldValue;
 		PS.LightTranslatedWorldPosition = FVector4f(LightSceneInfo->Proxy->GetPosition() + View.ViewMatrices.GetPreViewTranslation());
+
+		return PS;
 	}
 };
 
@@ -320,11 +322,10 @@ void FDeferredShadingSceneRenderer::RenderLightFunctionForVolumetricFog(
 
 					SetGraphicsPipelineState(RHICmdList, GraphicsPSOInit, 0);
 
-					FVolumetricFogLightFunctionPS::FParameters PS;
-					PixelShader->SetParameters(RHICmdList, PS, View, DirectionalLightSceneInfo, MaterialProxyForRendering, 
-						FVector2D(1.0f / LightFunctionResolution.X, 1.0f / LightFunctionResolution.Y), LightFunctionTranslatedWorldToShadowMatrix.Inverse());
+					FVolumetricFogLightFunctionPS::FParameters PS = PixelShader->GetParameters(View, DirectionalLightSceneInfo, FVector2D(1.0f / LightFunctionResolution.X, 1.0f / LightFunctionResolution.Y), LightFunctionTranslatedWorldToShadowMatrix.Inverse());
+
 					ClearUnusedGraphResources(PixelShader, &PS);
-					SetShaderParameters(RHICmdList, PixelShader, PixelShader.GetPixelShader(), PS);
+					SetShaderParametersMixedPS(RHICmdList, PixelShader, PS, View, MaterialProxyForRendering);
 
 					DrawRectangle(
 						RHICmdList,
@@ -443,11 +444,10 @@ void FDeferredShadingSceneRenderer::RenderLightFunctionForVolumetricFog(
 
 					SetGraphicsPipelineState(RHICmdList, GraphicsPSOInit, 0);
 
-					FVolumetricFogLightFunctionPS::FParameters PS;
-					PixelShader->SetParameters(RHICmdList, PS, View, LightSceneInfo, MaterialProxyForRendering,
-						FVector2D(1.0f / LightFunctionResolution, 1.0f / LightFunctionResolution), LightFunctionInfo.LightFunctionTranslatedWorldToLightMatrix.Inverse());
+					FVolumetricFogLightFunctionPS::FParameters PS = PixelShader->GetParameters(View, LightSceneInfo, FVector2D(1.0f / LightFunctionResolution, 1.0f / LightFunctionResolution), LightFunctionInfo.LightFunctionTranslatedWorldToLightMatrix.Inverse());
+
 					ClearUnusedGraphResources(PixelShader, &PS);
-					SetShaderParameters(RHICmdList, PixelShader, PixelShader.GetPixelShader(), PS);
+					SetShaderParametersMixedPS(RHICmdList, PixelShader, PS, View, MaterialProxyForRendering);
 
 					FIntPoint RectSize = LightFunctionInfo.AtlasTile.RectBound.Size();
 					DrawRectangle(
diff --git a/Engine/Source/Runtime/Renderer/Public/GlobalDistanceFieldParameters.h b/Engine/Source/Runtime/Renderer/Public/GlobalDistanceFieldParameters.h
--- a/Engine/Source/Runtime/Renderer/Public/GlobalDistanceFieldParameters.h
+++ b/Engine/Source/Runtime/Renderer/Public/GlobalDistanceFieldParameters.h
@@ -7,6 +7,7 @@
 #include "GlobalRenderResources.h"
 #include "ShaderParameterMacros.h"
 #include "ShaderParameters.h"
+#include "ShaderParameterUtils.h"
 #include "RenderUtils.h"
 #include "RHIStaticStates.h"
 #include "RenderGraphResources.h"
@@ -151,22 +152,26 @@ public:
 	{
 		if (IsBound())
 		{
-			SetTextureParameter(RHICmdList, ShaderRHI, GlobalDistanceFieldPageAtlasTexture, ParameterData.PageAtlasTexture ? ParameterData.PageAtlasTexture : GBlackVolumeTexture->TextureRHI.GetReference());
-			SetTextureParameter(RHICmdList, ShaderRHI, GlobalDistanceFieldPageTableTexture, ParameterData.PageTableTexture ? ParameterData.PageTableTexture : GBlackVolumeTexture->TextureRHI.GetReference());
-			SetTextureParameter(RHICmdList, ShaderRHI, GlobalDistanceFieldMipTexture, ParameterData.MipTexture ? ParameterData.MipTexture : GBlackVolumeTexture->TextureRHI.GetReference());
-
-			SetShaderValueArray(RHICmdList, ShaderRHI, GlobalVolumeTranslatedCenterAndExtent, ParameterData.TranslatedCenterAndExtent, GlobalDistanceField::MaxClipmaps);
-			SetShaderValueArray(RHICmdList, ShaderRHI, GlobalVolumeTranslatedWorldToUVAddAndMul, ParameterData.TranslatedWorldToUVAddAndMul, GlobalDistanceField::MaxClipmaps);
-			SetShaderValueArray(RHICmdList, ShaderRHI, GlobalDistanceFieldMipTranslatedWorldToUVScale, ParameterData.MipTranslatedWorldToUVScale, GlobalDistanceField::MaxClipmaps);
-			SetShaderValueArray(RHICmdList, ShaderRHI, GlobalDistanceFieldMipTranslatedWorldToUVBias, ParameterData.MipTranslatedWorldToUVBias, GlobalDistanceField::MaxClipmaps);
-			SetShaderValue(RHICmdList, ShaderRHI, GlobalDistanceFieldMipFactor, ParameterData.MipFactor);
-			SetShaderValue(RHICmdList, ShaderRHI, GlobalDistanceFieldMipTransition, ParameterData.MipTransition);
-			SetShaderValue(RHICmdList, ShaderRHI, GlobalDistanceFieldClipmapSizeInPages, ParameterData.ClipmapSizeInPages);
-			SetShaderValue(RHICmdList, ShaderRHI, GlobalDistanceFieldInvPageAtlasSize, (FVector3f)ParameterData.InvPageAtlasSize);
-			SetShaderValue(RHICmdList, ShaderRHI, GlobalVolumeDimension, ParameterData.GlobalDFResolution);
-			SetShaderValue(RHICmdList, ShaderRHI, GlobalVolumeTexelSize, 1.0f / ParameterData.GlobalDFResolution);
-			SetShaderValue(RHICmdList, ShaderRHI, MaxGlobalDFAOConeDistance, ParameterData.MaxDFAOConeDistance);
-			SetShaderValue(RHICmdList, ShaderRHI, NumGlobalSDFClipmaps, ParameterData.NumGlobalSDFClipmaps);
+			FRHIBatchedShaderParameters& BatchedParameters = RHICmdList.GetScratchShaderParameters();
+
+			SetTextureParameter(BatchedParameters, GlobalDistanceFieldPageAtlasTexture, ParameterData.PageAtlasTexture ? ParameterData.PageAtlasTexture : GBlackVolumeTexture->TextureRHI.GetReference());
+			SetTextureParameter(BatchedParameters, GlobalDistanceFieldPageTableTexture, ParameterData.PageTableTexture ? ParameterData.PageTableTexture : GBlackVolumeTexture->TextureRHI.GetReference());
+			SetTextureParameter(BatchedParameters, GlobalDistanceFieldMipTexture, ParameterData.MipTexture ? ParameterData.MipTexture : GBlackVolumeTexture->TextureRHI.GetReference());
+
+			SetShaderValueArray(BatchedParameters, GlobalVolumeTranslatedCenterAndExtent, ParameterData.TranslatedCenterAndExtent, GlobalDistanceField::MaxClipmaps);
+			SetShaderValueArray(BatchedParameters, GlobalVolumeTranslatedWorldToUVAddAndMul, ParameterData.TranslatedWorldToUVAddAndMul, GlobalDistanceField::MaxClipmaps);
+			SetShaderValueArray(BatchedParameters, GlobalDistanceFieldMipTranslatedWorldToUVScale, ParameterData.MipTranslatedWorldToUVScale, GlobalDistanceField::MaxClipmaps);
+			SetShaderValueArray(BatchedParameters, GlobalDistanceFieldMipTranslatedWorldToUVBias, ParameterData.MipTranslatedWorldToUVBias, GlobalDistanceField::MaxClipmaps);
+			SetShaderValue(BatchedParameters, GlobalDistanceFieldMipFactor, ParameterData.MipFactor);
+			SetShaderValue(BatchedParameters, GlobalDistanceFieldMipTransition, ParameterData.MipTransition);
+			SetShaderValue(BatchedParameters, GlobalDistanceFieldClipmapSizeInPages, ParameterData.ClipmapSizeInPages);
+			SetShaderValue(BatchedParameters, GlobalDistanceFieldInvPageAtlasSize, (FVector3f)ParameterData.InvPageAtlasSize);
+			SetShaderValue(BatchedParameters, GlobalVolumeDimension, ParameterData.GlobalDFResolution);
+			SetShaderValue(BatchedParameters, GlobalVolumeTexelSize, 1.0f / ParameterData.GlobalDFResolution);
+			SetShaderValue(BatchedParameters, MaxGlobalDFAOConeDistance, ParameterData.MaxDFAOConeDistance);
+			SetShaderValue(BatchedParameters, NumGlobalSDFClipmaps, ParameterData.NumGlobalSDFClipmaps);
+
+			RHICmdList.SetBatchedShaderParameters(ShaderRHI, BatchedParameters);
 		}
 	}
 
