World Partition: downgrade an error to a log when encountering duplicate runtime grid actors, since the cook will still be valid.

diff --git a/Engine/Source/Runtime/Engine/Private/WorldPartition/WorldPartitionRuntimeSpatialHash.cpp b/Engine/Source/Runtime/Engine/Private/WorldPartition/WorldPartitionRuntimeSpatialHash.cpp
--- a/Engine/Source/Runtime/Engine/Private/WorldPartition/WorldPartitionRuntimeSpatialHash.cpp
+++ b/Engine/Source/Runtime/Engine/Private/WorldPartition/WorldPartitionRuntimeSpatialHash.cpp
@@ -1122,7 +1122,7 @@ bool UWorldPartitionRuntimeSpatialHash::GenerateStreaming(UWorldPartitionStreami
 
 		if (const FSpatialHashRuntimeGrid* ExistingRuntimeGrid = AllGrids.FindByPredicate([RuntimeGridActor](const FSpatialHashRuntimeGrid& RuntimeGrid) { return RuntimeGrid.GridName == RuntimeGridActor->GridSettings.GridName; }))
 		{
-			UE_LOG(LogWorldPartition, Error, TEXT("Got duplicated runtime grid actor '%s' for grid '%s'"), *SpatialHashRuntimeGridInfo->GetActorLabelOrName().ToString(), *RuntimeGridActor->GridSettings.GridName.ToString());
+			UE_LOG(LogWorldPartition, Log, TEXT("Got duplicated runtime grid actor '%s' for grid '%s'"), *SpatialHashRuntimeGridInfo->GetActorLabelOrName().ToString(), *RuntimeGridActor->GridSettings.GridName.ToString());
 			check(*ExistingRuntimeGrid == RuntimeGridActor->GridSettings);
 		}
 		else
