set the cluster union actor's root component properly

diff --git a/Engine/Source/Runtime/Engine/Private/PhysicsEngine/ClusterUnionActor.cpp b/Engine/Source/Runtime/Engine/Private/PhysicsEngine/ClusterUnionActor.cpp
--- a/Engine/Source/Runtime/Engine/Private/PhysicsEngine/ClusterUnionActor.cpp
+++ b/Engine/Source/Runtime/Engine/Private/PhysicsEngine/ClusterUnionActor.cpp
@@ -9,6 +9,7 @@ AClusterUnionActor::AClusterUnionActor(const FObjectInitializer& ObjectInitializ
 	: AActor(ObjectInitializer)
 {
 	ClusterUnion = CreateDefaultSubobject<UClusterUnionComponent>(TEXT("ClusterUnion"));
+	SetRootComponent(ClusterUnion);
 
 	SetRemoteRoleForBackwardsCompat(ROLE_SimulatedProxy);
 	bReplicates = true;
