Include editor builds in WITH_DEBUG_VIEW_MODES. Fixes issues with Nanite wireframe in Shipping Editor builds.

#localization none
#tests confirmed in editor
#rnx

diff --git a/Engine/Source/Runtime/Engine/Private/SceneView.cpp b/Engine/Source/Runtime/Engine/Private/SceneView.cpp
--- a/Engine/Source/Runtime/Engine/Private/SceneView.cpp
+++ b/Engine/Source/Runtime/Engine/Private/SceneView.cpp
@@ -2883,7 +2883,9 @@ FSceneViewFamily::FSceneViewFamily(const ConstructionValues& CVS)
 	{
 		Time = FGameTime::CreateDilated(0.0, Time.GetDeltaRealTimeSeconds(), 0.0, Time.GetDeltaWorldTimeSeconds());
 	}
+#endif
 
+#if WITH_DEBUG_VIEW_MODES
 	DebugViewShaderMode = ChooseDebugViewShaderMode();
 	ViewModeParam = CVS.ViewModeParam;
 	ViewModeParamName = CVS.ViewModeParamName;
@@ -3035,7 +3037,7 @@ FSceneViewFamilyContext::~FSceneViewFamilyContext()
 	}
 }
 
-#if !(UE_BUILD_SHIPPING || UE_BUILD_TEST)
+#if WITH_DEBUG_VIEW_MODES
 
 EDebugViewShaderMode FSceneViewFamily::ChooseDebugViewShaderMode() const
 {
@@ -3089,4 +3091,4 @@ EDebugViewShaderMode FSceneViewFamily::ChooseDebugViewShaderMode() const
 	return DVSM_None;
 }
 
-#endif // !(UE_BUILD_SHIPPING || UE_BUILD_TEST)
+#endif // WITH_DEBUG_VIEW_MODES
diff --git a/Engine/Source/Runtime/Engine/Public/DebugViewModeHelpers.h b/Engine/Source/Runtime/Engine/Public/DebugViewModeHelpers.h
--- a/Engine/Source/Runtime/Engine/Public/DebugViewModeHelpers.h
+++ b/Engine/Source/Runtime/Engine/Public/DebugViewModeHelpers.h
@@ -12,7 +12,7 @@
 #include "SceneTypes.h"
 #endif
 
-#define WITH_DEBUG_VIEW_MODES !(UE_BUILD_SHIPPING || UE_BUILD_TEST)
+#define WITH_DEBUG_VIEW_MODES (WITH_EDITOR || !(UE_BUILD_SHIPPING || UE_BUILD_TEST))
 
 class UMaterialInterface;
 struct FSlowTask;
diff --git a/Engine/Source/Runtime/Engine/Public/SceneView.h b/Engine/Source/Runtime/Engine/Public/SceneView.h
--- a/Engine/Source/Runtime/Engine/Public/SceneView.h
+++ b/Engine/Source/Runtime/Engine/Public/SceneView.h
@@ -2024,7 +2024,7 @@ public:
 
 	EShaderPlatform GetShaderPlatform() const { return GShaderPlatformForFeatureLevel[GetFeatureLevel()]; }
 
-#if !(UE_BUILD_SHIPPING || UE_BUILD_TEST)
+#if WITH_DEBUG_VIEW_MODES
 	EDebugViewShaderMode DebugViewShaderMode;
 	int32 ViewModeParam;
 	FName ViewModeParamName;
