Added code in MaterialInterface serialization function that validates the associated EditorOnlyData name to be "canonical" (name of the material followed by "EditorOnlyData"). If name isn't canonical it may cause problems with the editor running on cooked content. The validation code will try to rename the EditorOnlyData to the correct name.

diff --git a/Engine/Source/Runtime/Engine/Private/Materials/MaterialInterface.cpp b/Engine/Source/Runtime/Engine/Private/Materials/MaterialInterface.cpp
--- a/Engine/Source/Runtime/Engine/Private/Materials/MaterialInterface.cpp
+++ b/Engine/Source/Runtime/Engine/Private/Materials/MaterialInterface.cpp
@@ -116,6 +116,18 @@ void FMaterialRelevance::SetPrimitiveViewRelevance(FPrimitiveViewRelevance& OutV
 
 //////////////////////////////////////////////////////////////////////////
 
+#if WITH_EDITORONLY_DATA
+
+namespace MaterialInterface
+{
+	FString GetEditorOnlyDataName(const TCHAR* InMaterialName)
+	{
+		return FString::Printf(TEXT("%sEditorOnlyData"), InMaterialName);
+	}
+}
+
+#endif // WITH_EDITORONLY_DATA
+
 UMaterialInterface::UMaterialInterface() = default;
 
 UMaterialInterface::UMaterialInterface(FVTableHelper& Helper)
@@ -690,6 +702,18 @@ void UMaterialInterface::PostDuplicate(bool bDuplicateForPIE)
 	Super::PostDuplicate(bDuplicateForPIE);
 
 	SetLightingGuid();
+
+#if WITH_EDITORONLY_DATA
+	// Duplication will create a new editor only data but this MI already links a populated data. This makes sure that the
+	// duplicate instance EOD has the correct name (necessary for running the editor on cooked data).
+	FString EditorOnlyDataName = MaterialInterface::GetEditorOnlyDataName(*GetName());
+	UObject* EditorOnlyExisting = StaticFindObject(/*Class=*/ nullptr, this, *EditorOnlyDataName, true);
+	if (EditorOnlyExisting)
+	{
+		EditorOnlyExisting->Rename(*(EditorOnlyDataName + "_Unused"), nullptr, REN_NonTransactional | REN_DontCreateRedirectors);
+	}
+	GetEditorOnlyData()->Rename(*EditorOnlyDataName, nullptr, REN_NonTransactional | REN_DontCreateRedirectors);
+#endif // WITH_EDITORONLY_DATA
 }
 
 #if WITH_EDITOR
@@ -1495,6 +1519,24 @@ void UMaterialInterface::PreSave(FObjectPreSaveContext ObjectSaveContext)
 	{
 		SortTextureStreamingData(true, true);
 	}
+
+#if WITH_EDITORONLY_DATA
+	// Make sure the EditorOnlyData is named correctly. Some MI assets have a differently named editor only data that can cause problems in
+	// the editor running on cooked data.
+	UMaterialInterfaceEditorOnlyData* EditorOnly = GetEditorOnlyData();
+	FString EditorOnlyDataName = MaterialInterface::GetEditorOnlyDataName(*GetName());
+	if (EditorOnly->GetName() != EditorOnlyDataName)
+	{
+		UE_LOG(LogMaterial, Display, TEXT("MaterialInterface %s has a incorrectly name EditorOnlyData '%s'. This may cause issues when running the editor on cooked data. Trying to rename it to the correct name '%s'."), *GetName(), *EditorOnly->GetName(), *EditorOnlyDataName);
+		UObject* EditorOnlyExisting = StaticFindObject(/*Class=*/ nullptr, this, *EditorOnlyDataName, true);
+		if (EditorOnlyExisting)
+		{
+			EditorOnlyExisting->Rename(*(EditorOnlyDataName + "_Unused"), nullptr, REN_NonTransactional | REN_DontCreateRedirectors);
+		}
+		EditorOnly->Rename(*EditorOnlyDataName, nullptr, REN_NonTransactional | REN_DontCreateRedirectors);
+	}
+#endif // WITH_EDITORONLY_DATA
+
 }
 
 void UMaterialInterface::AddAssetUserData(UAssetUserData* InUserData)
@@ -1588,13 +1630,6 @@ void UMaterialInterface::FilterOutPlatformShadingModels(const EShaderPlatform In
 
 #if WITH_EDITORONLY_DATA
 
-namespace MaterialInterface
-{
-	FString GetEditorOnlyDataName(const TCHAR* InMaterialName)
-	{
-		return FString::Printf(TEXT("%sEditorOnlyData"), InMaterialName);
-	}
-}
 
 const UClass* UMaterialInterface::GetEditorOnlyDataClass() const
 {
