Fixes for child actor replication with -game

diff --git a/Engine/Source/Runtime/Engine/Classes/Engine/NetDriver.h b/Engine/Source/Runtime/Engine/Classes/Engine/NetDriver.h
--- a/Engine/Source/Runtime/Engine/Classes/Engine/NetDriver.h
+++ b/Engine/Source/Runtime/Engine/Classes/Engine/NetDriver.h
@@ -356,6 +356,22 @@ namespace UE::Net
 }
 #endif // UE_WITH_IRIS
 
+namespace UE::Net
+{
+	class FScopedIgnoreStaticActorDestruction
+	{
+	public:
+		FScopedIgnoreStaticActorDestruction();
+		~FScopedIgnoreStaticActorDestruction();
+
+		UE_NONCOPYABLE(FScopedIgnoreStaticActorDestruction);
+
+	private:
+		bool bCachedValue = false;
+	};
+
+	bool ShouldIgnoreStaticActorDestruction();
+}
 
 using FConnectionMap = TMap<TSharedRef<const FInternetAddr>, UNetConnection*, FDefaultSetAllocator, FInternetAddrConstKeyMapFuncs<UNetConnection*>>;
 
diff --git a/Engine/Source/Runtime/Engine/Private/Components/ChildActorComponent.cpp b/Engine/Source/Runtime/Engine/Private/Components/ChildActorComponent.cpp
--- a/Engine/Source/Runtime/Engine/Private/Components/ChildActorComponent.cpp
+++ b/Engine/Source/Runtime/Engine/Private/Components/ChildActorComponent.cpp
@@ -75,6 +75,8 @@ void UChildActorComponent::OnRegister()
 			bool bChildActorPackageWasDirty = ChildActor->GetPackage()->IsDirty();
 			bool bPackageWasDirty = GetPackage()->IsDirty();
 
+			UE::Net::FScopedIgnoreStaticActorDestruction ScopedIgnoreDestruction;
+
 			DestroyChildActor();
 			CreateChildActor();
 			
@@ -703,7 +705,7 @@ void UChildActorComponent::CreateChildActor(TFunction<void(AActor*)> CustomizerF
 
 	if (MyOwner && !MyOwner->HasAuthority())
 	{
-		if (IsChildActorReplicated())
+		if (IsChildActorReplicated() && !GIsReconstructingBlueprintInstances)
 		{
 			// If we belong to an actor that is not authoritative and the child class is replicated then we expect that Actor will be replicated across so don't spawn one
 			return;
diff --git a/Engine/Source/Runtime/Engine/Private/Level.cpp b/Engine/Source/Runtime/Engine/Private/Level.cpp
--- a/Engine/Source/Runtime/Engine/Private/Level.cpp
+++ b/Engine/Source/Runtime/Engine/Private/Level.cpp
@@ -34,6 +34,7 @@ Level.cpp: Level-related functions
 #include "Engine/AssetUserData.h"
 #include "Engine/LevelScriptBlueprint.h"
 #include "Engine/LevelScriptActor.h"
+#include "Engine/NetDriver.h"
 #include "Engine/WorldComposition.h"
 #include "StaticLighting.h"
 #include "TickTaskManagerInterface.h"
@@ -744,7 +745,7 @@ void ULevel::Serialize( FArchive& Ar )
 
 void ULevel::CreateReplicatedDestructionInfo(AActor* const Actor)
 {
-	if (Actor == nullptr)
+	if ((Actor == nullptr) || GIsReconstructingBlueprintInstances || UE::Net::ShouldIgnoreStaticActorDestruction())
 	{
 		return;
 	}
diff --git a/Engine/Source/Runtime/Engine/Private/NetDriver.cpp b/Engine/Source/Runtime/Engine/Private/NetDriver.cpp
--- a/Engine/Source/Runtime/Engine/Private/NetDriver.cpp
+++ b/Engine/Source/Runtime/Engine/Private/NetDriver.cpp
@@ -185,7 +185,7 @@ namespace UE::Net::Private
 	static const FString PerfCounter_NumConnections(TEXT("NumConnections"));
 }
 
-namespace NetEmulationHelper
+namespace UE::Net::EmulationHelper
 {
 #if DO_ENABLE_NET_TEST
 	/** Global that stores the network emulation values outside the NetDriver lifetime */
@@ -214,7 +214,7 @@ namespace NetEmulationHelper
 #endif //#if DO_ENABLE_NET_TEST
 }
 
-namespace NetDriverInternal
+namespace UE::Net::Private
 {
 	struct FAutoDestructProperty
 	{
@@ -245,9 +245,9 @@ namespace NetDriverInternal
 		}
 	};
 
-	TArray<NetDriverInternal::FAutoDestructProperty> CopyOutParametersToLocalParameters(UFunction* Function, FOutParmRec* OutParms, void* LocalParms, UObject* TargetObj)
+	TArray<UE::Net::Private::FAutoDestructProperty> CopyOutParametersToLocalParameters(UFunction* Function, FOutParmRec* OutParms, void* LocalParms, UObject* TargetObj)
 	{
-		TArray<NetDriverInternal::FAutoDestructProperty> CopiedProperties;
+		TArray<UE::Net::Private::FAutoDestructProperty> CopiedProperties;
 
 		// Look for CPF_OutParm's, we'll need to copy these into the local parameter memory manually
 		// The receiving side will pull these back out when needed
@@ -287,7 +287,28 @@ namespace NetDriverInternal
 		return CopiedProperties;
 	}
 
-} //namespace NetDriverInternal
+	static bool bIgnoreStaticActorDestruction = false;
+
+} //namespace UE::Net::Private
+
+namespace UE::Net
+{
+	FScopedIgnoreStaticActorDestruction::FScopedIgnoreStaticActorDestruction()
+	{
+		bCachedValue = UE::Net::Private::bIgnoreStaticActorDestruction;
+		UE::Net::Private::bIgnoreStaticActorDestruction = true;
+	}
+
+	FScopedIgnoreStaticActorDestruction::~FScopedIgnoreStaticActorDestruction()
+	{
+		UE::Net::Private::bIgnoreStaticActorDestruction = bCachedValue;
+	}
+
+	bool ShouldIgnoreStaticActorDestruction()
+	{
+		return UE::Net::Private::bIgnoreStaticActorDestruction;
+	}
+}
 
 #if UE_BUILD_SHIPPING
 #define DEBUG_REMOTEFUNCTION(Format, ...)
@@ -524,9 +545,9 @@ void UNetDriver::InitPacketSimulationSettings()
 		return;
 	}
 
-	if (NetEmulationHelper::PersistentPacketSimulationSettings.IsSet())
+	if (UE::Net::EmulationHelper::PersistentPacketSimulationSettings.IsSet())
 	{
-		SetPacketSimulationSettings(NetEmulationHelper::PersistentPacketSimulationSettings.GetValue());
+		SetPacketSimulationSettings(UE::Net::EmulationHelper::PersistentPacketSimulationSettings.GetValue());
 		return;
 	}
 
@@ -3282,7 +3303,7 @@ void UNetDriver::NotifyActorDestroyed( AActor* ThisActor, bool IsSeamlessTravel
 
 		const bool bIsActorStatic = !GuidCache->IsDynamicObject( ThisActor );
 		const bool bActorHasRole = ThisActor->GetRemoteRole() != ROLE_None;
-		const bool bShouldCreateDestructionInfo = bIsServer && bIsActorStatic && bActorHasRole && !IsSeamlessTravel && !GIsReconstructingBlueprintInstances;
+		const bool bShouldCreateDestructionInfo = bIsServer && bIsActorStatic && bActorHasRole && !IsSeamlessTravel && !GIsReconstructingBlueprintInstances && !UE::Net::ShouldIgnoreStaticActorDestruction();
 
 		if (bShouldCreateDestructionInfo)
 		{
@@ -6605,12 +6626,12 @@ void UNetDriver::ProcessRemoteFunction(
 		++TotalRPCsCalled;
 
 		// Copy Any Out Params to Local Params 
-		TArray<NetDriverInternal::FAutoDestructProperty> LocalOutParms;
+		TArray<UE::Net::Private::FAutoDestructProperty> LocalOutParms;
 		if (Stack == nullptr)
 		{
 			// If we have a subobject, thats who we are actually calling this on. If no subobject, we are calling on the actor.
 			UObject* TargetObj = SubObject ? SubObject : Actor;
-			LocalOutParms = NetDriverInternal::CopyOutParametersToLocalParameters(Function, OutParms, Parameters, TargetObj);
+			LocalOutParms = UE::Net::Private::CopyOutParametersToLocalParameters(Function, OutParms, Parameters, TargetObj);
 		}
 
 #if UE_WITH_IRIS
@@ -7342,13 +7363,13 @@ FAutoConsoleCommandWithWorldArgsAndOutputDevice NetEmulationPktEmulationProfile(
 	{
 		FString CmdParams = FString::Printf(TEXT("PktEmulationProfile=%s"), *(Args[0]));
 
-		NetEmulationHelper::CreatePersistentSimulationSettings();
+		UE::Net::EmulationHelper::CreatePersistentSimulationSettings();
 
-		bProfileApplied = NetEmulationHelper::PersistentPacketSimulationSettings.GetValue().ParseSettings(*CmdParams, nullptr);
+		bProfileApplied = UE::Net::EmulationHelper::PersistentPacketSimulationSettings.GetValue().ParseSettings(*CmdParams, nullptr);
 		
 		if (bProfileApplied)
 		{
-			NetEmulationHelper::ApplySimulationSettingsOnNetDrivers(World, NetEmulationHelper::PersistentPacketSimulationSettings.GetValue());
+			UE::Net::EmulationHelper::ApplySimulationSettingsOnNetDrivers(World, UE::Net::EmulationHelper::PersistentPacketSimulationSettings.GetValue());
 		}
 		else
 		{
@@ -7376,9 +7397,9 @@ FAutoConsoleCommandWithWorldArgsAndOutputDevice NetEmulationPktEmulationProfile(
 FAutoConsoleCommandWithWorld NetEmulationOff(TEXT("NetEmulation.Off"), TEXT("Turn off network emulation"),
 	FConsoleCommandWithWorldDelegate::CreateStatic([](UWorld* World)
 {
-	NetEmulationHelper::CreatePersistentSimulationSettings();
-	NetEmulationHelper::PersistentPacketSimulationSettings.GetValue().ResetSettings();
-	NetEmulationHelper::ApplySimulationSettingsOnNetDrivers(World, NetEmulationHelper::PersistentPacketSimulationSettings.GetValue());
+	UE::Net::EmulationHelper::CreatePersistentSimulationSettings();
+	UE::Net::EmulationHelper::PersistentPacketSimulationSettings.GetValue().ResetSettings();
+	UE::Net::EmulationHelper::ApplySimulationSettingsOnNetDrivers(World, UE::Net::EmulationHelper::PersistentPacketSimulationSettings.GetValue());
 }));
 
 #define BUILD_NETEMULATION_CONSOLE_COMMAND(CommandName, CommandHelp) FAutoConsoleCommandWithWorldAndArgs NetEmulation##CommandName(TEXT("NetEmulation."#CommandName), TEXT(CommandHelp), \
@@ -7386,10 +7407,10 @@ FAutoConsoleCommandWithWorld NetEmulationOff(TEXT("NetEmulation.Off"), TEXT("Tur
 	{ \
 		if (Args.Num() > 0) \
 		{ \
-			NetEmulationHelper::CreatePersistentSimulationSettings(); \
+			UE::Net::EmulationHelper::CreatePersistentSimulationSettings(); \
 			FString CmdParams = FString::Printf(TEXT(#CommandName"=%s"), *(Args[0])); \
-			NetEmulationHelper::PersistentPacketSimulationSettings.GetValue().ParseSettings(*CmdParams, nullptr); \
-			NetEmulationHelper::ApplySimulationSettingsOnNetDrivers(World, NetEmulationHelper::PersistentPacketSimulationSettings.GetValue()); \
+			UE::Net::EmulationHelper::PersistentPacketSimulationSettings.GetValue().ParseSettings(*CmdParams, nullptr); \
+			UE::Net::EmulationHelper::ApplySimulationSettingsOnNetDrivers(World, UE::Net::EmulationHelper::PersistentPacketSimulationSettings.GetValue()); \
 		} \
 	}));
 
