Fixed a crash when adding curves to the output skeleton from USkeletalMergingLibrary::MergeSkeletons.

Made some minor changes in USkeleton to accommodate this fix:
- USkeleton::AddCurveMetaData now adds a metadata object to the skeleton if one doesn't exist already.
- Did some slight refactoring so that the UAnimCurveMetaData is only created in one function.

diff --git a/Engine/Plugins/Runtime/SkeletalMerging/Source/SkeletalMerging/Private/SkeletalMergingLibrary.cpp b/Engine/Plugins/Runtime/SkeletalMerging/Source/SkeletalMerging/Private/SkeletalMergingLibrary.cpp
--- a/Engine/Plugins/Runtime/SkeletalMerging/Source/SkeletalMerging/Private/SkeletalMergingLibrary.cpp
+++ b/Engine/Plugins/Runtime/SkeletalMerging/Source/SkeletalMerging/Private/SkeletalMergingLibrary.cpp
@@ -332,6 +332,8 @@ void USkeletalMergingLibrary::AddCurveNames(USkeleton* InSkeleton, const TMap<FN
 	{
 		if (CurveMetaDataPair.Value)
 		{
+			InSkeleton->AddCurveMetaData(CurveMetaDataPair.Key);
+
 			FCurveMetaData& SkeletonCurveMetaData = *InSkeleton->GetCurveMetaData(CurveMetaDataPair.Key);
 			SkeletonCurveMetaData = *CurveMetaDataPair.Value;
 			for (FBoneReference& BoneReference : SkeletonCurveMetaData.LinkedBones)
diff --git a/Engine/Source/Runtime/Engine/Classes/Animation/Skeleton.h b/Engine/Source/Runtime/Engine/Classes/Animation/Skeleton.h
--- a/Engine/Source/Runtime/Engine/Classes/Animation/Skeleton.h
+++ b/Engine/Source/Runtime/Engine/Classes/Animation/Skeleton.h
@@ -25,6 +25,7 @@
 #include "Delegates/DelegateCombinations.h"
 #include "Skeleton.generated.h"
 
+class UAnimCurveMetaData;
 class UAnimSequence;
 class UBlendProfile;
 class URig;
@@ -493,6 +494,9 @@ private:
 	// Refresh skeleton metadata (updates bone indices for linked bone references)
 	void RefreshSkeletonMetaData();
 
+	// Returns the UAnimCurveMetaData from this skeleton's AssetUserData, creating one if it doesn't exist
+	UAnimCurveMetaData* GetOrCreateCurveMetaDataObject();
+
 public:
 	//////////////////////////////////////////////////////////////////////////
 	// Blend Profiles
diff --git a/Engine/Source/Runtime/Engine/Private/Animation/Skeleton.cpp b/Engine/Source/Runtime/Engine/Private/Animation/Skeleton.cpp
--- a/Engine/Source/Runtime/Engine/Private/Animation/Skeleton.cpp
+++ b/Engine/Source/Runtime/Engine/Private/Animation/Skeleton.cpp
@@ -376,7 +376,7 @@ void USkeleton::PostLoad()
 	if(GetLinkerCustomVersion(FUE5MainStreamObjectVersion::GUID) < FUE5MainStreamObjectVersion::AnimationRemoveSmartNames)
 	{
 		// Move curve metadata over to asset user data
-		UAnimCurveMetaData* AnimCurveMetaData = NewObject<UAnimCurveMetaData>(this, NAME_None, RF_Transactional);
+		UAnimCurveMetaData* AnimCurveMetaData = GetOrCreateCurveMetaDataObject();
 		for(const TPair<FName, FSmartNameMapping>& NameMappingPair : SmartNames_DEPRECATED.NameMappings)
 		{
 			for(const TPair<FName, FCurveMetaData>& NameMetaDataPair : NameMappingPair.Value.CurveMetaDataMap)
@@ -384,7 +384,6 @@ void USkeleton::PostLoad()
 				AnimCurveMetaData->CurveMetaData.Add(NameMetaDataPair.Key, NameMetaDataPair.Value);
 			}
 		}
-		AddAssetUserData(AnimCurveMetaData);
 	}
 #endif
 }
@@ -1598,11 +1597,8 @@ void USkeleton::RenameSlotName(const FName& OldName, const FName& NewName)
 
 bool USkeleton::AddCurveMetaData(FName CurveName)
 {
-	if(UAnimCurveMetaData* AnimCurveMetaData = GetAssetUserData<UAnimCurveMetaData>())
-	{
-		return AnimCurveMetaData->AddCurveMetaData(CurveName);
-	}
-	return false;
+	UAnimCurveMetaData* AnimCurveMetaData = GetOrCreateCurveMetaDataObject();
+	return AnimCurveMetaData->AddCurveMetaData(CurveName);
 }
 
 #if WITH_EDITOR
@@ -1719,30 +1715,22 @@ void USkeleton::AccumulateCurveMetaData(FName CurveName, bool bMaterialSet, bool
 {
 	if (bMaterialSet || bMorphtargetSet)
 	{
-		// Add user data if not already added
-		UAnimCurveMetaData* AnimCurveMetaData = GetAssetUserData<UAnimCurveMetaData>();
-		if(AnimCurveMetaData == nullptr)
-		{
-			AnimCurveMetaData = NewObject<UAnimCurveMetaData>(this, NAME_None, RF_Transactional);
-			AddAssetUserData(AnimCurveMetaData);
-		}
-
 		// Add curve if not already present
 		AddCurveMetaData(CurveName);
 
-		if(FCurveMetaData* FoundCurveMetaData = GetCurveMetaData(CurveName))
+		FCurveMetaData* FoundCurveMetaData = GetCurveMetaData(CurveName);
+		check(FoundCurveMetaData);
+
+		bool bOldMaterial = FoundCurveMetaData->Type.bMaterial;
+		bool bOldMorphtarget = FoundCurveMetaData->Type.bMorphtarget;
+		// we don't want to undo previous flags, if it was true, we just allow more to it. 
+		FoundCurveMetaData->Type.bMaterial |= bMaterialSet;
+		FoundCurveMetaData->Type.bMorphtarget |= bMorphtargetSet;
+
+		if (bOldMaterial != FoundCurveMetaData->Type.bMaterial 
+			|| bOldMorphtarget != FoundCurveMetaData->Type.bMorphtarget)
 		{
-			bool bOldMaterial = FoundCurveMetaData->Type.bMaterial;
-			bool bOldMorphtarget = FoundCurveMetaData->Type.bMorphtarget;
-			// we don't want to undo previous flags, if it was true, we just allow more to it. 
-			FoundCurveMetaData->Type.bMaterial |= bMaterialSet;
-			FoundCurveMetaData->Type.bMorphtarget |= bMorphtargetSet;
-
-			if (bOldMaterial != FoundCurveMetaData->Type.bMaterial 
-				|| bOldMorphtarget != FoundCurveMetaData->Type.bMorphtarget)
-			{
-				MarkPackageDirty();
-			}
+			MarkPackageDirty();
 		}
 	}
 }
@@ -2216,5 +2204,17 @@ void USkeleton::RefreshSkeletonMetaData()
 	}
 }
 
+UAnimCurveMetaData* USkeleton::GetOrCreateCurveMetaDataObject()
+{
+	UAnimCurveMetaData* AnimCurveMetaData = GetAssetUserData<UAnimCurveMetaData>();
+	if (AnimCurveMetaData == nullptr)
+	{
+		AnimCurveMetaData = NewObject<UAnimCurveMetaData>(this, NAME_None, RF_Transactional);
+		AddAssetUserData(AnimCurveMetaData);
+	}
+
+	return AnimCurveMetaData;
+}
+
 #undef LOCTEXT_NAMESPACE 
 
