//UE5/Release-5.2 - Build Lyra Win64 - EXCEPTION_ACCESS_VIOLATION reading address 0x0000000000000048

#fix forcefully retrieve unreachable skeleton from weak-ptr, as compression happens concurrently with GC mark phase

diff --git a/Engine/Source/Runtime/Engine/Private/Animation/AnimationRuntime.cpp b/Engine/Source/Runtime/Engine/Private/Animation/AnimationRuntime.cpp
--- a/Engine/Source/Runtime/Engine/Private/Animation/AnimationRuntime.cpp
+++ b/Engine/Source/Runtime/Engine/Private/Animation/AnimationRuntime.cpp
@@ -2577,7 +2577,9 @@ void FAnimationRuntime::RetargetBoneTransform(const USkeleton* SourceSkeleton, c
 	check(!RetargetTransforms.IsEmpty());
 	if (SourceSkeleton)
 	{
-		const USkeleton* TargetSkeleton = RequiredBones.GetSkeletonAsset();
+		// Retrieve skeleton, even if it is unreachable (but not GC-ed yet)
+		const bool bEvenIfUnreachable = true; 
+		const USkeleton* TargetSkeleton = RequiredBones.GetSkeletonAsset(bEvenIfUnreachable);
 		const FSkeletonRemapping& SkeletonRemapping = UE::Anim::FSkeletonRemappingRegistry::Get().GetRemapping(SourceSkeleton, TargetSkeleton);
 
 		const int32 TargetSkeletonBoneIndex = RequiredBones.GetSkeletonIndex(BoneIndex);
