Added a delegate to SkeletalMeshComponent, so external systems can be notified when the RequiredBones get recalculated for a LOD.
Also added a missing ENGINE_API to SkeletonRemapping class, now that it is public in the Engine

diff --git a/Engine/Source/Runtime/Engine/Classes/Components/SkeletalMeshComponent.h b/Engine/Source/Runtime/Engine/Classes/Components/SkeletalMeshComponent.h
--- a/Engine/Source/Runtime/Engine/Classes/Components/SkeletalMeshComponent.h
+++ b/Engine/Source/Runtime/Engine/Classes/Components/SkeletalMeshComponent.h
@@ -65,6 +65,9 @@ DECLARE_DYNAMIC_MULTICAST_DELEGATE(FOnBoneTransformsFinalized);  // Deprecated,
 
 DECLARE_MULTICAST_DELEGATE(FOnBoneTransformsFinalizedMultiCast);
 
+DECLARE_MULTICAST_DELEGATE_ThreeParams(FOnLODRequiredBonesUpdateMulticast, USkeletalMeshComponent*, int32, const TArray<FBoneIndexType>&);
+typedef FOnLODRequiredBonesUpdateMulticast::FDelegate FOnLODRequiredBonesUpdate;
+
 /** Method used when retrieving a attribute value*/
 UENUM()
 enum class ECustomBoneAttributeLookup : uint8
@@ -2556,6 +2559,10 @@ public:
 	FDelegateHandle RegisterOnBoneTransformsFinalizedDelegate(const FOnBoneTransformsFinalizedMultiCast::FDelegate& Delegate);
 	void UnregisterOnBoneTransformsFinalizedDelegate(const FDelegateHandle& DelegateHandle);
 
+	/** Register/Unregister for OnLODRequiredBonesUpdate callback */
+	FDelegateHandle RegisterOnLODRequiredBonesUpdate(const FOnLODRequiredBonesUpdate& Delegate);
+	void UnregisterOnLODRequiredBonesUpdate(const FDelegateHandle& DelegateHandle);
+	
 private:
 
 #if WITH_EDITORONLY_DATA
@@ -2591,6 +2598,9 @@ private:
 	/** Multicaster fired when this component bone transforms are finalized */
 	FOnBoneTransformsFinalizedMultiCast OnBoneTransformsFinalizedMC;
 
+	/** Static Multicaster fired when this component finalizes the regeneration of the required bones list for the current LOD*/
+	static FOnLODRequiredBonesUpdateMulticast OnLODRequiredBonesUpdate;
+
 	/** Mark current anim UID version to up-to-date. Called when it's recalculated */
 	void MarkRequiredCurveUpToDate();
 	
diff --git a/Engine/Source/Runtime/Engine/Private/Components/SkeletalMeshComponent.cpp b/Engine/Source/Runtime/Engine/Private/Components/SkeletalMeshComponent.cpp
--- a/Engine/Source/Runtime/Engine/Private/Components/SkeletalMeshComponent.cpp
+++ b/Engine/Source/Runtime/Engine/Private/Components/SkeletalMeshComponent.cpp
@@ -102,6 +102,9 @@ FAutoConsoleTaskPriority CPrio_ParallelAnimationEvaluationTask(
 	ENamedThreads::HighTaskPriority // if we don't have hi pri threads, then use normal priority threads at high task priority instead
 	);
 
+/** Static Multicaster fired when SkeletalMeshComponent finalizes the regeneration of the required bones list for the current LOD*/
+/*static*/ FOnLODRequiredBonesUpdateMulticast USkeletalMeshComponent::OnLODRequiredBonesUpdate;
+
 class FParallelAnimationEvaluationTask
 {
 	TWeakObjectPtr<USkeletalMeshComponent> SkeletalMeshComponent;
@@ -1957,6 +1960,8 @@ void USkeletalMeshComponent::RecalcRequiredBones(int32 LODIndex)
 
 	// Invalidate cached bones.
 	ClearCachedAnimProperties();
+
+	OnLODRequiredBonesUpdate.Broadcast(this, LODIndex, RequiredBones);
 }
 
 void USkeletalMeshComponent::MarkRequiredCurveUpToDate()
@@ -4325,6 +4330,16 @@ void USkeletalMeshComponent::UnregisterOnBoneTransformsFinalizedDelegate(const F
 	OnBoneTransformsFinalizedMC.Remove(DelegateHandle);
 }
 
+FDelegateHandle USkeletalMeshComponent::RegisterOnLODRequiredBonesUpdate(const FOnLODRequiredBonesUpdate& Delegate)
+{
+	return OnLODRequiredBonesUpdate.Add(Delegate);
+}
+
+void USkeletalMeshComponent::UnregisterOnLODRequiredBonesUpdate(const FDelegateHandle& DelegateHandle)
+{
+	OnLODRequiredBonesUpdate.Remove(DelegateHandle);
+}
+
 bool USkeletalMeshComponent::MoveComponentImpl(const FVector& Delta, const FQuat& NewRotation, bool bSweep, FHitResult* OutHit /*= nullptr*/, EMoveComponentFlags MoveFlags /*= MOVECOMP_NoFlags*/, ETeleportType Teleport /*= ETeleportType::None*/)
 {
 #if WITH_EDITOR
diff --git a/Engine/Source/Runtime/Engine/Public/Animation/SkeletonRemapping.h b/Engine/Source/Runtime/Engine/Public/Animation/SkeletonRemapping.h
--- a/Engine/Source/Runtime/Engine/Public/Animation/SkeletonRemapping.h
+++ b/Engine/Source/Runtime/Engine/Public/Animation/SkeletonRemapping.h
@@ -7,7 +7,7 @@
 
 class USkeleton;
 
-struct FSkeletonRemapping
+struct ENGINE_API FSkeletonRemapping
 {
 	FSkeletonRemapping() = default;
 	FSkeletonRemapping(const USkeleton* InSourceSkeleton, const USkeleton* InTargetSkeleton);
