Templated FWorldPartitionHelpers::GetWorldPartition so it's inlined and not recursive.

diff --git a/Engine/Source/Editor/DataLayerEditor/Private/DataLayer/DataLayerHierarchy.cpp b/Engine/Source/Editor/DataLayerEditor/Private/DataLayer/DataLayerHierarchy.cpp
--- a/Engine/Source/Editor/DataLayerEditor/Private/DataLayer/DataLayerHierarchy.cpp
+++ b/Engine/Source/Editor/DataLayerEditor/Private/DataLayer/DataLayerHierarchy.cpp
@@ -204,7 +204,11 @@ void FDataLayerHierarchy::CreateItems(TArray<FSceneOutlinerTreeItemPtr>& OutItem
 
 		// Create DataLayerTreeItems for World and for Current Level (if any).
 		CreateDataLayerTreeItems(GetOwningWorld()->GetDataLayerManager());
-		CreateDataLayerTreeItems(UDataLayerManager::GetDataLayerManager(CurrentLevel));
+
+		if (CurrentLevel)
+		{
+			CreateDataLayerTreeItems(UDataLayerManager::GetDataLayerManager(CurrentLevel));
+		}
 
 		if (bShowDataLayerActors)
 		{
diff --git a/Engine/Source/Runtime/Engine/Private/ActorEditor.cpp b/Engine/Source/Runtime/Engine/Private/ActorEditor.cpp
--- a/Engine/Source/Runtime/Engine/Private/ActorEditor.cpp
+++ b/Engine/Source/Runtime/Engine/Private/ActorEditor.cpp
@@ -1652,6 +1652,12 @@ void AActor::FixupDataLayers(bool bRevertChangesOnLockedDataLayer /*= false*/)
 		return;
 	}
 
+	// Don't fixup if the actor is not part of a level (template actors, etc.)
+	if (!GetLevel())
+	{
+		return;
+	}
+
 	// Don't fixup in game world
 	UWorld* World = GetWorld();
 	if ((GetPackage()->HasAnyPackageFlags(PKG_PlayInEditor)) ||
diff --git a/Engine/Source/Runtime/Engine/Private/WorldPartition/DataLayer/DataLayerManager.cpp b/Engine/Source/Runtime/Engine/Private/WorldPartition/DataLayer/DataLayerManager.cpp
--- a/Engine/Source/Runtime/Engine/Private/WorldPartition/DataLayer/DataLayerManager.cpp
+++ b/Engine/Source/Runtime/Engine/Private/WorldPartition/DataLayer/DataLayerManager.cpp
@@ -543,12 +543,6 @@ TArray<UDataLayerInstance*> UDataLayerManager::ConvertArgsToDataLayers(const TAr
 	return OutDataLayers.Array();
 }
 
-UDataLayerManager* UDataLayerManager::GetDataLayerManager(const UObject* InObject)
-{
-	UWorldPartition* WorldPartition = FWorldPartitionHelpers::GetWorldPartition(InObject);
-	return WorldPartition ? WorldPartition->GetDataLayerManager() : nullptr;
-}
-
 void UDataLayerManager::ForEachDataLayerInstance(TFunctionRef<bool(UDataLayerInstance*)> Func)
 {
 	if (AWorldDataLayers* WorldDataLayers = GetWorldDataLayers())
diff --git a/Engine/Source/Runtime/Engine/Private/WorldPartition/WorldPartitionHelpers.cpp b/Engine/Source/Runtime/Engine/Private/WorldPartition/WorldPartitionHelpers.cpp
--- a/Engine/Source/Runtime/Engine/Private/WorldPartition/WorldPartitionHelpers.cpp
+++ b/Engine/Source/Runtime/Engine/Private/WorldPartition/WorldPartitionHelpers.cpp
@@ -3,6 +3,7 @@
 #include "WorldPartition/WorldPartitionHelpers.h"
 
 #include "WorldPartition/WorldPartition.h"
+#include "WorldPartition/WorldPartitionLog.h"
 #include "WorldPartition/WorldPartitionEditorHash.h"
 #include "WorldPartition/WorldPartitionRuntimeCell.h"
 #include "Engine/World.h"
@@ -14,41 +15,32 @@
 
 DEFINE_LOG_CATEGORY_STATIC(LogWorldPartitionHelpers, Log, All);
 
-UWorldPartition* FWorldPartitionHelpers::GetWorldPartition(const UObject* InObject)
+namespace FWorldPartitionHelpersPrivate
 {
-	if (!IsValid(InObject))
-	{
-		return nullptr;
-	}
-
-	const UWorld* OuterWorld = nullptr;
-
-	if (const UWorld* InWorld = Cast<const UWorld>(InObject))
+	UWorldPartition* GetWorldPartitionFromObject(const UObject* InObject)
 	{
-		return GetWorldPartition(InWorld->PersistentLevel);
-	}
-	else if (const ULevel* Level = Cast<const ULevel>(InObject))
-	{
-		if (const IWorldPartitionCell* Cell = Level->GetWorldPartitionRuntimeCell())
+		for (const UObject* Object = InObject; IsValid(Object); Object = Object->GetOuter())
 		{
-			return GetWorldPartition(Cast<UObject>(Cell));
+			if (const AActor* Actor = Cast<const AActor>(Object))
+			{
+				return GetWorldPartition(Actor);
+			}
+			else if (const ULevel* Level = Cast<const ULevel>(Object))
+			{
+				return GetWorldPartition(Level);
+			}
+			else if (const UWorld* World = Cast<const UWorld>(Object))
+			{
+				return GetWorldPartition(World);
+			}
+			else if (const UWorldPartition* WorldPartition = Cast<const UWorldPartition>(Object))
+			{
+				return const_cast<UWorldPartition*>(WorldPartition);
+			}
 		}
-	}
-	else if (const ULevel* OuterLevel = InObject->GetTypedOuter<ULevel>())
-	{
-		return GetWorldPartition(OuterLevel);
-	}
-	else if (const UWorldPartitionRuntimeCell* Cell = Cast<const UWorldPartitionRuntimeCell>(InObject))
-	{
-		OuterWorld = Cell->GetOuterWorld();
-	}
 
-	if (!OuterWorld)
-	{
-		OuterWorld = InObject ? InObject->GetTypedOuter<UWorld>() : nullptr;
+		return nullptr;
 	}
-
-	return OuterWorld ? OuterWorld->GetWorldPartition() : nullptr;
 }
 
 #if WITH_EDITOR
@@ -329,4 +321,4 @@ bool FWorldPartitionHelpers::ConvertEditorPathToRuntimePath(const FSoftObjectPat
 	return false;
 }
 
-#endif // #if WITH_EDITOR
+#endif // #if WITH_EDITOR
\ No newline at end of file
diff --git a/Engine/Source/Runtime/Engine/Private/WorldPartition/WorldPartitionStreamingGeneration.cpp b/Engine/Source/Runtime/Engine/Private/WorldPartition/WorldPartitionStreamingGeneration.cpp
--- a/Engine/Source/Runtime/Engine/Private/WorldPartition/WorldPartitionStreamingGeneration.cpp
+++ b/Engine/Source/Runtime/Engine/Private/WorldPartition/WorldPartitionStreamingGeneration.cpp
@@ -254,7 +254,7 @@ class FWorldPartitionStreamingGenerator
 
 	void ResolveRuntimeDataLayers(FWorldPartitionActorDescView& ActorDescView, const FActorDescViewMap& ActorDescViewMap)
 	{
-		const UDataLayerManager* DataLayerManager = UDataLayerManager::GetDataLayerManager(WorldPartitionContext);
+		const UDataLayerManager* DataLayerManager = WorldPartitionContext->GetDataLayerManager();
 
 		// Resolve DataLayerInstanceNames of ActorDescView only when necessary (i.e. when container is a template)
 		if (!ActorDescView.GetActorDesc()->HasResolvedDataLayerInstanceNames())
@@ -546,7 +546,7 @@ class FWorldPartitionStreamingGenerator
 				}
 
 				// Validate data layers
-				if (UDataLayerManager* DataLayerManager = UDataLayerManager::GetDataLayerManager(WorldPartitionContext))
+				if (UDataLayerManager* DataLayerManager = WorldPartitionContext->GetDataLayerManager())
 				{
 					DataLayerManager->ForEachDataLayerInstance([this](const UDataLayerInstance* DataLayerInstance)
 					{
@@ -975,7 +975,7 @@ public:
 	TArray<const UDataLayerInstance*> GetRuntimeDataLayerInstances(const TArray<FName>& RuntimeDataLayers) const
 	{
 		static TArray<const UDataLayerInstance*> EmptyArray;
-		UDataLayerManager* DataLayerManager = UDataLayerManager::GetDataLayerManager(WorldPartitionContext);
+		UDataLayerManager* DataLayerManager = WorldPartitionContext->GetDataLayerManager();
 		return DataLayerManager ? DataLayerManager->GetRuntimeDataLayerInstances(RuntimeDataLayers) : EmptyArray;
 	}
 
diff --git a/Engine/Source/Runtime/Engine/Public/WorldPartition/DataLayer/DataLayerManager.h b/Engine/Source/Runtime/Engine/Public/WorldPartition/DataLayer/DataLayerManager.h
--- a/Engine/Source/Runtime/Engine/Public/WorldPartition/DataLayer/DataLayerManager.h
+++ b/Engine/Source/Runtime/Engine/Public/WorldPartition/DataLayer/DataLayerManager.h
@@ -4,6 +4,7 @@
 
 #include "WorldPartition/WorldPartition.h"
 #include "WorldPartition/WorldPartitionHandle.h"
+#include "WorldPartition/WorldPartitionHelpers.h"
 #include "WorldPartition/DataLayer/ActorDataLayer.h"
 #include "WorldPartition/DataLayer/DataLayerInstance.h"
 #include "WorldPartition/DataLayer/DataLayerEditorContext.h"
@@ -44,8 +45,12 @@ class ENGINE_API UDataLayerManager : public UObject
 	GENERATED_BODY()
 
 public:
-
-	static UDataLayerManager* GetDataLayerManager(const UObject* InObject);
+	template <class T>
+	static UDataLayerManager* GetDataLayerManager(const T* InObject)
+	{
+		UWorldPartition* WorldPartition = IsValid(InObject) ? FWorldPartitionHelpers::GetWorldPartition(InObject) : nullptr;
+		return WorldPartition ? WorldPartition->GetDataLayerManager() : nullptr;
+	}
 
 	//~ Begin Blueprint interface
 
diff --git a/Engine/Source/Runtime/Engine/Public/WorldPartition/WorldPartitionHelpers.h b/Engine/Source/Runtime/Engine/Public/WorldPartition/WorldPartitionHelpers.h
--- a/Engine/Source/Runtime/Engine/Public/WorldPartition/WorldPartitionHelpers.h
+++ b/Engine/Source/Runtime/Engine/Public/WorldPartition/WorldPartitionHelpers.h
@@ -3,6 +3,7 @@
 
 #include "CoreMinimal.h"
 #include "GameFramework/Actor.h"
+#include "WorldPartition/WorldPartitionRuntimeCellInterface.h"
 
 #if WITH_EDITOR
 #include "WorldPartition/WorldPartitionHandle.h"
@@ -15,11 +16,56 @@ class UWorldPartition;
 class FWorldPartitionActorDesc;
 #endif
 
+namespace FWorldPartitionHelpersPrivate
+{
+	ENGINE_API UWorldPartition* GetWorldPartitionFromObject(const UObject* InObject);
+
+	template <class T>
+	inline UWorldPartition* GetWorldPartition(const T* InObject)
+	{
+		return IsValid(InObject) ? GetWorldPartitionFromObject(InObject) : nullptr;
+	}
+
+	template <>
+	inline UWorldPartition* GetWorldPartition(const ULevel* InLevel)
+	{
+		if (IsValid(InLevel))
+		{
+			if (const IWorldPartitionCell* Cell = InLevel->GetWorldPartitionRuntimeCell())
+			{
+				return Cell->GetOuterWorld()->GetWorldPartition();
+			}		
+			else if (const UWorld* OuterWorld = Cast<UWorld>(InLevel->GetOuter()))
+			{
+				return OuterWorld->GetWorldPartition();
+			}
+		}
+			
+		return nullptr;
+	}
+
+	template <>
+	inline UWorldPartition* GetWorldPartition(const UWorld* InWorld)
+	{
+		return IsValid(InWorld) ? GetWorldPartition<ULevel>(InWorld->PersistentLevel) : nullptr;
+	}
+
+	template <>
+	inline UWorldPartition* GetWorldPartition(const AActor* InActor)
+	{
+		return IsValid(InActor) ? GetWorldPartition<ULevel>(InActor->GetLevel()) : nullptr;
+	}
+}
+
 class ENGINE_API FWorldPartitionHelpers
 {
 public:
 	/** Returns the owning World Partition for this object. */
-	static UWorldPartition* GetWorldPartition(const UObject* InObject);
+	template <class T>
+	static inline UWorldPartition* GetWorldPartition(const T* InObject)
+	{
+		return FWorldPartitionHelpersPrivate::GetWorldPartition(InObject);
+	}
 
 #if WITH_EDITOR
 private:
@@ -90,7 +136,6 @@ public:
 	// Editor/Runtime conversions
 	static bool ConvertRuntimePathToEditorPath(const FSoftObjectPath& InPath, FSoftObjectPath& OutPath);
 	static bool ConvertEditorPathToRuntimePath(const FSoftObjectPath& InPath, FSoftObjectPath& OutPath);
-
 #endif // WITH_EDITOR
 };
 
diff --git a/Engine/Source/Runtime/Engine/Public/WorldPartition/WorldPartitionRuntimeCell.h b/Engine/Source/Runtime/Engine/Public/WorldPartition/WorldPartitionRuntimeCell.h
--- a/Engine/Source/Runtime/Engine/Public/WorldPartition/WorldPartitionRuntimeCell.h
+++ b/Engine/Source/Runtime/Engine/Public/WorldPartition/WorldPartitionRuntimeCell.h
@@ -167,9 +167,6 @@ class ENGINE_API UWorldPartitionRuntimeCell : public UObject, public IWorldParti
 		return Super::GetTypedOuter<T>();
 	}
 
-	UWorld* GetOwningWorld() const;
-	UWorld* GetOuterWorld() const;
-
 	virtual void Load() const PURE_VIRTUAL(UWorldPartitionRuntimeCell::Load,);
 	virtual void Unload() const PURE_VIRTUAL(UWorldPartitionRuntimeCell::Unload,);
 	virtual bool CanUnload() const PURE_VIRTUAL(UWorldPartitionRuntimeCell::CanUnload, return true;);
@@ -201,6 +198,8 @@ class ENGINE_API UWorldPartitionRuntimeCell : public UObject, public IWorldParti
 	}
 	virtual FName GetLevelPackageName() const override;
 	virtual FString GetDebugName() const override;
+	virtual UWorld* GetOwningWorld() const override;
+	virtual UWorld* GetOuterWorld() const override;
 	//~End IWorldPartitionCell Interface
 
 	inline bool HasDataLayers() const { return !DataLayers.IsEmpty(); }
diff --git a/Engine/Source/Runtime/Engine/Public/WorldPartition/WorldPartitionRuntimeCellInterface.h b/Engine/Source/Runtime/Engine/Public/WorldPartition/WorldPartitionRuntimeCellInterface.h
--- a/Engine/Source/Runtime/Engine/Public/WorldPartition/WorldPartitionRuntimeCellInterface.h
+++ b/Engine/Source/Runtime/Engine/Public/WorldPartition/WorldPartitionRuntimeCellInterface.h
@@ -39,6 +39,10 @@ public:
 	virtual FName GetLevelPackageName() const = 0;
 	/** Returns the debug name associated with this cell. */
 	virtual FString GetDebugName() const = 0;
+	/**  Returns the owning world of this cell. */
+	virtual UWorld* GetOwningWorld() const = 0;
+	/**  REturns the outer world of this cell. */
+	virtual UWorld* GetOuterWorld() const = 0;
 
 #if WITH_EDITOR
 	/** Returns the referenced actor packages. */
