World Partition: show the map/plugin name when generating streaming logs.

diff --git a/Engine/Source/Runtime/Engine/Private/WorldPartition/RuntimeSpatialHash/RuntimeSpatialHashGridHelper.cpp b/Engine/Source/Runtime/Engine/Private/WorldPartition/RuntimeSpatialHash/RuntimeSpatialHashGridHelper.cpp
--- a/Engine/Source/Runtime/Engine/Private/WorldPartition/RuntimeSpatialHash/RuntimeSpatialHashGridHelper.cpp
+++ b/Engine/Source/Runtime/Engine/Private/WorldPartition/RuntimeSpatialHash/RuntimeSpatialHashGridHelper.cpp
@@ -3,7 +3,6 @@
 #include "WorldPartition/RuntimeSpatialHash/RuntimeSpatialHashGridHelper.h"
 #include "ActorPartition/PartitionActor.h"
 #include "HAL/IConsoleManager.h"
-#include "ProfilingDebugging/ScopedTimers.h"
 #include "WorldPartition/WorldPartitionLog.h"
 #include "WorldPartition/WorldPartitionStreamingGeneration.h"
 #include "WorldPartition/WorldPartitionStreamingSource.h"
@@ -148,7 +147,6 @@ FSquare2DGridHelper GetGridHelper(const FBox& WorldBounds, int64 GridCellSize)
 FSquare2DGridHelper GetPartitionedActors(const FBox& WorldBounds, const FSpatialHashRuntimeGrid& Grid, const TArray<const IStreamingGenerationContext::FActorSetInstance*>& ActorSetInstances)
 {
 	TRACE_CPUPROFILER_EVENT_SCOPE(GetPartitionedActors);
-	UE_SCOPED_TIMER(TEXT("GetPartitionedActors"), LogWorldPartition, Display);
 
 	//
 	// Create the hierarchical grids for the game
diff --git a/Engine/Source/Runtime/Engine/Private/WorldPartition/WorldPartitionRuntimeSpatialHash.cpp b/Engine/Source/Runtime/Engine/Private/WorldPartition/WorldPartitionRuntimeSpatialHash.cpp
--- a/Engine/Source/Runtime/Engine/Private/WorldPartition/WorldPartitionRuntimeSpatialHash.cpp
+++ b/Engine/Source/Runtime/Engine/Private/WorldPartition/WorldPartitionRuntimeSpatialHash.cpp
@@ -29,7 +29,6 @@
 #include "Math/TranslationMatrix.h"
 #include "Misc/ArchiveMD5.h"
 #include "Misc/HierarchicalLogArchive.h"
-#include "ProfilingDebugging/ScopedTimers.h"
 #include "UObject/UnrealType.h"
 #include "UObject/ObjectSaveContext.h"
 
@@ -1098,8 +1097,6 @@ bool UWorldPartitionRuntimeSpatialHash::GenerateStreaming(UWorldPartitionStreami
 	TRACE_CPUPROFILER_EVENT_SCOPE(UWorldPartitionRuntimeSpatialHash::GenerateStreaming);
 	UWorldPartition* WorldPartition = GetOuterUWorldPartition();
 
-	UE_SCOPED_TIMER(TEXT("GenerateStreaming"), LogWorldPartition, Display);
-	
 	if (!Grids.Num())
 	{
 		UE_LOG(LogWorldPartition, Error, TEXT("Invalid partition grids setup"));
diff --git a/Engine/Source/Runtime/Engine/Private/WorldPartition/WorldPartitionStreamingGeneration.cpp b/Engine/Source/Runtime/Engine/Private/WorldPartition/WorldPartitionStreamingGeneration.cpp
--- a/Engine/Source/Runtime/Engine/Private/WorldPartition/WorldPartitionStreamingGeneration.cpp
+++ b/Engine/Source/Runtime/Engine/Private/WorldPartition/WorldPartitionStreamingGeneration.cpp
@@ -11,6 +11,7 @@
 #include "Misc/PackageName.h"
 #include "ReferenceCluster.h"
 #include "Misc/Paths.h"
+#include "ProfilingDebugging/ScopedTimers.h"
 #include "WorldPartition/WorldPartitionRuntimeHash.h"
 #include "WorldPartition/WorldPartitionStreamingPolicy.h"
 #include "WorldPartition/DataLayer/DataLayerManager.h"
@@ -1025,21 +1026,25 @@ bool UWorldPartition::GenerateContainerStreaming(const UActorDescContainer* InAc
 		ErrorHandler = &MapCheckErrorHandler;
 	}
 
-	// Dump state log
-	TStringBuilder<256> StateLogSuffix;
-	StateLogSuffix += bIsPIE ? TEXT("PIE") : (IsRunningGame() ? TEXT("Game") : (IsRunningCookCommandlet() ? TEXT("Cook") : TEXT("Manual")));
-	StateLogSuffix += TEXT("_");
-	FString ContainerPackageName = InActorDescContainer->ContainerPackageName.ToString();
-	StateLogSuffix += FPackageName::GetShortName(ContainerPackageName);
+	const FString ContainerPackageName = InActorDescContainer->ContainerPackageName.ToString();
+	FString ContainerShortName = FPackageName::GetShortName(ContainerPackageName);
 	if (!ContainerPackageName.StartsWith(TEXT("/Game/")))
 	{
 		TArray<FString> SplitContainerPath;
 		if (ContainerPackageName.ParseIntoArray(SplitContainerPath, TEXT("/")))
 		{
-			StateLogSuffix += TEXT("_");
-			StateLogSuffix += SplitContainerPath[0];
+			ContainerShortName += TEXT(".");
+			ContainerShortName += SplitContainerPath[0];
 		}
 	}
+
+	UE_SCOPED_TIMER(*FString::Printf(TEXT("GenerateStreaming for '%s'"), *ContainerShortName), LogWorldPartition, Display);
+
+	// Dump state log
+	TStringBuilder<256> StateLogSuffix;
+	StateLogSuffix += bIsPIE ? TEXT("PIE") : (IsRunningGame() ? TEXT("Game") : (IsRunningCookCommandlet() ? TEXT("Cook") : TEXT("Manual")));
+	StateLogSuffix += TEXT("_");
+	StateLogSuffix += ContainerShortName;
 	TUniquePtr<FArchive> LogFileAr = FWorldPartitionStreamingGenerator::CreateDumpStateLogArchive(*StateLogSuffix);
 	FHierarchicalLogArchive HierarchicalLogAr(*LogFileAr);
 
diff --git a/Engine/Source/Runtime/Engine/Public/WorldPartition/WorldPartitionHelpers.h b/Engine/Source/Runtime/Engine/Public/WorldPartition/WorldPartitionHelpers.h
--- a/Engine/Source/Runtime/Engine/Public/WorldPartition/WorldPartitionHelpers.h
+++ b/Engine/Source/Runtime/Engine/Public/WorldPartition/WorldPartitionHelpers.h
@@ -3,6 +3,7 @@
 
 #include "CoreMinimal.h"
 #include "GameFramework/Actor.h"
+#include "Engine/Level.h"
 #include "WorldPartition/WorldPartitionRuntimeCellInterface.h"
 
 #if WITH_EDITOR
