Fix non-unity build.

diff --git a/Engine/Source/Runtime/Engine/Private/WorldPartition/WorldPartitionHelpers.cpp b/Engine/Source/Runtime/Engine/Private/WorldPartition/WorldPartitionHelpers.cpp
--- a/Engine/Source/Runtime/Engine/Private/WorldPartition/WorldPartitionHelpers.cpp
+++ b/Engine/Source/Runtime/Engine/Private/WorldPartition/WorldPartitionHelpers.cpp
@@ -6,8 +6,6 @@
 #include "WorldPartition/WorldPartitionLog.h"
 #include "WorldPartition/WorldPartitionEditorHash.h"
 #include "WorldPartition/WorldPartitionRuntimeCell.h"
-#include "Engine/World.h"
-#include "Engine/Level.h"
 #include "Algo/AnyOf.h"
 
 #include "Commandlets/Commandlet.h"
diff --git a/Engine/Source/Runtime/Engine/Public/WorldPartition/WorldPartitionHelpers.h b/Engine/Source/Runtime/Engine/Public/WorldPartition/WorldPartitionHelpers.h
--- a/Engine/Source/Runtime/Engine/Public/WorldPartition/WorldPartitionHelpers.h
+++ b/Engine/Source/Runtime/Engine/Public/WorldPartition/WorldPartitionHelpers.h
@@ -3,6 +3,7 @@
 
 #include "CoreMinimal.h"
 #include "GameFramework/Actor.h"
+#include "Engine/World.h"
 #include "Engine/Level.h"
 #include "WorldPartition/WorldPartitionRuntimeCellInterface.h"
 
