Merging CL#24484412

#UE Add CVar to debug DDC determinism problems with a specific asset

diff --git a/Engine/Source/Runtime/Engine/Private/Materials/MaterialShader.cpp b/Engine/Source/Runtime/Engine/Private/Materials/MaterialShader.cpp
--- a/Engine/Source/Runtime/Engine/Private/Materials/MaterialShader.cpp
+++ b/Engine/Source/Runtime/Engine/Private/Materials/MaterialShader.cpp
@@ -43,6 +43,15 @@ static TAutoConsoleVariable<FString> CVarMaterialShaderMapDump(
 	TEXT("Files (.txt extension) will be dumped to Saved\\MaterialShaderMaps named with the DDC key hash.\n"),
 	ECVF_ReadOnly);
 
+#if WITH_EDITOR
+static TAutoConsoleVariable<FString> CVarShaderCompilerDebugDDCKeyAsset(
+	TEXT("r.ShaderCompiler.DebugDDCKeyAsset"),
+	FString(),
+	TEXT("if set, assets containing this name will print their full DDC key hash information when requested"),
+	ECVF_Default
+);
+#endif // WITH_EDITOR
+
 #if ENABLE_COOK_STATS
 namespace MaterialShaderCookStats
 {
@@ -1473,6 +1482,17 @@ TSharedRef<FMaterialShaderMap::FAsyncLoadContext> FMaterialShaderMap::BeginLoadF
 				UE_LOG(LogMaterial, Log, TEXT("%s-%s-%s: %s"), *Material->GetAssetName(), *LexToString(ShaderMapId.FeatureLevel), *LexToString(ShaderMapId.QualityLevel), *SpecialEngineDDCKey);
 			}
 
+			if (UNLIKELY(!CVarShaderCompilerDebugDDCKeyAsset.GetValueOnAnyThread().IsEmpty()))
+			{
+				FString DebugDDCKeyAsset = CVarShaderCompilerDebugDDCKeyAsset.GetValueOnAnyThread();
+				if (Material->GetAssetName().Contains(DebugDDCKeyAsset))
+				{
+					FString DataKey = Result->DataKey;
+					DataKey.RemoveFromEnd(TEXT("\n"));
+					UE_LOG(LogMaterial, Display, TEXT("%s-%s-%s: %s"), *Material->GetAssetName(), *LexToString(ShaderMapId.FeatureLevel), *LexToString(ShaderMapId.QualityLevel), *DataKey);
+				}
+			}
+
 			bool bCheckCache = true;
 
 			// If NoShaderDDC then don't check for a material the first time we encounter it to simulate
