Fixed issues with materials and reference replacement

* Various material expressions were storing data in PreEditChangeProperty in a shared variable rather than per-instance.
* MaterialFunctionCall now always refreshes its inputs/outputs when a reference replacement has happened on it.
* MaterialInstance now ensures that its base material has valid expression data when rebuilding after a reference replacement has happened on it, as these PECP can happen out-of-order (or not at all).

diff --git a/Engine/Source/Runtime/Engine/Classes/Materials/MaterialExpressionFunctionInput.h b/Engine/Source/Runtime/Engine/Classes/Materials/MaterialExpressionFunctionInput.h
--- a/Engine/Source/Runtime/Engine/Classes/Materials/MaterialExpressionFunctionInput.h
+++ b/Engine/Source/Runtime/Engine/Classes/Materials/MaterialExpressionFunctionInput.h
@@ -134,6 +134,9 @@ private:
 #if WITH_EDITOR
 	/** Helper function which compiles this expression for previewing. */
 	int32 CompilePreviewValue(FMaterialCompiler* Compiler);
+
+	/** Stashed data between a Pre/PostEditChange event */
+	FName InputNameBackup;
 #endif // WITH_EDITOR
 };
 
diff --git a/Engine/Source/Runtime/Engine/Classes/Materials/MaterialExpressionFunctionOutput.h b/Engine/Source/Runtime/Engine/Classes/Materials/MaterialExpressionFunctionOutput.h
--- a/Engine/Source/Runtime/Engine/Classes/Materials/MaterialExpressionFunctionOutput.h
+++ b/Engine/Source/Runtime/Engine/Classes/Materials/MaterialExpressionFunctionOutput.h
@@ -78,6 +78,12 @@ class UMaterialExpressionFunctionOutput : public UMaterialExpression
 	/** Validate OutputName.  Must be called after OutputName is changed to prevent duplicate outputs. */
 	ENGINE_API void ValidateName();
 #endif // WITH_EDITOR
+
+private:
+#if WITH_EDITOR
+	/** Stashed data between a Pre/PostEditChange event */
+	FName OutputNameBackup;
+#endif // WITH_EDITOR
 };
 
 
diff --git a/Engine/Source/Runtime/Engine/Classes/Materials/MaterialExpressionMaterialFunctionCall.h b/Engine/Source/Runtime/Engine/Classes/Materials/MaterialExpressionMaterialFunctionCall.h
--- a/Engine/Source/Runtime/Engine/Classes/Materials/MaterialExpressionMaterialFunctionCall.h
+++ b/Engine/Source/Runtime/Engine/Classes/Materials/MaterialExpressionMaterialFunctionCall.h
@@ -176,6 +176,9 @@ private:
 
 	/** Temporary compilation state shared between function calls */
 	FMaterialFunctionCompileState* SharedCompileState;
+
+	/** Stashed data between a Pre/PostEditChange event */
+	UMaterialFunctionInterface* SavedMaterialFunction = nullptr;
 #endif // WITH_EDITOR
 };
 
diff --git a/Engine/Source/Runtime/Engine/Private/Materials/MaterialExpressions.cpp b/Engine/Source/Runtime/Engine/Private/Materials/MaterialExpressions.cpp
--- a/Engine/Source/Runtime/Engine/Private/Materials/MaterialExpressions.cpp
+++ b/Engine/Source/Runtime/Engine/Private/Materials/MaterialExpressions.cpp
@@ -15980,8 +15980,6 @@ void FMaterialLayersFunctions::Validate(const FMaterialLayersFunctionsRuntimeDat
 // UMaterialExpressionMaterialFunctionCall
 ///////////////////////////////////////////////////////////////////////////////
 
-UMaterialFunctionInterface* SavedMaterialFunction = nullptr;
-
 UMaterialExpressionMaterialFunctionCall::UMaterialExpressionMaterialFunctionCall(const FObjectInitializer& ObjectInitializer)
 	: Super(ObjectInitializer)
 {
@@ -16067,6 +16065,11 @@ void UMaterialExpressionMaterialFunctionCall::PostEditChangeProperty(FPropertyCh
 		SetMaterialFunctionEx(SavedMaterialFunction, MaterialFunction);
 		SavedMaterialFunction = nullptr;
 	}
+	else if (PropertyChangedEvent.ChangeType == EPropertyChangeType::Redirected)
+	{
+		// Refresh from the current material function as it may have been redirected to a different value
+		UpdateFromFunctionResource();
+	}
 
 	Super::PostEditChangeProperty(PropertyChangedEvent);
 }
@@ -16684,8 +16687,6 @@ void UMaterialExpressionFunctionInput::PostEditImport()
 	ConditionallyGenerateId(true);
 }
 
-FName InputNameBackup;
-
 void UMaterialExpressionFunctionInput::PreEditChange(FProperty* PropertyAboutToChange)
 {
 	if (PropertyAboutToChange && PropertyAboutToChange->GetFName() == GET_MEMBER_NAME_CHECKED(UMaterialExpressionFunctionInput, InputName))
@@ -17079,8 +17080,6 @@ void UMaterialExpressionFunctionOutput::PostEditImport()
 }
 #endif	//#if WITH_EDITOR
 
-FName OutputNameBackup;
-
 #if WITH_EDITOR
 void UMaterialExpressionFunctionOutput::PreEditChange(FProperty* PropertyAboutToChange)
 {
diff --git a/Engine/Source/Runtime/Engine/Private/Materials/MaterialInstance.cpp b/Engine/Source/Runtime/Engine/Private/Materials/MaterialInstance.cpp
--- a/Engine/Source/Runtime/Engine/Private/Materials/MaterialInstance.cpp
+++ b/Engine/Source/Runtime/Engine/Private/Materials/MaterialInstance.cpp
@@ -3911,6 +3911,15 @@ void UMaterialInstance::PostEditChangeProperty(FPropertyChangedEvent& PropertyCh
 
 	// Force UpdateStaticPermutation when change type is Redirected as this probably means a Material or MaterialInstance parent asset was deleted.
 	const bool bForceStaticPermutationUpdate = PropertyChangedEvent.ChangeType == EPropertyChangeType::Redirected;
+	if (bForceStaticPermutationUpdate)
+	{
+		// This can run before UMaterial::PostEditChangeProperty has a chance to run, so explicitly call UpdateCachedExpressionData here
+		UMaterial* BaseMaterial = GetMaterial();
+		if (BaseMaterial && !BaseMaterial->GetPackage()->HasAnyPackageFlags(PKG_Cooked))
+		{
+			BaseMaterial->UpdateCachedExpressionData();
+		}
+	}
 	UpdateStaticPermutation(GetStaticParameters(), BasePropertyOverrides, bForceStaticPermutationUpdate);
 
 	if (PropertyChangedEvent.ChangeType & (EPropertyChangeType::ValueSet | EPropertyChangeType::ArrayClear | EPropertyChangeType::ArrayRemove | EPropertyChangeType::ArrayMove | EPropertyChangeType::Unspecified | EPropertyChangeType::Duplicate))
