RHI Parameter Batching
- Moving some loose uniform buffer sets into their shader SetParameters methods.
- Changing some deprecated functions to use batched parameters internally.

diff --git a/Engine/Plugins/Experimental/ColorCorrectRegions/Source/ColorCorrectRegions/Public/ColorCorrectRegionsPostProcessMaterial.h b/Engine/Plugins/Experimental/ColorCorrectRegions/Source/ColorCorrectRegions/Public/ColorCorrectRegionsPostProcessMaterial.h
--- a/Engine/Plugins/Experimental/ColorCorrectRegions/Source/ColorCorrectRegions/Public/ColorCorrectRegionsPostProcessMaterial.h
+++ b/Engine/Plugins/Experimental/ColorCorrectRegions/Source/ColorCorrectRegions/Public/ColorCorrectRegionsPostProcessMaterial.h
@@ -126,7 +126,9 @@ public:
 	UE_DEPRECATED(5.3, "SetParameters with FRHIBatchedShaderParameters should be used.")
 	void SetParameters(FRHICommandList& RHICmdList, const FSceneView& View)
 	{
-		FGlobalShader::SetParameters<FViewUniformShaderParameters>(RHICmdList, RHICmdList.GetBoundVertexShader(), View.ViewUniformBuffer);
+		FRHIBatchedShaderParameters& BatchedParameters = RHICmdList.GetScratchShaderParameters();
+		SetParameters(BatchedParameters, View);
+		RHICmdList.SetBatchedShaderParameters(RHICmdList.GetBoundVertexShader(), BatchedParameters);
 	}
 };
 
@@ -160,7 +162,9 @@ public:
 	UE_DEPRECATED(5.3, "SetParameters with FRHIBatchedShaderParameters should be used.")
 	void SetParameters(FRHICommandList& RHICmdList, const FSceneView& View)
 	{
-		FGlobalShader::SetParameters<FViewUniformShaderParameters>(RHICmdList, RHICmdList.GetBoundPixelShader(), View.ViewUniformBuffer);
+		FRHIBatchedShaderParameters& BatchedParameters = RHICmdList.GetScratchShaderParameters();
+		SetParameters(BatchedParameters, View);
+		RHICmdList.SetBatchedShaderParameters(RHICmdList.GetBoundPixelShader(), BatchedParameters);
 	}
 };
 
diff --git a/Engine/Source/Runtime/Engine/Private/GPUSort.cpp b/Engine/Source/Runtime/Engine/Private/GPUSort.cpp
--- a/Engine/Source/Runtime/Engine/Private/GPUSort.cpp
+++ b/Engine/Source/Runtime/Engine/Private/GPUSort.cpp
@@ -294,8 +294,10 @@ public:
 		return RadixSortParameterBuffer.IsBound();
 	}
 
-	void SetParameters(FRHIBatchedShaderParameters& BatchedParameters, FRHIUnorderedAccessView* OutOffsetsUAV, FRHIShaderResourceView* InKeysSRV, const FRadixSortUniformBufferRef& RadixSortUniformBuffer, FRHIShaderResourceView* RadixSortParameterBufferSRV)
+	void SetParameters(FRHIBatchedShaderParameters& BatchedParameters, const FRadixSortUniformBufferRef& SortUniformBufferRef, FRHIUnorderedAccessView* OutOffsetsUAV, FRHIShaderResourceView* InKeysSRV, const FRadixSortUniformBufferRef& RadixSortUniformBuffer, FRHIShaderResourceView* RadixSortParameterBufferSRV)
 	{
+		SetUniformBufferParameter(BatchedParameters, GetUniformBufferParameter<FRadixSortParameters>(), SortUniformBufferRef);
+
 		SetUAVParameter(BatchedParameters, OutOffsets, OutOffsetsUAV);
 		SetSRVParameter(BatchedParameters, InKeys, InKeysSRV);
 		SetSRVParameter(BatchedParameters, RadixSortParameterBuffer, RadixSortParameterBufferSRV);
@@ -420,6 +422,7 @@ public:
 
 	void SetParameters(
 		FRHIBatchedShaderParameters& BatchedParameters,
+		const FRadixSortUniformBufferRef& SortUniformBufferRef,
 		FRHIUnorderedAccessView* OutKeysUAV,
 		FRHIUnorderedAccessView* OutValuesUAV,
 		FRHIShaderResourceView* InKeysSRV,
@@ -427,6 +430,8 @@ public:
 		FRHIShaderResourceView* InOffsetsSRV,
 		FRHIShaderResourceView* RadixSortParameterBufferSRV )
 	{
+		SetUniformBufferParameter(BatchedParameters, GetUniformBufferParameter<FRadixSortParameters>(), SortUniformBufferRef);
+
 		SetUAVParameter(BatchedParameters, OutKeys, OutKeysUAV);
 		SetUAVParameter(BatchedParameters, OutValues, OutValuesUAV);
 
@@ -595,8 +600,7 @@ int32 SortGPUBuffers(FRHICommandList& RHICmdList, FGPUSortBuffers SortBuffers, i
 			// Phase 1: Scan upsweep to compute per-digit totals.
 			SetComputePipelineState(RHICmdList, UpsweepCS.GetComputeShader());
 
-			SetUniformBufferParameter(RHICmdList, UpsweepCS.GetComputeShader(), UpsweepCS->GetUniformBufferParameter<FRadixSortParameters>(), SortUniformBufferRef);
-			SetShaderParametersLegacyCS(RHICmdList, UpsweepCS, GSortOffsetBuffers.BufferUAVs[0], SortBuffers.RemoteKeySRVs[BufferIndex], SortUniformBufferRef, GRadixSortParametersBuffer.SortParametersBufferSRV);
+			SetShaderParametersLegacyCS(RHICmdList, UpsweepCS, SortUniformBufferRef, GSortOffsetBuffers.BufferUAVs[0], SortBuffers.RemoteKeySRVs[BufferIndex], SortUniformBufferRef, GRadixSortParametersBuffer.SortParametersBufferSRV);
 
 			DispatchComputeShader(RHICmdList, UpsweepCS.GetShader(), GroupCount, 1, 1);
 
@@ -638,8 +642,6 @@ int32 SortGPUBuffers(FRHICommandList& RHICmdList, FGPUSortBuffers SortBuffers, i
 			// Phase 3: Downsweep to compute final offsets and scatter keys.
 			SetComputePipelineState(RHICmdList, DownsweepCS.GetComputeShader());
 
-			SetUniformBufferParameter(RHICmdList, DownsweepCS.GetComputeShader(), DownsweepCS->GetUniformBufferParameter<FRadixSortParameters>(), SortUniformBufferRef);
-
 			{
 				FRHIUnorderedAccessView* ValuesUAV = nullptr;
 				if (bIsLastPass && SortBuffers.FinalValuesUAV)
@@ -655,7 +657,7 @@ int32 SortGPUBuffers(FRHICommandList& RHICmdList, FGPUSortBuffers SortBuffers, i
 
 				FRHIShaderResourceView* ValuesSRV = (PassIndex == 0 && SortBuffers.FirstValuesSRV) ? SortBuffers.FirstValuesSRV : SortBuffers.RemoteValueSRVs[BufferIndex];
 
-				SetShaderParametersLegacyCS(RHICmdList, DownsweepCS, SortBuffers.RemoteKeyUAVs[BufferIndex ^ 0x1], ValuesUAV, SortBuffers.RemoteKeySRVs[BufferIndex], ValuesSRV, GSortOffsetBuffers.BufferSRVs[1], GRadixSortParametersBuffer.SortParametersBufferSRV);
+				SetShaderParametersLegacyCS(RHICmdList, DownsweepCS, SortUniformBufferRef, SortBuffers.RemoteKeyUAVs[BufferIndex ^ 0x1], ValuesUAV, SortBuffers.RemoteKeySRVs[BufferIndex], ValuesSRV, GSortOffsetBuffers.BufferSRVs[1], GRadixSortParametersBuffer.SortParametersBufferSRV);
 			}
 			DispatchComputeShader(RHICmdList, DownsweepCS.GetShader(), GroupCount, 1, 1 );
 			UnsetShaderParametersLegacyCS(RHICmdList, DownsweepCS);
diff --git a/Engine/Source/Runtime/Engine/Private/Particles/ParticleSortingGPU.cpp b/Engine/Source/Runtime/Engine/Private/Particles/ParticleSortingGPU.cpp
--- a/Engine/Source/Runtime/Engine/Private/Particles/ParticleSortingGPU.cpp
+++ b/Engine/Source/Runtime/Engine/Private/Particles/ParticleSortingGPU.cpp
@@ -78,12 +78,16 @@ public:
 
 	void SetParameters(
 		FRHIBatchedShaderParameters& BatchedParameters,
+		const FParticleKeyGenParameters& KeyGenParameters,
 		FRHIUnorderedAccessView* OutKeysUAV,
 		FRHIUnorderedAccessView* OutIndicesUAV,
 		FRHITexture2D* PositionTextureRHI,
 		FRHIShaderResourceView* InIndicesSRV
 		)
 	{
+		FParticleKeyGenUniformBufferRef KeyGenUniformBuffer = FParticleKeyGenUniformBufferRef::CreateUniformBufferImmediate(KeyGenParameters, UniformBuffer_SingleDraw);
+		SetUniformBufferParameter(BatchedParameters, GetUniformBufferParameter<FParticleKeyGenParameters>(), KeyGenUniformBuffer);
+
 		SetUAVParameter(BatchedParameters, OutKeys, OutKeysUAV);
 		SetUAVParameter(BatchedParameters, OutParticleIndices, OutIndicesUAV);
 		SetTextureParameter(BatchedParameters, PositionTexture, PositionTextureRHI);
@@ -159,10 +163,7 @@ int32 GenerateParticleSortKeys(
 			KeyGenParameters.EmitterKey = (uint32)SortInfo.AllocationInfo.ElementIndex << 16;
 			KeyGenParameters.KeyCount = ParticleCount;
 
-			FParticleKeyGenUniformBufferRef KeyGenUniformBuffer = FParticleKeyGenUniformBufferRef::CreateUniformBufferImmediate( KeyGenParameters, UniformBuffer_SingleDraw );
-			SetUniformBufferParameter(RHICmdList, KeyGenCS.GetComputeShader(), KeyGenCS->GetUniformBufferParameter<FParticleKeyGenParameters>(), KeyGenUniformBuffer);
-
-			SetShaderParametersLegacyCS(RHICmdList, KeyGenCS, KeyBufferUAV, SortedVertexBufferUAV, PositionTextureRHI, SortInfo.VertexBufferSRV);
+			SetShaderParametersLegacyCS(RHICmdList, KeyGenCS, KeyGenParameters, KeyBufferUAV, SortedVertexBufferUAV, PositionTextureRHI, SortInfo.VertexBufferSRV);
 
 			DispatchComputeShader(RHICmdList, KeyGenCS.GetShader(), GroupCount, 1, 1);
 		}
diff --git a/Engine/Source/Runtime/Engine/Private/VectorField.cpp b/Engine/Source/Runtime/Engine/Private/VectorField.cpp
--- a/Engine/Source/Runtime/Engine/Private/VectorField.cpp
+++ b/Engine/Source/Runtime/Engine/Private/VectorField.cpp
@@ -998,10 +998,13 @@ public:
 
 	void SetParameters(
 		FRHIBatchedShaderParameters& BatchedParameters,
+		const FCompositeAnimatedVectorFieldUniformBufferRef& UniformBuffer,
 		FRHIUnorderedAccessView* VolumeTextureUAV,
 		FRHITexture* AtlasTextureRHI,
 		FRHITexture* NoiseVolumeTextureRHI)
 	{
+		SetUniformBufferParameter(BatchedParameters, GetUniformBufferParameter<FCompositeAnimatedVectorFieldUniformParameters>(), UniformBuffer);
+
 		FRHISamplerState* SamplerStateLinear = TStaticSamplerState<SF_Bilinear, AM_Clamp, AM_Clamp, AM_Clamp>::GetRHI();
 
 		SetUAVParameter(BatchedParameters, OutVolumeTexture, VolumeTextureUAV);
@@ -1166,11 +1169,10 @@ public:
 			RHICmdList.Transition(FRHITransitionInfo(VolumeTextureUAV, ERHIAccess::Unknown, ERHIAccess::UAVCompute));
 			SetComputePipelineState(RHICmdList, CompositeCS.GetComputeShader());
 
-			SetUniformBufferParameter(RHICmdList, CompositeCS.GetComputeShader(), CompositeCS->GetUniformBufferParameter<FCompositeAnimatedVectorFieldUniformParameters>(), UniformBuffer);
-
 			SetShaderParametersLegacyCS(
 				RHICmdList,
 				CompositeCS,
+				UniformBuffer,
 				VolumeTextureUAV,
 				AnimatedVectorField->Texture->GetResource()->TextureRHI,
 				NoiseVolumeTextureRHI);
diff --git a/Engine/Source/Runtime/Engine/Public/ScreenRendering.h b/Engine/Source/Runtime/Engine/Public/ScreenRendering.h
--- a/Engine/Source/Runtime/Engine/Public/ScreenRendering.h
+++ b/Engine/Source/Runtime/Engine/Public/ScreenRendering.h
@@ -414,7 +414,9 @@ public:
 	UE_DEPRECATED(5.3, "SetParameters with FRHIBatchedShaderParameters should be used.")
 	void SetParameters(FRHICommandList& RHICmdList, FRHIUniformBuffer* ViewUniformBuffer)
 	{
-		FGlobalShader::SetParameters<FViewUniformShaderParameters>(RHICmdList, RHICmdList.GetBoundVertexShader(), ViewUniformBuffer);
+		FRHIBatchedShaderParameters& BatchedParameters = RHICmdList.GetScratchShaderParameters();
+		FGlobalShader::SetParameters<FViewUniformShaderParameters>(BatchedParameters, ViewUniformBuffer);
+		RHICmdList.SetBatchedShaderParameters(RHICmdList.GetBoundVertexShader(), BatchedParameters);
 	}
 };
 
diff --git a/Engine/Source/Runtime/RenderCore/Public/GlobalShader.h b/Engine/Source/Runtime/RenderCore/Public/GlobalShader.h
--- a/Engine/Source/Runtime/RenderCore/Public/GlobalShader.h
+++ b/Engine/Source/Runtime/RenderCore/Public/GlobalShader.h
@@ -298,10 +298,12 @@ public:
 	}
 
 	template<typename TViewUniformShaderParameters, typename ShaderRHIParamRef, typename TRHICmdList>
+	UE_DEPRECATED(5.3, "SetParameters with FRHIBatchedShaderParameters should be used.")
 	inline void SetParameters(TRHICmdList& RHICmdList, const ShaderRHIParamRef ShaderRHI, FRHIUniformBuffer* ViewUniformBuffer)
 	{
-		const auto& ViewUniformBufferParameter = static_cast<const FShaderUniformBufferParameter&>(GetUniformBufferParameter<TViewUniformShaderParameters>());
-		SetUniformBufferParameter(RHICmdList, ShaderRHI, ViewUniformBufferParameter, ViewUniformBuffer);
+		FRHIBatchedShaderParameters& BatchedParameters = RHICmdList.GetScratchShaderParameters();
+		SetParameters<TViewUniformShaderParameters>(BatchedParameters, ViewUniformBuffer);
+		RHICmdList.SetBatchedShaderParameters(ShaderRHI, BatchedParameters);
 	}
 
 	static inline bool ShouldCompilePermutation(const FGlobalShaderPermutationParameters& Parameters)
diff --git a/Engine/Source/Runtime/Renderer/Private/LightRendering.cpp b/Engine/Source/Runtime/Renderer/Private/LightRendering.cpp
--- a/Engine/Source/Runtime/Renderer/Private/LightRendering.cpp
+++ b/Engine/Source/Runtime/Renderer/Private/LightRendering.cpp
@@ -885,6 +885,11 @@ class FDeferredLightPS : public FGlobalShader
 			FForwardLightingParameters::ModifyCompilationEnvironment(Parameters.Platform, OutEnvironment);
 		}
 	}
+
+	void SetParameters(FRHIBatchedShaderParameters& BatchedParameters, const FDeferredLightUniformStruct& DeferredLightUniformsValue)
+	{
+		SetUniformBufferParameterImmediate(BatchedParameters, GetUniformBufferParameter<FDeferredLightUniformStruct>(), DeferredLightUniformsValue);
+	}
 };
 
 IMPLEMENT_GLOBAL_SHADER(FDeferredLightPS, "/Engine/Private/DeferredLightPixelShaders.usf", "DeferredLightPixelMain", SF_Pixel);
@@ -2837,8 +2842,8 @@ static void InternalRenderSimpleLightsStandardDeferred(
 
 			// Update the light parameters with a custom uniform buffer
 			FDeferredLightUniformStruct DeferredLightUniformsValue = GetSimpleDeferredLightParameters(View, SimpleLight, SimpleLightPerViewData);
-			SetShaderParameters(RHICmdList, PixelShader, PixelShader.GetPixelShader(), PassParameters->PS);
-			SetUniformBufferParameterImmediate(RHICmdList, RHICmdList.GetBoundPixelShader(), PixelShader->GetUniformBufferParameter<FDeferredLightUniformStruct>(), DeferredLightUniformsValue);
+
+			SetShaderParametersMixed(RHICmdList, PixelShader, PixelShader.GetPixelShader(), PassParameters->PS, DeferredLightUniformsValue);
 
 			// Update vertex shader parameters with custom parameters/uniform buffer
 			FDeferredLightVS::FParameters ParametersVS = FDeferredLightVS::GetParameters(View, LightBounds);
