Create Heterogeneous Volume actor, component and rendering interface

diff --git a/Engine/Source/Runtime/Engine/Classes/Components/HeterogeneousVolumeComponent.h b/Engine/Source/Runtime/Engine/Classes/Components/HeterogeneousVolumeComponent.h
new file mode 100644
--- /dev/null
+++ b/Engine/Source/Runtime/Engine/Classes/Components/HeterogeneousVolumeComponent.h
@@ -0,0 +1,70 @@
+// Copyright Epic Games, Inc. All Rights Reserved.
+
+#pragma once
+
+#include "CoreMinimal.h"
+#include "Components/MeshComponent.h"
+#include "EngineDefines.h"
+#include "GameFramework/Info.h"
+#include "Misc/Guid.h"
+#include "RenderResource.h"
+
+#include "HeterogeneousVolumeComponent.generated.h"
+
+/**
+ * A component that represents a heterogeneous volume.
+ */
+UCLASS(Blueprintable, ClassGroup = (Rendering, Common), hidecategories = (Object, Activation, "Components|Activation"), ShowCategories = (Mobility), editinlinenew, meta = (BlueprintSpawnableComponent))
+class UHeterogeneousVolumeComponent : public UMeshComponent
+{
+	GENERATED_UCLASS_BODY()
+
+	UPROPERTY(EditAnywhere, BlueprintReadOnly, Category = Volume)
+	FIntVector VolumeResolution;
+
+	UPROPERTY(EditAnywhere, BlueprintReadOnly, Category = Volume)
+	float MinimumVoxelSize;
+
+	UPROPERTY(EditAnywhere, BlueprintReadOnly, Category = Volume)
+	uint32 bAnimate : 1;
+
+	UPROPERTY(EditAnywhere, BlueprintReadOnly, Category = Lighting)
+	float LightingDownsampleFactor;
+
+	~UHeterogeneousVolumeComponent() {}
+
+public:
+	//~ Begin USceneComponent Interface.
+	virtual FBoxSphereBounds CalcBounds(const FTransform& LocalToWorld) const override;
+	virtual void TickComponent(float DeltaTime, ELevelTick TickType, FActorComponentTickFunction* ThisTickFunction) override;
+	//~ End USceneComponent Interface.
+
+	//~ Begin UPrimitiveComponent Interface.
+	virtual FPrimitiveSceneProxy* CreateSceneProxy() override;
+	virtual int32 GetNumMaterials() const override { return 1; }
+	virtual bool ShouldRenderSelected() const override { return true; }
+	//~ End UPrimitiveComponent Interface.
+
+private:
+	float Time;
+	float Framerate;
+};
+
+/**
+ * A placeable actor that represents a heterogeneous volume.
+ */
+UCLASS(showcategories = (Movement, Rendering, Transformation, DataLayers, "Input|MouseInput", "Input|TouchInput"), ClassGroup = Fog, hidecategories = (Info, Object, Input), MinimalAPI)
+class AHeterogeneousVolume : public AInfo
+{
+	GENERATED_UCLASS_BODY()
+
+private:
+
+	UPROPERTY(BlueprintReadOnly, VisibleAnywhere, Category = Volume, meta = (AllowPrivateAccess = "true"))
+	TObjectPtr<class UHeterogeneousVolumeComponent> HeterogeneousVolumeComponent;
+
+#if WITH_EDITOR
+	virtual bool ActorTypeSupportsDataLayer() const override { return true; }
+#endif
+
+};
\ No newline at end of file
diff --git a/Engine/Source/Runtime/Engine/Private/Components/HeterogeneousVolumeComponent.cpp b/Engine/Source/Runtime/Engine/Private/Components/HeterogeneousVolumeComponent.cpp
new file mode 100644
--- /dev/null
+++ b/Engine/Source/Runtime/Engine/Private/Components/HeterogeneousVolumeComponent.cpp
@@ -0,0 +1,315 @@
+// Copyright Epic Games, Inc. All Rights Reserved.
+
+#include "Components/HeterogeneousVolumeComponent.h"
+
+#include "Engine/Texture2D.h"
+#include "HeterogeneousVolumeInterface.h"
+#include "MaterialDomain.h"
+#include "Materials/MaterialInstanceDynamic.h"
+#include "Components/BillboardComponent.h"
+#include "Materials/MaterialInterface.h"
+#include "UObject/ConstructorHelpers.h"
+#include "SparseVolumeTexture/SparseVolumeTexture.h"
+
+#include UE_INLINE_GENERATED_CPP_BY_NAME(HeterogeneousVolumeComponent)
+
+#define LOCTEXT_NAMESPACE "HeterogeneousVolumeComponent"
+
+class FHeterogeneousVolumeSceneProxy : public FPrimitiveSceneProxy
+{
+public:
+	FHeterogeneousVolumeSceneProxy(UHeterogeneousVolumeComponent* InComponent);
+	virtual ~FHeterogeneousVolumeSceneProxy();
+
+	//~ Begin FPrimitiveSceneProxy Interface.
+	virtual FPrimitiveViewRelevance GetViewRelevance(const FSceneView* View) const override;
+
+	virtual void GetDynamicMeshElements(
+		const TArray<const FSceneView*>& Views,
+		const FSceneViewFamily& ViewFamily,
+		uint32 VisibilityMap,
+		FMeshElementCollector& Collector) const override;
+
+	virtual SIZE_T GetTypeHash() const override
+	{
+		static size_t UniquePointer;
+		return reinterpret_cast<size_t>(&UniquePointer);
+	}
+
+	virtual uint32 GetMemoryFootprint(void) const override { return sizeof(*this) + GetAllocatedSize(); }
+	uint32 GetAllocatedSize(void) const { return FPrimitiveSceneProxy::GetAllocatedSize(); }
+	//~ End FPrimitiveSceneProxy Interface.
+
+	virtual const IHeterogeneousVolumeInterface* GetHeterogeneousVolumeInterface() const override { return &HeterogeneousVolumeData; }
+
+private:
+	UMaterialInterface* MaterialInterface;
+	FLocalVertexFactory VertexFactory;
+	FStaticMeshVertexBuffers StaticMeshVertexBuffers;
+
+	FHeterogeneousVolumeData HeterogeneousVolumeData;
+};
+
+/*=============================================================================
+	FHeterogeneousVolumeSceneProxy implementation.
+=============================================================================*/
+
+FHeterogeneousVolumeSceneProxy::FHeterogeneousVolumeSceneProxy(UHeterogeneousVolumeComponent* InComponent)
+	: FPrimitiveSceneProxy(InComponent)
+	, MaterialInterface(InComponent->GetMaterial(0))
+	, VertexFactory(GetScene().GetFeatureLevel(), "FHeterogeneousVolumeSceneProxy")
+{
+	bIsHeterogeneousVolume = true;
+
+	HeterogeneousVolumeData.VoxelResolution = InComponent->VolumeResolution;
+	HeterogeneousVolumeData.MinimumVoxelSize = InComponent->MinimumVoxelSize;
+	HeterogeneousVolumeData.LightingDownsampleFactor = InComponent->LightingDownsampleFactor;
+
+	// Update material assignment to include Heterogeneous Volumes
+	if (MaterialInterface)
+	{
+		const UMaterial* Material = MaterialInterface->GetMaterial();
+		if (Material && Material->MaterialDomain == EMaterialDomain::MD_Volume)
+		{
+			Material->GetRenderProxy();
+			MaterialInterface->CheckMaterialUsage(MATUSAGE_HeterogeneousVolumes);
+		}
+	}
+
+	// Initialize vertex buffer data for a quad
+	StaticMeshVertexBuffers.PositionVertexBuffer.Init(4);
+	StaticMeshVertexBuffers.StaticMeshVertexBuffer.Init(4, 1);
+	StaticMeshVertexBuffers.ColorVertexBuffer.Init(4);
+
+	for (uint32 VertexIndex = 0; VertexIndex < 4; ++VertexIndex)
+	{
+		StaticMeshVertexBuffers.ColorVertexBuffer.VertexColor(VertexIndex) = FColor::White;
+	}
+
+	StaticMeshVertexBuffers.PositionVertexBuffer.VertexPosition(0) = FVector3f(-1.0, -1.0, -1.0);
+	StaticMeshVertexBuffers.PositionVertexBuffer.VertexPosition(1) = FVector3f(-1.0, 1.0, -1.0);
+	StaticMeshVertexBuffers.PositionVertexBuffer.VertexPosition(2) = FVector3f(1.0, -1.0, -1.0);
+	StaticMeshVertexBuffers.PositionVertexBuffer.VertexPosition(3) = FVector3f(1.0, 1.0, -1.0);
+
+	StaticMeshVertexBuffers.StaticMeshVertexBuffer.SetVertexUV(0, 0, FVector2f(0, 0));
+	StaticMeshVertexBuffers.StaticMeshVertexBuffer.SetVertexUV(1, 0, FVector2f(0, 1));
+	StaticMeshVertexBuffers.StaticMeshVertexBuffer.SetVertexUV(2, 0, FVector2f(1, 0));
+	StaticMeshVertexBuffers.StaticMeshVertexBuffer.SetVertexUV(3, 0, FVector2f(1, 1));
+
+	FHeterogeneousVolumeSceneProxy* Self = this;
+	ENQUEUE_RENDER_COMMAND(FHeterogeneousVolumeSceneProxyInit)(
+		[Self](FRHICommandListImmediate& RHICmdList)
+		{
+			Self->StaticMeshVertexBuffers.PositionVertexBuffer.InitResource();
+	Self->StaticMeshVertexBuffers.StaticMeshVertexBuffer.InitResource();
+	Self->StaticMeshVertexBuffers.ColorVertexBuffer.InitResource();
+
+	FLocalVertexFactory::FDataType Data;
+	Self->StaticMeshVertexBuffers.PositionVertexBuffer.BindPositionVertexBuffer(&Self->VertexFactory, Data);
+	Self->StaticMeshVertexBuffers.StaticMeshVertexBuffer.BindTangentVertexBuffer(&Self->VertexFactory, Data);
+	Self->StaticMeshVertexBuffers.StaticMeshVertexBuffer.BindPackedTexCoordVertexBuffer(&Self->VertexFactory, Data);
+	Self->StaticMeshVertexBuffers.StaticMeshVertexBuffer.BindLightMapVertexBuffer(&Self->VertexFactory, Data, 0);
+	Self->StaticMeshVertexBuffers.ColorVertexBuffer.BindColorVertexBuffer(&Self->VertexFactory, Data);
+	Self->VertexFactory.SetData(Data);
+
+	Self->VertexFactory.InitResource();
+		});
+}
+
+/** Virtual destructor. */
+FHeterogeneousVolumeSceneProxy::~FHeterogeneousVolumeSceneProxy()
+{
+	VertexFactory.ReleaseResource();
+	StaticMeshVertexBuffers.PositionVertexBuffer.ReleaseResource();
+	StaticMeshVertexBuffers.StaticMeshVertexBuffer.ReleaseResource();
+	StaticMeshVertexBuffers.ColorVertexBuffer.ReleaseResource();
+}
+
+FPrimitiveViewRelevance FHeterogeneousVolumeSceneProxy::GetViewRelevance(const FSceneView* View) const
+{
+	FPrimitiveViewRelevance Result;
+
+	if (MaterialInterface)
+	{
+		FMaterialRelevance MaterialRelevance = MaterialInterface->GetRelevance_Concurrent(View->GetFeatureLevel());
+		MaterialRelevance.SetPrimitiveViewRelevance(Result);
+	}
+
+	Result.bDrawRelevance = IsShown(View);
+	Result.bOpaque = false;
+	Result.bStaticRelevance = false;
+	Result.bDynamicRelevance = true;
+	Result.bShadowRelevance = IsShadowCast(View);
+	Result.bRenderInMainPass = ShouldRenderInMainPass();
+
+	return Result;
+}
+
+void FHeterogeneousVolumeSceneProxy::GetDynamicMeshElements(
+	const TArray<const FSceneView*>& Views,
+	const FSceneViewFamily& ViewFamily,
+	uint32 VisibilityMap,
+	FMeshElementCollector& Collector) const
+{
+	check(IsInRenderingThread());
+
+	// Create a dummy MeshBatch to make the system happy..
+	if (MaterialInterface)
+	{
+		// Set up MeshBatch
+		FMeshBatch& Mesh = Collector.AllocateMesh();
+
+		Mesh.VertexFactory = &VertexFactory;
+		Mesh.MaterialRenderProxy = MaterialInterface->GetRenderProxy();
+		Mesh.LCI = NULL;
+		Mesh.ReverseCulling = IsLocalToWorldDeterminantNegative() ? true : false;
+		Mesh.CastShadow = false;
+		//Mesh.DepthPriorityGroup = (ESceneDepthPriorityGroup)GetDepthPriorityGroup(View);
+		Mesh.Type = PT_TriangleStrip;
+		Mesh.bDisableBackfaceCulling = true;
+
+		// Set up the FMeshBatchElement.
+		FMeshBatchElement& BatchElement = Mesh.Elements[0];
+		BatchElement.IndexBuffer = NULL;
+		BatchElement.FirstIndex = 0;
+		BatchElement.MinVertexIndex = 0;
+		BatchElement.MaxVertexIndex = 3;
+		BatchElement.NumPrimitives = 2;
+		BatchElement.BaseVertexIndex = 0;
+
+		Mesh.bCanApplyViewModeOverrides = true;
+		Mesh.bUseWireframeSelectionColoring = IsSelected();
+
+		Collector.AddMesh(0, Mesh);
+	}
+}
+
+/*=============================================================================
+	HeterogeneousVolumeComponent implementation.
+=============================================================================*/
+
+UHeterogeneousVolumeComponent::UHeterogeneousVolumeComponent(const FObjectInitializer& ObjectInitializer)
+	: Super(ObjectInitializer)
+	, VolumeResolution(128)
+	, MinimumVoxelSize(0.1f)
+	, bAnimate(false)
+	, LightingDownsampleFactor(1.0f)
+	, Time(0.0f)
+	, Framerate(24.0f)
+{
+	PrimaryComponentTick.bCanEverTick = true;
+	// What is this?
+	PrimaryComponentTick.TickGroup = TG_DuringPhysics;
+
+#if WITH_EDITORONLY_DATA
+	bTickInEditor = true;
+#endif // WITH_EDITORONLY_DATA
+}
+
+FPrimitiveSceneProxy* UHeterogeneousVolumeComponent::CreateSceneProxy()
+{
+	return new FHeterogeneousVolumeSceneProxy(this);
+}
+
+FBoxSphereBounds UHeterogeneousVolumeComponent::CalcBounds(const FTransform& LocalToWorld) const
+{
+	FBoxSphereBounds NewBounds;
+	NewBounds.Origin = FVector::ZeroVector;
+	NewBounds.BoxExtent = FVector(50, 50, 50);
+	NewBounds.SphereRadius = NewBounds.BoxExtent.Length();
+
+	return NewBounds.TransformBy(LocalToWorld);
+}
+
+void UHeterogeneousVolumeComponent::TickComponent(float DeltaTime, ELevelTick TickType, FActorComponentTickFunction* ThisTickFunction)
+{
+	Super::TickComponent(DeltaTime, TickType, ThisTickFunction);
+
+	uint32 MaterialIndex = 0;
+	UMaterialInterface* Material = GetMaterial(MaterialIndex);
+	if (Material)
+	{
+		// Get all SVT params in the material
+		TArray<FMaterialParameterInfo> SVTParameterInfo;
+		TArray<FGuid> SVTParameterIds;
+		Material->GetAllSparseVolumeTextureParameterInfo(SVTParameterInfo, SVTParameterIds);
+
+		if (!SVTParameterInfo.IsEmpty())
+		{
+			// Create a MID if this isn't one already
+			UMaterialInstanceDynamic* MaterialInstanceDynamic = nullptr;
+			UMaterialInterface* ParentMaterial = nullptr;
+			if (!Material->IsA<UMaterialInstanceDynamic>())
+			{
+				MaterialInstanceDynamic = CreateAndSetMaterialInstanceDynamicFromMaterial(MaterialIndex, Material);
+				ParentMaterial = Material;
+			}
+			else
+			{
+				MaterialInstanceDynamic = CastChecked<UMaterialInstanceDynamic>(Material);
+				ParentMaterial = MaterialInstanceDynamic->Parent;
+			}
+
+			USparseVolumeTexture* DefaultSVT = nullptr;
+			MaterialInstanceDynamic->GetSparseVolumeTextureParameterDefaultValue(SVTParameterInfo[MaterialIndex], DefaultSVT);
+
+			if (bAnimate)
+			{
+			int32 FrameIndex = int32(Time * (1000.0f / Framerate)) % DefaultSVT->GetNumFrames();
+			int32 MipLevel = 0;
+			USparseVolumeTextureFrame* SVTFrame = USparseVolumeTextureFrame::CreateFrame(DefaultSVT, FrameIndex, MipLevel);
+			MaterialInstanceDynamic->SetSparseVolumeTextureParameterValue(SVTParameterInfo[MaterialIndex].Name, SVTFrame);
+
+				VolumeResolution = SVTFrame->GetVolumeResolution();
+			}
+			else
+			{
+				VolumeResolution = DefaultSVT->GetVolumeResolution();
+			}
+		}
+	}
+
+	Time += DeltaTime;
+}
+
+AHeterogeneousVolume::AHeterogeneousVolume(const FObjectInitializer& ObjectInitializer)
+	: Super(ObjectInitializer)
+{
+	HeterogeneousVolumeComponent = CreateDefaultSubobject<UHeterogeneousVolumeComponent>(TEXT("HeterogeneousVolumeComponent"));
+	RootComponent = HeterogeneousVolumeComponent;
+
+#if WITH_EDITORONLY_DATA
+
+	if (!IsRunningCommandlet())
+	{
+		// Structure to hold one-time initialization
+		struct FConstructorStatics
+		{
+			ConstructorHelpers::FObjectFinderOptional<UTexture2D> HeterogeneousVolumeTextureObject;
+			FName ID_HeterogeneousVolume;
+			FText NAME_HeterogeneousVolume;
+			FConstructorStatics()
+				: HeterogeneousVolumeTextureObject(TEXT("/Engine/EditorResources/S_HeterogeneousVolume"))
+				, ID_HeterogeneousVolume(TEXT("Fog"))
+				, NAME_HeterogeneousVolume(NSLOCTEXT("SpriteCategory", "Fog", "Fog"))
+			{
+			}
+		};
+		static FConstructorStatics ConstructorStatics;
+
+		if (GetSpriteComponent())
+		{
+			GetSpriteComponent()->Sprite = ConstructorStatics.HeterogeneousVolumeTextureObject.Get();
+			GetSpriteComponent()->SetRelativeScale3D(FVector(0.5f, 0.5f, 0.5f));
+			GetSpriteComponent()->SpriteInfo.Category = ConstructorStatics.ID_HeterogeneousVolume;
+			GetSpriteComponent()->SpriteInfo.DisplayName = ConstructorStatics.NAME_HeterogeneousVolume;
+			GetSpriteComponent()->SetupAttachment(HeterogeneousVolumeComponent);
+		}
+	}
+#endif // WITH_EDITORONLY_DATA
+
+	PrimaryActorTick.bCanEverTick = true;
+	SetHidden(false);
+}
+
+#undef LOCTEXT_NAMESPACE
\ No newline at end of file
