World Partition: fix regression when validating changelists for unloaded worlds.

diff --git a/Engine/Source/Runtime/Engine/Private/WorldPartition/WorldPartitionStreamingGeneration.cpp b/Engine/Source/Runtime/Engine/Private/WorldPartition/WorldPartitionStreamingGeneration.cpp
--- a/Engine/Source/Runtime/Engine/Private/WorldPartition/WorldPartitionStreamingGeneration.cpp
+++ b/Engine/Source/Runtime/Engine/Private/WorldPartition/WorldPartitionStreamingGeneration.cpp
@@ -255,7 +255,7 @@ class FWorldPartitionStreamingGenerator
 
 	void ResolveRuntimeDataLayers(FWorldPartitionActorDescView& ActorDescView, const FActorDescViewMap& ActorDescViewMap)
 	{
-		const UDataLayerManager* DataLayerManager = WorldPartitionContext->GetDataLayerManager();
+		const UDataLayerManager* DataLayerManager = WorldPartitionContext ? WorldPartitionContext->GetDataLayerManager() : nullptr;
 
 		// Resolve DataLayerInstanceNames of ActorDescView only when necessary (i.e. when container is a template)
 		if (!ActorDescView.GetActorDesc()->HasResolvedDataLayerInstanceNames())
@@ -547,7 +547,7 @@ class FWorldPartitionStreamingGenerator
 				}
 
 				// Validate data layers
-				if (UDataLayerManager* DataLayerManager = WorldPartitionContext->GetDataLayerManager())
+				if (const UDataLayerManager* DataLayerManager = WorldPartitionContext ? WorldPartitionContext->GetDataLayerManager() : nullptr)
 				{
 					DataLayerManager->ForEachDataLayerInstance([this](const UDataLayerInstance* DataLayerInstance)
 					{
@@ -976,7 +976,7 @@ public:
 	TArray<const UDataLayerInstance*> GetRuntimeDataLayerInstances(const TArray<FName>& RuntimeDataLayers) const
 	{
 		static TArray<const UDataLayerInstance*> EmptyArray;
-		UDataLayerManager* DataLayerManager = WorldPartitionContext->GetDataLayerManager();
+		const UDataLayerManager* DataLayerManager = WorldPartitionContext ? WorldPartitionContext->GetDataLayerManager() : nullptr;
 		return DataLayerManager ? DataLayerManager->GetRuntimeDataLayerInstances(RuntimeDataLayers) : EmptyArray;
 	}
 
