Increase timeout threshold for distributed shader compiler from 5 to 15 min.

#rnx

diff --git a/Engine/Source/Runtime/Engine/Private/ShaderCompiler/ShaderCompilerDistributed.cpp b/Engine/Source/Runtime/Engine/Private/ShaderCompiler/ShaderCompilerDistributed.cpp
--- a/Engine/Source/Runtime/Engine/Private/ShaderCompiler/ShaderCompilerDistributed.cpp
+++ b/Engine/Source/Runtime/Engine/Private/ShaderCompiler/ShaderCompilerDistributed.cpp
@@ -23,7 +23,7 @@ namespace DistributedShaderCompilerVariables
 		TEXT("Smaller number of shaders will compile locally."),
 		ECVF_Default);
 
-	static int32 GDistributedControllerTimeout = 5 * 60;	// 
+	static int32 GDistributedControllerTimeout = 15 * 60;
 	static FAutoConsoleVariableRef CVarDistributedControllerTimeout(
 		TEXT("r.ShaderCompiler.DistributedControllerTimeout"),
 		GDistributedControllerTimeout,
