Reduce missing mips error to warning in the editor as there are some legitimate cases this can happen (not generating thumbnails for example).

diff --git a/Engine/Source/Runtime/Engine/Private/Texture2D.cpp b/Engine/Source/Runtime/Engine/Private/Texture2D.cpp
--- a/Engine/Source/Runtime/Engine/Private/Texture2D.cpp
+++ b/Engine/Source/Runtime/Engine/Private/Texture2D.cpp
@@ -978,7 +978,7 @@ FTextureResource* UTexture2D::CreateResource()
 				}
 				else
 				{
-					UE_LOG(LogTexture, Error, TEXT("%s contains no miplevels! Please delete. (Format: %d)"), *GetFullName(), (int)PixelFormat);
+					UE_LOG(LogTexture, Warning, TEXT("%s contains no miplevels! This could happen if this texture is a thumbnail and hasn't been generated (Format: %d)"), *GetFullName(), (int)PixelFormat);
 				}
 #else
 				UE_LOG(LogTexture, Error, TEXT("%s contains no miplevels! Please delete. (Format: %d)"), *GetFullName(), (int)PixelFormat);
