Optimization for CreateGeometry. Cuts time in half and saves around 5% time taken to load neon city.

- Fixed bug not counting tapered capsules in the simple shapes count.
- Added array reserve to avoid resize allocs.
- Factored duplicated math operations out of convex loop.

diff --git a/Engine/Source/Runtime/Engine/Private/PhysicsEngine/Experimental/ChaosInterfaceUtils.cpp b/Engine/Source/Runtime/Engine/Private/PhysicsEngine/Experimental/ChaosInterfaceUtils.cpp
--- a/Engine/Source/Runtime/Engine/Private/PhysicsEngine/Experimental/ChaosInterfaceUtils.cpp
+++ b/Engine/Source/Runtime/Engine/Private/PhysicsEngine/Experimental/ChaosInterfaceUtils.cpp
@@ -65,8 +65,9 @@ namespace ChaosInterface
 			CollisionTraceType = UPhysicsSettings::Get()->DefaultShapeComplexity;
 		}
 
-		float CollisionMarginFraction = FMath::Max(0.0f, UPhysicsSettingsCore::Get()->SolverOptions.CollisionMarginFraction);
-		float CollisionMarginMax = FMath::Max(0.0f, UPhysicsSettingsCore::Get()->SolverOptions.CollisionMarginMax);
+		const FChaosSolverConfiguration& SolverOptions = UPhysicsSettingsCore::Get()->SolverOptions;
+		float CollisionMarginFraction = FMath::Max(0.0f, SolverOptions.CollisionMarginFraction);
+		float CollisionMarginMax = FMath::Max(0.0f, SolverOptions.CollisionMarginMax);
 
 		// Test margins without changing physics settings
 		if (Chaos_Collision_MarginFraction >= 0.0f)
@@ -79,14 +80,18 @@ namespace ChaosInterface
 		}
 
 		// Complex as simple should not create simple geometry, unless there is no complex geometry.  Otherwise both get queried against.
-		bool bMakeSimpleGeometry = (CollisionTraceType != CTF_UseComplexAsSimple) || (InParams.ChaosTriMeshes.Num() == 0);
+		const bool bMakeSimpleGeometry = (CollisionTraceType != CTF_UseComplexAsSimple) || (InParams.ChaosTriMeshes.Num() == 0);
 
 		// The reverse is true for Simple as Complex.
-		const int32 SimpleShapeCount = InParams.Geometry->SphereElems.Num() + InParams.Geometry->BoxElems.Num() + InParams.Geometry->ConvexElems.Num() + InParams.Geometry->SphylElems.Num();
-		bool bMakeComplexGeometry = (CollisionTraceType != CTF_UseSimpleAsComplex) || (SimpleShapeCount == 0);
+		const int32 SimpleShapeCount = InParams.Geometry->SphereElems.Num() + InParams.Geometry->BoxElems.Num() + InParams.Geometry->SphylElems.Num() + InParams.Geometry->TaperedCapsuleElems.Num() + InParams.Geometry->ConvexElems.Num();
+		const bool bMakeComplexGeometry = (CollisionTraceType != CTF_UseSimpleAsComplex) || (SimpleShapeCount == 0);
 
 		ensure(bMakeComplexGeometry || bMakeSimpleGeometry);
 
+		const int32 NumNewShapes = (bMakeSimpleGeometry ? SimpleShapeCount : 0) + (bMakeComplexGeometry ? InParams.ChaosTriMeshes.Num() : 0);
+		Shapes.Reserve(Shapes.Num() + NumNewShapes);
+		Geoms.Reserve(Geoms.Num() + NumNewShapes);
+
 		auto NewShapeHelper = [&InParams, &CollisionTraceType](Chaos::TSerializablePtr<Chaos::FImplicitObject> InGeom, int32 ShapeIdx, void* UserData, ECollisionEnabled::Type ShapeCollisionEnabled, bool bComplexShape = false)
 		{
 			TUniquePtr<Chaos::FPerShapeData> NewShape = Chaos::FShapeInstanceProxy::Make(ShapeIdx, InGeom);
@@ -110,20 +115,17 @@ namespace ChaosInterface
 
 		if (bMakeSimpleGeometry)
 		{
-			for (uint32 i = 0; i < static_cast<uint32>(InParams.Geometry->SphereElems.Num()); ++i)
+			for (const FKSphereElem& SphereElem : InParams.Geometry->SphereElems)
 			{
-				const FKSphereElem& SphereElem = InParams.Geometry->SphereElems[i];
 				const FKSphereElem ScaledSphereElem = SphereElem.GetFinalScaled(Scale, InParams.LocalTransform);
 				const float UseRadius = FMath::Max(ScaledSphereElem.Radius, UE_KINDA_SMALL_NUMBER);
 				auto ImplicitSphere = MakeUnique<Chaos::TSphere<Chaos::FReal, 3>>(ScaledSphereElem.Center, UseRadius);
 				TUniquePtr<Chaos::FPerShapeData> NewShape = NewShapeHelper(MakeSerializable(ImplicitSphere), Shapes.Num(), (void*)SphereElem.GetUserData(), SphereElem.GetCollisionEnabled());
 				Shapes.Emplace(MoveTemp(NewShape));
-				Geoms.Add(MoveTemp(ImplicitSphere));
+				Geoms.Emplace(MoveTemp(ImplicitSphere));
 			}
-
-			for (uint32 i = 0; i < static_cast<uint32>(InParams.Geometry->BoxElems.Num()); ++i)
+			for (const FKBoxElem& BoxElem : InParams.Geometry->BoxElems)
 			{
-				const FKBoxElem& BoxElem = InParams.Geometry->BoxElems[i];
 				const FKBoxElem ScaledBoxElem = BoxElem.GetFinalScaled(Scale, InParams.LocalTransform);
 				const FTransform& BoxTransform = ScaledBoxElem.GetTransform();
 				Chaos::FVec3 HalfExtents = Chaos::FVec3(ScaledBoxElem.X * 0.5f, ScaledBoxElem.Y * 0.5f, ScaledBoxElem.Z * 0.5f);
@@ -148,68 +150,60 @@ namespace ChaosInterface
 
 				TUniquePtr<Chaos::FPerShapeData> NewShape = NewShapeHelper(MakeSerializable(Implicit),Shapes.Num(), (void*)BoxElem.GetUserData(), BoxElem.GetCollisionEnabled());
 				Shapes.Emplace(MoveTemp(NewShape));
-				Geoms.Add(MoveTemp(Implicit));
+				Geoms.Emplace(MoveTemp(Implicit));
 			}
-			for (uint32 i = 0; i < static_cast<uint32>(InParams.Geometry->SphylElems.Num()); ++i)
+			for (const FKSphylElem& UnscaledSphyl : InParams.Geometry->SphylElems)
 			{
-				const FKSphylElem& UnscaledSphyl = InParams.Geometry->SphylElems[i];
 				const FKSphylElem ScaledSphylElem = UnscaledSphyl.GetFinalScaled(Scale, InParams.LocalTransform);
 				Chaos::FReal HalfHeight = FMath::Max(ScaledSphylElem.Length * 0.5f, UE_KINDA_SMALL_NUMBER);
 				const Chaos::FReal Radius = FMath::Max(ScaledSphylElem.Radius, UE_KINDA_SMALL_NUMBER);
 
+				TUniquePtr<Chaos::FImplicitObject> Object;
 				if (HalfHeight < UE_KINDA_SMALL_NUMBER)
 				{
 					//not a capsule just use a sphere
-					auto ImplicitSphere = MakeUnique<Chaos::TSphere<Chaos::FReal, 3>>(ScaledSphylElem.Center, Radius);
-					TUniquePtr<Chaos::FPerShapeData> NewShape = NewShapeHelper(MakeSerializable(ImplicitSphere),Shapes.Num(), (void*)UnscaledSphyl.GetUserData(), UnscaledSphyl.GetCollisionEnabled());
-					Shapes.Emplace(MoveTemp(NewShape));
-					Geoms.Add(MoveTemp(ImplicitSphere));
+					Object = MakeUnique<Chaos::TSphere<Chaos::FReal, 3>>(ScaledSphylElem.Center, Radius);
 				}
 				else
 				{
 					Chaos::FVec3 HalfExtents = ScaledSphylElem.Rotation.RotateVector(Chaos::FVec3(0, 0, HalfHeight));
-
-					auto ImplicitCapsule = MakeUnique<Chaos::FCapsule>(ScaledSphylElem.Center - HalfExtents, ScaledSphylElem.Center + HalfExtents, Radius);
-					TUniquePtr<Chaos::FPerShapeData> NewShape = NewShapeHelper(MakeSerializable(ImplicitCapsule),Shapes.Num(), (void*)UnscaledSphyl.GetUserData(), UnscaledSphyl.GetCollisionEnabled());
-					Shapes.Emplace(MoveTemp(NewShape));
-					Geoms.Add(MoveTemp(ImplicitCapsule));
+					Object = MakeUnique<Chaos::FCapsule>(ScaledSphylElem.Center - HalfExtents, ScaledSphylElem.Center + HalfExtents, Radius);
 				}
+
+				TUniquePtr<Chaos::FPerShapeData> NewShape = NewShapeHelper(MakeSerializable(Object), Shapes.Num(), (void*)UnscaledSphyl.GetUserData(), UnscaledSphyl.GetCollisionEnabled());
+				Shapes.Emplace(MoveTemp(NewShape));
+				Geoms.Emplace(MoveTemp(Object));
 			}
 #ifdef CREATE_STRAIGHT_CAPSULE_GEOMETRY_FOR_TAPERED_CAPSULES
-			for (uint32 i = 0; i < static_cast<uint32>(InParams.Geometry->TaperedCapsuleElems.Num()); ++i)
+			for (const FKTaperedCapsuleElem& UnscaledTaperedCapsule : InParams.Geometry->TaperedCapsuleElems)
 			{
-				const FKTaperedCapsuleElem& UnscaledTaperedCapsule = InParams.Geometry->TaperedCapsuleElems[i];
-
 				const FKTaperedCapsuleElem ScaledTaperedCapsule = UnscaledTaperedCapsule.GetFinalScaled(Scale, InParams.LocalTransform);
 				Chaos::FReal HalfHeight = FMath::Max(ScaledTaperedCapsule.Length * 0.5f, UE_KINDA_SMALL_NUMBER);
 				const Chaos::FReal Radius0 = FMath::Max(ScaledTaperedCapsule.Radius0, UE_KINDA_SMALL_NUMBER);
 				const Chaos::FReal Radius1 = FMath::Max(ScaledTaperedCapsule.Radius1, UE_KINDA_SMALL_NUMBER);
 
+				TUniquePtr<Chaos::FImplicitObject> Object;
 				if (HalfHeight < UE_KINDA_SMALL_NUMBER)
 				{
 					//not a capsule just use a sphere
 					const Chaos::FReal MaxRadius = FMath::Max(Radius0, Radius1);
-					auto ImplicitSphere = MakeUnique<Chaos::TSphere<Chaos::FReal, 3>>(ScaledTaperedCapsule.Center, MaxRadius);
-					TUniquePtr<Chaos::FPerShapeData> NewShape = NewShapeHelper(MakeSerializable(ImplicitSphere), Shapes.Num(), (void*)UnscaledTaperedCapsule.GetUserData(), UnscaledTaperedCapsule.GetCollisionEnabled());
-					Shapes.Emplace(MoveTemp(NewShape));
-					Geoms.Add(MoveTemp(ImplicitSphere));
+					Object = MakeUnique<Chaos::TSphere<Chaos::FReal, 3>>(ScaledTaperedCapsule.Center, MaxRadius);
 				}
 				else
 				{
 					Chaos::FVec3 HalfExtents = ScaledTaperedCapsule.Rotation.RotateVector(Chaos::FVec3(0, 0, HalfHeight));
 					const Chaos::FReal MeanRadius = 0.5f * (Radius0 + Radius1);
-
-					auto ImplicitCapsule = MakeUnique<Chaos::FCapsule>(ScaledTaperedCapsule.Center - HalfExtents, ScaledTaperedCapsule.Center + HalfExtents, MeanRadius);
-					TUniquePtr<Chaos::FPerShapeData> NewShape = NewShapeHelper(MakeSerializable(ImplicitCapsule), Shapes.Num(), (void*)UnscaledTaperedCapsule.GetUserData(), UnscaledTaperedCapsule.GetCollisionEnabled());
-					Shapes.Emplace(MoveTemp(NewShape));
-					Geoms.Add(MoveTemp(ImplicitCapsule));
+					Object = MakeUnique<Chaos::FCapsule>(ScaledTaperedCapsule.Center - HalfExtents, ScaledTaperedCapsule.Center + HalfExtents, MeanRadius);
 				}
+
+				TUniquePtr<Chaos::FPerShapeData> NewShape = NewShapeHelper(MakeSerializable(Object), Shapes.Num(), (void*)UnscaledTaperedCapsule.GetUserData(), UnscaledTaperedCapsule.GetCollisionEnabled());
+				Shapes.Emplace(MoveTemp(NewShape));
+				Geoms.Emplace(MoveTemp(Object));
 			}
 #elif 0
-			for (uint32 i = 0; i < static_cast<uint32>(InParams.Geometry->TaperedCapsuleElems.Num()); ++i)
+			for (const FKTaperedCapsuleElem& TaperedCapsule : InParams.Geometry->TaperedCapsuleElems)
 			{
 				ensure(FMath::IsNearlyEqual(Scale[0], Scale[1]) && FMath::IsNearlyEqual(Scale[1], Scale[2]));
-				const auto& TaperedCapsule = InParams.Geometry->TaperedCapsuleElems[i];
 				if (TaperedCapsule.Length == 0)
 				{
 					Chaos::TSphere<FReal, 3>* ImplicitSphere = new Chaos::TSphere<FReal, 3>(-half_extents, TaperedCapsule.Radius * Scale[0]);
@@ -233,55 +227,62 @@ namespace ChaosInterface
 				}
 			}
 #endif
-
-			for (uint32 i = 0; i < static_cast<uint32>(InParams.Geometry->ConvexElems.Num()); ++i)
+			if (!InParams.Geometry->ConvexElems.IsEmpty())
 			{
-				const FKConvexElem& CollisionBody = InParams.Geometry->ConvexElems[i];
-				if (const auto& ConvexImplicit = CollisionBody.GetChaosConvexMesh())
-				{
-					// Extract the scale from the transform - we have separate wrapper classes for scale versus translate/rotate 
-					FVector NetScale = Scale * InParams.LocalTransform.GetScale3D();
-
-					// If Scale is zero in any component, set minimum positive instead
-					NetScale.X = FMath::Abs(NetScale.X) < UE_KINDA_SMALL_NUMBER ? UE_KINDA_SMALL_NUMBER : NetScale.X;
-					NetScale.Y = FMath::Abs(NetScale.Y) < UE_KINDA_SMALL_NUMBER ? UE_KINDA_SMALL_NUMBER : NetScale.Y;
-					NetScale.Z = FMath::Abs(NetScale.Z) < UE_KINDA_SMALL_NUMBER ? UE_KINDA_SMALL_NUMBER : NetScale.Z;
-
-					FTransform ConvexTransform = FTransform(InParams.LocalTransform.GetRotation(), Scale * InParams.LocalTransform.GetLocation(), FVector(1, 1, 1));
-					const FVector ScaledSize = (NetScale.GetAbs() * CollisionBody.ElemBox.GetSize());	// Note: Scale can be negative
-					const Chaos::FReal CollisionMargin = FMath::Min<Chaos::FReal>(ScaledSize.GetMin() * CollisionMarginFraction, CollisionMarginMax);
-
-					// Wrap the convex in a scaled or instanced wrapper depending on scale value, and add a margin
-					// NOTE: CollisionMargin is on the Instance/Scaled wrapper, not the inner convex (which is shared and should not have a margin).
-					TUniquePtr<Chaos::FImplicitObject> Implicit;
-					if (NetScale == FVector(1))
-					{
-						Implicit = TUniquePtr<Chaos::FImplicitObject>(new Chaos::TImplicitObjectInstanced<Chaos::FConvex>(ConvexImplicit, CollisionMargin));
-					}
-					else
-					{
-						Implicit = TUniquePtr<Chaos::FImplicitObject>(new Chaos::TImplicitObjectScaled<Chaos::FConvex>(ConvexImplicit, NetScale, CollisionMargin));
-					}
+				// Extract the scale from the transform - we have separate wrapper classes for scale versus translate/rotate 
+				FVector NetScale = Scale * InParams.LocalTransform.GetScale3D();
 
-					// Wrap the convex in a non-scaled transform if necessary (the scale is pulled out above)
-					if (!ConvexTransform.GetTranslation().IsNearlyZero() || !ConvexTransform.GetRotation().IsIdentity())
+				// If Scale is zero in any component, set minimum positive instead
+				NetScale.X = FMath::Abs(NetScale.X) < UE_KINDA_SMALL_NUMBER ? UE_KINDA_SMALL_NUMBER : NetScale.X;
+				NetScale.Y = FMath::Abs(NetScale.Y) < UE_KINDA_SMALL_NUMBER ? UE_KINDA_SMALL_NUMBER : NetScale.Y;
+				NetScale.Z = FMath::Abs(NetScale.Z) < UE_KINDA_SMALL_NUMBER ? UE_KINDA_SMALL_NUMBER : NetScale.Z;
+
+				const FVector NetScaleAbs = NetScale.GetAbs();
+
+				const FTransform ConvexTransform = FTransform(InParams.LocalTransform.GetRotation(), Scale * InParams.LocalTransform.GetTranslation(), FVector(1, 1, 1));
+				const bool bHasTranslationOrRotation = !ConvexTransform.GetTranslation().IsNearlyZero() || !ConvexTransform.GetRotation().IsIdentity();
+				const bool bNoScale = NetScale == FVector(1);
+
+				for (const FKConvexElem& CollisionBody : InParams.Geometry->ConvexElems)
+				{
+					if (const auto& ConvexImplicit = CollisionBody.GetChaosConvexMesh())
 					{
-						Implicit = TUniquePtr<Chaos::FImplicitObject>(new Chaos::TImplicitObjectTransformed<Chaos::FReal, 3>(MoveTemp(Implicit), ConvexTransform));
+						const FVector ScaledSize = (NetScaleAbs * CollisionBody.ElemBox.GetSize());	// Note: Scale can be negative
+						const Chaos::FReal CollisionMargin = FMath::Min<Chaos::FReal>(ScaledSize.GetMin() * CollisionMarginFraction, CollisionMarginMax);
+
+						// Wrap the convex in a scaled or instanced wrapper depending on scale value, and add a margin
+						// NOTE: CollisionMargin is on the Instance/Scaled wrapper, not the inner convex (which is shared and should not have a margin).
+						TUniquePtr<Chaos::FImplicitObject> Implicit;
+						if (bNoScale)
+						{
+							Implicit = TUniquePtr<Chaos::FImplicitObject>(new Chaos::TImplicitObjectInstanced<Chaos::FConvex>(ConvexImplicit, CollisionMargin));
+						}
+						else
+						{
+							Implicit = TUniquePtr<Chaos::FImplicitObject>(new Chaos::TImplicitObjectScaled<Chaos::FConvex>(ConvexImplicit, NetScale, CollisionMargin));
+						}
+
+						// Wrap the convex in a non-scaled transform if necessary (the scale is pulled out above)
+						if (bHasTranslationOrRotation)
+						{
+							Implicit = TUniquePtr<Chaos::FImplicitObject>(new Chaos::TImplicitObjectTransformed<Chaos::FReal, 3>(MoveTemp(Implicit), ConvexTransform));
+						}
+
+						TUniquePtr<Chaos::FPerShapeData> NewShape = NewShapeHelper(MakeSerializable(Implicit), Shapes.Num(), (void*)CollisionBody.GetUserData(), CollisionBody.GetCollisionEnabled());
+						Shapes.Emplace(MoveTemp(NewShape));
+						Geoms.Emplace(MoveTemp(Implicit));
 					}
-
-					TUniquePtr<Chaos::FPerShapeData> NewShape = NewShapeHelper(MakeSerializable(Implicit), Shapes.Num(), (void*)CollisionBody.GetUserData(), CollisionBody.GetCollisionEnabled());
-					Shapes.Emplace(MoveTemp(NewShape));
-					Geoms.Add(MoveTemp(Implicit));
 				}
 			}
 		}
 
 		if (bMakeComplexGeometry)
 		{
+			const bool bNoScale = Scale == FVector(1);
 			for (auto& ChaosTriMesh : InParams.ChaosTriMeshes)
 			{
 				TUniquePtr<Chaos::FImplicitObject> Implicit;
-				if (Scale == FVector(1))
+				if (bNoScale)
 				{
 					Implicit = TUniquePtr<Chaos::FImplicitObject>(new Chaos::TImplicitObjectInstanced<Chaos::FTriangleMeshImplicitObject>(ChaosTriMesh));
 				}
@@ -294,7 +295,7 @@ namespace ChaosInterface
 
 				TUniquePtr<Chaos::FPerShapeData> NewShape = NewShapeHelper(MakeSerializable(Implicit),Shapes.Num(), nullptr, ECollisionEnabled::QueryAndPhysics, true);
 				Shapes.Emplace(MoveTemp(NewShape));
-				Geoms.Add(MoveTemp(Implicit));
+				Geoms.Emplace(MoveTemp(Implicit));
 			}
 		}
 	}
