Added a missing ENGINE_API tag to SkeletetonRemappingRegistry, as it is required to use it from outside of the engine

diff --git a/Engine/Source/Runtime/Engine/Public/Animation/SkeletonRemappingRegistry.h b/Engine/Source/Runtime/Engine/Public/Animation/SkeletonRemappingRegistry.h
--- a/Engine/Source/Runtime/Engine/Public/Animation/SkeletonRemappingRegistry.h
+++ b/Engine/Source/Runtime/Engine/Public/Animation/SkeletonRemappingRegistry.h
@@ -17,7 +17,7 @@ struct FSkeletonRemappingRegistryPrivate;
 // Global registry of skeleton remappings
 // Remappings are created on-demand when calling GetRemapping.
 // Calling public functions from multiple threads is expected. Data races are guarded by a FRWLock.
-class FSkeletonRemappingRegistry
+class ENGINE_API FSkeletonRemappingRegistry
 {
 public:
 	// Access the global registry
