Fix for missing includes causing a compile error in HeterogeneousVolumeComponent.cpp

#rnx

diff --git a/Engine/Source/Runtime/Engine/Private/Components/HeterogeneousVolumeComponent.cpp b/Engine/Source/Runtime/Engine/Private/Components/HeterogeneousVolumeComponent.cpp
--- a/Engine/Source/Runtime/Engine/Private/Components/HeterogeneousVolumeComponent.cpp
+++ b/Engine/Source/Runtime/Engine/Private/Components/HeterogeneousVolumeComponent.cpp
@@ -10,6 +10,9 @@
 #include "Materials/MaterialInterface.h"
 #include "UObject/ConstructorHelpers.h"
 #include "SparseVolumeTexture/SparseVolumeTexture.h"
+#include "PrimitiveViewRelevance.h"
+#include "PrimitiveSceneProxy.h"
+#include "SceneManagement.h"
 
 #include UE_INLINE_GENERATED_CPP_BY_NAME(HeterogeneousVolumeComponent)
 
