Added a new DecompressPose function at AnimBoneCompressionCodec, taking a FAnimBoneDecompressionData parameter with the decompression data in separate arrays format (translations, rotations, scales3D)
Implemented a default implementation that codecs have to implement in order to impreove the performance of the functiuon (right now it is just making an in place conversion between formats).

diff --git a/Engine/Source/Runtime/Engine/Classes/Animation/AnimBoneCompressionCodec.h b/Engine/Source/Runtime/Engine/Classes/Animation/AnimBoneCompressionCodec.h
--- a/Engine/Source/Runtime/Engine/Classes/Animation/AnimBoneCompressionCodec.h
+++ b/Engine/Source/Runtime/Engine/Classes/Animation/AnimBoneCompressionCodec.h
@@ -12,6 +12,12 @@ class UAnimBoneCompressionCodec;
 struct FBlendedCurve;
 class FBoneData;
 
+namespace UE::Anim
+{
+struct FAnimPoseDecompressionData;
+}
+
+
 /*
  * Base class for all bone compression codecs.
  */
@@ -90,4 +96,11 @@ class ENGINE_API UAnimBoneCompressionCodec : public UObject
 
 	/** Decompress a single bone. */
 	virtual void DecompressBone(FAnimSequenceDecompressionContext& DecompContext, int32 TrackIndex, FTransform& OutAtom) const PURE_VIRTUAL(UAnimCurveCompressionCodec::DecompressBone, );
+
+	/**
+	 * Decompresses all the specified bone tracks.
+	 * The caller is responsible for pre-filling the output pose with sensible values (e.g. reference/bind/additive identity pose) as
+	 * the codec will only decompress and write out tracks that are contained in the compressed data.
+	 */
+	virtual void DecompressPose(FAnimSequenceDecompressionContext& DecompContext, const UE::Anim::FAnimPoseDecompressionData& DecompressionData) const;
 };
diff --git a/Engine/Source/Runtime/Engine/Classes/Animation/AnimBoneDecompressionData.h b/Engine/Source/Runtime/Engine/Classes/Animation/AnimBoneDecompressionData.h
new file mode 100644
--- /dev/null
+++ b/Engine/Source/Runtime/Engine/Classes/Animation/AnimBoneDecompressionData.h
@@ -0,0 +1,91 @@
+// Copyright Epic Games, Inc. All Rights Reserved.
+
+#pragma once
+
+#include "CoreMinimal.h"
+
+namespace UE::Anim
+{
+
+/*
+ * Struct containing the generated pairs to decompress a Pose, with the output atoms arranged in separate arrays
+ */
+struct FAnimPoseDecompressionData
+{
+	FAnimPoseDecompressionData(const BoneTrackArray& InRotationPairs
+			, const BoneTrackArray& InTranslationPairs
+			, const BoneTrackArray& InScalePairs
+			, TArrayView<FQuat>& InOutAtomRotations
+			, TArrayView<FVector>& InOutAtomTranslations
+			, TArrayView<FVector>& InOutAtomScales3D)
+		: RotationPairs(InRotationPairs)
+		, TranslationPairs(InTranslationPairs)
+		, ScalePairs(InScalePairs)
+		, OutAtomRotations(InOutAtomRotations)
+		, OutAtomTranslations(InOutAtomTranslations)
+		, OutAtomScales3D(InOutAtomScales3D)
+	{}
+
+	const BoneTrackArray& GetRotationPairs() const
+	{
+		return RotationPairs;
+	}
+	const BoneTrackArray& GetTranslationPairs() const
+	{
+		return TranslationPairs;
+	}
+	const BoneTrackArray& GetScalePairs() const
+	{
+		return ScalePairs;
+	}
+
+	TArrayView<FQuat>& GetOutAtomRotations() const
+	{
+		return OutAtomRotations;
+	}
+	TArrayView<FVector>& GetOutAtomTranslations() const
+	{
+		return OutAtomTranslations;
+	}
+	TArrayView<FVector>& GetOutAtomScales3D() const
+	{
+		return OutAtomScales3D;
+	}
+
+
+protected:
+	const BoneTrackArray& RotationPairs;
+	const BoneTrackArray& TranslationPairs;
+	const BoneTrackArray& ScalePairs;
+	TArrayView<FQuat>& OutAtomRotations;
+	TArrayView<FVector>& OutAtomTranslations;
+	TArrayView<FVector>& OutAtomScales3D;
+};
+
+} // end namespace UE::Anim
+
+/*
+* Helper macro to iterate two NON OVERLAPPING arrays/array views using restricted raw pointers
+* - ItFirstType and ItFirstMember refer to the first array type and variable. It will be accessed through It pointer (all pointers are incremented automatically)
+* - ItSecondType and ItSecondMember refer to the second array type variable. It can be accessed through ItRet pointer (all pointers are incremented automatically)
+* - ItNumElements is the number of elements that we want to process and that have to exist on BOTH arrays
+* Inside the Loop, we can use two pointers :
+* - ItFirst : Pointer to the element currently iterated on the first array
+* - ItSecond : Pointer to the element currently iterated on the second array 
+* 
+* Example of iteration over a FQuat array and writing values in a FTransform array
+*         ITERATE_ATOMS_START(FQuat, DecompressionData.OutAtomRotations, FTransform, OutAtoms, NumTransforms)
+*             ItSecond->SetRotation(*ItFirst);
+*         ITERATE_ATOMS_END()
+*/
+#define ITERATE_NON_OVERLAPPING_ARRAYS_START(ItFirstType, ItFirstMember, ItSecondType, ItSecondMember, ItNumElements) \
+if (ItNumElements > 0) \
+{ \
+	check(ItFirstMember.Num() >= ItNumElements && ItSecondMember.Num() >= ItNumElements); \
+	const ItFirstType* RESTRICT ItFirst = ItFirstMember.GetData(); \
+	const ItFirstType* RESTRICT ItFirstEnd = ItFirst + ItNumElements; \
+	ItSecondType* RESTRICT ItSecond = ItSecondMember.GetData(); \
+	for (; ItFirst != ItFirstEnd; ++ItFirst, ++ItSecond) \
+	{
+
+#define ITERATE_NON_OVERLAPPING_ARRAYS_END() }}
diff --git a/Engine/Source/Runtime/Engine/Private/Animation/AnimBoneCompressionCodec.cpp b/Engine/Source/Runtime/Engine/Private/Animation/AnimBoneCompressionCodec.cpp
--- a/Engine/Source/Runtime/Engine/Private/Animation/AnimBoneCompressionCodec.cpp
+++ b/Engine/Source/Runtime/Engine/Private/Animation/AnimBoneCompressionCodec.cpp
@@ -2,6 +2,8 @@
 
 #include "Animation/AnimBoneCompressionCodec.h"
 #include "Animation/AnimBoneCompressionSettings.h"
+#include "Animation/AnimBoneDecompressionData.h"
+#include "Misc/MemStack.h"
 
 #include UE_INLINE_GENERATED_CPP_BY_NAME(AnimBoneCompressionCodec)
 
@@ -58,3 +60,48 @@ void UAnimBoneCompressionCodec::PopulateDDCKey(FArchive& Ar)
 }
 #endif
 
+
+// Default implementation that codecs should override and implement in a more performant way
+// It performs a conversion between SoA and AoS, making copies of the data between the formats
+void UAnimBoneCompressionCodec::DecompressPose(FAnimSequenceDecompressionContext& DecompContext, const UE::Anim::FAnimPoseDecompressionData& DecompressionData) const
+{
+	const int32 NumTransforms = DecompressionData.GetOutAtomRotations().Num();
+
+	FMemMark Mark(FMemStack::Get());
+
+	TArray<FTransform, TMemStackAllocator<>> OutAtoms;
+	OutAtoms.SetNum(NumTransforms);
+
+	// Right now we only support same size arrays
+	check(DecompressionData.GetOutAtomRotations().Num() == DecompressionData.GetOutAtomTranslations().Num() && DecompressionData.GetOutAtomRotations().Num() == DecompressionData.GetOutAtomScales3D().Num());
+
+	// Copy atoms data into the FTransform array
+	ITERATE_NON_OVERLAPPING_ARRAYS_START(FQuat, DecompressionData.GetOutAtomRotations(), FTransform, OutAtoms, NumTransforms)
+		ItSecond->SetRotation(*ItFirst);
+	ITERATE_NON_OVERLAPPING_ARRAYS_END()
+
+	ITERATE_NON_OVERLAPPING_ARRAYS_START(FVector, DecompressionData.GetOutAtomTranslations(), FTransform, OutAtoms, NumTransforms)
+		ItSecond->SetTranslation(*ItFirst);
+	ITERATE_NON_OVERLAPPING_ARRAYS_END()
+
+	ITERATE_NON_OVERLAPPING_ARRAYS_START(FVector, DecompressionData.GetOutAtomScales3D(), FTransform, OutAtoms, NumTransforms)
+		ItSecond->SetScale3D(*ItFirst);
+	ITERATE_NON_OVERLAPPING_ARRAYS_END()
+
+	// Decompress using default FTransform function
+	TArrayView<FTransform> OutAtomsView = OutAtoms;
+	DecompressPose(DecompContext, DecompressionData.GetRotationPairs(), DecompressionData.GetTranslationPairs(), DecompressionData.GetScalePairs(), OutAtomsView);
+
+	// Copy back the result
+	ITERATE_NON_OVERLAPPING_ARRAYS_START(FTransform, OutAtoms, FQuat, DecompressionData.GetOutAtomRotations(), NumTransforms)
+		*ItSecond = ItFirst->GetRotation();
+	ITERATE_NON_OVERLAPPING_ARRAYS_END()
+
+	ITERATE_NON_OVERLAPPING_ARRAYS_START(FTransform, OutAtoms, FVector, DecompressionData.GetOutAtomTranslations(), NumTransforms)
+		*ItSecond = ItFirst->GetTranslation();
+	ITERATE_NON_OVERLAPPING_ARRAYS_END()
+
+	ITERATE_NON_OVERLAPPING_ARRAYS_START(FTransform, OutAtoms, FVector, DecompressionData.GetOutAtomScales3D(), NumTransforms)
+		*ItSecond = ItFirst->GetScale3D();
+	ITERATE_NON_OVERLAPPING_ARRAYS_END()
+}
