Deprecate replication pausing

diff --git a/Engine/Plugins/NetcodeUnitTest/NetcodeUnitTest/Source/NetcodeUnitTest/Private/MinimalClient.cpp b/Engine/Plugins/NetcodeUnitTest/NetcodeUnitTest/Source/NetcodeUnitTest/Private/MinimalClient.cpp
--- a/Engine/Plugins/NetcodeUnitTest/NetcodeUnitTest/Source/NetcodeUnitTest/Private/MinimalClient.cpp
+++ b/Engine/Plugins/NetcodeUnitTest/NetcodeUnitTest/Source/NetcodeUnitTest/Private/MinimalClient.cpp
@@ -282,7 +282,9 @@ bool UMinimalClient::SendRawBunch(FOutBunch& Bunch, bool bAllowPartial/*=false*/
 				NewBunch->bOpen = Bunch.bOpen;
 				NewBunch->bClose = Bunch.bClose;
 				NewBunch->CloseReason = Bunch.CloseReason;
+				PRAGMA_DISABLE_DEPRECATION_WARNINGS
 				NewBunch->bIsReplicationPaused = Bunch.bIsReplicationPaused;
+				PRAGMA_ENABLE_DEPRECATION_WARNINGS
 				NewBunch->ChIndex = Bunch.ChIndex;
 				NewBunch->ChName = Bunch.ChName;
 
diff --git a/Engine/Source/Developer/TraceInsights/Private/Insights/NetworkingProfiler/ViewModels/PacketContentViewDrawHelper.cpp b/Engine/Source/Developer/TraceInsights/Private/Insights/NetworkingProfiler/ViewModels/PacketContentViewDrawHelper.cpp
--- a/Engine/Source/Developer/TraceInsights/Private/Insights/NetworkingProfiler/ViewModels/PacketContentViewDrawHelper.cpp
+++ b/Engine/Source/Developer/TraceInsights/Private/Insights/NetworkingProfiler/ViewModels/PacketContentViewDrawHelper.cpp
@@ -213,10 +213,12 @@ void FPacketContentViewDrawStateBuilder::AddEvent(const TraceServices::FNetProfi
 				{
 					Builder.Appendf(TEXT(" | Partial%s%s"), Info.bPartialInitial ? TEXT("Initial") : TEXT(""), Info.bPartialFinal ? TEXT("Final") : TEXT(""));
 				}
+				PRAGMA_DISABLE_DEPRECATION_WARNINGS
 				if (Event.BunchInfo.bIsReplicationPaused)
 				{
 					Builder.Append(TEXT(" | ReplicationPaused"));
 				}
+				PRAGMA_ENABLE_DEPRECATION_WARNINGS
 				if (Event.BunchInfo.bHasMustBeMappedGUIDs)
 				{
 					Builder.Append(TEXT(" | HasMustBeMappedGUIDs"));
diff --git a/Engine/Source/Developer/TraceServices/Public/TraceServices/Model/NetProfiler.h b/Engine/Source/Developer/TraceServices/Public/TraceServices/Model/NetProfiler.h
--- a/Engine/Source/Developer/TraceServices/Public/TraceServices/Model/NetProfiler.h
+++ b/Engine/Source/Developer/TraceServices/Public/TraceServices/Model/NetProfiler.h
@@ -104,6 +104,7 @@ union FNetProfilerBunchInfo
 		uint64 bPartial : 1;
 		uint64 bPartialInitial : 1;
 		uint64 bPartialFinal : 1;
+		UE_DEPRECATED(5.3, "Replication pausing is now deprecated.")
 		uint64 bIsReplicationPaused : 1;
 		uint64 bOpen : 1;
 		uint64 bClose : 1;
diff --git a/Engine/Source/Runtime/Engine/Classes/Engine/Channel.h b/Engine/Source/Runtime/Engine/Classes/Engine/Channel.h
--- a/Engine/Source/Runtime/Engine/Classes/Engine/Channel.h
+++ b/Engine/Source/Runtime/Engine/Classes/Engine/Channel.h
@@ -53,11 +53,11 @@ ENUM_CLASS_FLAGS(EChannelCreateFlags);
  * Base class of communication channels.
  */
 UCLASS(abstract, transient)
-class ENGINE_API UChannel
-	: public UObject
+class ENGINE_API UChannel : public UObject
 {
-	GENERATED_UCLASS_BODY()
+	GENERATED_BODY()
 	
+public:
 	UPROPERTY()
 	TObjectPtr<class UNetConnection>	Connection;		// Owner connection.
 
@@ -65,6 +65,7 @@ class ENGINE_API UChannel
 	uint32				OpenAcked:1;		// If OpenedLocally is true, this means we have acknowledged the packet we sent the bOpen bunch on. Otherwise, it means we have received the bOpen bunch from the server.
 	uint32				Closing:1;			// State of the channel.
 	uint32				Dormant:1;			// Channel is going dormant (it will close but the client will not destroy
+	UE_DEPRECATED(5.3, "Replication pausing is deprecated and will be removed")
 	uint32				bIsReplicationPaused:1;	// Replication is being paused, but channel will not be closed
 	uint32				OpenTemporary:1;	// Opened temporarily.
 	uint32				Broken:1;			// Has encountered errors and is ignoring subsequent packets.
@@ -85,15 +86,11 @@ class ENGINE_API UChannel
 	class FOutBunch*	OutRec;				// Outgoing reliable unacked data.
 	class FInBunch*		InPartialBunch;		// Partial bunch we are receiving (incoming partial bunches are appended to this)
 
-public:
-
 	// UObject overrides
 
 	virtual void BeginDestroy() override;
 	virtual void Serialize(FArchive& Ar) override;
 
-public:	
-
 	/** UChannel interface. */
 	/** Initialize this channel for the given connection and index. */
 	virtual void Init(UNetConnection* InConnection, int32 InChIndex, EChannelCreateFlags CreateFlags);
@@ -162,6 +159,7 @@ public:
 	/** Puts the channel in a state to start becoming dormant. It will not become dormant until ReadyForDormancy returns true in Tick */
 	virtual void StartBecomingDormant() { }
 
+	UE_DEPRECATED(5.3, "Will be removed from test/shipping in the future.")
 	void PrintReliableBunchBuffer();
 
 	/* Notification that this channel has been placed in a channel pool and needs to reset to its original state so it can be used again like a new channel */
@@ -176,10 +174,22 @@ protected:
 	virtual bool CleanUp( const bool bForDestroy, EChannelCloseReason CloseReason );
 
 	/** Sets whether replication is currently paused on this channel or not */
-	virtual void SetReplicationPaused(bool InbIsReplicationPaused) { bIsReplicationPaused = InbIsReplicationPaused; }
+	UE_DEPRECATED(5.3, "Replication pausing is deprecated and will be removed")
+	virtual void SetReplicationPaused(bool InbIsReplicationPaused) 
+	{ 
+		PRAGMA_DISABLE_DEPRECATION_WARNINGS
+		bIsReplicationPaused = InbIsReplicationPaused;
+		PRAGMA_ENABLE_DEPRECATION_WARNINGS
+	}
 
 	/** Returns whether replication is currently paused on this channel */
-	virtual bool IsReplicationPaused() { return bIsReplicationPaused; }
+	UE_DEPRECATED(5.3, "Replication pausing is deprecated and will be removed")
+	virtual bool IsReplicationPaused() 
+	{
+		PRAGMA_DISABLE_DEPRECATION_WARNINGS
+		return bIsReplicationPaused; 
+		PRAGMA_ENABLE_DEPRECATION_WARNINGS
+	}
 
 private:
 
diff --git a/Engine/Source/Runtime/Engine/Classes/GameFramework/Actor.h b/Engine/Source/Runtime/Engine/Classes/GameFramework/Actor.h
--- a/Engine/Source/Runtime/Engine/Classes/GameFramework/Actor.h
+++ b/Engine/Source/Runtime/Engine/Classes/GameFramework/Actor.h
@@ -692,9 +692,11 @@ public:
 	TEnumAsByte<enum ENetDormancy> NetDormancy;
 
 	/** Gives the actor a chance to pause replication to a player represented by the passed in actor - only called on server */
+	UE_DEPRECATED(5.3, "Replication pausing is deprecated.")
 	virtual bool IsReplicationPausedForConnection(const FNetViewer& ConnectionOwnerNetViewer);
 
 	/** Called on the client when the replication paused value is changed */
+	UE_DEPRECATED(5.3, "Replication pausing is deprecated.")
 	virtual void OnReplicationPausedChanged(bool bIsReplicationPaused);
 
 	/** Controls how to handle spawning this actor in a situation where it's colliding with something else. "Default" means AlwaysSpawn here. */
diff --git a/Engine/Source/Runtime/Engine/Private/DataBunch.cpp b/Engine/Source/Runtime/Engine/Private/DataBunch.cpp
--- a/Engine/Source/Runtime/Engine/Private/DataBunch.cpp
+++ b/Engine/Source/Runtime/Engine/Private/DataBunch.cpp
@@ -17,6 +17,7 @@ const int32 MAX_BUNCH_SIZE = 1024 * 1024;
 	FInBunch implementation.
 -----------------------------------------------------------------------------*/
 
+PRAGMA_DISABLE_DEPRECATION_WARNINGS
 FInBunch::FInBunch( UNetConnection* InConnection, uint8* Src, int64 CountBits )
 :	FNetBitReader	(InConnection->PackageMap, Src, CountBits)
 ,	PacketId	( 0 )
@@ -27,6 +28,7 @@ FInBunch::FInBunch( UNetConnection* InConnection, uint8* Src, int64 CountBits )
 ,	ChSequence ( 0 )
 ,	bOpen ( 0 )
 ,	bClose ( 0 )
+,	bIsReplicationPaused( 0 )
 ,	bReliable ( 0 )
 ,	bPartial ( 0 )
 ,	bPartialInitial ( 0 )
@@ -77,6 +79,7 @@ FInBunch::FInBunch( FInBunch &InBunch, bool CopyBuffer )
 
 	Pos = 0;
 }
+PRAGMA_ENABLE_DEPRECATION_WARNINGS
 
 void FInBunch::CountMemory(FArchive& Ar) const
 {
@@ -99,6 +102,8 @@ void FInBunch::CountMemory(FArchive& Ar) const
 FOutBunch::FOutBunch()
 : FNetBitWriter( 0 )
 {}
+
+PRAGMA_DISABLE_DEPRECATION_WARNINGS
 FOutBunch::FOutBunch( UChannel* InChannel, bool bInClose )
 :	FNetBitWriter	( InChannel->Connection->PackageMap, InChannel->Connection->GetMaxSingleBunchSizeBits())
 ,	Next		( nullptr )
@@ -155,6 +160,7 @@ FOutBunch::FOutBunch( UPackageMap *InPackageMap, int64 MaxBits )
 ,	CloseReason( EChannelCloseReason::Destroyed )
 {
 }
+PRAGMA_ENABLE_DEPRECATION_WARNINGS
 
 FOutBunch::FOutBunch(int64 InMaxBits)
 : FOutBunch(static_cast<UPackageMap*>(nullptr), InMaxBits)
diff --git a/Engine/Source/Runtime/Engine/Private/DataChannel.cpp b/Engine/Source/Runtime/Engine/Private/DataChannel.cpp
--- a/Engine/Source/Runtime/Engine/Private/DataChannel.cpp
+++ b/Engine/Source/Runtime/Engine/Private/DataChannel.cpp
@@ -252,11 +252,6 @@ namespace DataChannelInternal
 	UChannel implementation.
 -----------------------------------------------------------------------------*/
 
-UChannel::UChannel(const FObjectInitializer& ObjectInitializer)
-	: Super(ObjectInitializer)
-{
-}
-
 void UChannel::Init(UNetConnection* InConnection, int32 InChIndex, EChannelCreateFlags CreateFlags)
 {
 	// if child connection then use its parent
@@ -884,7 +879,9 @@ bool UChannel::ReceivedNextBunch( FInBunch & Bunch, bool & bOutSkipAck )
 					InPartialBunch->bPartialFinal			= true;
 					InPartialBunch->bClose					= Bunch.bClose;
 					InPartialBunch->CloseReason				= Bunch.CloseReason;
+					PRAGMA_DISABLE_DEPRECATION_WARNINGS
 					InPartialBunch->bIsReplicationPaused	= Bunch.bIsReplicationPaused;
+					PRAGMA_ENABLE_DEPRECATION_WARNINGS
 					InPartialBunch->bHasMustBeMappedGUIDs	= Bunch.bHasMustBeMappedGUIDs;
 				}
 				else
@@ -1332,7 +1329,9 @@ FPacketIdRange UChannel::SendBunch( FOutBunch* Bunch, bool Merge )
 	{
 		UE_LOG(LogNetPartialBunch, Warning, TEXT("SendBunch: Reliable partial bunch overflows reliable buffer! %s"), *Describe() );
 		UE_LOG(LogNetPartialBunch, Warning, TEXT("   Num OutgoingBunches: %d. NumOutRec: %d"), OutgoingBunches.Num(), NumOutRec );
+		PRAGMA_DISABLE_DEPRECATION_WARNINGS
 		PrintReliableBunchBuffer();
+		PRAGMA_ENABLE_DEPRECATION_WARNINGS
 
 		// Bail out, we can't recover from this (without increasing RELIABLE_BUFFER)
 		FString ErrorMsg = NSLOCTEXT("NetworkErrors", "ClientReliableBufferOverflow", "Outgoing reliable buffer overflow").ToString();
@@ -1354,7 +1353,9 @@ FPacketIdRange UChannel::SendBunch( FOutBunch* Bunch, bool Merge )
 		NextBunch->bOpen = Bunch->bOpen;
 		NextBunch->bClose = Bunch->bClose;
 		NextBunch->CloseReason = Bunch->CloseReason;
+		PRAGMA_DISABLE_DEPRECATION_WARNINGS
 		NextBunch->bIsReplicationPaused = Bunch->bIsReplicationPaused;
+		PRAGMA_ENABLE_DEPRECATION_WARNINGS
 		NextBunch->ChIndex = Bunch->ChIndex;
 		NextBunch->ChName = Bunch->ChName;
 
@@ -1434,7 +1435,9 @@ FOutBunch* UChannel::PrepBunch(FOutBunch* Bunch, FOutBunch* OutBunch, bool Merge
 			if (!(NumOutRec<RELIABLE_BUFFER-1+Bunch->bClose))
 			{
 				UE_LOG(LogNetTraffic, Warning, TEXT("PrepBunch: Reliable buffer overflow! %s"), *Describe());
+				PRAGMA_DISABLE_DEPRECATION_WARNINGS
 				PrintReliableBunchBuffer();
+				PRAGMA_ENABLE_DEPRECATION_WARNINGS
 			}
 #else
 			check(NumOutRec<RELIABLE_BUFFER-1+Bunch->bClose);
@@ -1466,7 +1469,9 @@ FOutBunch* UChannel::PrepBunch(FOutBunch* Bunch, FOutBunch* OutBunch, bool Merge
 		if (CVarNetReliableDebug.GetValueOnAnyThread() == 2)
 		{
 			UE_LOG(LogNetTraffic, Warning, TEXT("%s. Reliable: %s"), *Describe(), *Bunch->DebugString);
+			PRAGMA_DISABLE_DEPRECATION_WARNINGS
 			PrintReliableBunchBuffer();
+			PRAGMA_ENABLE_DEPRECATION_WARNINGS
 			UE_LOG(LogNetTraffic, Warning, TEXT(""));
 		}
 #endif
@@ -1576,7 +1581,9 @@ void UChannel::AddedToChannelPool()
 	OpenAcked = false;
 	Closing = false;
 	Dormant = false;
+	PRAGMA_DISABLE_DEPRECATION_WARNINGS
 	bIsReplicationPaused = false;
+	PRAGMA_ENABLE_DEPRECATION_WARNINGS
 	OpenTemporary = false;
 	Broken = false;
 	bTornOff = false;
@@ -3044,12 +3051,14 @@ void UActorChannel::ProcessBunch( FInBunch & Bunch )
 		UE_LOG(LogNetTraffic, Log, TEXT("      Actor %s:"), *Actor->GetFullName() );
 	}
 
+	PRAGMA_DISABLE_DEPRECATION_WARNINGS
 	bool bLatestIsReplicationPaused = Bunch.bIsReplicationPaused != 0;
 	if (bLatestIsReplicationPaused != IsReplicationPaused())
 	{
 		Actor->OnReplicationPausedChanged(bLatestIsReplicationPaused);
 		SetReplicationPaused(bLatestIsReplicationPaused);
 	}
+	PRAGMA_ENABLE_DEPRECATION_WARNINGS
 
 	// Owned by connection's player?
 	UNetConnection* ActorConnection = Actor->GetNetConnection();
@@ -3313,10 +3322,11 @@ int64 UActorChannel::ReplicateActor()
 		return 0;
 	}
 
+	PRAGMA_DISABLE_DEPRECATION_WARNINGS
 	const TArray<FNetViewer>& NetViewers = ActorWorld->GetWorldSettings()->ReplicationViewers;
 	bool bIsNewlyReplicationPaused = false;
 	bool bIsNewlyReplicationUnpaused = false;
-	
+
 	if (OpenPacketId.First != INDEX_NONE && NetViewers.Num() > 0)
 	{
 		bool bNewPaused = true;
@@ -3342,6 +3352,7 @@ int64 UActorChannel::ReplicateActor()
 		bIsNewlyReplicationPaused = !bOldPaused && bNewPaused;
 		SetReplicationPaused(bNewPaused);
 	}
+	PRAGMA_ENABLE_DEPRECATION_WARNINGS
 
 	// The package map shouldn't have any carry over guids
 	// Static cast is fine here, since we check above.
@@ -3376,7 +3387,9 @@ int64 UActorChannel::ReplicateActor()
 	if (bIsNewlyReplicationPaused)
 	{
 		Bunch.bReliable = true;
+		PRAGMA_DISABLE_DEPRECATION_WARNINGS
 		Bunch.bIsReplicationPaused = true;
+		PRAGMA_ENABLE_DEPRECATION_WARNINGS
 	}
 
 #if !(UE_BUILD_SHIPPING || UE_BUILD_TEST)
diff --git a/Engine/Source/Runtime/Engine/Private/NetConnection.cpp b/Engine/Source/Runtime/Engine/Private/NetConnection.cpp
--- a/Engine/Source/Runtime/Engine/Private/NetConnection.cpp
+++ b/Engine/Source/Runtime/Engine/Private/NetConnection.cpp
@@ -3012,7 +3012,10 @@ void UNetConnection::ReceivedPacket( FBitReader& Reader, bool bIsReinjectedPacke
 														EChannelCloseReason::Destroyed;
 				}
 
-				Bunch.bIsReplicationPaused  = Reader.ReadBit();
+				PRAGMA_DISABLE_DEPRECATION_WARNINGS
+				// Removing this will require a new network version for replay compatabilty
+				Bunch.bIsReplicationPaused = Reader.ReadBit();
+				PRAGMA_ENABLE_DEPRECATION_WARNINGS
 				Bunch.bReliable				= Reader.ReadBit();
 
 				if (Bunch.EngineNetVer() < FEngineNetworkCustomVersion::MaxActorChannelsCustomization)
@@ -3941,7 +3944,9 @@ int32 UNetConnection::SendRawBunch(FOutBunch& Bunch, bool InAllowMerge, const FN
 			SendBunchHeader.SerializeInt(Value, (uint32)EChannelCloseReason::MAX);
 		}
 	}
+	PRAGMA_DISABLE_DEPRECATION_WARNINGS
 	SendBunchHeader.WriteBit(Bunch.bIsReplicationPaused);
+	PRAGMA_ENABLE_DEPRECATION_WARNINGS
 	SendBunchHeader.WriteBit(Bunch.bReliable);
 
 	uint32 ChIndex = Bunch.ChIndex;
diff --git a/Engine/Source/Runtime/Engine/Public/Net/DataBunch.h b/Engine/Source/Runtime/Engine/Public/Net/DataBunch.h
--- a/Engine/Source/Runtime/Engine/Public/Net/DataBunch.h
+++ b/Engine/Source/Runtime/Engine/Public/Net/DataBunch.h
@@ -33,6 +33,7 @@ public:
 	uint8					ReceivedAck:1;
 	uint8					bOpen:1;
 	uint8					bClose:1;
+	UE_DEPRECATED(5.3, "Replication pausing is deprecated")
 	uint8					bIsReplicationPaused:1;   // Replication on this channel is being paused by the server
 	uint8					bReliable:1;
 	uint8					bPartial:1;				// Not a complete bunch
@@ -73,6 +74,11 @@ public:
 	FOutBunch( class UChannel* InChannel, bool bClose );
 	FOutBunch( UPackageMap * PackageMap, int64 InMaxBits = 1024 );
 
+	PRAGMA_DISABLE_DEPRECATION_WARNINGS
+	FOutBunch(const FOutBunch&) = default;
+	FOutBunch& operator=(const FOutBunch&) = default;
+	PRAGMA_ENABLE_DEPRECATION_WARNINGS
+
 	virtual ~FOutBunch();
 
 	FString	ToString()
@@ -90,7 +96,9 @@ public:
 		{
 			Str += FString::Printf(TEXT("CloseReason: %s "), LexToString(CloseReason));
 		}
+		PRAGMA_DISABLE_DEPRECATION_WARNINGS
 		Str += FString::Printf(TEXT("bIsReplicationPaused: %d "), bIsReplicationPaused);
+		PRAGMA_ENABLE_DEPRECATION_WARNINGS
 		Str += FString::Printf(TEXT("bReliable: %d "), bReliable);
 		Str += FString::Printf(TEXT("bPartial: %d//%d//%d "), bPartial, bPartialInitial, bPartialFinal);
 		Str += FString::Printf( TEXT( "bHasPackageMapExports: %d " ), bHasPackageMapExports );
@@ -119,6 +127,7 @@ public:
 	int32				ChSequence;
 	uint8				bOpen:1;
 	uint8				bClose:1;
+	UE_DEPRECATED(5.3, "Replication pausing is deprecated")
 	uint8				bIsReplicationPaused:1;		// Replication on this channel is being paused by the server
 	uint8				bReliable:1;
 	uint8				bPartial:1;					// Not a complete bunch
@@ -145,7 +154,9 @@ public:
 		{
 			Str += FString::Printf(TEXT("CloseReason: %s "), LexToString(CloseReason));
 		}
+		PRAGMA_DISABLE_DEPRECATION_WARNINGS
 		Str += FString::Printf(TEXT("bIsReplicationPaused: %d "), bIsReplicationPaused);
+		PRAGMA_ENABLE_DEPRECATION_WARNINGS
 		Str += FString::Printf(TEXT("bReliable: %d "), bReliable);
 		Str += FString::Printf(TEXT("bPartial: %d//%d//%d "), bPartial, bPartialInitial, bPartialFinal);
 		Str += FString::Printf(TEXT("bHasPackageMapExports: %d "), bHasPackageMapExports );
@@ -161,6 +172,10 @@ public:
 	FInBunch( UNetConnection* InConnection, uint8* Src=NULL, int64 CountBits=0 );
 	FInBunch( FInBunch &InBunch, bool CopyBuffer );
 
+	PRAGMA_DISABLE_DEPRECATION_WARNINGS
+	FInBunch(const FInBunch&) = default;
+	PRAGMA_ENABLE_DEPRECATION_WARNINGS
+
 	virtual void CountMemory(FArchive& Ar) const override;
 };
 
diff --git a/Engine/Source/Runtime/Net/Core/Public/Net/Core/Trace/Private/NetTraceInternal.h b/Engine/Source/Runtime/Net/Core/Public/Net/Core/Trace/Private/NetTraceInternal.h
--- a/Engine/Source/Runtime/Net/Core/Public/Net/Core/Trace/Private/NetTraceInternal.h
+++ b/Engine/Source/Runtime/Net/Core/Public/Net/Core/Trace/Private/NetTraceInternal.h
@@ -75,6 +75,7 @@ union FNetTraceBunchInfo
 		uint64 bPartial : 1;
 		uint64 bPartialInitial : 1;
 		uint64 bPartialFinal : 1;
+		UE_DEPRECATED(5.3, "Replication pausing is deprecated")
 		uint64 bIsReplicationPaused : 1;
 		uint64 bOpen : 1;
 		uint64 bClose : 1;
@@ -299,7 +300,9 @@ FNetTraceBunchInfo MakeBunchInfo(const T& Bunch)
 	BunchInfo.bPartial = Bunch.bPartial;
 	BunchInfo.bPartialInitial = Bunch.bPartialInitial;
 	BunchInfo.bPartialFinal = Bunch.bPartialFinal;
+	PRAGMA_DISABLE_DEPRECATION_WARNINGS
 	BunchInfo.bIsReplicationPaused = Bunch.bIsReplicationPaused;
+	PRAGMA_ENABLE_DEPRECATION_WARNINGS
 	BunchInfo.bOpen = Bunch.bOpen;
 	BunchInfo.bClose = Bunch.bClose;
 	BunchInfo.bReliable = Bunch.bReliable;
