Added physics listing command to ISM components to check how many ISMCs have active physics bodies and shapes and output stats relating to the number created.
Fixed a number of cases that did not correctly handle nullptr body instances in ISMCs - that is a valid state when instances have no collision

diff --git a/Engine/Source/Runtime/Engine/Private/InstancedStaticMesh.cpp b/Engine/Source/Runtime/Engine/Private/InstancedStaticMesh.cpp
--- a/Engine/Source/Runtime/Engine/Private/InstancedStaticMesh.cpp
+++ b/Engine/Source/Runtime/Engine/Private/InstancedStaticMesh.cpp
@@ -145,7 +145,70 @@ class FISMExecHelper : public FSelfRegisteringExec
 {
 	virtual bool Exec_Runtime(class UWorld* InWorld, const TCHAR* Cmd, FOutputDevice& Ar)
 	{
-		if (FParse::Command(&Cmd, TEXT("LIST ISM")))
+		if(FParse::Command(&Cmd, TEXT("LIST ISM PHYSICS")))
+		{
+			Ar.Logf(TEXT("----- BEGIN ISM PHYSICS LISTING -------------------------"));
+			Ar.Logf(TEXT("Component Name, Component Class Name, Mesh Name, Num Bodies, Num Shapes"));
+
+			int32 BodiesTotal = 0;
+			int32 ShapesTotal = 0;
+			for(TObjectIterator<UInstancedStaticMeshComponent> It; It; ++It)
+			{
+				UInstancedStaticMeshComponent* ISMComponent = *It;
+				if((!ISMComponent->GetWorld() || !ISMComponent->GetWorld()->IsGameWorld()))
+				{
+					continue;
+				}
+
+				const FString ComponentName = ISMComponent->GetPathName(ISMComponent->GetOwner()->GetOuter());
+
+				if(ISMComponent->ShouldCreatePhysicsState())
+				{
+					const TArray<FBodyInstance*>& Bodies = ISMComponent->InstanceBodies;
+					if(Bodies.Num() == 0)
+					{
+						continue;
+					}
+
+					int32 NumValidInstanceBodies = 0;
+					int32 NumInstanceShapes = 0;
+					TArray<FPhysicsShapeHandle> Shapes;
+					for(const FBodyInstance* Body : Bodies)
+					{
+						if(!Body)
+						{
+							continue;
+						}
+						++NumValidInstanceBodies;
+						
+						FPhysicsCommand::ExecuteRead(Body->ActorHandle, [&Body, &Shapes, &NumInstanceShapes](const FPhysicsActorHandle& Actor)
+						{
+							Shapes.Reset();
+							Body->GetAllShapes_AssumesLocked(Shapes);
+
+							NumInstanceShapes += Shapes.Num();
+						});
+					}
+
+					UStaticMesh* Mesh = ISMComponent->GetStaticMesh();
+
+					Ar.Logf(TEXT("%s, %s, %s, %d, %d"),
+							*ComponentName,
+							*ISMComponent->GetClass()->GetName(),
+							Mesh ? *Mesh->GetPathName() : TEXT("None"),
+							NumValidInstanceBodies,
+							NumInstanceShapes);
+
+					BodiesTotal += NumValidInstanceBodies;
+					ShapesTotal += NumInstanceShapes;
+				}
+			}
+			Ar.Logf(TEXT("\nTotal Bodies, Total Shapes\n%d, %d"), BodiesTotal, ShapesTotal);
+			Ar.Logf(TEXT("----- END ISM PHYSICS LISTING ---------------------------"));
+
+			return true;
+		}
+		else if(FParse::Command(&Cmd, TEXT("LIST ISM")))
 		{
 			// Flush commands because we will be touching the proxy.
 			FlushRenderingCommands();
@@ -3383,9 +3446,13 @@ bool UInstancedStaticMeshComponent::RemoveInstanceInternal(int32 InstanceIndex,
 	{
 		if (FBodyInstance*& InstanceBody = InstanceBodies[InstanceIndex])
 		{
-			InstanceBody->TermBody();
-			delete InstanceBody;
-			InstanceBody = nullptr;
+			// Not having a body is a valid case when our physics state cannot be created (see CreateAllInstanceBodies)
+			if(InstanceBody)
+			{
+				InstanceBody->TermBody();
+				delete InstanceBody;
+				InstanceBody = nullptr;
+			}
 
 			InstanceBodies.RemoveAt(InstanceIndex);
 
@@ -4687,6 +4754,12 @@ bool UInstancedStaticMeshComponent::LineTraceComponent(FHitResult& OutHit, const
 	// TODO: use spatial acceleration instead
 	for (FBodyInstance* Body : InstanceBodies)
 	{
+		// Not having a body is a valid case when our physics state cannot be created (see CreateAllInstanceBodies)
+		if(!Body)
+		{
+			continue;
+		}
+
 		if (Body->LineTrace(Hit, Start, End, Params.bTraceComplex, Params.bReturnPhysicalMaterial))
 		{
 			bHaveHit = true;
@@ -4722,6 +4795,12 @@ bool UInstancedStaticMeshComponent::SweepComponent(FHitResult& OutHit, const FVe
 	// TODO: use spatial acceleration instead
 	for (FBodyInstance* Body : InstanceBodies)
 	{
+		// Not having a body is a valid case when our physics state cannot be created (see CreateAllInstanceBodies)
+		if(!Body)
+		{
+			continue;
+		}
+
 		if (Body->Sweep(Hit, Start, End, ShapeWorldRotation, CollisionShape, bTraceComplex))
 		{
 			if (!bHaveHit || Hit.Time < OutHit.Time)
@@ -4738,6 +4817,12 @@ bool UInstancedStaticMeshComponent::OverlapComponent(const FVector& Pos, const F
 {	
 	for (FBodyInstance* Body : InstanceBodies)
 	{
+		// Not having a body is a valid case when our physics state cannot be created (see CreateAllInstanceBodies)
+		if(!Body)
+		{
+			continue;
+		}
+
 		if (Body->OverlapTest(Pos, Rot, CollisionShape))
 		{
 			return true;
@@ -4785,7 +4870,12 @@ bool UInstancedStaticMeshComponent::ComponentOverlapMultiImpl(TArray<struct FOve
 	bool bHaveBlockingHit = false;
 	for (FBodyInstance* Body : InstanceBodies)
 	{
-		checkSlow(Body);
+		// Not having a body is a valid case when our physics state cannot be created (see CreateAllInstanceBodies)
+		if(!Body)
+		{
+			continue;
+		}
+
 		if (Body->OverlapMulti(OutOverlaps, InWorld, &WorldToComponent, Pos, Rot, TestChannel, ParamsWithSelf, ResponseParams, ObjectQueryParams))
 		{
 			bHaveBlockingHit = true;
