Levels can now have more than one AWorldDataLayers
ActorDescContainers restriced to having a single one
RegisterActorDescContainer fitler duplicated WorldDataLayers
RegisterActorDescContainer callers can pass an additionnal ActorDescFilter function

diff --git a/Engine/Source/Runtime/Engine/Private/World.cpp b/Engine/Source/Runtime/Engine/Private/World.cpp
--- a/Engine/Source/Runtime/Engine/Private/World.cpp
+++ b/Engine/Source/Runtime/Engine/Private/World.cpp
@@ -1675,7 +1675,6 @@ void UWorld::RepairWorldSettings()
 void UWorld::RepairSingletonActors()
 {
 	RepairWorldSettings();
-	RepairSingletonActorOfClass<AWorldDataLayers>();
 }
 
 #if WITH_EDITOR
diff --git a/Engine/Source/Runtime/Engine/Private/WorldPartition/ContentBundle/ContentBundleEditor.cpp b/Engine/Source/Runtime/Engine/Private/WorldPartition/ContentBundle/ContentBundleEditor.cpp
--- a/Engine/Source/Runtime/Engine/Private/WorldPartition/ContentBundle/ContentBundleEditor.cpp
+++ b/Engine/Source/Runtime/Engine/Private/WorldPartition/ContentBundle/ContentBundleEditor.cpp
@@ -58,7 +58,8 @@ void FContentBundleEditor::DoInjectContent()
 		UnsavedActorMonitor->Initialize(*this);
 
 		UWorldPartition* WorldPartition = GetInjectedWorld()->GetWorldPartition();
-		ActorDescContainer = WorldPartition->RegisterActorDescContainer(FName(*ActorDescContainerPackage));
+		UWorldPartition::FContainerRegistrationParams RegistrationsParams(*ActorDescContainerPackage);
+		ActorDescContainer = WorldPartition->RegisterActorDescContainer(RegistrationsParams);
 		if (ActorDescContainer.IsValid())
 		{
 			UE_LOG(LogContentBundle, Log, TEXT("%s ExternalActors in %s found. %u actors were injected"), *ContentBundle::Log::MakeDebugInfoString(*this), *ActorDescContainer->GetExternalActorPath(), ActorDescContainer->GetActorDescCount());
diff --git a/Engine/Source/Runtime/Engine/Private/WorldPartition/WorldPartition.cpp b/Engine/Source/Runtime/Engine/Private/WorldPartition/WorldPartition.cpp
--- a/Engine/Source/Runtime/Engine/Private/WorldPartition/WorldPartition.cpp
+++ b/Engine/Source/Runtime/Engine/Private/WorldPartition/WorldPartition.cpp
@@ -29,6 +29,7 @@
 #include "LocationVolume.h"
 #include "Engine/LevelScriptBlueprint.h"
 #include "ActorReferencesUtils.h"
+#include "WorldPartition/DataLayer/WorldDataLayersActorDesc.h"
 #include "WorldPartition/IWorldPartitionEditorModule.h"
 #include "WorldPartition/WorldPartitionLevelHelper.h"
 #include "WorldPartition/WorldPartitionEditorHash.h"
@@ -517,7 +518,8 @@ void UWorldPartition::Initialize(UWorld* InWorld, const FTransform& InTransform)
 			);
 		}
 
-		ActorDescContainer = RegisterActorDescContainer(PackageName);
+		FContainerRegistrationParams ContainerInitParams(PackageName);
+		ActorDescContainer = RegisterActorDescContainer(ContainerInitParams);
 
 		{
 			TRACE_CPUPROFILER_EVENT_SCOPE(UActorDescContainer::Hash);
@@ -1736,14 +1738,38 @@ void UWorldPartition::AppendAssetRegistryTags(TArray<FAssetRegistryTag>& OutTags
 	}
 }
 
-UActorDescContainer* UWorldPartition::RegisterActorDescContainer(const FName& ContainerPackage)
+UActorDescContainer* UWorldPartition::RegisterActorDescContainer(const FContainerRegistrationParams& InRegistrationParameters)
 {
-	if (!Contains(ContainerPackage))
-	{
-		UActorDescContainer* ContainerToRegister = NewObject<UActorDescContainer>(this);
+	if (!Contains(InRegistrationParameters.PackageName))
+	{	
+		UActorDescContainer::FInitializeParams ContainerInitParams(GetWorld(), InRegistrationParameters.PackageName);
+
+		const FWorldDataLayersActorDesc* WorldDataLayerActorsDesc = nullptr;
+		ContainerInitParams.FilterActorDesc = [this, &WorldDataLayerActorsDesc, &InRegistrationParameters](const FWorldPartitionActorDesc* ActorDesc)
+		{
+			if (InRegistrationParameters.FilterActorDescFunc && !InRegistrationParameters.FilterActorDescFunc(ActorDesc))
+			{
+				return false;
+			}
+
+			// Filter duplicate WorldDataLayers
+			if (ActorDesc->GetActorNativeClass()->IsChildOf<AWorldDataLayers>())
+			{
+				const FWorldDataLayersActorDesc* FoundWorldDataLayerActorsDesc = StaticCast<const FWorldDataLayersActorDesc*>(ActorDesc);
+				if (FoundWorldDataLayerActorsDesc != nullptr && WorldDataLayerActorsDesc != nullptr)
+				{
+					UE_LOG(LogWorldPartition, Warning, TEXT("Extra World Data Layer '%s' actor found. Clean up invalid actors to remove the error."), *ActorDesc->GetActorPackage().ToString());
+					return false;
+				}
 
-		UActorDescContainer::FInitializeParams ContainerInitParams(GetWorld(), ContainerPackage);
-		ContainerInitParams.FilterActorDesc = [this](const FWorldPartitionActorDesc* ActorDesc) { return !GetActorDesc(ActorDesc->GetGuid()); };
+				WorldDataLayerActorsDesc = FoundWorldDataLayerActorsDesc;
+			}
+
+			// Filter actors with duplicated GUID in WorldPartition
+			return !GetActorDesc(ActorDesc->GetGuid());
+		};
+
+		UActorDescContainer* ContainerToRegister = NewObject<UActorDescContainer>(this);
 		ContainerToRegister->Initialize(ContainerInitParams);
 
 		AddContainer(ContainerToRegister);
diff --git a/Engine/Source/Runtime/Engine/Public/WorldPartition/WorldPartition.h b/Engine/Source/Runtime/Engine/Public/WorldPartition/WorldPartition.h
--- a/Engine/Source/Runtime/Engine/Public/WorldPartition/WorldPartition.h
+++ b/Engine/Source/Runtime/Engine/Public/WorldPartition/WorldPartition.h
@@ -222,12 +222,28 @@ public:
 
 	void AppendAssetRegistryTags(TArray<FAssetRegistryTag>& OutTags) const;
 
+	struct FContainerRegistrationParams
+	{
+		FContainerRegistrationParams(FName InPackageName)
+			: PackageName(InPackageName)
+		{}
+
+		/* The long package name of the container package on disk. */
+		FName PackageName;
+
+		/* Custom filter function used to filter actors descriptors. */
+		TUniqueFunction<bool(const FWorldPartitionActorDesc*)> FilterActorDescFunc;
+	};
+	UActorDescContainer* RegisterActorDescContainer(const FContainerRegistrationParams& InRegistrationParameters);
+	bool UnregisterActorDescContainer(UActorDescContainer* Container);
+	void UninitializeActorDescContainers();
+
 	DECLARE_MULTICAST_DELEGATE_OneParam(FActorDescContainerRegistrationDelegate, UActorDescContainer*);
 	FActorDescContainerRegistrationDelegate OnActorDescContainerRegistered;
 	FActorDescContainerRegistrationDelegate OnActorDescContainerUnregistered;
-	UActorDescContainer* RegisterActorDescContainer(const FName& ContainerPackage);
-	bool UnregisterActorDescContainer(UActorDescContainer* Container);
-	void UninitializeActorDescContainers();
+
+	UE_DEPRECATED(5.3, "Use RegisterActorDescContainer with FContainerRegistrationParams instead.")
+	UActorDescContainer* RegisterActorDescContainer(const FName& ContainerPackage) { return RegisterActorDescContainer(FContainerRegistrationParams(ContainerPackage)); }
 
 	// Actors pinning
 	void PinActors(const TArray<FGuid>& ActorGuids);
