FWorldPartitionHelpers: renamed HasExceededMaxMemory to ShouldCollectGarbage for consistency.

diff --git a/Engine/Source/Editor/UnrealEd/Private/WorldPartition/WorldPartitionBuilder.cpp b/Engine/Source/Editor/UnrealEd/Private/WorldPartition/WorldPartitionBuilder.cpp
--- a/Engine/Source/Editor/UnrealEd/Private/WorldPartition/WorldPartitionBuilder.cpp
+++ b/Engine/Source/Editor/UnrealEd/Private/WorldPartition/WorldPartitionBuilder.cpp
@@ -294,7 +294,7 @@ bool UWorldPartitionBuilder::Run(UWorld* World, FPackageSourceControlHelper& Pac
 
 					bResult = RunInternal(World, CellInfo, PackageHelper);
 
-					if (FWorldPartitionHelpers::HasExceededMaxMemory())
+					if (FWorldPartitionHelpers::ShouldCollectGarbage())
 					{
 						LoaderAdapters.Empty();
 						FWorldPartitionHelpers::DoCollectGarbage();
diff --git a/Engine/Source/Editor/UnrealEd/Private/WorldPartition/WorldPartitionHLODsBuilder.cpp b/Engine/Source/Editor/UnrealEd/Private/WorldPartition/WorldPartitionHLODsBuilder.cpp
--- a/Engine/Source/Editor/UnrealEd/Private/WorldPartition/WorldPartitionHLODsBuilder.cpp
+++ b/Engine/Source/Editor/UnrealEd/Private/WorldPartition/WorldPartitionHLODsBuilder.cpp
@@ -501,7 +501,7 @@ bool UWorldPartitionHLODsBuilder::BuildHLODActors()
 			}
 		}
 
-		if (FWorldPartitionHelpers::HasExceededMaxMemory())
+		if (FWorldPartitionHelpers::ShouldCollectGarbage())
 		{
 			FWorldPartitionHelpers::DoCollectGarbage();
 		}
diff --git a/Engine/Source/Editor/UnrealEd/Private/WorldPartition/WorldPartitionRenameDuplicateBuilder.cpp b/Engine/Source/Editor/UnrealEd/Private/WorldPartition/WorldPartitionRenameDuplicateBuilder.cpp
--- a/Engine/Source/Editor/UnrealEd/Private/WorldPartition/WorldPartitionRenameDuplicateBuilder.cpp
+++ b/Engine/Source/Editor/UnrealEd/Private/WorldPartition/WorldPartitionRenameDuplicateBuilder.cpp
@@ -363,7 +363,7 @@ bool UWorldPartitionRenameDuplicateBuilder::RunInternal(UWorld* World, const FCe
 					}
 				}
 						
-				if (FWorldPartitionHelpers::HasExceededMaxMemory())
+				if (FWorldPartitionHelpers::ShouldCollectGarbage())
 				{
 					if (!ProcessLoadedActors(ActorReferences))
 					{
diff --git a/Engine/Source/Editor/UnrealEd/Private/WorldPartition/WorldPartitionResaveActorsBuilder.cpp b/Engine/Source/Editor/UnrealEd/Private/WorldPartition/WorldPartitionResaveActorsBuilder.cpp
--- a/Engine/Source/Editor/UnrealEd/Private/WorldPartition/WorldPartitionResaveActorsBuilder.cpp
+++ b/Engine/Source/Editor/UnrealEd/Private/WorldPartition/WorldPartitionResaveActorsBuilder.cpp
@@ -316,7 +316,7 @@ bool UWorldPartitionResaveActorsBuilder::RunInternal(UWorld* World, const FCellI
 				}
 			}
 
-			if (FWorldPartitionHelpers::HasExceededMaxMemory())
+			if (FWorldPartitionHelpers::ShouldCollectGarbage())
 			{
 				if (!UWorldPartitionBuilder::SavePackages(PackagesToSave, PackageHelper))
 				{
diff --git a/Engine/Source/Runtime/Engine/Private/WorldPartition/RuntimeSpatialHash/RuntimeSpatialHashHLOD.cpp b/Engine/Source/Runtime/Engine/Private/WorldPartition/RuntimeSpatialHash/RuntimeSpatialHashHLOD.cpp
--- a/Engine/Source/Runtime/Engine/Private/WorldPartition/RuntimeSpatialHash/RuntimeSpatialHashHLOD.cpp
+++ b/Engine/Source/Runtime/Engine/Private/WorldPartition/RuntimeSpatialHash/RuntimeSpatialHashHLOD.cpp
@@ -106,7 +106,7 @@ static void GameTick(UWorld* InWorld)
 	static const int32 FlushRenderingFrequency = 256;
 
 	// Perform a GC when memory usage exceeds a given threshold
-	if (FWorldPartitionHelpers::HasExceededMaxMemory())
+	if (FWorldPartitionHelpers::ShouldCollectGarbage())
 	{
 		FWorldPartitionHelpers::DoCollectGarbage();
 	}
diff --git a/Engine/Source/Runtime/Engine/Private/WorldPartition/WorldPartitionHelpers.cpp b/Engine/Source/Runtime/Engine/Private/WorldPartition/WorldPartitionHelpers.cpp
--- a/Engine/Source/Runtime/Engine/Private/WorldPartition/WorldPartitionHelpers.cpp
+++ b/Engine/Source/Runtime/Engine/Private/WorldPartition/WorldPartitionHelpers.cpp
@@ -178,7 +178,7 @@ void FWorldPartitionHelpers::ForEachActorWithLoading(UWorldPartition* WorldParti
 					return false;
 				}
 
-				if (!Params.bKeepReferences && (Params.bGCPerActor || FWorldPartitionHelpers::HasExceededMaxMemory()))
+				if (!Params.bKeepReferences && (Params.bGCPerActor || FWorldPartitionHelpers::ShouldCollectGarbage()))
 				{
 					CallGarbageCollect();
 				}
@@ -224,15 +224,6 @@ void FWorldPartitionHelpers::ForEachActorWithLoading(UWorldPartition* WorldParti
 
 bool FWorldPartitionHelpers::HasExceededMaxMemory()
 {
-	// Even if we're not exhausting memory, GC should be run at periodic intervals
-	static double LastTime = FPlatformTime::Seconds();
-	double ThisTime = FPlatformTime::Seconds();
-	if ((ThisTime - LastTime) > 30)
-	{
-		LastTime = ThisTime;
-		return true;
-	}
-
 	const FPlatformMemoryStats MemStats = FPlatformMemory::GetStats();
 
 	const uint64 MemoryMinFreePhysical = 1llu * 1024 * 1024 * 1024;
@@ -243,7 +234,21 @@ bool FWorldPartitionHelpers::HasExceededMaxMemory()
 	const bool bHasExceededMaxMemory = bHasExceededMinFreePhysical || bHasExceededMaxUsedPhysical;
 
 	return bHasExceededMaxMemory;
-};
+}
+
+bool FWorldPartitionHelpers::ShouldCollectGarbage()
+{
+	// Even if we're not exhausting memory, GC should be run at periodic intervals
+	static double LastTime = FPlatformTime::Seconds();
+	double ThisTime = FPlatformTime::Seconds();
+	if ((ThisTime - LastTime) > 30)
+	{
+		LastTime = ThisTime;
+		return true;
+	}
+
+	return HasExceededMaxMemory();
+}
 
 void FWorldPartitionHelpers::DoCollectGarbage()
 {
@@ -255,7 +260,7 @@ void FWorldPartitionHelpers::DoCollectGarbage()
 		(int64)MemStatsAfter.AvailablePhysical / (1024.0 * 1024.0 * 1024.0),
 		(int64)MemStatsAfter.AvailableVirtual / (1024.0 * 1024.0 * 1024.0)
 	);
-};
+}
 
 void FWorldPartitionHelpers::FakeEngineTick(UWorld* InWorld)
 {
diff --git a/Engine/Source/Runtime/Engine/Public/WorldPartition/WorldPartitionHelpers.h b/Engine/Source/Runtime/Engine/Public/WorldPartition/WorldPartitionHelpers.h
--- a/Engine/Source/Runtime/Engine/Public/WorldPartition/WorldPartitionHelpers.h
+++ b/Engine/Source/Runtime/Engine/Public/WorldPartition/WorldPartitionHelpers.h
@@ -130,6 +130,7 @@ public:
 	static void ForEachActorWithLoading(UWorldPartition* WorldPartition, TFunctionRef<bool(const FWorldPartitionActorDesc*)> Func, const FForEachActorWithLoadingParams& Params, FForEachActorWithLoadingResult& Result);
 	
 	static bool HasExceededMaxMemory();
+	static bool ShouldCollectGarbage();
 	static void DoCollectGarbage();
 
 	// Simulate an engine frame tick
