Fixing Non-Unity build

diff --git a/Engine/Source/Runtime/Engine/Private/Components/HeterogeneousVolumeComponent.cpp b/Engine/Source/Runtime/Engine/Private/Components/HeterogeneousVolumeComponent.cpp
--- a/Engine/Source/Runtime/Engine/Private/Components/HeterogeneousVolumeComponent.cpp
+++ b/Engine/Source/Runtime/Engine/Private/Components/HeterogeneousVolumeComponent.cpp
@@ -2,17 +2,22 @@
 
 #include "Components/HeterogeneousVolumeComponent.h"
 
+#include "Components/BillboardComponent.h"
 #include "Engine/Texture2D.h"
 #include "HeterogeneousVolumeInterface.h"
+#include "LocalVertexFactory.h"
 #include "MaterialDomain.h"
+#include "Materials/Material.h"
 #include "Materials/MaterialInstanceDynamic.h"
-#include "Components/BillboardComponent.h"
 #include "Materials/MaterialInterface.h"
-#include "UObject/ConstructorHelpers.h"
-#include "SparseVolumeTexture/SparseVolumeTexture.h"
 #include "PrimitiveViewRelevance.h"
 #include "PrimitiveSceneProxy.h"
+#include "SceneInterface.h"
 #include "SceneManagement.h"
+#include "SceneView.h"
+#include "SparseVolumeTexture/SparseVolumeTexture.h"
+#include "StaticMeshResources.h"
+#include "UObject/ConstructorHelpers.h"
 
 #include UE_INLINE_GENERATED_CPP_BY_NAME(HeterogeneousVolumeComponent)
 
