Texture factor GetBuiltTextureSize logic out of GetDefaultTextureFormatName
allows query of platform built texture size
behavior NOP

diff --git a/Engine/Source/Runtime/Engine/Classes/Engine/Texture.h b/Engine/Source/Runtime/Engine/Classes/Engine/Texture.h
--- a/Engine/Source/Runtime/Engine/Classes/Engine/Texture.h
+++ b/Engine/Source/Runtime/Engine/Classes/Engine/Texture.h
@@ -1501,6 +1501,12 @@ public:
 	virtual TMap<FString, FTexturePlatformData*>* GetCookedPlatformData() { return NULL; }
 
 	void CleanupCachedRunningPlatformData();
+	
+	/**
+	 * Get the dimensions of the largest mip of the texture when built for the target platform
+	 *   accounting for LODBias and other constraints
+	 */
+	void GetBuiltTextureSize( const ITargetPlatform* TargetPlatform , int32 & OutSizeX, int32 & OutSizeY ) const;
 
 	/**
 	 * Serializes cooked platform data.
diff --git a/Engine/Source/Runtime/Engine/Private/Texture.cpp b/Engine/Source/Runtime/Engine/Private/Texture.cpp
--- a/Engine/Source/Runtime/Engine/Private/Texture.cpp
+++ b/Engine/Source/Runtime/Engine/Private/Texture.cpp
@@ -3221,6 +3221,89 @@ static FName ConditionalGetPrefixedFormat(FName TextureFormatName, const ITarget
 	return TextureFormatName;
 }
 
+void UTexture::GetBuiltTextureSize( const ITargetPlatform* TargetPlatform , int32 & OutSizeX, int32 & OutSizeY ) const
+{
+	// @todo Oodle : SizeZ
+	// @todo Oodle : verify against TextureCompressorModule
+	// @todo Oodle : with cinematic mips or not? maybe add a bool arg
+	
+	int32 SizeX,SizeY;
+
+#if WITH_EDITORONLY_DATA
+	SizeX = Source.GetSizeX();
+	SizeY = Source.GetSizeY();
+
+	if (PowerOfTwoMode == ETexturePowerOfTwoSetting::PadToPowerOfTwo || PowerOfTwoMode == ETexturePowerOfTwoSetting::PadToSquarePowerOfTwo)
+	{
+		SizeX = FMath::RoundUpToPowerOfTwo(SizeX);
+		SizeY = FMath::RoundUpToPowerOfTwo(SizeY);
+
+		if (PowerOfTwoMode == ETexturePowerOfTwoSetting::PadToSquarePowerOfTwo)
+		{
+			SizeX = SizeY = FMath::Max(SizeX, SizeY);
+		}
+	}
+	else
+	{
+		checkf(PowerOfTwoMode == ETexturePowerOfTwoSetting::None, TEXT("Unknown entry in ETexturePowerOfTwoSetting::Type"));
+	}
+
+	if (Source.IsLongLatCubemap())
+	{
+		// this should be kept in sync with ComputeLongLatCubemapExtents()
+		SizeX = SizeY = FMath::Max(1 << FMath::FloorLog2(SizeX / 2), 32);
+	}
+
+	//we need to really have the actual top mip size of output platformdata
+	//	(hence the LODBias check below)
+	// trying to reproduce here exactly what TextureCompressor + serialization will do = brittle
+
+	if ( MaxTextureSize != 0 )
+	{
+		while( SizeX > MaxTextureSize || SizeY > MaxTextureSize )
+		{
+			SizeX = FMath::Max(SizeX>>1,1);
+			SizeY = FMath::Max(SizeY>>1,1);
+		}
+	}
+
+	static const auto CVarVirtualTexturesEnabled = IConsoleManager::Get().FindTConsoleVariableDataInt(TEXT("r.VirtualTextures")); check(CVarVirtualTexturesEnabled);
+	const bool bVirtualTextureStreaming = CVarVirtualTexturesEnabled->GetValueOnAnyThread() && TargetPlatform->SupportsFeature(ETargetPlatformFeatures::VirtualTextureStreaming) && VirtualTextureStreaming;
+
+	const UTextureLODSettings& LODSettings = TargetPlatform->GetTextureLODSettings();
+ 	const uint32 LODBiasNoCinematics = FMath::Max<int32>(LODSettings.CalculateLODBias(SizeX, SizeY, MaxTextureSize, LODGroup, LODBias, 0, MipGenSettings, bVirtualTextureStreaming), 0);
+	SizeX = FMath::Max<int32>(SizeX >> LODBiasNoCinematics, 1);
+	SizeY = FMath::Max<int32>(SizeY >> LODBiasNoCinematics, 1);
+
+	// @todo Oodle : check against GetCookedPlatformData ?
+
+#else // WITH_EDITORONLY_DATA
+
+	// no Editor data
+	//	can't support a query about arbitrary target platform
+	//	must be a query of current running platform
+
+	// note that using PlatformData size is not right in Editor
+	//	because of the different way that LODBias acts in Editor vs cook (as a mip selection, not a size change)
+
+	FTexturePlatformData ** PPlatformData = const_cast<UTexture *>(this)->GetRunningPlatformData();
+	if ( PPlatformData && *PPlatformData )
+	{
+		SizeX = (*PPlatformData)->SizeX;
+		SizeY = (*PPlatformData)->SizeY;
+	}
+	else
+	{
+		SizeX = 0;
+		SizeY = 0;
+	}
+
+#endif // WITH_EDITORONLY_DATA
+
+	OutSizeX = SizeX;
+	OutSizeY = SizeY;
+}
+
 // this should not be called directly; it is called from TargetPlatform GetTextureFormats
 //	entry point API is GetPlatformTextureFormatNamesWithPrefix
 FName GetDefaultTextureFormatName( const ITargetPlatform* TargetPlatform, const UTexture* Texture, int32 LayerIndex, 
@@ -3257,9 +3340,6 @@ FName GetDefaultTextureFormatName( const ITargetPlatform* TargetPlatform, const
 
 	check(TargetPlatform);
 
-	static const auto CVarVirtualTexturesEnabled = IConsoleManager::Get().FindTConsoleVariableDataInt(TEXT("r.VirtualTextures")); check(CVarVirtualTexturesEnabled);
-	const bool bVirtualTextureStreaming = CVarVirtualTexturesEnabled->GetValueOnAnyThread() && TargetPlatform->SupportsFeature(ETargetPlatformFeatures::VirtualTextureStreaming) && Texture->VirtualTextureStreaming;
-
 	FTextureFormatSettings FormatSettings;
 	Texture->GetLayerFormatSettings(LayerIndex, FormatSettings);
 
@@ -3279,49 +3359,10 @@ FName GetDefaultTextureFormatName( const ITargetPlatform* TargetPlatform, const
 
 	if (!bNoCompression)
 	{
-		int32 SizeX = Texture->Source.GetSizeX();
-		int32 SizeY = Texture->Source.GetSizeY();
-
-		if (Texture->PowerOfTwoMode == ETexturePowerOfTwoSetting::PadToPowerOfTwo || Texture->PowerOfTwoMode == ETexturePowerOfTwoSetting::PadToSquarePowerOfTwo)
-		{
-			SizeX = FMath::RoundUpToPowerOfTwo(SizeX);
-			SizeY = FMath::RoundUpToPowerOfTwo(SizeY);
-
-			if (Texture->PowerOfTwoMode == ETexturePowerOfTwoSetting::PadToSquarePowerOfTwo)
-			{
-				SizeX = SizeY = FMath::Max(SizeX, SizeY);
-			}
-		}
-		else
-		{
-			checkf(Texture->PowerOfTwoMode == ETexturePowerOfTwoSetting::None, TEXT("Unknown entry in ETexturePowerOfTwoSetting::Type"));
-		}
-
-		if (Texture->Source.IsLongLatCubemap())
-		{
-			// this should be kept in sync with ComputeLongLatCubemapExtents()
-			SizeX = SizeY = FMath::Max(1 << FMath::FloorLog2(SizeX / 2), 32);
-		}
-
-		//@todo Oodle: we need to really have the actual top mip size of output platformdata
-		//	(hence the LODBias check below)
-		// trying to reproduce here exactly what TextureCompressor + serialization will do = brittle
-
-		if ( Texture->MaxTextureSize != 0 )
-		{
-			while( SizeX > Texture->MaxTextureSize || SizeY > Texture->MaxTextureSize )
-			{
-				SizeX = FMath::Max(SizeX>>1,1);
-				SizeY = FMath::Max(SizeY>>1,1);
-			}
-		}
+		int32 SizeX,SizeY;
+		Texture->GetBuiltTextureSize(TargetPlatform,SizeX,SizeY);
 
-#if WITH_EDITORONLY_DATA
-		const UTextureLODSettings& LODSettings = TargetPlatform->GetTextureLODSettings();
- 		const uint32 LODBiasNoCinematics = FMath::Max<int32>(LODSettings.CalculateLODBias(SizeX, SizeY, Texture->MaxTextureSize, Texture->LODGroup, Texture->LODBias, 0, Texture->MipGenSettings, bVirtualTextureStreaming), 0);
-		SizeX = FMath::Max<int32>(SizeX >> LODBiasNoCinematics, 1);
-		SizeY = FMath::Max<int32>(SizeY >> LODBiasNoCinematics, 1);
-#endif
+	
 		// Don't compress textures smaller than the DXT block size.
 		// Also force uncompressed if size of top mip is not a multiple of 4
 		// note that even if top mip is a multiple of 4, lower may not be
