Track (first) owners of FShaderMapResource_SharedCode to improve leak investigation.

#rnx

diff --git a/Engine/Plugins/Compositing/OpenColorIO/Source/OpenColorIO/Private/OpenColorIOShared.cpp b/Engine/Plugins/Compositing/OpenColorIO/Source/OpenColorIO/Private/OpenColorIOShared.cpp
--- a/Engine/Plugins/Compositing/OpenColorIO/Source/OpenColorIO/Private/OpenColorIOShared.cpp
+++ b/Engine/Plugins/Compositing/OpenColorIO/Source/OpenColorIO/Private/OpenColorIOShared.cpp
@@ -190,6 +190,7 @@ void FOpenColorIOTransformResource::SerializeShaderMap(FArchive& Ar)
 				if (bSuccessfullyLoaded && FApp::CanEverRender())
 				{
 					GameThreadShaderMap = RenderingThreadShaderMap = LoadedShaderMap;
+					GameThreadShaderMap->GetResource()->SetOwnerName(GetOwnerFName());
 #if WITH_EDITOR
 					GameThreadShaderMap->AssociateWithAsset(AssetPath);
 #endif
@@ -216,9 +217,7 @@ void FOpenColorIOTransformResource::SetupResource(ERHIFeatureLevel::Type InFeatu
 	ShaderCode = InShadercode;
 	FriendlyName = InFriendlyName.Replace(TEXT("/"), TEXT(""));
 	WorkingColorSpaceTransformType = InWorkingColorSpaceTransformType;
-#if WITH_EDITOR
 	AssetPath = InAssetPath;
-#endif // WITH_EDITOR
 
 	SetFeatureLevel(InFeatureLevel);
 }
diff --git a/Engine/Plugins/Compositing/OpenColorIO/Source/OpenColorIO/Public/OpenColorIOShared.h b/Engine/Plugins/Compositing/OpenColorIO/Source/OpenColorIO/Public/OpenColorIOShared.h
--- a/Engine/Plugins/Compositing/OpenColorIO/Source/OpenColorIO/Public/OpenColorIOShared.h
+++ b/Engine/Plugins/Compositing/OpenColorIO/Source/OpenColorIO/Public/OpenColorIOShared.h
@@ -401,11 +401,21 @@ public:
 		return GameThreadShaderMap;
 	}
 
+	/** Returns owner name for tracking */
+	FName GetOwnerFName() const
+	{
+		return AssetPath;
+	}
+
 	/** Note: SetRenderingThreadShaderMap must also be called with the same value, but from the rendering thread. */
 	void SetGameThreadShaderMap(FOpenColorIOShaderMap* InShaderMap)
 	{
 		checkSlow(IsInGameThread() || IsInAsyncLoadingThread());
 		GameThreadShaderMap = InShaderMap;
+		if (LIKELY(GameThreadShaderMap))
+		{
+			GameThreadShaderMap->GetResource()->SetOwnerName(GetOwnerFName());
+		}
 	}
 
 	/** Note: SetGameThreadShaderMap must also be called with the same value, but from the game thread. */
@@ -422,6 +432,10 @@ public:
 	{
 		checkSlow(IsInGameThread() || IsInAsyncLoadingThread());
 		GameThreadShaderMap = InShaderMap;
+		if (LIKELY(GameThreadShaderMap))
+		{
+			GameThreadShaderMap->GetResource()->SetOwnerName(GetOwnerFName());
+		}
 		bContainsInlineShaders = true;
 		bLoadedCookedShaderMapId = true;
 		CookedShaderMapId = InShaderMap->GetShaderMapId();
@@ -566,10 +580,8 @@ private:
 
 	FString FriendlyName;
 
-#if WITH_EDITOR
 	/** Asset using this resource */
 	FName AssetPath;
-#endif // WITH_EDITOR
 
 	friend class FOpenColorIOShaderMap;
 	friend class FShaderCompilingManager;
diff --git a/Engine/Plugins/FX/Niagara/Source/NiagaraShader/Private/NiagaraShared.cpp b/Engine/Plugins/FX/Niagara/Source/NiagaraShader/Private/NiagaraShared.cpp
--- a/Engine/Plugins/FX/Niagara/Source/NiagaraShader/Private/NiagaraShared.cpp
+++ b/Engine/Plugins/FX/Niagara/Source/NiagaraShader/Private/NiagaraShared.cpp
@@ -377,6 +377,7 @@ void FNiagaraShaderScript::SerializeShaderMap(FArchive& Ar)
 				if (FApp::CanEverRender() && bLoaded)
 				{
 					GameThreadShaderMap = RenderingThreadShaderMap = LoadedShaderMap;
+					GameThreadShaderMap->GetResource()->SetOwnerName(GetOwnerFName());
 
 					UpdateCachedData_PostCompile(true);
 				}
@@ -389,6 +390,12 @@ void FNiagaraShaderScript::SerializeShaderMap(FArchive& Ar)
 	}
 }
 
+FName FNiagaraShaderScript::GetOwnerFName() const
+{
+	const UNiagaraScriptBase* Owner = GetBaseVMScript();
+	return Owner ? Owner->GetOutermost()->GetFName() : NAME_None;
+}
+
 #if WITH_EDITOR
 void FNiagaraShaderScript::SaveShaderStableKeys(EShaderPlatform TargetShaderPlatform, FStableShaderKeyAndValue& SaveKeyVal)
 {
diff --git a/Engine/Plugins/FX/Niagara/Source/NiagaraShader/Public/NiagaraShared.h b/Engine/Plugins/FX/Niagara/Source/NiagaraShader/Public/NiagaraShared.h
--- a/Engine/Plugins/FX/Niagara/Source/NiagaraShader/Public/NiagaraShared.h
+++ b/Engine/Plugins/FX/Niagara/Source/NiagaraShader/Public/NiagaraShared.h
@@ -817,6 +817,10 @@ public:
 	{
 		checkSlow(IsInGameThread() || IsInAsyncLoadingThread());
 		GameThreadShaderMap = InShaderMap;
+		if (LIKELY(GameThreadShaderMap))
+		{
+			GameThreadShaderMap->GetResource()->SetOwnerName(GetOwnerFName());
+		}
 	}
 
 	/** Note: SetGameThreadShaderMap must also be called with the same value, but from the game thread. */
@@ -837,6 +841,10 @@ public:
 	{
 		checkSlow(IsInGameThread() || IsInAsyncLoadingThread());
 		GameThreadShaderMap = InShaderMap;
+		if (LIKELY(GameThreadShaderMap))
+		{
+			GameThreadShaderMap->GetResource()->SetOwnerName(GetOwnerFName());
+		}
 		bLoadedCookedShaderMapId = true;
 		CookedShaderMapId = InShaderMap->GetShaderMapId();
 	}
@@ -883,6 +891,8 @@ public:
 
 	UNiagaraScriptBase* GetBaseVMScript() { return BaseVMScript; }
 	const UNiagaraScriptBase* GetBaseVMScript() const { return BaseVMScript; }
+	/** Returns owner name for tracking */
+	NIAGARASHADER_API FName GetOwnerFName() const;
 
 	NIAGARASHADER_API  FNiagaraShaderRef GetShader(int32 PermutationId) const;
 	NIAGARASHADER_API  FNiagaraShaderRef GetShaderGameThread(int32 PermutationId) const;
diff --git a/Engine/Plugins/Runtime/ComputeFramework/Source/ComputeFramework/Private/ComputeFramework/ComputeKernelShared.h b/Engine/Plugins/Runtime/ComputeFramework/Source/ComputeFramework/Private/ComputeFramework/ComputeKernelShared.h
--- a/Engine/Plugins/Runtime/ComputeFramework/Source/ComputeFramework/Private/ComputeFramework/ComputeKernelShared.h
+++ b/Engine/Plugins/Runtime/ComputeFramework/Source/ComputeFramework/Private/ComputeFramework/ComputeKernelShared.h
@@ -374,11 +374,21 @@ public:
 		return GameThreadShaderMap;
 	}
 
+	/** Returns owner name for tracking */
+	FName GetOwnerFName() const
+	{
+		return AssetPath;
+	}
+
 	/** Note: SetRenderingThreadShaderMap must also be called with the same value, but from the rendering thread. */
 	void SetGameThreadShaderMap(FComputeKernelShaderMap* InShaderMap)
 	{
 		checkSlow(IsInGameThread() || IsInAsyncLoadingThread());
 		GameThreadShaderMap = InShaderMap;
+		if (LIKELY(GameThreadShaderMap))
+		{
+			GameThreadShaderMap->GetResource()->SetOwnerName(GetOwnerFName());
+		}
 	}
 
 	/** Note: SetGameThreadShaderMap must also be called with the same value, but from the game thread. */
@@ -393,6 +403,10 @@ public:
 	{
 		checkSlow(IsInGameThread() || IsInAsyncLoadingThread());
 		GameThreadShaderMap = InShaderMap;
+		if (LIKELY(GameThreadShaderMap))
+		{
+			GameThreadShaderMap->GetResource()->SetOwnerName(GetOwnerFName());
+		}
 		bContainsInlineShaders = true;
 		bLoadedCookedShaderMapId = true;
 		CookedShaderMapId = InShaderMap->GetShaderMapId();
@@ -555,10 +569,8 @@ private:
 
 	FString FriendlyName;
 
-#if WITH_EDITOR
 	/** Asset using this resource */
 	FName AssetPath;
-#endif
 
 	friend class FComputeKernelShaderMap;
 	friend class FShaderCompilingManager;
diff --git a/Engine/Plugins/Runtime/ComputeFramework/Source/ComputeFramework/Private/ComputeKernelShared.cpp b/Engine/Plugins/Runtime/ComputeFramework/Source/ComputeFramework/Private/ComputeKernelShared.cpp
--- a/Engine/Plugins/Runtime/ComputeFramework/Source/ComputeFramework/Private/ComputeKernelShared.cpp
+++ b/Engine/Plugins/Runtime/ComputeFramework/Source/ComputeFramework/Private/ComputeKernelShared.cpp
@@ -218,6 +218,7 @@ void FComputeKernelResource::SerializeShaderMap(FArchive& Ar)
 				if (bSuccessfullyLoaded && FApp::CanEverRender())
 				{
 					GameThreadShaderMap = RenderingThreadShaderMap = LoadedShaderMap;
+					GameThreadShaderMap->GetResource()->SetOwnerName(GetOwnerFName());
 				}
 			}
 		}
@@ -251,9 +252,7 @@ void FComputeKernelResource::SetupResource(
 	ShaderParameterMetadataAllocations = MoveTemp(InShaderParameterMetadataAllocations);
 	ShaderParameterMetadata = InShaderParameterMetadata;
 	CompilationResults.Messages.Reset();
-#if WITH_EDITOR
 	AssetPath = InAssetPath;
-#endif
 }
 
 int32 FComputeKernelResource::GetNumPermutations() const 
diff --git a/Engine/Source/Runtime/Engine/Private/Materials/MaterialShared.cpp b/Engine/Source/Runtime/Engine/Private/Materials/MaterialShared.cpp
--- a/Engine/Source/Runtime/Engine/Private/Materials/MaterialShared.cpp
+++ b/Engine/Source/Runtime/Engine/Private/Materials/MaterialShared.cpp
@@ -1122,6 +1122,10 @@ void FMaterial::SetGameThreadShaderMap(FMaterialShaderMap* InMaterialShaderMap)
 	checkSlow(IsInGameThread() || IsInAsyncLoadingThread());
 	const bool bIsComplete = InMaterialShaderMap ? InMaterialShaderMap->IsComplete(this, true) : false;
 	GameThreadShaderMap = InMaterialShaderMap;
+	if (LIKELY(GameThreadShaderMap))
+	{
+		GameThreadShaderMap->GetResource()->SetOwnerName(GetOwnerFName());
+	}
 	bGameThreadShaderMapIsComplete = bIsComplete;
 
 	TRefCountPtr<FMaterial> Material = this;
@@ -1155,6 +1159,7 @@ void FMaterial::SetInlineShaderMap(FMaterialShaderMap* InMaterialShaderMap)
 	check(InMaterialShaderMap);
 
 	GameThreadShaderMap = InMaterialShaderMap;
+	GameThreadShaderMap->GetResource()->SetOwnerName(GetOwnerFName());
 	bContainsInlineShaders = true;
 	bLoadedCookedShaderMapId = true;
 
@@ -1376,6 +1381,7 @@ void FMaterial::SerializeInlineShaderMap(FArchive& Ar)
 				if (LoadedShaderMap->Serialize(Ar, true, bCooked && Ar.IsLoading()))
 				{
 					GameThreadShaderMap = MoveTemp(LoadedShaderMap);
+					GameThreadShaderMap->GetResource()->SetOwnerName(GetOwnerFName());
 #if WITH_EDITOR
 					GameThreadShaderMap->AssociateWithAsset(GetAssetPath());
 #endif
@@ -1404,6 +1410,12 @@ void FMaterial::RegisterInlineShaderMap(bool bLoadedByCookedMaterial)
 	}
 }
 
+FName FMaterial::GetOwnerFName() const
+{
+	UMaterialInterface* Owner = GetMaterialInterface();
+	return Owner ? Owner->GetOutermost()->GetFName() : NAME_None;
+}
+
 void FMaterialResource::LegacySerialize(FArchive& Ar)
 {
 	FMaterial::LegacySerialize(Ar);
@@ -4844,7 +4856,15 @@ void FMaterial::SetShaderMapsOnMaterialResources(const TMap<TRefCountPtr<FMateri
 		FMaterial* Material = It.Key;
 		const TRefCountPtr<FMaterialShaderMap>& ShaderMap = It.Value;
 		Material->GameThreadShaderMap = ShaderMap;
-		Material->bGameThreadShaderMapIsComplete = ShaderMap ? ShaderMap->IsComplete(Material, true) : false;
+		if (LIKELY(Material->GameThreadShaderMap))
+		{
+			Material->GameThreadShaderMap->GetResource()->SetOwnerName(Material->GetOwnerFName());
+			Material->bGameThreadShaderMapIsComplete = ShaderMap->IsComplete(Material, true);
+		}
+		else
+		{
+			Material->bGameThreadShaderMapIsComplete = false;
+		}
 	}
 
 	ENQUEUE_RENDER_COMMAND(FSetShaderMapOnMaterialResources)(
diff --git a/Engine/Source/Runtime/Engine/Public/MaterialShared.h b/Engine/Source/Runtime/Engine/Public/MaterialShared.h
--- a/Engine/Source/Runtime/Engine/Public/MaterialShared.h
+++ b/Engine/Source/Runtime/Engine/Public/MaterialShared.h
@@ -1953,6 +1953,8 @@ public:
 	 */
 	virtual bool IsPersistent() const = 0;
 	virtual UMaterialInterface* GetMaterialInterface() const { return NULL; }
+	/** For tracking ownership data */
+	FName GetOwnerFName() const;
 
 	virtual bool IsPreview() const { return false; }
 
diff --git a/Engine/Source/Runtime/RenderCore/Private/ShaderCodeLibrary.cpp b/Engine/Source/Runtime/RenderCore/Private/ShaderCodeLibrary.cpp
--- a/Engine/Source/Runtime/RenderCore/Private/ShaderCodeLibrary.cpp
+++ b/Engine/Source/Runtime/RenderCore/Private/ShaderCodeLibrary.cpp
@@ -1228,10 +1228,11 @@ void FShaderMapResource_SharedCode::ReleaseRHI()
 	// on assumption that we aren't going to get resurrected
 	LibraryInstance = nullptr;
 
-	if (GetNumRefs()> 0)
+	if (GetNumRefs() > 0)
 	{
 		ensureMsgf(false, TEXT("FShaderMapResource_SharedCode::ReleaseRHI is still referenced (Num of references %d). Invoking OnSharedShaderMapResourceExplicitRelease delegate."), GetNumRefs());
-		UE_LOG(LogShaderLibrary, Warning, TEXT("FShaderMapResource_SharedCode::ReleaseRHI is still referenced (Num of references %d). Invoking OnSharedShaderMapResourceExplicitRelease delegate."), GetNumRefs());
+		UE_LOG(LogShaderLibrary, Warning, TEXT("FShaderMapResource_SharedCode::ReleaseRHI is still referenced (Num of references %d, owner %s). Invoking OnSharedShaderMapResourceExplicitRelease delegate."), 
+			GetNumRefs(), *GetOwnerName().ToString());
 		
 		// Invoke delegate to notify shader map resource needs to be forced released
 		OnSharedShaderMapResourceExplicitRelease.ExecuteIfBound(this);
