Correctly re-add persitent linked layers when we reuse them

#tests PIE
#rnx

diff --git a/Engine/Source/Runtime/Engine/Private/Animation/AnimSubsystem_SharedLinkedAnimLayers.cpp b/Engine/Source/Runtime/Engine/Private/Animation/AnimSubsystem_SharedLinkedAnimLayers.cpp
--- a/Engine/Source/Runtime/Engine/Private/Animation/AnimSubsystem_SharedLinkedAnimLayers.cpp
+++ b/Engine/Source/Runtime/Engine/Private/Animation/AnimSubsystem_SharedLinkedAnimLayers.cpp
@@ -92,6 +92,14 @@ UAnimInstance* FLinkedAnimLayerClassData::FindOrAddInstanceForLinking(UAnimInsta
 		// Check if function is already linked
 		if (!LayerInstanceData.GetLinkedFunctions().Contains(Function))
 		{
+			// Re-add persistent instance of first function re-bind
+			if (LayerInstanceData.IsPersistent() && LayerInstanceData.GetLinkedFunctions().Num() == 0)
+			{
+				USkeletalMeshComponent* Mesh = OwningInstance->GetSkelMeshComponent();
+				check(!Mesh->GetLinkedAnimInstances().Contains(LayerInstanceData.Instance));
+				Mesh->GetLinkedAnimInstances().Add(LayerInstanceData.Instance);
+			}
+
 			bIsNewInstance = false;
 			LayerInstanceData.AddLinkedFunction(Function, LayerInstanceData.Instance);
 			// not linked, use this instance
