Fix RHI validation layer error in SkeletalRenderGPUSkin.cpp
 - MorphVertexBuffer was not in UAV state before being cleared with RHIClearUAVUint.

Fix RDGImmediate not uploading buffers if no render passes were enqueued.
 - Also remove stray debug string compare / debug break

diff --git a/Engine/Source/Runtime/Engine/Private/SkeletalRenderGPUSkin.cpp b/Engine/Source/Runtime/Engine/Private/SkeletalRenderGPUSkin.cpp
--- a/Engine/Source/Runtime/Engine/Private/SkeletalRenderGPUSkin.cpp
+++ b/Engine/Source/Runtime/Engine/Private/SkeletalRenderGPUSkin.cpp
@@ -610,8 +610,9 @@ void FSkeletalMeshObjectGPUSkin::ProcessUpdatedDynamicData(EGPUSkinCacheEntryMod
 		QUICK_SCOPE_CYCLE_COUNTER(STAT_FSkeletalMeshObjectGPUSkin_ProcessUpdatedDynamicData_ClearMorphBuffer);
 		if (MorphVertexBuffer.GetUAV())
 		{
+			RHICmdList.Transition(FRHITransitionInfo(MorphVertexBuffer.GetUAV(), ERHIAccess::Unknown, ERHIAccess::UAVCompute));
 			RHICmdList.ClearUAVUint(MorphVertexBuffer.GetUAV(), FUintVector4(0, 0, 0, 0));
-			RHICmdList.Transition(FRHITransitionInfo(MorphVertexBuffer.GetUAV(), ERHIAccess::Unknown, ERHIAccess::SRVMask));
+			RHICmdList.Transition(FRHITransitionInfo(MorphVertexBuffer.GetUAV(), ERHIAccess::UAVCompute, ERHIAccess::SRVMask));
 		}
 	}
 	MorphVertexBuffer.bNeedsInitialClear = false;
diff --git a/Engine/Source/Runtime/RenderCore/Private/RenderGraphBuilder.cpp b/Engine/Source/Runtime/RenderCore/Private/RenderGraphBuilder.cpp
--- a/Engine/Source/Runtime/RenderCore/Private/RenderGraphBuilder.cpp
+++ b/Engine/Source/Runtime/RenderCore/Private/RenderGraphBuilder.cpp
@@ -1850,6 +1850,12 @@ void FRDGBuilder::Execute()
 			}
 		}
 	}
+	else
+	{
+		PrepareBufferUploads();
+		SubmitBufferUploads();
+		CreateUniformBuffers();
+	}
 
 	{
 		SCOPED_NAMED_EVENT_TEXT("FRDGBuilder::Finalize", FColor::Magenta);
@@ -2245,11 +2251,6 @@ void FRDGBuilder::SetupPassResources(FRDGPass* Pass)
 			return;
 		}
 
-		if (!FCString::Stricmp(Buffer->Name, TEXT("DefaultBuffer")))
-		{
-			UE_DEBUG_BREAK();
-		}
-
 		const FRDGViewHandle NoUAVBarrierHandle = GetHandleIfNoUAVBarrier(BufferView);
 
 		FRDGPass::FBufferState* PassState;
