Adding the ability to force garbage collection if the engine is close to reaching UObject count limit

#tests Introduced 10K per frame UObject leak on the server and observed the new code behavior

diff --git a/Engine/Source/Runtime/CoreUObject/Public/UObject/UObjectArray.h b/Engine/Source/Runtime/CoreUObject/Public/UObject/UObjectArray.h
--- a/Engine/Source/Runtime/CoreUObject/Public/UObject/UObjectArray.h
+++ b/Engine/Source/Runtime/CoreUObject/Public/UObject/UObjectArray.h
@@ -918,6 +918,22 @@ public:
 	}
 #endif
 
+	/**
+	* Returns the estimated number of object indices available for allocation
+	*/
+	int32 GetObjectArrayCapacity() const
+	{
+		return ObjObjects.Capacity();
+	}
+
+	/**
+	* Returns the estimated number of object indices available for allocation
+	*/
+	int32 GetObjectArrayNumAvailable() const
+	{
+		return ObjObjects.Capacity() - ObjObjects.Num() + ObjAvailableList.Num();
+	}
+
 	/**
 	 * Clears some internal arrays to get rid of false memory leaks
 	 */
diff --git a/Engine/Source/Runtime/Engine/Classes/Engine/Engine.h b/Engine/Source/Runtime/Engine/Classes/Engine/Engine.h
--- a/Engine/Source/Runtime/Engine/Classes/Engine/Engine.h
+++ b/Engine/Source/Runtime/Engine/Classes/Engine/Engine.h
@@ -705,6 +705,16 @@ DECLARE_MULTICAST_DELEGATE_OneParam(FPreRenderDelegateEx, class FRDGBuilder&);
 DECLARE_MULTICAST_DELEGATE(FPostRenderDelegate);
 DECLARE_MULTICAST_DELEGATE_OneParam(FPostRenderDelegateEx, class FRDGBuilder&);
 
+/**
+ * Type of UObject purge type to be performed by the engine
+ */
+enum class EGarbageCollectionType
+{
+	None,
+	Incremental,
+	Full
+};
+
 /**
  * Abstract base class of all Engine classes, responsible for management of systems critical to editor or game systems.
  * Also defines default classes for certain engine systems.
@@ -2670,6 +2680,14 @@ protected:
 	  */
 	 static void SendWorldEndOfFrameUpdates();
 
+	 /**
+	  * Allows derived classes to force garbage collection based on various factors (low on available UObject slots / other resources)
+	  */
+	 virtual EGarbageCollectionType ShouldForceGarbageCollection()
+	 {
+		 return EGarbageCollectionType::None;
+	 }
+
 public:
 	/** @return the GIsEditor flag setting */
 	bool IsEditor();
diff --git a/Engine/Source/Runtime/Engine/Private/UnrealEngine.cpp b/Engine/Source/Runtime/Engine/Private/UnrealEngine.cpp
--- a/Engine/Source/Runtime/Engine/Private/UnrealEngine.cpp
+++ b/Engine/Source/Runtime/Engine/Private/UnrealEngine.cpp
@@ -1742,6 +1742,13 @@ void UEngine::ConditionalCollectGarbage()
 		}
 		else
 #endif
+		{
+			EGarbageCollectionType ForceTriggerPurge = ShouldForceGarbageCollection();
+			if (ForceTriggerPurge != EGarbageCollectionType::None)
+			{
+				ForceGarbageCollection(ForceTriggerPurge == EGarbageCollectionType::Full);
+			}
+
 			if (bFullPurgeTriggered)
 			{
 				if (TryCollectGarbage(GARBAGE_COLLECTION_KEEPFLAGS, true))
@@ -1822,6 +1829,7 @@ void UEngine::ConditionalCollectGarbage()
 					}
 				}
 			}
+		}
 
 		if (CVarCollectGarbageEveryFrame.GetValueOnGameThread() > 0)
 		{
