Prevent hazard when FRenderAssetStreamingManager::LevelRenderAssetManagers array is modified during recursion

diff --git a/Engine/Source/Runtime/Engine/Private/Streaming/StreamingManagerTexture.cpp b/Engine/Source/Runtime/Engine/Private/Streaming/StreamingManagerTexture.cpp
--- a/Engine/Source/Runtime/Engine/Private/Streaming/StreamingManagerTexture.cpp
+++ b/Engine/Source/Runtime/Engine/Private/Streaming/StreamingManagerTexture.cpp
@@ -103,6 +103,7 @@ FRenderAssetStreamingManager::FRenderAssetStreamingManager()
 :	CurrentUpdateStreamingRenderAssetIndex(0)
 ,	AsyncWork( nullptr )
 ,	CurrentPendingMipCopyRequestIdx(0)
+,	LevelRenderAssetManagersLock(nullptr)
 ,	ProcessingStage( 0 )
 ,	NumRenderAssetProcessingStages(5)
 ,	bUseDynamicStreaming( false )
@@ -237,6 +238,7 @@ void FRenderAssetStreamingManager::OnPreGarbageCollect()
 	FRemovedRenderAssetArray RemovedRenderAssets;
 
 	// Check all levels for pending kills.
+	check(!LevelRenderAssetManagersLock);
 	for (int32 Index = 0; Index < LevelRenderAssetManagers.Num(); ++Index)
 	{
 		if (LevelRenderAssetManagers[Index] == nullptr)
@@ -407,11 +409,16 @@ void FRenderAssetStreamingManager::IncrementalUpdate(float Percentage, bool bUpd
 		NumStepsLeftForIncrementalBuild = MAX_int64;
 	}
 
-	for (FLevelRenderAssetManager* LevelManager : LevelRenderAssetManagers)
 	{
-		if (LevelManager != nullptr)
+		// Prevent hazard if levels are added or removed during iteration
+		FScopedLevelRenderAssetManagersLock ScopedLevelRenderAssetManagersLock(this);
+
+		for (FLevelRenderAssetManager* LevelManager : LevelRenderAssetManagers)
 		{
-			LevelManager->IncrementalUpdate(DynamicComponentManager, RemovedRenderAssets, NumStepsLeftForIncrementalBuild, Percentage, bUseDynamicStreaming); // Complete the incremental update.
+			if (LevelManager != nullptr)
+			{
+				LevelManager->IncrementalUpdate(DynamicComponentManager, RemovedRenderAssets, NumStepsLeftForIncrementalBuild, Percentage, bUseDynamicStreaming); // Complete the incremental update.
+			}
 		}
 	}
 
@@ -574,6 +581,7 @@ void FRenderAssetStreamingManager::ConditionalUpdateStaticData()
 			TArray<ULevel*, TInlineAllocator<32> > Levels;
 
 			// RemoveLevel data
+			check(!LevelRenderAssetManagersLock);
 			for (FLevelRenderAssetManager* LevelManager : LevelRenderAssetManagers)
 			{
 				if (LevelManager!=nullptr)
@@ -740,6 +748,36 @@ void FRenderAssetStreamingManager::BoostTextures( AActor* Actor, float BoostFact
 	}
 }
 
+FRenderAssetStreamingManager::FScopedLevelRenderAssetManagersLock::FScopedLevelRenderAssetManagersLock(FRenderAssetStreamingManager* InStreamingManager)
+	: StreamingManager(InStreamingManager)
+{
+	check(!StreamingManager->LevelRenderAssetManagersLock);
+	StreamingManager->LevelRenderAssetManagersLock = this;
+}
+
+FRenderAssetStreamingManager::FScopedLevelRenderAssetManagersLock::~FScopedLevelRenderAssetManagersLock()
+{
+	StreamingManager->ProcessPendingLevelManagers();
+	StreamingManager->LevelRenderAssetManagersLock = nullptr;
+}
+
+void FRenderAssetStreamingManager::ProcessPendingLevelManagers()
+{
+	check(LevelRenderAssetManagersLock);
+
+	LevelRenderAssetManagers.Append(LevelRenderAssetManagersLock->PendingAddLevelManagers);
+
+	FRemovedRenderAssetArray RemovedRenderAssets;
+
+	for (FLevelRenderAssetManager* LevelManager : LevelRenderAssetManagersLock->PendingRemoveLevelManagers)
+	{
+		LevelManager->Remove(&RemovedRenderAssets);
+		// Delete the level manager. The async task view will still be valid as it is ref-counted.
+		delete LevelManager;
+	}
+	SetRenderAssetsRemovedTimestamp(RemovedRenderAssets);
+}
+
 /** Adds a ULevel to the streaming manager. This is called from 2 paths : after PostPostLoad and after AddToWorld */
 void FRenderAssetStreamingManager::AddLevel( ULevel* Level )
 {
@@ -764,6 +802,19 @@ void FRenderAssetStreamingManager::AddLevel( ULevel* Level )
 				return;
 			}
 		}
+
+		if (LevelRenderAssetManagersLock)
+		{
+			for (const FLevelRenderAssetManager* LevelManager : LevelRenderAssetManagersLock->PendingAddLevelManagers)
+			{
+				check(LevelManager);
+				if (LevelManager->GetLevel() == Level)
+				{
+					// Nothing to do, since the incremental update automatically manages what needs to be done.
+					return;
+				}
+			}
+		}
 	}
 
 	// If the level was not already there, create a new one, find an available slot or add a new one.
@@ -775,6 +826,11 @@ void FRenderAssetStreamingManager::AddLevel( ULevel* Level )
 	{
 		LevelRenderAssetManagers[LevelIndex] = LevelRenderAssetManager;
 	}
+	else if (LevelRenderAssetManagersLock)
+	{
+		// Cannot add during recursion as it may cause the array to realloc and invalidate the iterator
+		LevelRenderAssetManagersLock->PendingAddLevelManagers.Add(LevelRenderAssetManager);
+	}
 	else
 	{
 		LevelRenderAssetManagers.Add(LevelRenderAssetManager);
@@ -797,13 +853,22 @@ void FRenderAssetStreamingManager::RemoveLevel( ULevel* Level )
 			FLevelRenderAssetManager* LevelManager = LevelRenderAssetManagers[Index];
 			if (LevelManager!=nullptr && LevelManager->GetLevel() == Level)
 			{
-				FRemovedRenderAssetArray RemovedRenderAssets;
-				LevelManager->Remove(&RemovedRenderAssets);
-				SetRenderAssetsRemovedTimestamp(RemovedRenderAssets);
-
-				// Remove the level entry. The async task view will still be valid as it uses a shared ptr.
 				LevelRenderAssetManagers[Index] = nullptr;
-				delete LevelManager;
+
+				if (LevelRenderAssetManagersLock)
+				{
+					// Do not delete during recursion in case the level manager is updating its internal states
+					LevelRenderAssetManagersLock->PendingRemoveLevelManagers.Add(LevelManager);
+				}
+				else
+				{
+					FRemovedRenderAssetArray RemovedRenderAssets;
+					LevelManager->Remove(&RemovedRenderAssets);
+					SetRenderAssetsRemovedTimestamp(RemovedRenderAssets);
+
+					// Delete the level manager. The async task view will still be valid as it is ref-counted.
+					delete LevelManager;
+				}
 				break;
 			}
 		}
@@ -918,6 +983,7 @@ void FRenderAssetStreamingManager::NotifyActorDestroyed( AActor* Actor )
 	ULevel* Level = !GIsEditor ? Actor->GetLevel() : nullptr;
 
 	// Remove any reference in the level managers.
+	check(!LevelRenderAssetManagersLock);
 	for (FLevelRenderAssetManager* LevelManager : LevelRenderAssetManagers)
 	{
 		if (LevelManager!=nullptr && (!Level || LevelManager->GetLevel() == Level))
@@ -953,6 +1019,7 @@ void FRenderAssetStreamingManager::RemoveStaticReferences(const UPrimitiveCompon
 	{
 		FRemovedRenderAssetArray RemovedRenderAssets;
 		ULevel* Level = Primitive->GetComponentLevel();
+		check(!LevelRenderAssetManagersLock);
 		for (FLevelRenderAssetManager* LevelManager : LevelRenderAssetManagers)
 		{
 			if (LevelManager != nullptr && (!Level || LevelManager->GetLevel() == Level))
@@ -998,6 +1065,7 @@ void FRenderAssetStreamingManager::NotifyPrimitiveDetached( const UPrimitiveComp
 		// Unless in editor, we don't want to remove reference in static level data when toggling visibility.
 		else if (GIsEditor || !IsValid(Primitive) || Primitive->HasAnyFlags(RF_BeginDestroyed|RF_FinishDestroyed))
 		{
+			check(!LevelRenderAssetManagersLock);
 			for (FLevelRenderAssetManager* LevelManager : LevelRenderAssetManagers)
 			{
 				if (LevelManager != nullptr && (!Level || LevelManager->GetLevel() == Level))
@@ -1972,6 +2040,7 @@ void FRenderAssetStreamingManager::PropagateLightingScenarioChange()
 	// Note that dynamic components don't need to be handled because their renderstates are updated, which triggers and update.
 	
 	TArray<ULevel*, TInlineAllocator<32> > Levels;
+	check(!LevelRenderAssetManagersLock);
 	for (FLevelRenderAssetManager* LevelManager : LevelRenderAssetManagers)
 	{
 		if (LevelManager!=nullptr)
diff --git a/Engine/Source/Runtime/Engine/Private/Streaming/StreamingManagerTexture.h b/Engine/Source/Runtime/Engine/Private/Streaming/StreamingManagerTexture.h
--- a/Engine/Source/Runtime/Engine/Private/Streaming/StreamingManagerTexture.h
+++ b/Engine/Source/Runtime/Engine/Private/Streaming/StreamingManagerTexture.h
@@ -309,6 +309,8 @@ private:
 	 */
 	void TickDeferredMipLevelChangeCallbacks();
 
+	void ProcessPendingLevelManagers();
+
 	/** Cached from the system settings. */
 	int32 NumStreamedMips_Texture[TEXTUREGROUP_MAX];
 	TArray<int32> NumStreamedMips_StaticMesh;
@@ -362,6 +364,22 @@ private:
 	/** Level data */
 	TArray<FLevelRenderAssetManager*> LevelRenderAssetManagers;
 
+	// Used to prevent hazard when LevelRenderAssetManagers array is modified through recursion
+	struct FScopedLevelRenderAssetManagersLock
+	{
+		FRenderAssetStreamingManager* StreamingManager;
+		TArray<FLevelRenderAssetManager*> PendingAddLevelManagers;
+		TArray<FLevelRenderAssetManager*> PendingRemoveLevelManagers;
+
+		FScopedLevelRenderAssetManagersLock(FRenderAssetStreamingManager* InStreamingManager);
+
+		~FScopedLevelRenderAssetManagersLock();
+	};
+
+	friend struct FScopedLevelRenderAssetManagersLock;
+
+	FScopedLevelRenderAssetManagersLock* LevelRenderAssetManagersLock;
+
 	/** Stages [0,N-2] is non-threaded data collection, Stage N-1 is wait-for-AsyncWork-and-finalize. */
 	int32					ProcessingStage;
 
