Fix allowing child classes to track additional stats.

diff --git a/Engine/Source/Runtime/Engine/Private/ProfilingDebugging/LevelStreamingProfilingSubsystem.cpp b/Engine/Source/Runtime/Engine/Private/ProfilingDebugging/LevelStreamingProfilingSubsystem.cpp
--- a/Engine/Source/Runtime/Engine/Private/ProfilingDebugging/LevelStreamingProfilingSubsystem.cpp
+++ b/Engine/Source/Runtime/Engine/Private/ProfilingDebugging/LevelStreamingProfilingSubsystem.cpp
@@ -129,6 +129,15 @@ TUniquePtr<ULevelStreamingProfilingSubsystem::FActiveLevel> ULevelStreamingProfi
 	return MoveTemp(Level);
 }
 
+void ULevelStreamingProfilingSubsystem::UpdateLevelState(ULevelStreamingProfilingSubsystem::FActiveLevel& Level, const ULevelStreaming* StreamingLevel, ULevelStreamingProfilingSubsystem::ELevelState NewState, double Time)
+{	
+	ELevelState OldState = Level.State;
+	Level.State = NewState;
+	Level.StateStartTime = Time;
+
+	UpdateTrackingData(Level.StatsIndex, LevelStats[Level.StatsIndex], StreamingLevel, OldState, NewState);
+}
+
 void ULevelStreamingProfilingSubsystem::StartTracking()
 {
 	if (!ReportWritingTask.IsCompleted())
@@ -452,8 +461,7 @@ void ULevelStreamingProfilingSubsystem::OnLevelStartedAsyncLoading(UWorld* World
 	double Time = FPlatformTime::Seconds();
 	LevelStats[Level->StatsIndex].bValid = true;
 	LevelStats[Level->StatsIndex].TimeQueuedForLoading = Time - Level->StateStartTime;
-	Level->State = ELevelState::Loading;
-	Level->StateStartTime = FPlatformTime::Seconds();
+	UpdateLevelState(*Level, StreamingLevel, ELevelState::Loading, FPlatformTime::Seconds());
 }
 
 void ULevelStreamingProfilingSubsystem::OnLevelFinishedAsyncLoading(UWorld* World, const ULevelStreaming* StreamingLevel, ULevel* LoadedLevel)
@@ -479,8 +487,7 @@ void ULevelStreamingProfilingSubsystem::OnLevelFinishedAsyncLoading(UWorld* Worl
 		// Fill in stats so we calculate the time til now as queued time, and loading time as 0
 		LevelStats[Level->StatsIndex].bValid = true;
 		LevelStats[Level->StatsIndex].TimeQueuedForLoading = Time - Level->StateStartTime;
-		Level->State = ELevelState::Loading;
-		Level->StateStartTime = Time;
+		UpdateLevelState(*Level, StreamingLevel, ELevelState::Loading, Time);
 	}
 	else if (Level->State != ELevelState::Loading)
 	{
@@ -495,8 +502,7 @@ void ULevelStreamingProfilingSubsystem::OnLevelFinishedAsyncLoading(UWorld* Worl
 			*StreamingLevel->GetWorldAssetPackageName());
 		LevelStats[Level->StatsIndex].TimeLoading = Time - Level->StateStartTime;
 		LevelStats[Level->StatsIndex].bValid = true; // Record this as a cancelled load
-		Level->State = ELevelState::None;
-		Level->StateStartTime = Time;
+		UpdateLevelState(*Level, StreamingLevel, ELevelState::None, Time);
 		return;
 	}
 
@@ -507,8 +513,7 @@ void ULevelStreamingProfilingSubsystem::OnLevelFinishedAsyncLoading(UWorld* Worl
 	LevelStats[Level->StatsIndex].PackageNameInMemory = LoadedLevel->GetOutermost()->GetFName();
 	LevelStats[Level->StatsIndex].bValid = true;
 	LevelStats[Level->StatsIndex].TimeLoading = Time - Level->StateStartTime;
-	Level->State = ELevelState::Loaded;
-	Level->StateStartTime = Time;
+	UpdateLevelState(*Level, StreamingLevel, ELevelState::Loaded, Time);
 }
 
 void ULevelStreamingProfilingSubsystem::OnLevelQueuedForAddToWorld(UWorld* World, const ULevelStreaming* StreamingLevel, ULevel* LoadedLevel)
@@ -550,8 +555,7 @@ void ULevelStreamingProfilingSubsystem::OnLevelQueuedForAddToWorld(UWorld* World
 		*StreamingLevel->GetWorldAssetPackageName());
 	double Time = FPlatformTime::Seconds();
 	LevelStats[Level->StatsIndex].bValid = true;
-	Level->State = ELevelState::QueuedForAddToWorld;
-	Level->StateStartTime = Time;
+	UpdateLevelState(*Level, StreamingLevel, ELevelState::QueuedForAddToWorld, Time);
 }
 
 void ULevelStreamingProfilingSubsystem::OnLevelUnqueuedForAddToWorld(UWorld* World, const ULevelStreaming* StreamingLevel, ULevel* LoadedLevel)
@@ -581,8 +585,7 @@ void ULevelStreamingProfilingSubsystem::OnLevelUnqueuedForAddToWorld(UWorld* Wor
 		*StreamingLevel->GetWorldAssetPackageName());
 	double Time = FPlatformTime::Seconds();
 	LevelStats[Level->StatsIndex].bValid = true;
-	Level->State = ELevelState::Loaded;
-	Level->StateStartTime = Time;
+	UpdateLevelState(*Level, StreamingLevel, ELevelState::Loaded, Time);
 }
 
 void ULevelStreamingProfilingSubsystem::OnLevelStartedAddToWorld(UWorld* World, const ULevelStreaming* StreamingLevel, ULevel* LoadedLevel)
@@ -613,8 +616,7 @@ void ULevelStreamingProfilingSubsystem::OnLevelStartedAddToWorld(UWorld* World,
 	double Time = FPlatformTime::Seconds();
 	LevelStats[Level->StatsIndex].bValid = true;
 	LevelStats[Level->StatsIndex].TimeQueueudForAddToWorld = Time - Level->StateStartTime; 
-	Level->State = ELevelState::AddingToWorld;
-	Level->StateStartTime = Time;
+	UpdateLevelState(*Level, StreamingLevel, ELevelState::AddingToWorld, Time);
 }
 
 void ULevelStreamingProfilingSubsystem::OnLevelFinishedAddToWorld(UWorld* World, const ULevelStreaming* StreamingLevel, ULevel* LoadedLevel)
@@ -647,16 +649,12 @@ void ULevelStreamingProfilingSubsystem::OnLevelFinishedAddToWorld(UWorld* World,
 	double Time = FPlatformTime::Seconds();
 	LevelStats[Level->StatsIndex].bValid = true;
 	LevelStats[Level->StatsIndex].TimeAddingToWorld = Time - Level->StateStartTime;
-	Level->State = ELevelState::AddedToWorld;
-	Level->StateStartTime = Time;
 	Level->TimeAddedToWorld = Time;
 
 	FBox CellBounds = LevelStats[Level->StatsIndex].CellBounds;
 	if (CellBounds.IsValid)
 	{
 		double MinDistance = MAX_dbl;
-
-
 		MinDistance = MAX_dbl;
 		FVector Location;
 		for (const FWorldPartitionStreamingSource& Source : World->GetWorldPartition()->GetStreamingSources())
@@ -693,6 +691,7 @@ void ULevelStreamingProfilingSubsystem::OnLevelFinishedAddToWorld(UWorld* World,
 			LevelStats[Level->StatsIndex].FinalStreamInDistance_Content = FMath::Sqrt(MinDistance);
 		}
 	}
+	UpdateLevelState(*Level, StreamingLevel, ELevelState::AddedToWorld, Time);
 }
 
 void ULevelStreamingProfilingSubsystem::OnLevelQueuedForRemoveFromWorld(UWorld* World, const ULevelStreaming* StreamingLevel, ULevel* LoadedLevel)
@@ -722,8 +721,7 @@ void ULevelStreamingProfilingSubsystem::OnLevelQueuedForRemoveFromWorld(UWorld*
 		*StreamingLevel->GetWorldAssetPackageName());
 	double Time = FPlatformTime::Seconds();
 	LevelStats[Level->StatsIndex].bValid = true;
-	Level->State = ELevelState::QueuedForRemoveFromWorld;
-	Level->StateStartTime = Time;
+	UpdateLevelState(*Level, StreamingLevel, ELevelState::QueuedForRemoveFromWorld, Time);
 }
 
 void ULevelStreamingProfilingSubsystem::OnLevelUnqueuedForRemoveFromWorld(UWorld* World, const ULevelStreaming* StreamingLevel, ULevel* LoadedLevel)
@@ -753,8 +751,7 @@ void ULevelStreamingProfilingSubsystem::OnLevelUnqueuedForRemoveFromWorld(UWorld
 		*StreamingLevel->GetWorldAssetPackageName());
 	double Time = FPlatformTime::Seconds();
 	LevelStats[Level->StatsIndex].bValid = true;
-	Level->State = ELevelState::AddedToWorld;
-	Level->StateStartTime = Time;
+	UpdateLevelState(*Level, StreamingLevel, ELevelState::AddedToWorld, Time);
 }
 
 void ULevelStreamingProfilingSubsystem::OnLevelStartedRemoveFromWorld(UWorld* World, const ULevelStreaming* StreamingLevel, ULevel*)
@@ -787,8 +784,7 @@ void ULevelStreamingProfilingSubsystem::OnLevelStartedRemoveFromWorld(UWorld* Wo
 	{
 		LevelStats[Level->StatsIndex].TimeInWorld = Time - Level->TimeAddedToWorld.GetValue(); // Not really a perf thing but may give an idea of churn 
 	}
-	Level->State = ELevelState::RemovingFromWorld;
-	Level->StateStartTime = Time;
+	UpdateLevelState(*Level, StreamingLevel, ELevelState::RemovingFromWorld, Time);
 }
 
 void ULevelStreamingProfilingSubsystem::OnLevelFinishedRemoveFromWorld(UWorld* World, const ULevelStreaming* StreamingLevel, ULevel*)
@@ -817,7 +813,6 @@ void ULevelStreamingProfilingSubsystem::OnLevelFinishedRemoveFromWorld(UWorld* W
 	double Time = FPlatformTime::Seconds();
 	LevelStats[Level->StatsIndex].bValid = true;
 	LevelStats[Level->StatsIndex].TimeRemovingFromWorld = Time - Level->StateStartTime; 
-	Level->State = ELevelState::RemovedFromWorld;
-	Level->StateStartTime = Time;
+	UpdateLevelState(*Level, StreamingLevel, ELevelState::RemovedFromWorld, Time);
 }
 #endif // !UE_BUILD_SHIPPING
diff --git a/Engine/Source/Runtime/Engine/Public/ProfilingDebugging/LevelStreamingProfilingSubsystem.h b/Engine/Source/Runtime/Engine/Public/ProfilingDebugging/LevelStreamingProfilingSubsystem.h
--- a/Engine/Source/Runtime/Engine/Public/ProfilingDebugging/LevelStreamingProfilingSubsystem.h
+++ b/Engine/Source/Runtime/Engine/Public/ProfilingDebugging/LevelStreamingProfilingSubsystem.h
@@ -63,6 +63,7 @@ protected:
 
 #if !UE_BUILD_SHIPPING
 	enum class ELevelState;
+	struct FLevelStats;
 
 	/** 
 	 * Gives child classes an opportunity to track additional data along with the main report.
@@ -75,6 +76,7 @@ protected:
 	 */
 	ENGINE_API virtual void UpdateTrackingData(
 		int32 TrackingIndex, 
+		FLevelStats& BaseStats,
 		const ULevelStreaming* StreamingLevel, 
 		ELevelState PreviousState, 
 		ELevelState NewState) {}
@@ -96,7 +98,6 @@ protected:
 
 	ENGINE_API static const TCHAR* EnumToString(ULevelStreamingProfilingSubsystem::ELevelState State);
 
-	struct FLevelStats;
 	ENGINE_API TConstArrayView<FLevelStats> GetLevelStats() const;
 
 	enum class ELevelState
@@ -178,6 +179,7 @@ private:
 	void OnStreamingLevelRemoved(UWorld* World, const ULevelStreaming* StreamingLevel);
 
 	TUniquePtr<FActiveLevel> MakeActiveLevel(const ULevelStreaming* StreamingLevel, ELevelState InitialState, ULevel* LoadedLevel=nullptr);
+	void UpdateLevelState(FActiveLevel& Level, const ULevelStreaming* StreamingLevel, ELevelState NewState, double Time);
 
 	// Handles for registered callbacks
 	FDelegateHandle Handle_OnLevelStreamingTargetStateChanged;
