Adjusted debug camera options for the editor viewport:
-- Reduced the speed scale adjustment value from 0.5 to 0.05 to allow for more fine-grained adjustment.
-- Bound FOV increase and decrease to IE_Repeat as well as IE_Pressed so that the key can be held down to rapidly adjust the FOV.

#rnx

diff --git a/Engine/Source/Runtime/Engine/Private/DebugCameraController.cpp b/Engine/Source/Runtime/Engine/Private/DebugCameraController.cpp
--- a/Engine/Source/Runtime/Engine/Private/DebugCameraController.cpp
+++ b/Engine/Source/Runtime/Engine/Private/DebugCameraController.cpp
@@ -26,7 +26,7 @@
 
 #include UE_INLINE_GENERATED_CPP_BY_NAME(DebugCameraController)
 
-static const float SPEED_SCALE_ADJUSTMENT = 0.5f;
+static const float SPEED_SCALE_ADJUSTMENT = 0.05f;
 static const float MIN_ORBIT_RADIUS = 30.0f;
 
 ADebugCameraController::ADebugCameraController(const FObjectInitializer& ObjectInitializer)
@@ -130,7 +130,9 @@ void ADebugCameraController::SetupInputComponent()
 	InputComponent->BindAction("DebugCamera_DecreaseSpeed", IE_Pressed, this, &ADebugCameraController::DecreaseCameraSpeed);
 
 	InputComponent->BindAction("DebugCamera_IncreaseFOV", IE_Pressed, this, &ADebugCameraController::IncreaseFOV);
+	InputComponent->BindAction("DebugCamera_IncreaseFOV", IE_Repeat, this, &ADebugCameraController::IncreaseFOV);
 	InputComponent->BindAction("DebugCamera_DecreaseFOV", IE_Pressed, this, &ADebugCameraController::DecreaseFOV);
+	InputComponent->BindAction("DebugCamera_DecreaseFOV", IE_Repeat, this, &ADebugCameraController::DecreaseFOV);
 
 	InputComponent->BindAction("DebugCamera_ToggleDisplay", IE_Pressed, this, &ADebugCameraController::ToggleDisplay);
 	InputComponent->BindAction("DebugCamera_FreezeRendering", IE_Pressed, this, &ADebugCameraController::ToggleFreezeRendering);
