Correctly reset the minimal view info to prevent parameters from "leaking" from one view target to the next

diff --git a/Engine/Source/Runtime/Engine/Private/PlayerCameraManager.cpp b/Engine/Source/Runtime/Engine/Private/PlayerCameraManager.cpp
--- a/Engine/Source/Runtime/Engine/Private/PlayerCameraManager.cpp
+++ b/Engine/Source/Runtime/Engine/Private/PlayerCameraManager.cpp
@@ -340,19 +340,18 @@ void APlayerCameraManager::UpdateViewTarget(FTViewTarget& OutVT, float DeltaTime
 		return;
 	}
 
-	// store previous POV, in case we need it later
+	// Store previous POV, in case we need it later
 	FMinimalViewInfo OrigPOV = OutVT.POV;
 
-	//@TODO: CAMERA: Should probably reset the view target POV fully here
+	// Reset the view target POV fully
+	static const FMinimalViewInfo DefaultViewInfo;
+	OutVT.POV = DefaultViewInfo;
 	OutVT.POV.FOV = DefaultFOV;
 	OutVT.POV.OrthoWidth = DefaultOrthoWidth;
 	OutVT.POV.AspectRatio = DefaultAspectRatio;
 	OutVT.POV.bConstrainAspectRatio = bDefaultConstrainAspectRatio;
-	OutVT.POV.bUseFieldOfViewForLOD = true;
 	OutVT.POV.ProjectionMode = bIsOrthographic ? ECameraProjectionMode::Orthographic : ECameraProjectionMode::Perspective;
-	OutVT.POV.PostProcessSettings.SetBaseValues();
 	OutVT.POV.PostProcessBlendWeight = 1.0f;
-	OutVT.POV.PreviousViewTransform.Reset();
 
 	bool bDoNotApplyModifiers = false;
 
