Fix curve inertial blending regression from curve refactor

3-way curve merge requires 2 unions to work correctly. Unfortunately this does take up a bit more memory as we need to hold the curve value as well as the delta & derivative.

diff --git a/Engine/Source/Runtime/Engine/Classes/Animation/AnimNode_Inertialization.h b/Engine/Source/Runtime/Engine/Classes/Animation/AnimNode_Inertialization.h
--- a/Engine/Source/Runtime/Engine/Classes/Animation/AnimNode_Inertialization.h
+++ b/Engine/Source/Runtime/Engine/Classes/Animation/AnimNode_Inertialization.h
@@ -238,8 +238,7 @@ struct FInertializationBoneDiff
 
 struct FInertializationCurveDiffElement : public UE::Anim::FCurveElement
 {
-	UE_DEPRECATED(5.3, "Please use Value instead.")
-	static float Delta;
+	float Delta = 0.0f;
 	float Derivative = 0.0f;
 
 	FInertializationCurveDiffElement() = default;
@@ -247,6 +246,7 @@ struct FInertializationCurveDiffElement : public UE::Anim::FCurveElement
 	void Clear()
 	{
 		Value = 0.0f;
+		Delta = 0.0f;
 		Derivative = 0.0f;
 	}
 };
diff --git a/Engine/Source/Runtime/Engine/Private/Animation/AnimNode_Inertialization.cpp b/Engine/Source/Runtime/Engine/Private/Animation/AnimNode_Inertialization.cpp
--- a/Engine/Source/Runtime/Engine/Private/Animation/AnimNode_Inertialization.cpp
+++ b/Engine/Source/Runtime/Engine/Private/Animation/AnimNode_Inertialization.cpp
@@ -28,9 +28,6 @@ TAutoConsoleVariable<int32> CVarAnimInertializationIgnoreDeficit(TEXT("a.AnimNod
 static constexpr int32 INERTIALIZATION_MAX_POSE_SNAPSHOTS = 2;
 static constexpr float INERTIALIZATION_TIME_EPSILON = 1.0e-7f;
 
-// DEPRECATED: Used for backwards compatibility
-float FInertializationCurveDiffElement::Delta = 0.0f;
-
 namespace UE { namespace Anim {
 
 // Inertialization request event bound to a node
@@ -717,22 +714,24 @@ void FInertializationPoseDiff::InitFrom(const FCompactPose& Pose, const FBlended
 		UE::Anim::FCurveUtils::Filter(CurveDiffs, CurveFilter);
 	}
 
-	// Compute differences & derivatives
-	UE::Anim::FNamedValueArrayUtils::Union(CurveDiffs, Prev1.Curves.BlendedCurve, Prev2.Curves.BlendedCurve,
-		[DeltaTime = Prev1.DeltaTime](FInertializationCurveDiffElement& OutResultElement, const UE::Anim::FCurveElement& InElement0, const UE::Anim::FCurveElement& InElement1, UE::Anim::ENamedValueUnionFlags InFlags)
+	// Compute differences
+	UE::Anim::FNamedValueArrayUtils::Union(CurveDiffs, Prev1.Curves.BlendedCurve,
+		[](FInertializationCurveDiffElement& OutResultElement, const UE::Anim::FCurveElement& InElement1, UE::Anim::ENamedValueUnionFlags InFlags)
 		{
-			const float CurrWeight = OutResultElement.Value;
-			const float Prev1Weight = InElement0.Value;
-
-			// Store delta in curve value
-			OutResultElement.Value = Prev1Weight - CurrWeight;
+			OutResultElement.Delta = InElement1.Value - OutResultElement.Value;
+		});
 
-			if(DeltaTime > UE_KINDA_SMALL_NUMBER)
+	// Compute derivatives
+	if(Prev1.DeltaTime > UE_KINDA_SMALL_NUMBER)
+	{
+		UE::Anim::FNamedValueArrayUtils::Union(CurveDiffs, Prev2.Curves.BlendedCurve,
+			[DeltaTime = Prev1.DeltaTime](FInertializationCurveDiffElement& OutResultElement, const UE::Anim::FCurveElement& InElement1, UE::Anim::ENamedValueUnionFlags InFlags)
 			{
+				const float Prev1Weight = OutResultElement.Delta - OutResultElement.Value;
 				const float Prev2Weight = InElement1.Value;
 				OutResultElement.Derivative = (Prev1Weight - Prev2Weight) / DeltaTime;
-			}
-		});
+			});
+	}
 }
 
 
@@ -775,7 +774,7 @@ void FInertializationPoseDiff::ApplyTo(FCompactPose& Pose, FBlendedCurve& Curves
 	UE::Anim::FNamedValueArrayUtils::Union(Curves, CurveDiffs,
 		[&InertializationElapsedTime, &InertializationDuration](UE::Anim::FCurveElement& OutResultElement, const FInertializationCurveDiffElement& InParamElement, UE::Anim::ENamedValueUnionFlags InFlags)
 		{
-			OutResultElement.Value = CalcInertialFloat(InParamElement.Value, InParamElement.Derivative, InertializationElapsedTime, InertializationDuration);
+			OutResultElement.Value += CalcInertialFloat(InParamElement.Delta, InParamElement.Derivative, InertializationElapsedTime, InertializationDuration);
 			OutResultElement.Flags |= InParamElement.Flags;
 		});
 }
