Add exposure as new light unit.

diff --git a/Engine/Plugins/Enterprise/DatasmithImporter/Source/DatasmithImporter/Private/DatasmithLightImporter.cpp b/Engine/Plugins/Enterprise/DatasmithImporter/Source/DatasmithImporter/Private/DatasmithLightImporter.cpp
--- a/Engine/Plugins/Enterprise/DatasmithImporter/Source/DatasmithImporter/Private/DatasmithLightImporter.cpp
+++ b/Engine/Plugins/Enterprise/DatasmithImporter/Source/DatasmithImporter/Private/DatasmithLightImporter.cpp
@@ -361,6 +361,9 @@ void FDatasmithLightImporter::SetupPointLightComponent( UPointLightComponent* Po
 	case EDatasmithLightUnits::Lumens:
 		PointLightComponentTemplate->IntensityUnits = ELightUnits::Lumens;
 		break;
+	case EDatasmithLightUnits::EV:
+		PointLightComponentTemplate->IntensityUnits = ELightUnits::EV;
+		break;
 	default:
 		PointLightComponentTemplate->IntensityUnits = ELightUnits::Unitless;
 		break;
diff --git a/Engine/Plugins/Enterprise/DatasmithImporter/Source/DatasmithImporter/Private/Utility/DatasmithImporterUtils.cpp b/Engine/Plugins/Enterprise/DatasmithImporter/Source/DatasmithImporter/Private/Utility/DatasmithImporterUtils.cpp
--- a/Engine/Plugins/Enterprise/DatasmithImporter/Source/DatasmithImporter/Private/Utility/DatasmithImporterUtils.cpp
+++ b/Engine/Plugins/Enterprise/DatasmithImporter/Source/DatasmithImporter/Private/Utility/DatasmithImporterUtils.cpp
@@ -669,6 +669,9 @@ namespace FDatasmithImporterUtilsHelper
 		case ELightUnits::Lumens:
 			PointLightElement->SetIntensityUnits( EDatasmithLightUnits::Lumens );
 			break;
+		case ELightUnits::EV:
+			PointLightElement->SetIntensityUnits( EDatasmithLightUnits::EV );
+			break;
 		default:
 			PointLightElement->SetIntensityUnits( EDatasmithLightUnits::Unitless );
 			break;
diff --git a/Engine/Plugins/Experimental/Enterprise/DatasmithInterchange/Source/DatasmithInterchange/Private/InterchangeDatasmithTranslator.cpp b/Engine/Plugins/Experimental/Enterprise/DatasmithInterchange/Source/DatasmithInterchange/Private/InterchangeDatasmithTranslator.cpp
--- a/Engine/Plugins/Experimental/Enterprise/DatasmithInterchange/Source/DatasmithInterchange/Private/InterchangeDatasmithTranslator.cpp
+++ b/Engine/Plugins/Experimental/Enterprise/DatasmithInterchange/Source/DatasmithInterchange/Private/InterchangeDatasmithTranslator.cpp
@@ -438,6 +438,7 @@ UInterchangeBaseLightNode* UInterchangeDatasmithTranslator::AddLightNode(UInterc
 	static_assert(FCommonLightUnits(EInterchangeLightUnits::Unitless) == FCommonLightUnits(EDatasmithLightUnits::Unitless), "EDatasmithLightUnits::Unitless differs from EInterchangeLightUnits::Unitless");
 	static_assert(FCommonLightUnits(EInterchangeLightUnits::Lumens) == FCommonLightUnits(EDatasmithLightUnits::Lumens), "EDatasmithLightUnits::Lumens differs from EInterchangeLightUnits::Lumens");
 	static_assert(FCommonLightUnits(EInterchangeLightUnits::Candelas) == FCommonLightUnits(EDatasmithLightUnits::Candelas), "EDatasmithLightUnits::Candelas differs from EInterchangeLightUnits::Candelas");
+	static_assert(FCommonLightUnits(EInterchangeLightUnits::EV) == FCommonLightUnits(EDatasmithLightUnits::EV), "EDatasmithLightUnits::EV differs from EInterchangeLightUnits::EV");
 
 	// TODO Add properties currently missing from the UInterchangeLightNode: everything
 	UInterchangeBaseLightNode* LightNode = nullptr;
diff --git a/Engine/Plugins/Experimental/Enterprise/DatasmithRuntime/Source/Private/LightImporter.cpp b/Engine/Plugins/Experimental/Enterprise/DatasmithRuntime/Source/Private/LightImporter.cpp
--- a/Engine/Plugins/Experimental/Enterprise/DatasmithRuntime/Source/Private/LightImporter.cpp
+++ b/Engine/Plugins/Experimental/Enterprise/DatasmithRuntime/Source/Private/LightImporter.cpp
@@ -221,6 +221,9 @@ namespace DatasmithRuntime
 				case EDatasmithLightUnits::Lumens:
 					PointLightComponent->IntensityUnits = ELightUnits::Lumens;
 					break;
+				case EDatasmithLightUnits::EV:
+					PointLightComponent->IntensityUnits = ELightUnits::EV;
+					break;
 				default:
 					PointLightComponent->IntensityUnits = ELightUnits::Unitless;
 					break;
diff --git a/Engine/Plugins/Importers/USDImporter/Source/USDUtilities/Private/USDLightConversion.cpp b/Engine/Plugins/Importers/USDImporter/Source/USDUtilities/Private/USDLightConversion.cpp
--- a/Engine/Plugins/Importers/USDImporter/Source/USDUtilities/Private/USDLightConversion.cpp
+++ b/Engine/Plugins/Importers/USDImporter/Source/USDUtilities/Private/USDLightConversion.cpp
@@ -25,6 +25,7 @@
 #include "EditorFramework/AssetImportData.h"
 #include "Engine/TextureCube.h"
 #include "Misc/Paths.h"
+#include "RenderUtils.h"
 
 #if USE_USD_SDK
 
@@ -66,6 +67,10 @@ namespace LightConversionImpl
 			// https://docs.unrealengine.com/en-US/Engine/Rendering/LightingAndShadows/PhysicalLightUnits/index.html#point,spot,andrectlights
 			return Intensity / ( Steradians * AreaInSqMeters );
 			break;
+		case ELightUnits::EV:
+			// Nit = luminance (cd/m2)
+			return EV100ToLuminance(Intensity);
+			break;
 		case ELightUnits::Unitless:
 			// Nit = ( unitless / 625 ) / area = candela / area
 			// https://docs.unrealengine.com/en-US/Engine/Rendering/LightingAndShadows/PhysicalLightUnits/index.html#point,spot,andrectlights
@@ -531,6 +536,10 @@ bool UnrealToUsd::ConvertRectLightComponent( const URectLightComponent& LightCom
 			// https://docs.unrealengine.com/en-US/Engine/Rendering/LightingAndShadows/PhysicalLightUnits/index.html#point,spot,andrectlights
 			FinalIntensityNits = OldIntensity / ( PI * AreaInSqMeters );
 			break;
+		case ELightUnits::EV:
+			// Nit = luminance (cd/m2)
+			FinalIntensityNits = EV100ToLuminance(OldIntensity);
+			break;
 		case ELightUnits::Unitless:
 			// Nit = (unitless/625) / area = candela / area
 			// https://docs.unrealengine.com/en-US/Engine/Rendering/LightingAndShadows/PhysicalLightUnits/index.html#point,spot,andrectlights
@@ -601,6 +610,10 @@ bool UnrealToUsd::ConvertPointLightComponent( const UPointLightComponent& LightC
 			// Nit = lumen / (sr * area); For a full sphere sr = 4PI
 			FinalIntensityNits = OldIntensity / ( SolidAngle * AreaInSqMeters );
 			break;
+		case ELightUnits::EV:
+			// Nit = luminance (cd/m2)
+			FinalIntensityNits = EV100ToLuminance(OldIntensity);
+			break;
 		case ELightUnits::Unitless:
 			// Nit = (unitless/625) / area = candela / area
 			// https://docs.unrealengine.com/en-US/Engine/Rendering/LightingAndShadows/PhysicalLightUnits/index.html#point,spot,andrectlights
diff --git a/Engine/Plugins/Interchange/Runtime/Source/Nodes/Public/InterchangeLightNode.h b/Engine/Plugins/Interchange/Runtime/Source/Nodes/Public/InterchangeLightNode.h
--- a/Engine/Plugins/Interchange/Runtime/Source/Nodes/Public/InterchangeLightNode.h
+++ b/Engine/Plugins/Interchange/Runtime/Source/Nodes/Public/InterchangeLightNode.h
@@ -14,6 +14,7 @@ enum class EInterchangeLightUnits : uint8
 	Unitless,
 	Candelas,
 	Lumens,
+	EV,
 };
 
 UCLASS(BlueprintType)
diff --git a/Engine/Plugins/Interchange/Runtime/Source/Pipelines/Private/InterchangeGenericScenesPipeline.cpp b/Engine/Plugins/Interchange/Runtime/Source/Pipelines/Private/InterchangeGenericScenesPipeline.cpp
--- a/Engine/Plugins/Interchange/Runtime/Source/Pipelines/Private/InterchangeGenericScenesPipeline.cpp
+++ b/Engine/Plugins/Interchange/Runtime/Source/Pipelines/Private/InterchangeGenericScenesPipeline.cpp
@@ -331,6 +331,7 @@ void UInterchangeGenericLevelPipeline::SetUpFactoryNode(UInterchangeActorFactory
 			static_assert(FCommonLightUnits(EInterchangeLightUnits::Unitless) == FCommonLightUnits(ELightUnits::Unitless), "EInterchangeLightUnits::Unitless differs from ELightUnits::Unitless");
 			static_assert(FCommonLightUnits(EInterchangeLightUnits::Lumens) == FCommonLightUnits(ELightUnits::Lumens), "EInterchangeLightUnits::Lumens differs from ELightUnits::Lumens");
 			static_assert(FCommonLightUnits(EInterchangeLightUnits::Candelas) == FCommonLightUnits(ELightUnits::Candelas), "EInterchangeLightUnits::Candelas differs from ELightUnits::Candelas");
+			static_assert(FCommonLightUnits(EInterchangeLightUnits::EV) == FCommonLightUnits(ELightUnits::EV), "EInterchangeLightUnits::EV differs from ELightUnits::EV");
 
 			if (const UInterchangeLightNode* LightNode = Cast<UInterchangeLightNode>(BaseLightNode))
 			{
diff --git a/Engine/Programs/UnrealHeaderTool/Config/DefaultEngine.ini b/Engine/Programs/UnrealHeaderTool/Config/DefaultEngine.ini
--- a/Engine/Programs/UnrealHeaderTool/Config/DefaultEngine.ini
+++ b/Engine/Programs/UnrealHeaderTool/Config/DefaultEngine.ini
@@ -132,6 +132,7 @@ gc.MaxObjectsInProgram=262144
 +Units="CandelaPerMeterSquared"
 +Units="cd/m2"
 +Units="CandelaPerMeter2"
++Units="EV"
 +Units="Nanoseconds"
 +Units="ns"
 +Units="Microseconds"
diff --git a/Engine/Source/Editor/DetailCustomizations/Private/LocalLightComponentDetails.cpp b/Engine/Source/Editor/DetailCustomizations/Private/LocalLightComponentDetails.cpp
--- a/Engine/Source/Editor/DetailCustomizations/Private/LocalLightComponentDetails.cpp
+++ b/Engine/Source/Editor/DetailCustomizations/Private/LocalLightComponentDetails.cpp
@@ -52,9 +52,18 @@ void FLocalLightComponentDetails::CustomizeDetails( IDetailLayoutBuilder& Detail
 	IntensityUnitsProperty->SetOnPropertyValueChanged(FSimpleDelegate::CreateSP(this, &FLocalLightComponentDetails::OnIntensityUnitsChanged, Component));
 
 	// Inverse squared falloff point lights (the default) are in units of lumens, instead of just being a brightness scale
-	LightIntensityProperty->SetInstanceMetaData("UIMin",TEXT("0.0f"));
-	LightIntensityProperty->SetInstanceMetaData("UIMax",  *FString::SanitizeFloat(100000.0f * ConversionFactor));
-	LightIntensityProperty->SetInstanceMetaData("SliderExponent", TEXT("2.0f"));
+	if (Component->IntensityUnits == ELightUnits::EV)
+	{
+		// +/-32 stops give a large enough range
+		LightIntensityProperty->SetInstanceMetaData("UIMin",TEXT("-32.0f"));
+		LightIntensityProperty->SetInstanceMetaData("UIMax",TEXT("32.0f"));
+	}
+	else
+	{
+		LightIntensityProperty->SetInstanceMetaData("UIMin",TEXT("0.0f"));
+		LightIntensityProperty->SetInstanceMetaData("UIMax",  *FString::SanitizeFloat(100000.0f * ConversionFactor));
+		LightIntensityProperty->SetInstanceMetaData("SliderExponent", TEXT("2.0f"));
+	}
 	if (Component->IntensityUnits == ELightUnits::Lumens)
 	{
 		LightIntensityProperty->SetInstanceMetaData("Units", TEXT("lm"));
@@ -65,6 +74,11 @@ void FLocalLightComponentDetails::CustomizeDetails( IDetailLayoutBuilder& Detail
 		LightIntensityProperty->SetInstanceMetaData("Units", TEXT("cd"));
 		LightIntensityProperty->SetToolTipText(LOCTEXT("LightIntensityInCandelasToolTipText", "Luminous intensity in candelas"));
 	}
+	else if (Component->IntensityUnits == ELightUnits::EV)
+	{
+		LightIntensityProperty->SetInstanceMetaData("Units", TEXT("ev"));
+		LightIntensityProperty->SetToolTipText(LOCTEXT("LightIntensityInCandelasToolTipText", "Luminous intensity in EV100"));
+	}
 
 	// Make these come first
 	IDetailCategoryBuilder& LightCategory = DetailBuilder.EditCategory( "Light", FText::GetEmpty(), ECategoryPriority::TypeSpecific );
diff --git a/Engine/Source/Editor/UnrealEd/Private/Factories/ActorFactoryPointLight.cpp b/Engine/Source/Editor/UnrealEd/Private/Factories/ActorFactoryPointLight.cpp
--- a/Engine/Source/Editor/UnrealEd/Private/Factories/ActorFactoryPointLight.cpp
+++ b/Engine/Source/Editor/UnrealEd/Private/Factories/ActorFactoryPointLight.cpp
@@ -21,7 +21,7 @@ void UActorFactoryPointLight::PostSpawnActor(UObject* Asset, AActor* NewActor)
 			static const auto CVarDefaultLightUnits = IConsoleManager::Get().FindTConsoleVariableDataInt(TEXT("r.DefaultFeature.LightUnits"));
 			ELightUnits DefaultUnits = (ELightUnits)CVarDefaultLightUnits->GetValueOnAnyThread();
 
-			Component->Intensity *= UPointLightComponent::GetUnitsConversionFactor(Component->IntensityUnits, DefaultUnits);
+			Component->Intensity *= UPointLightComponent::GetUnitsConversionFactor(Component->IntensityUnits, DefaultUnits, -1.f);
 			Component->IntensityUnits = DefaultUnits;
 		}
 	}
diff --git a/Engine/Source/Runtime/Core/Private/Math/UnitConversion.cpp b/Engine/Source/Runtime/Core/Private/Math/UnitConversion.cpp
--- a/Engine/Source/Runtime/Core/Private/Math/UnitConversion.cpp
+++ b/Engine/Source/Runtime/Core/Private/Math/UnitConversion.cpp
@@ -81,6 +81,7 @@ FParseCandidate ParseCandidates[] = {
 	{ TEXT("Candela"),				EUnit::Candela },				{ TEXT("cd"),		EUnit::Candela },
 	{ TEXT("Lux"),					EUnit::Lux },					{ TEXT("lx"),		EUnit::Lux },
 	{ TEXT("CandelaPerMeterSquared"), EUnit::CandelaPerMeter2 },	{ TEXT("cd/m2"),	EUnit::CandelaPerMeter2 },		{ TEXT("CandelaPerMeter2"),		EUnit::CandelaPerMeter2 },
+	{ TEXT("EV"),					EUnit::ExposureValue },			{ TEXT("EV"),		EUnit::ExposureValue },
 
 	{ TEXT("Nanoseconds"),			EUnit::Nanoseconds },			{ TEXT("ns"),		EUnit::Nanoseconds },
 	{ TEXT("Microseconds"),			EUnit::Microseconds },			{ TEXT("us"),		EUnit::Microseconds },			{ TEXT("Microseconds"),			EUnit::Microseconds },			{ TEXT("\u00B5s"),	EUnit::Microseconds },
@@ -120,7 +121,7 @@ const TCHAR* const DisplayStrings[] = {
 
 	TEXT("B"), TEXT("KiB"), TEXT("MiB"), TEXT("GiB"), TEXT("TiB"),
 
-	TEXT("lm"), TEXT("cd"), TEXT("lux"), TEXT("cd/m2"),
+	TEXT("lm"), TEXT("cd"), TEXT("lux"), TEXT("cd/m2"), TEXT("EV"),
 
 	TEXT("ns"), TEXT("\u00B5s"), TEXT("ms"), TEXT("s"), TEXT("min"), TEXT("hr"), TEXT("dy"), TEXT("mth"), TEXT("yr"),
 
@@ -151,7 +152,7 @@ const EUnitType UnitTypes[] = {
 
 	EUnitType::DataSize,	EUnitType::DataSize,	EUnitType::DataSize,	EUnitType::DataSize,	EUnitType::DataSize,
 
-	EUnitType::LuminousFlux, EUnitType::LuminousIntensity, EUnitType::Illuminance, EUnitType::Luminance,
+	EUnitType::LuminousFlux, EUnitType::LuminousIntensity, EUnitType::Illuminance, EUnitType::Luminance, EUnitType::ExposureValue,
 
 	EUnitType::Time,		EUnitType::Time,		EUnitType::Time,		EUnitType::Time,		EUnitType::Time,		EUnitType::Time,		EUnitType::Time,		EUnitType::Time,		EUnitType::Time,
 
@@ -430,6 +431,7 @@ FUnitSettings::FUnitSettings()
 	DisplayUnits[(uint8)EUnitType::Frequency].Add(EUnit::Hertz);
 	DisplayUnits[(uint8)EUnitType::DataSize].Add(EUnit::Megabytes);
 	DisplayUnits[(uint8)EUnitType::LuminousFlux].Add(EUnit::Lumens);
+	DisplayUnits[(uint8)EUnitType::ExposureValue].Add(EUnit::ExposureValue);
 	DisplayUnits[(uint8)EUnitType::Time].Add(EUnit::Seconds);
 }
 
diff --git a/Engine/Source/Runtime/Core/Public/Math/UnitConversion.h b/Engine/Source/Runtime/Core/Public/Math/UnitConversion.h
--- a/Engine/Source/Runtime/Core/Public/Math/UnitConversion.h
+++ b/Engine/Source/Runtime/Core/Public/Math/UnitConversion.h
@@ -44,8 +44,8 @@ enum class EUnit : uint8
 	/** Data Size units */
 	Bytes, Kilobytes, Megabytes, Gigabytes, Terabytes,
 
-	/** Luminous flux units, luminous intensity, illuminance, luminance */
-	Lumens, Candela, Lux, CandelaPerMeter2,
+	/** Luminous flux units, luminous intensity, illuminance, luminance, exposure value */
+	Lumens, Candela, Lux, CandelaPerMeter2, ExposureValue,
 
 	/** Time units */
 	Nanoseconds, Microseconds, Milliseconds, Seconds, Minutes, Hours, Days, Months, Years,
@@ -63,7 +63,7 @@ enum class EUnit : uint8
 /** Enumeration that specifies particular classes of unit */
 enum class EUnitType
 {
-	Distance, Angle, Speed, Temperature, Mass, Force, Frequency, DataSize, LuminousFlux, LuminousIntensity, Illuminance, Luminance, Time, PixelDensity, Multipliers,
+	Distance, Angle, Speed, Temperature, Mass, Force, Frequency, DataSize, LuminousFlux, LuminousIntensity, Illuminance, Luminance, Time, PixelDensity, Multipliers, ExposureValue,
 
 	// Symbolic entry - do not use directly
 	NumberOf,
diff --git a/Engine/Source/Runtime/Datasmith/DatasmithCore/Public/DatasmithDefinitions.h b/Engine/Source/Runtime/Datasmith/DatasmithCore/Public/DatasmithDefinitions.h
--- a/Engine/Source/Runtime/Datasmith/DatasmithCore/Public/DatasmithDefinitions.h
+++ b/Engine/Source/Runtime/Datasmith/DatasmithCore/Public/DatasmithDefinitions.h
@@ -159,6 +159,7 @@ enum class EDatasmithLightUnits
 	Unitless,
 	Candelas,
 	Lumens,
+	EV,
 };
 
 /** Different usage for textures.  Note: Preserve enum order. */
diff --git a/Engine/Source/Runtime/Engine/Classes/Components/DirectionalLightComponent.h b/Engine/Source/Runtime/Engine/Classes/Components/DirectionalLightComponent.h
--- a/Engine/Source/Runtime/Engine/Classes/Components/DirectionalLightComponent.h
+++ b/Engine/Source/Runtime/Engine/Classes/Components/DirectionalLightComponent.h
@@ -338,6 +338,7 @@ class ENGINE_API UDirectionalLightComponent : public ULightComponent
 	{
 		return AtmosphereSunDiskColorScale;
 	}
+	virtual ELightUnits GetLightUnits() const { return ELightUnits::Unitless; /* Lux */ }
 	//~ End ULightComponent Interface
 
 	//~ Begin UObject Interface
diff --git a/Engine/Source/Runtime/Engine/Classes/Components/LightComponent.h b/Engine/Source/Runtime/Engine/Classes/Components/LightComponent.h
--- a/Engine/Source/Runtime/Engine/Classes/Components/LightComponent.h
+++ b/Engine/Source/Runtime/Engine/Classes/Components/LightComponent.h
@@ -377,6 +377,7 @@ public:
 
 	virtual float GetUniformPenumbraSize() const PURE_VIRTUAL(ULightComponent::GetUniformPenumbraSize,return 0;);
 
+	virtual ELightUnits GetLightUnits() const { return ELightUnits::Unitless; }
 
 	/**
 	 * Check whether a given primitive will cast shadows from this light.
diff --git a/Engine/Source/Runtime/Engine/Classes/Components/LocalLightComponent.h b/Engine/Source/Runtime/Engine/Classes/Components/LocalLightComponent.h
--- a/Engine/Source/Runtime/Engine/Classes/Components/LocalLightComponent.h
+++ b/Engine/Source/Runtime/Engine/Classes/Components/LocalLightComponent.h
@@ -73,6 +73,7 @@ public:
 	{
 		return LightmassSettings;
 	}
+	virtual ELightUnits GetLightUnits() const { return IntensityUnits; }
 
 	//~ Begin UObject Interface
 	virtual void Serialize(FArchive& Ar) override;
diff --git a/Engine/Source/Runtime/Engine/Classes/Engine/Scene.h b/Engine/Source/Runtime/Engine/Classes/Engine/Scene.h
--- a/Engine/Source/Runtime/Engine/Classes/Engine/Scene.h
+++ b/Engine/Source/Runtime/Engine/Classes/Engine/Scene.h
@@ -72,6 +72,7 @@ enum class ELightUnits : uint8
 	Unitless,
 	Candelas,
 	Lumens,
+	EV,
 };
 
 UENUM()
diff --git a/Engine/Source/Runtime/Engine/Private/Components/LightComponent.cpp b/Engine/Source/Runtime/Engine/Private/Components/LightComponent.cpp
--- a/Engine/Source/Runtime/Engine/Private/Components/LightComponent.cpp
+++ b/Engine/Source/Runtime/Engine/Private/Components/LightComponent.cpp
@@ -459,15 +459,19 @@ bool ULightComponent::IsShadowCast(UPrimitiveComponent* Primitive) const
 float ULightComponent::ComputeLightBrightness() const
 {
 	float LightBrightness = Intensity;
-
-	if(IESTexture)
+	if (IESTexture)
 	{
-		if(bUseIESBrightness)
+		// When using EV100 unit, do conversion back and force so that IES brigthness can be computed on linear value
+		const bool bEVUnit = GetLightUnits() == ELightUnits::EV;
+		if (bEVUnit) { LightBrightness = EV100ToLuminance(Intensity); }
+		
+		if (bUseIESBrightness)
 		{
 			LightBrightness = IESTexture->Brightness * IESBrightnessScale;
 		}
-
 		LightBrightness *= IESTexture->TextureMultiplier;
+
+		if (bEVUnit) { LightBrightness = LuminanceToEV100(LightBrightness); }
 	}
 
 	return LightBrightness;
@@ -478,7 +482,11 @@ void ULightComponent::SetLightBrightness(float InBrightness)
 {
 	if (IESTexture && IESTexture->TextureMultiplier > 0)
 	{
-		if(bUseIESBrightness && IESBrightnessScale > 0)
+		// When using EV100 unit, do conversion back and force so that IES brigthness can be computed on linear value
+		const bool bEVUnit = GetLightUnits() == ELightUnits::EV;
+		if (bEVUnit) { InBrightness = EV100ToLuminance(InBrightness); }
+
+		if (bUseIESBrightness && IESBrightnessScale > 0)
 		{
 			IESTexture->Brightness = InBrightness / IESBrightnessScale / IESTexture->TextureMultiplier;
 		}
@@ -486,6 +494,8 @@ void ULightComponent::SetLightBrightness(float InBrightness)
 		{
 			Intensity = InBrightness / IESTexture->TextureMultiplier;
 		}
+
+		if (bEVUnit) { Intensity = LuminanceToEV100(Intensity); }
 	}
 	else
 	{
@@ -644,7 +654,11 @@ void ULightComponent::PostEditChangeProperty(FPropertyChangedEvent& PropertyChan
 	FProperty* PropertyThatChanged = PropertyChangedEvent.MemberProperty;
 	const FString PropertyName = PropertyThatChanged ? PropertyThatChanged->GetName() : TEXT("");
 
-	Intensity = FMath::Max(0.0f, Intensity);
+
+	if (GetLightUnits() != ELightUnits::EV)
+	{
+		Intensity = FMath::Max(0.0f, Intensity);
+	}
 	SpecularScale = FMath::Clamp( SpecularScale, 0.0f, 1.0f );
 
 	if (HasStaticLighting())
diff --git a/Engine/Source/Runtime/Engine/Private/Components/LocalLightComponent.cpp b/Engine/Source/Runtime/Engine/Private/Components/LocalLightComponent.cpp
--- a/Engine/Source/Runtime/Engine/Private/Components/LocalLightComponent.cpp
+++ b/Engine/Source/Runtime/Engine/Private/Components/LocalLightComponent.cpp
@@ -117,7 +117,11 @@ bool ULocalLightComponent::CanEditChange(const FProperty* InProperty) const
  */
 void ULocalLightComponent::PostEditChangeProperty(FPropertyChangedEvent& PropertyChangedEvent)
 {
-	Intensity = FMath::Max(0.0f, Intensity);
+	// Clamp intensity to 0 only for non-EV unit, as EV value are negative for small luminance value.
+	if (IntensityUnits != ELightUnits::EV)
+	{
+		Intensity = FMath::Max(0.0f, Intensity);
+	}
 	LightmassSettings.IndirectLightingSaturation = FMath::Max(LightmassSettings.IndirectLightingSaturation, 0.0f);
 	LightmassSettings.ShadowExponent = FMath::Clamp(LightmassSettings.ShadowExponent, .5f, 8.0f);
 
@@ -167,6 +171,30 @@ void ULocalLightComponent::PushRadiusToRenderThread()
 
 float ULocalLightComponent::GetUnitsConversionFactor(ELightUnits SrcUnits, ELightUnits TargetUnits, float CosHalfConeAngle)
 {
+	// Notes
+	// -----
+	// * UE light operates at constant 'luminous intensity' i.e., the intensity will remain constant when changing the light's 
+	//   size (radius/width/height/...). When dealing with EV, we use an implicit 1m2 surface are for conversion, which allows 
+	//   to keep EV constant under light's size change.
+	// * UE unit is in centimeters (CM), while SI unit version are in meter (M), hence the conversion unit (100*100) in the 
+	//   formula below
+	// * When chaning unit, first GetUnitsConversionFactor() is called, then SetBrightness(). When switching to EV unit, 
+	//   we for convert the intensity to luminance (assuming an implicity 1m2 surface) and then apply the luminance -> EV unit 
+	//   in SetBrightness()
+	// Reminder
+	// --------
+	// Light units (in terms of candela)
+	//  Flux        = Lm = cd.sr
+	//  Intensity   = Cd
+	//  Luminance   = Cd/m2
+	//  Illuminance = Cd.sr/m2 = Lux
+	//
+	// Light units (in terms of Lumen)
+	//  Flux        = Lm = cd.sr
+	//  Intensity   = Lm/sr
+	//  Luminance   = Lm/sr/m2
+	//  Illuminance = Lm/m2 = Lux
+	
 	CosHalfConeAngle = FMath::Clamp<float>(CosHalfConeAngle, -1, 1 - UE_KINDA_SMALL_NUMBER);
 
 	if (SrcUnits == TargetUnits)
@@ -185,6 +213,10 @@ float ULocalLightComponent::GetUnitsConversionFactor(ELightUnits SrcUnits, ELigh
 		{
 			CnvFactor = 100.f * 100.f / 2.f / UE_PI / (1.f - CosHalfConeAngle);
 		}
+		else if (SrcUnits == ELightUnits::EV)
+		{
+			CnvFactor = 100.f * 100.f;
+		}
 		else
 		{
 			CnvFactor = 16.f;
@@ -198,6 +230,10 @@ float ULocalLightComponent::GetUnitsConversionFactor(ELightUnits SrcUnits, ELigh
 		{
 			CnvFactor *= 2.f  * UE_PI * (1.f - CosHalfConeAngle) / 100.f / 100.f;
 		}
+		else if (TargetUnits == ELightUnits::EV)
+		{
+			CnvFactor *= 1.f / 100.f / 100.f;
+		}
 		else
 		{
 			CnvFactor *= 1.f / 16.f;
diff --git a/Engine/Source/Runtime/Engine/Private/Components/PointLightComponent.cpp b/Engine/Source/Runtime/Engine/Private/Components/PointLightComponent.cpp
--- a/Engine/Source/Runtime/Engine/Private/Components/PointLightComponent.cpp
+++ b/Engine/Source/Runtime/Engine/Private/Components/PointLightComponent.cpp
@@ -207,6 +207,10 @@ float UPointLightComponent::ComputeLightBrightness() const
 		{
 			LightBrightness *= (100.f * 100.f / 4 / UE_PI); // Conversion from cm2 to m2 and 4PI from the sphere area in the 1/r2 attenuation
 		}
+		else if (IntensityUnits == ELightUnits::EV)
+		{
+			LightBrightness = EV100ToLuminance(LightBrightness) * (100.f * 100.f);
+		}
 		else
 		{
 			LightBrightness *= 16; // Legacy scale of 16
@@ -228,6 +232,10 @@ void UPointLightComponent::SetLightBrightness(float InBrightness)
 		{
 			Super::SetLightBrightness(InBrightness / (100.f * 100.f / 4 / UE_PI)); // Conversion from cm2 to m2 and 4PI from the sphere area in the 1/r2 attenuation
 		}
+		else if (IntensityUnits == ELightUnits::EV)
+		{
+			Super::SetLightBrightness(LuminanceToEV100(InBrightness / (100.f * 100.f)));
+		}
 		else
 		{
 			Super::SetLightBrightness(InBrightness / 16); // Legacy scale of 16
diff --git a/Engine/Source/Runtime/Engine/Private/Components/RectLightComponent.cpp b/Engine/Source/Runtime/Engine/Private/Components/RectLightComponent.cpp
--- a/Engine/Source/Runtime/Engine/Private/Components/RectLightComponent.cpp
+++ b/Engine/Source/Runtime/Engine/Private/Components/RectLightComponent.cpp
@@ -117,6 +117,10 @@ float URectLightComponent::ComputeLightBrightness() const
 	{
 		LightBrightness *= (100.f * 100.f / UE_PI); // Conversion from cm2 to m2 and PI from the cosine distribution
 	}
+	else if (IntensityUnits == ELightUnits::EV)
+	{
+		LightBrightness *= (100.f * 100.f) * EV100ToLuminance(LightBrightness);
+	}
 	else
 	{
 		LightBrightness *= 16; // Legacy scale of 16
@@ -136,6 +140,10 @@ void URectLightComponent::SetLightBrightness(float InBrightness)
 	{
 		Super::SetLightBrightness(InBrightness / (100.f * 100.f / UE_PI)); // Conversion from cm2 to m2 and PI from the cosine distribution
 	}
+	else if (IntensityUnits == ELightUnits::EV)
+	{
+		Super::SetLightBrightness(LuminanceToEV100(InBrightness / (100.f * 100.f)));
+	}
 	else
 	{
 		Super::SetLightBrightness(InBrightness / 16); // Legacy scale of 16
diff --git a/Engine/Source/Runtime/Engine/Private/Components/SpotLightComponent.cpp b/Engine/Source/Runtime/Engine/Private/Components/SpotLightComponent.cpp
--- a/Engine/Source/Runtime/Engine/Private/Components/SpotLightComponent.cpp
+++ b/Engine/Source/Runtime/Engine/Private/Components/SpotLightComponent.cpp
@@ -245,6 +245,10 @@ float USpotLightComponent::ComputeLightBrightness() const
 		{
 			LightBrightness *= (100.f * 100.f / 2.f / UE_PI / (1.f - GetCosHalfConeAngle())); // Conversion from cm2 to m2 and cone remapping.
 		}
+		else if (IntensityUnits == ELightUnits::EV)
+		{
+			LightBrightness = EV100ToLuminance(LightBrightness * (100.f * 100.f));
+		}
 		else
 		{
 			LightBrightness *= 16; // Legacy scale of 16
@@ -266,6 +270,10 @@ void USpotLightComponent::SetLightBrightness(float InBrightness)
 		{
 			ULightComponent::SetLightBrightness(InBrightness / (100.f * 100.f / 2.f / UE_PI / (1.f - GetCosHalfConeAngle()))); // Conversion from cm2 to m2 and cone remapping
 		}
+		else if (IntensityUnits == ELightUnits::EV)
+		{
+			ULightComponent::SetLightBrightness(LuminanceToEV100(InBrightness / (100.f * 100.f)));
+		}
 		else
 		{
 			ULightComponent::SetLightBrightness(InBrightness / 16); // Legacy scale of 16
diff --git a/Engine/Source/Runtime/Engine/Private/SceneView.cpp b/Engine/Source/Runtime/Engine/Private/SceneView.cpp
--- a/Engine/Source/Runtime/Engine/Private/SceneView.cpp
+++ b/Engine/Source/Runtime/Engine/Private/SceneView.cpp
@@ -339,12 +339,6 @@ static TAutoConsoleVariable<float> CVarOverrideTimeMaterialExpressions(
 
 #endif
 
-// Conversion factor used when "r.DefaultFeature.AutoExposure.ExtendDefaultLuminanceRange" 
-FORCEINLINE float LuminanceToEV100(float Luminance)
-{
-	return FMath::Log2(Luminance / 1.2f);
-}
-
 /** Global vertex color view mode setting when SHOW_VertexColors show flag is set */
 EVertexColorViewMode::Type GVertexColorViewMode = EVertexColorViewMode::Color;
 
@@ -1954,8 +1948,9 @@ void FSceneView::StartFinalPostprocessSettings(FVector InViewLocation)
 			FinalPostProcessSettings.AutoExposureMaxBrightness = 1;
 			if (CVarDefaultAutoExposureExtendDefaultLuminanceRange.GetValueOnGameThread())
 			{
-				FinalPostProcessSettings.AutoExposureMinBrightness = LuminanceToEV100(FinalPostProcessSettings.AutoExposureMinBrightness);
-				FinalPostProcessSettings.AutoExposureMaxBrightness = LuminanceToEV100(FinalPostProcessSettings.AutoExposureMaxBrightness);
+				const float MaxLuminance = 1.2f; // Should we use const LuminanceMaxFromLensAttenuation() instead?
+				FinalPostProcessSettings.AutoExposureMinBrightness = LuminanceToEV100(MaxLuminance, FinalPostProcessSettings.AutoExposureMinBrightness);
+				FinalPostProcessSettings.AutoExposureMaxBrightness = LuminanceToEV100(MaxLuminance, FinalPostProcessSettings.AutoExposureMaxBrightness);
 			}
 		}
 		else
diff --git a/Engine/Source/Runtime/RenderCore/Public/RenderUtils.h b/Engine/Source/Runtime/RenderCore/Public/RenderUtils.h
--- a/Engine/Source/Runtime/RenderCore/Public/RenderUtils.h
+++ b/Engine/Source/Runtime/RenderCore/Public/RenderUtils.h
@@ -499,4 +499,38 @@ namespace Strata
 
 	RENDERCORE_API uint32 GetShadingQuality();
 	RENDERCORE_API uint32 GetShadingQuality(EShaderPlatform InPlatform);
+}
+
+// LuminanceMax is the amount of light that will cause the sensor to saturate at EV100.
+//  See also https://en.wikipedia.org/wiki/Film_speed and https://en.wikipedia.org/wiki/Exposure_value for more info.
+FORCEINLINE float EV100ToLuminance(float LuminanceMax, float EV100)
+{
+	return LuminanceMax * FMath::Pow(2.0f, EV100);
+}
+
+FORCEINLINE float EV100ToLuminance(float EV100)
+{
+	// LuminanceMax set to 1 for lighting to be unitless (1.0cd/m^2 becomes 1.0 at EV100)
+	return EV100ToLuminance(1.0f, EV100);
+}
+
+FORCEINLINE float EV100ToLog2(float LuminanceMax, float EV100)
+{
+	return EV100 + FMath::Log2(LuminanceMax);
+}
+
+FORCEINLINE float LuminanceToEV100(float LuminanceMax, float Luminance)
+{
+	return FMath::Log2(Luminance / LuminanceMax);
+}
+
+FORCEINLINE float LuminanceToEV100(float Luminance)
+{
+	// LuminanceMax set to 1 for lighting to be unitless (1.0cd/m^2 becomes 1.0 at EV100)
+	return FMath::Log2(Luminance);
+}
+
+FORCEINLINE float Log2ToEV100(float LuminanceMax, float Log2)
+{
+	return Log2 - FMath::Log2(LuminanceMax);
 }
\ No newline at end of file
diff --git a/Engine/Source/Runtime/Renderer/Private/PostProcess/PostProcessEyeAdaptation.h b/Engine/Source/Runtime/Renderer/Private/PostProcess/PostProcessEyeAdaptation.h
--- a/Engine/Source/Runtime/Renderer/Private/PostProcess/PostProcessEyeAdaptation.h
+++ b/Engine/Source/Runtime/Renderer/Private/PostProcess/PostProcessEyeAdaptation.h
@@ -4,28 +4,6 @@
 
 #include "ScreenPass.h"
 
-// LuminanceMax is the amount of light that will cause the sensor to saturate at EV100.
-//  See also https://en.wikipedia.org/wiki/Film_speed and https://en.wikipedia.org/wiki/Exposure_value for more info.
-FORCEINLINE float EV100ToLuminance(float LuminanceMax, float EV100)
-{
-	return LuminanceMax * FMath::Pow(2.0f, EV100);
-}
-
-FORCEINLINE float EV100ToLog2(float LuminanceMax, float EV100)
-{
-	return EV100 + FMath::Log2(LuminanceMax);
-}
-
-FORCEINLINE float LuminanceToEV100(float LuminanceMax, float Luminance)
-{
-	return FMath::Log2(Luminance / LuminanceMax);
-}
-
-FORCEINLINE float Log2ToEV100(float LuminanceMax, float Log2)
-{
-	return Log2 - FMath::Log2(LuminanceMax);
-}
-
 namespace AutoExposurePermutation
 {
 	class FUsePrecalculatedLuminanceDim : SHADER_PERMUTATION_BOOL("USE_PRECALCULATED_LUMINANCE");
