Fix CIS issue.

diff --git a/Engine/Source/Runtime/Engine/Classes/Components/LightComponent.h b/Engine/Source/Runtime/Engine/Classes/Components/LightComponent.h
--- a/Engine/Source/Runtime/Engine/Classes/Components/LightComponent.h
+++ b/Engine/Source/Runtime/Engine/Classes/Components/LightComponent.h
@@ -377,7 +377,7 @@ public:
 
 	virtual float GetUniformPenumbraSize() const PURE_VIRTUAL(ULightComponent::GetUniformPenumbraSize,return 0;);
 
-	virtual ELightUnits GetLightUnits() const { return ELightUnits::Unitless; }
+	virtual ELightUnits GetLightUnits() const;
 
 	/**
 	 * Check whether a given primitive will cast shadows from this light.
diff --git a/Engine/Source/Runtime/Engine/Private/Components/LightComponent.cpp b/Engine/Source/Runtime/Engine/Private/Components/LightComponent.cpp
--- a/Engine/Source/Runtime/Engine/Private/Components/LightComponent.cpp
+++ b/Engine/Source/Runtime/Engine/Private/Components/LightComponent.cpp
@@ -865,6 +865,8 @@ bool ULightComponent::GetMaterialPropertyPath(int32 ElementIndex, UObject*& OutO
 }
 #endif // WITH_EDITOR
 
+ELightUnits ULightComponent::GetLightUnits() const { return ELightUnits::Unitless; }
+
 /** Set brightness of the light */
 void ULightComponent::SetIntensity(float NewIntensity)
 {
