Fix CIS issue.

diff --git a/Engine/Source/Runtime/Engine/Classes/Components/LightComponent.h b/Engine/Source/Runtime/Engine/Classes/Components/LightComponent.h
--- a/Engine/Source/Runtime/Engine/Classes/Components/LightComponent.h
+++ b/Engine/Source/Runtime/Engine/Classes/Components/LightComponent.h
@@ -18,6 +18,7 @@ class ULevel;
 class UMaterialInterface;
 class UPrimitiveComponent;
 class UTextureLightProfile;
+enum class ELightUnits : uint8;
 
 /** 
  * A texture containing depth values of static objects that was computed during the lighting build.
