Adding a batched parameters flush at the end of GPUSkinCache DoDispatch to make sure everything is comitted before the render command finishes.

#rnx

diff --git a/Engine/Source/Runtime/Engine/Private/GPUSkinCache.cpp b/Engine/Source/Runtime/Engine/Private/GPUSkinCache.cpp
--- a/Engine/Source/Runtime/Engine/Private/GPUSkinCache.cpp
+++ b/Engine/Source/Runtime/Engine/Private/GPUSkinCache.cpp
@@ -1388,6 +1388,9 @@ void FGPUSkinCache::DoDispatch(FRHICommandListImmediate& RHICmdList)
 	}
 
 	TransitionAllToReadable(RHICmdList, BuffersToTransitionToRead);
+
+	// Need to flush here since we're a scene-less process.
+	RHICmdList.FlushAllPendingComputeParameters();
 }
 
 void FGPUSkinCache::DoDispatch(FRHICommandListImmediate& RHICmdList, FGPUSkinCacheEntry* SkinCacheEntry, int32 Section, int32 RevisionNumber)
@@ -1441,6 +1444,9 @@ void FGPUSkinCache::DoDispatch(FRHICommandListImmediate& RHICmdList, FGPUSkinCac
 	SkinCacheEntry->UpdateVertexFactoryDeclaration(Section);
 
 	TransitionAllToReadable(RHICmdList, BuffersToTransitionToRead);
+
+	// Need to flush here since we're a scene-less process.
+	RHICmdList.FlushAllPendingComputeParameters();
 }
 
 bool FGPUSkinCache::ProcessEntry(
diff --git a/Engine/Source/Runtime/RHI/Public/RHICommandList.h b/Engine/Source/Runtime/RHI/Public/RHICommandList.h
--- a/Engine/Source/Runtime/RHI/Public/RHICommandList.h
+++ b/Engine/Source/Runtime/RHI/Public/RHICommandList.h
@@ -2175,6 +2175,12 @@ protected:
 		ParameterBatcher.FlushPendingParameters(*this, InShader);
 	}
 
+public:
+	void FlushAllPendingComputeParameters()
+	{
+		PreDispatch();
+	}
+
 public:
 	FRHIComputeCommandList(FRHIGPUMask GPUMask = FRHIGPUMask::All(), ERecordingThread InRecordingThread = ERecordingThread::Render)
 		: FRHICommandListBase(GPUMask, InRecordingThread)
