Disabled expensive curve debug checks

diff --git a/Engine/Source/Runtime/Engine/Public/Animation/NamedValueArray.h b/Engine/Source/Runtime/Engine/Public/Animation/NamedValueArray.h
--- a/Engine/Source/Runtime/Engine/Public/Animation/NamedValueArray.h
+++ b/Engine/Source/Runtime/Engine/Public/Animation/NamedValueArray.h
@@ -17,8 +17,8 @@
 #define CURVE_PROFILE_CYCLE_COUNTER(Stat)
 #endif
 
-#define DO_ANIM_NAMED_VALUE_SORTING_CHECKS		(DO_GUARD_SLOW || WITH_DEV_AUTOMATION_TESTS)
-#define DO_ANIM_NAMED_VALUE_DUPLICATE_CHECKS	(DO_GUARD_SLOW || WITH_DEV_AUTOMATION_TESTS)
+#define DO_ANIM_NAMED_VALUE_SORTING_CHECKS		0
+#define DO_ANIM_NAMED_VALUE_DUPLICATE_CHECKS	0
 
 namespace UE::Anim
 {
