Get correct PerPlatform ScreenSize for LOD when in preview

diff --git a/Engine/Source/Runtime/Engine/Private/SkeletalRender.cpp b/Engine/Source/Runtime/Engine/Private/SkeletalRender.cpp
--- a/Engine/Source/Runtime/Engine/Private/SkeletalRender.cpp
+++ b/Engine/Source/Runtime/Engine/Private/SkeletalRender.cpp
@@ -123,7 +123,7 @@ void FSkeletalMeshObject::UpdateMinDesiredLODLevel(const FSceneView* View, const
 		for(int32 LODLevel = SkeletalMeshRenderData->LODRenderData.Num()-1; LODLevel > 0; LODLevel--)
 		{
 			// Get ScreenSize for this LOD
-			float ScreenSize = SkeletalMeshLODInfo[LODLevel].ScreenSize.Default;
+			float ScreenSize = SkeletalMeshLODInfo[LODLevel].ScreenSize.GetValue();
 
 			// If we are considering shifting to a better (lower) LOD, bias with hysteresis.
 			if(LODLevel  <= CurrentLODLevel)
@@ -234,7 +234,7 @@ float FSkeletalMeshObject::GetScreenSize(int32 LODIndex) const
 {
 	if (SkeletalMeshLODInfo.IsValidIndex(LODIndex))
 	{
-		return SkeletalMeshLODInfo[LODIndex].ScreenSize.Default;
+		return SkeletalMeshLODInfo[LODIndex].ScreenSize.GetValue();
 	}
 	return 0.f;
 }
