Fix CIS error.

diff --git a/Engine/Source/Runtime/Engine/Classes/Components/DirectionalLightComponent.h b/Engine/Source/Runtime/Engine/Classes/Components/DirectionalLightComponent.h
--- a/Engine/Source/Runtime/Engine/Classes/Components/DirectionalLightComponent.h
+++ b/Engine/Source/Runtime/Engine/Classes/Components/DirectionalLightComponent.h
@@ -338,7 +338,7 @@ class ENGINE_API UDirectionalLightComponent : public ULightComponent
 	{
 		return AtmosphereSunDiskColorScale;
 	}
-	virtual ELightUnits GetLightUnits() const { return ELightUnits::Unitless; /* Lux */ }
+	virtual ELightUnits GetLightUnits() const;
 	//~ End ULightComponent Interface
 
 	//~ Begin UObject Interface
diff --git a/Engine/Source/Runtime/Engine/Private/Components/DirectionalLightComponent.cpp b/Engine/Source/Runtime/Engine/Private/Components/DirectionalLightComponent.cpp
--- a/Engine/Source/Runtime/Engine/Private/Components/DirectionalLightComponent.cpp
+++ b/Engine/Source/Runtime/Engine/Private/Components/DirectionalLightComponent.cpp
@@ -1430,3 +1430,4 @@ void UDirectionalLightComponent::InvalidateLightingCacheDetailed(bool bInvalidat
 	}
 }
 
+ELightUnits UDirectionalLightComponent::GetLightUnits() const { return ELightUnits::Unitless; /* Lux */ }
\ No newline at end of file
