PR #10207: Redundant InstanceUpdateCmdBuffer Reset (Contributed by [REDACTED_PERSON])

#rnx

diff --git a/Engine/Source/Runtime/Engine/Private/HierarchicalInstancedStaticMesh.cpp b/Engine/Source/Runtime/Engine/Private/HierarchicalInstancedStaticMesh.cpp
--- a/Engine/Source/Runtime/Engine/Private/HierarchicalInstancedStaticMesh.cpp
+++ b/Engine/Source/Runtime/Engine/Private/HierarchicalInstancedStaticMesh.cpp
@@ -2911,7 +2911,6 @@ void UHierarchicalInstancedStaticMeshComponent::ApplyBuildTree(FClusterBuilder&
 	check(BuiltInstanceData->GetNumInstances() == NumBuiltRenderInstances);
 
 	InstanceCountToRender = NumBuiltInstances;
-	InstanceUpdateCmdBuffer.Reset();
 
 	check(InstanceReorderTable.Num() == PerInstanceSMData.Num());
 
