Displays VGPU0 in stat unit with -VMGPU and MaxGPUCount=1

diff --git a/Engine/Source/Runtime/Engine/Private/UnrealClient.cpp b/Engine/Source/Runtime/Engine/Private/UnrealClient.cpp
--- a/Engine/Source/Runtime/Engine/Private/UnrealClient.cpp
+++ b/Engine/Source/Runtime/Engine/Private/UnrealClient.cpp
@@ -713,12 +713,13 @@ int32 FStatUnitData::DrawStat(FViewport* InViewport, FCanvas* InCanvas, int32 In
 		// Draw all GPU informations
 		{
 			const int32 ColumnPerGPU = ColumnCount;
-			const int32 GPURowCount = 1 + (GRHISupportsGPUUsage ? 3 : 0) + (GNumExplicitGPUsForRendering > 1 ? 1 : 0);
+			const bool bDisplayGPUIndexes = GNumExplicitGPUsForRendering > 1 || GVirtualMGPU;
+			const int32 GPURowCount = 1 + (GRHISupportsGPUUsage ? 3 : 0) + (bDisplayGPUIndexes ? 1 : 0);
 			const int32 GPUColumnCount = ColumnCount * GNumExplicitGPUsForRendering;
 
 			// Draw the different timings:
 			{
-				int32 GPURowId = GNumExplicitGPUsForRendering > 1 ? 1 : 0;
+				int32 GPURowId = bDisplayGPUIndexes ? 1 : 0;
 				DrawRowTitle(GPURowId++, GPUColumnCount, TEXT("GPU Time"), /* UnitGraphColor = */ FColor(255, 255, 100));
 
 				if (GRHISupportsGPUUsage)
@@ -734,7 +735,7 @@ int32 FStatUnitData::DrawStat(FViewport* InViewport, FCanvas* InCanvas, int32 In
 			for (uint32 GPUIndex : FRHIGPUMask::All())
 			{
 				int32 GPURowId = 0;
-				if (GNumExplicitGPUsForRendering > 1)
+				if (bDisplayGPUIndexes)
 				{
 					DrawCell(
 						GPURowId++, GPUIndex * ColumnPerGPU + AvgUnitColumnId, GPUColumnCount,
