Streamlined GetMovementBaseTransform to remove redundant base check and reduce duplicate code. No outward functional changes.

#rnx

diff --git a/Engine/Source/Runtime/Engine/Private/Character.cpp b/Engine/Source/Runtime/Engine/Private/Character.cpp
--- a/Engine/Source/Runtime/Engine/Private/Character.cpp
+++ b/Engine/Source/Runtime/Engine/Private/Character.cpp
@@ -657,35 +657,27 @@ namespace MovementBaseUtility
 	{
 		if (MovementBase)
 		{
+			bool bGotTransformOfIntendedBone = true;
+
 			if (BoneName != NAME_None)
 			{
-				bool bFoundBone = false;
-				if (MovementBase)
+				// Check if this socket or bone exists (DoesSocketExist checks for either, as does requesting the transform).
+				if (MovementBase->DoesSocketExist(BoneName))
 				{
-					// Check if this socket or bone exists (DoesSocketExist checks for either, as does requesting the transform).
-					if (MovementBase->DoesSocketExist(BoneName))
-					{
-						MovementBase->GetSocketWorldLocationAndRotation(BoneName, OutLocation, OutQuat);
-						bFoundBone = true;
-					}
-					else
-					{
-						UE_LOG(LogCharacter, Warning, TEXT("GetMovementBaseTransform(): Invalid bone or socket '%s' for PrimitiveComponent base %s"), *BoneName.ToString(), *GetPathNameSafe(MovementBase));
-					}
+					MovementBase->GetSocketWorldLocationAndRotation(BoneName, OutLocation, OutQuat);
+					return true;
 				}
-
-				if (!bFoundBone)
+				else
 				{
-					OutLocation = MovementBase->GetComponentLocation();
-					OutQuat = MovementBase->GetComponentQuat();
+					UE_LOG(LogCharacter, Warning, TEXT("GetMovementBaseTransform(): Invalid bone or socket '%s' for PrimitiveComponent base %s. Falling back to base's root transform."), *BoneName.ToString(), *GetPathNameSafe(MovementBase));
+					bGotTransformOfIntendedBone = false;
 				}
-				return bFoundBone;
 			}
 
-			// No bone supplied
+			// No bone supplied (or it was invalid)
 			OutLocation = MovementBase->GetComponentLocation();
 			OutQuat = MovementBase->GetComponentQuat();
-			return true;
+			return bGotTransformOfIntendedBone;
 		}
 
 		// nullptr MovementBase
