Deprecating public apis using FNetworkObjectInfo

diff --git a/Engine/Source/Runtime/Engine/Classes/Engine/NetDriver.h b/Engine/Source/Runtime/Engine/Classes/Engine/NetDriver.h
--- a/Engine/Source/Runtime/Engine/Classes/Engine/NetDriver.h
+++ b/Engine/Source/Runtime/Engine/Classes/Engine/NetDriver.h
@@ -1671,13 +1671,19 @@ public:
      *	Get the network object matching the given Actor.
 	 *	If the Actor is not present in the NetworkObjectInfo list, it will be added.
 	 */
+	UE_DEPRECATED(5.3, "Will be made private in a future release")
 	ENGINE_API FNetworkObjectInfo* FindOrAddNetworkObjectInfo(const AActor* InActor);
 
 	/** Get the network object matching the given Actor. */
+	UE_DEPRECATED(5.3, "Will be made private in a future release")
 	ENGINE_API FNetworkObjectInfo* FindNetworkObjectInfo(const AActor* InActor);
+
+	UE_DEPRECATED(5.3, "Will be made private in a future release")
 	ENGINE_API const FNetworkObjectInfo* FindNetworkObjectInfo(const AActor* InActor) const
 	{
+		PRAGMA_DISABLE_DEPRECATION_WARNINGS
 		return const_cast<UNetDriver*>(this)->FindNetworkObjectInfo(InActor);
+		PRAGMA_ENABLE_DEPRECATION_WARNINGS
 	}
 
 	/**
@@ -1690,11 +1696,18 @@ public:
 	ENGINE_API bool IsNetworkActorUpdateFrequencyThrottled(const AActor* InActor) const;
 
 	/** Returns true if adaptive net update frequency is enabled and the given actor is having its update rate lowered from its standard rate. */
+	UE_DEPRECATED(5.3, "Will be made private in a future release, please use version that takes an actor")
 	ENGINE_API bool IsNetworkActorUpdateFrequencyThrottled(const FNetworkObjectInfo& InNetworkActor) const;
 
 	/** Stop adaptive replication for the given actor if it's currently throttled. It maybe be allowed to throttle again later. */
+	UE_DEPRECATED(5.3, "Will be made private in a future release, please use version that takes an actor")
 	ENGINE_API void CancelAdaptiveReplication(FNetworkObjectInfo& InNetworkActor);
 
+	ENGINE_API void CancelAdaptiveReplication(const AActor* InActor);
+
+	/** Returns true if the driver's world time has exceeded the next replication update time for this actor, or if it is pending replication from a previous frame. */
+	ENGINE_API bool IsPendingNetUpdate(const AActor* InActor) const;
+
 	/** Returns the level ID/PIE instance ID for this netdriver to use. */
 	ENGINE_API int32 GetDuplicateLevelID() const { return DuplicateLevelID; }
 
diff --git a/Engine/Source/Runtime/Engine/Classes/GameFramework/PlayerController.h b/Engine/Source/Runtime/Engine/Classes/GameFramework/PlayerController.h
--- a/Engine/Source/Runtime/Engine/Classes/GameFramework/PlayerController.h
+++ b/Engine/Source/Runtime/Engine/Classes/GameFramework/PlayerController.h
@@ -2076,6 +2076,7 @@ public:
 	/** forces a full replication check of the specified Actor on only the client that owns this PlayerController
 	 * this function has no effect if this PC is not a remote client or if the Actor is not relevant to that client
 	 */
+	UE_DEPRECATED(5.3, "Deprecated in favor of using ForceNetUpdate on the target actor instead.")
 	void ForceSingleNetUpdateFor(class AActor* Target);
 
 	/** Set the view target
diff --git a/Engine/Source/Runtime/Engine/Private/Components/CharacterMovementComponent.cpp b/Engine/Source/Runtime/Engine/Private/Components/CharacterMovementComponent.cpp
--- a/Engine/Source/Runtime/Engine/Private/Components/CharacterMovementComponent.cpp
+++ b/Engine/Source/Runtime/Engine/Private/Components/CharacterMovementComponent.cpp
@@ -2703,13 +2703,11 @@ void UCharacterMovementComponent::PerformMovement(float DeltaSeconds)
 		UNetDriver* NetDriver = MyWorld->GetNetDriver();
 		if (NetDriver && NetDriver->IsServer())
 		{
-			FNetworkObjectInfo* NetActor = NetDriver->FindNetworkObjectInfo(CharacterOwner);
-				
-			if (NetActor && MyWorld->GetTimeSeconds() <= NetActor->NextUpdateTime && NetDriver->IsNetworkActorUpdateFrequencyThrottled(*NetActor))
+			if (!NetDriver->IsPendingNetUpdate(CharacterOwner) && NetDriver->IsNetworkActorUpdateFrequencyThrottled(CharacterOwner))
 			{
 				if (ShouldCancelAdaptiveReplication())
 				{
-					NetDriver->CancelAdaptiveReplication(*NetActor);
+					NetDriver->CancelAdaptiveReplication(CharacterOwner);
 				}
 			}
 		}
diff --git a/Engine/Source/Runtime/Engine/Private/DemoNetDriver.cpp b/Engine/Source/Runtime/Engine/Private/DemoNetDriver.cpp
--- a/Engine/Source/Runtime/Engine/Private/DemoNetDriver.cpp
+++ b/Engine/Source/Runtime/Engine/Private/DemoNetDriver.cpp
@@ -5368,13 +5368,15 @@ void UDemoNetDriver::ForceNetUpdate(AActor* Actor)
 	}
 	else
 	{
+		PRAGMA_DISABLE_DEPRECATION_WARNINGS
 		if (FNetworkObjectInfo* NetActor = FindNetworkObjectInfo(Actor))
 		{
 			// replays use update times relative to DemoCurrentTime and not World->TimeSeconds
 			NetActor->NextUpdateTime = GetDemoCurrentTime() - 0.01f;
 		}
+		PRAGMA_ENABLE_DEPRECATION_WARNINGS
 	}
-	}
+}
 
 UChannel* UDemoNetDriver::InternalCreateChannelByName(const FName& ChName)
 {
diff --git a/Engine/Source/Runtime/Engine/Private/NetDriver.cpp b/Engine/Source/Runtime/Engine/Private/NetDriver.cpp
--- a/Engine/Source/Runtime/Engine/Private/NetDriver.cpp
+++ b/Engine/Source/Runtime/Engine/Private/NetDriver.cpp
@@ -739,7 +739,9 @@ FNetworkObjectInfo* UNetDriver::FindOrAddNetworkObjectInfo(const AActor* InActor
 #if UE_WITH_IRIS
 	if (ReplicationSystem)
 	{
+		PRAGMA_DISABLE_DEPRECATION_WARNINGS
 		return FindNetworkObjectInfo(InActor);
+		PRAGMA_ENABLE_DEPRECATION_WARNINGS
 	}
 #endif // UE_WITH_IRIS
 
@@ -805,15 +807,30 @@ bool UNetDriver::IsNetworkActorUpdateFrequencyThrottled(const AActor* InActor) c
 	bool bThrottled = false;
 	if (InActor && IsAdaptiveNetUpdateFrequencyEnabled())
 	{
+		PRAGMA_DISABLE_DEPRECATION_WARNINGS
 		if (const FNetworkObjectInfo* NetActor = FindNetworkObjectInfo(InActor))
 		{
 			bThrottled = IsNetworkActorUpdateFrequencyThrottled(*NetActor);
 		}
+		PRAGMA_ENABLE_DEPRECATION_WARNINGS
 	}
 
 	return bThrottled;
 }
 
+void UNetDriver::CancelAdaptiveReplication(const AActor* InActor)
+{
+	if (InActor && IsAdaptiveNetUpdateFrequencyEnabled())
+	{
+		PRAGMA_DISABLE_DEPRECATION_WARNINGS
+		if (FNetworkObjectInfo* NetActor = FindNetworkObjectInfo(InActor))
+		{
+			CancelAdaptiveReplication(*NetActor);
+		}
+		PRAGMA_ENABLE_DEPRECATION_WARNINGS
+	}
+}
+
 void UNetDriver::CancelAdaptiveReplication(FNetworkObjectInfo& InNetworkActor)
 {
 	if (IsAdaptiveNetUpdateFrequencyEnabled())
@@ -834,6 +851,21 @@ void UNetDriver::CancelAdaptiveReplication(FNetworkObjectInfo& InNetworkActor)
 	}
 }
 
+bool UNetDriver::IsPendingNetUpdate(const AActor* InActor) const
+{
+	if (World)
+	{
+		PRAGMA_DISABLE_DEPRECATION_WARNINGS
+		if (const FNetworkObjectInfo* NetActor = FindNetworkObjectInfo(InActor))
+		{
+			return NetActor->bPendingNetUpdate || (NetActor->NextUpdateTime < World->GetTimeSeconds());
+		}
+		PRAGMA_ENABLE_BUFFER_OVERRUN_WARNING
+	}
+		
+	return false;
+}
+
 static TAutoConsoleVariable<int32> CVarOptimizedRemapping( TEXT( "net.OptimizedRemapping" ), 1, TEXT( "Uses optimized path to remap unmapped network guids" ) );
 static TAutoConsoleVariable<int32> CVarMaxClientGuidRemaps( TEXT( "net.MaxClientGuidRemaps" ), 100, TEXT( "Max client resolves of unmapped network guids per tick" ) );
 
@@ -3600,10 +3632,12 @@ void UNetDriver::ForceNetUpdate(AActor* Actor)
 #endif // UE_WITH_IRIS
 	
 	// Legacy implementation
+	PRAGMA_DISABLE_DEPRECATION_WARNINGS
 	if (FNetworkObjectInfo* NetActor = FindNetworkObjectInfo(Actor))
 	{
 		NetActor->NextUpdateTime = World ? (World->TimeSeconds - 0.01f) : 0.0;
 	}
+	PRAGMA_ENABLE_DEPRECATION_WARNINGS
 }
 
 void UNetDriver::ForceAllActorsNetUpdateTime(float NetUpdateTimeOffset, TFunctionRef<bool(const AActor* const)> ValidActorTestFunc)
@@ -5590,7 +5624,9 @@ void UNetDriver::DrawNetDriverDebug()
 			ExtraStateDrawColor = FColor::Red;
 		}
 		
+		PRAGMA_DISABLE_DEPRECATION_WARNINGS
 		const FNetworkObjectInfo* NetworkObjectInfo = Connection->Driver->FindNetworkObjectInfo( *It );
+		PRAGMA_ENABLE_DEPRECATION_WARNINGS
 
 		FColor DrawColor;
 		if ( NetworkObjectInfo && NetworkObjectInfo->DormantConnections.Contains( Connection ) )
diff --git a/Engine/Source/Runtime/Engine/Private/NetworkObjectList.cpp b/Engine/Source/Runtime/Engine/Private/NetworkObjectList.cpp
--- a/Engine/Source/Runtime/Engine/Private/NetworkObjectList.cpp
+++ b/Engine/Source/Runtime/Engine/Private/NetworkObjectList.cpp
@@ -417,9 +417,10 @@ void FNetworkObjectList::ClearRecentlyDormantConnection(AActor* const Actor, UNe
 void FNetworkObjectList::OnActorIsTraveling(AActor* TravelingAtor)
 {
 	TSharedPtr<FNetworkObjectInfo>* NetworkObjectInfoPtr = AllNetworkObjects.Find(TravelingAtor);
-	check(NetworkObjectInfoPtr);
-
-	SeamlessTravelingObjects.Add(*NetworkObjectInfoPtr);
+	if (NetworkObjectInfoPtr)
+	{
+		SeamlessTravelingObjects.Add(*NetworkObjectInfoPtr);
+	}
 }
 
 void FNetworkObjectList::OnPostSeamlessTravel()
diff --git a/Engine/Source/Runtime/Engine/Private/Pawn.cpp b/Engine/Source/Runtime/Engine/Private/Pawn.cpp
--- a/Engine/Source/Runtime/Engine/Private/Pawn.cpp
+++ b/Engine/Source/Runtime/Engine/Private/Pawn.cpp
@@ -395,7 +395,7 @@ void APawn::BecomeViewTarget(APlayerController* PC)
 
 	if (GetNetMode() != NM_Client)
 	{
-		PC->ForceSingleNetUpdateFor(this);
+		ForceNetUpdate();
 	}
 
 	if (GetNetMode() != NM_DedicatedServer)
diff --git a/Engine/Source/Runtime/Engine/Private/PlayerController.cpp b/Engine/Source/Runtime/Engine/Private/PlayerController.cpp
--- a/Engine/Source/Runtime/Engine/Private/PlayerController.cpp
+++ b/Engine/Source/Runtime/Engine/Private/PlayerController.cpp
@@ -681,31 +681,9 @@ void APlayerController::ForceSingleNetUpdateFor(AActor* Target)
 	{
 		UE_LOG(LogPlayerController, Warning, TEXT("PlayerController::ForceSingleNetUpdateFor(): No Target specified"));
 	}
-	else if (GetNetMode() == NM_Client)
-	{
-		UE_LOG(LogPlayerController, Warning, TEXT("PlayerController::ForceSingleNetUpdateFor(): Only valid on server"));
-	}
 	else
 	{
-		if (UNetConnection* Conn = Cast<UNetConnection>(Player))
-		{
-			if (Conn->GetUChildConnection() != NULL)
-			{
-				Conn = ((UChildConnection*)Conn)->Parent;
-				checkSlow(Conn != NULL);
-			}
-
-			if (UActorChannel* Channel = Conn->FindActorChannelRef(Target))
-			{
-				if (UNetDriver* NetDriver = Conn->GetDriver())
-				{
-					if (FNetworkObjectInfo* NetActor = NetDriver->FindOrAddNetworkObjectInfo(Target))
-					{
-						NetActor->bPendingNetUpdate = true; // will cause some other clients to do lesser checks too, but that's unavoidable with the current functionality
-					}
-				}
-			}
-		}
+		Target->ForceNetUpdate();
 	}
 }
 
@@ -1948,9 +1926,10 @@ bool APlayerController::SetPause( bool bPause, FCanUnpause CanUnpauseDelegate)
 				// Pause gamepad rumbling too if needed
 				bResult = GameMode->SetPause(this, CanUnpauseDelegate);
 
-				// Force an update, otherwise since the game time is not updating, the net driver
-				// might not see that it is time for the world settings actor to replicate
-				ForceSingleNetUpdateFor(GetWorldSettings());
+				if (AWorldSettings* WorldSettings = GetWorldSettings())
+				{
+					WorldSettings->ForceNetUpdate();
+				}
 			}
 			else if (!bPause && bCurrentPauseState)
 			{
diff --git a/Engine/Source/Runtime/Engine/Private/World.cpp b/Engine/Source/Runtime/Engine/Private/World.cpp
--- a/Engine/Source/Runtime/Engine/Private/World.cpp
+++ b/Engine/Source/Runtime/Engine/Private/World.cpp
@@ -7162,21 +7162,16 @@ UWorld* FSeamlessTravelHandler::Tick()
 			// Rename dynamic actors in the old world's PersistentLevel that we want to keep into the new world
 			auto ProcessActor = [this, &KeepAnnotation, &ActuallyKeptActors, NetDriver](AActor* TheActor) -> bool
 			{
-				const FNetworkObjectInfo* NetworkObjectInfo = NetDriver ? NetDriver->FindNetworkObjectInfo(TheActor) : nullptr;
-
 				const bool bIsInCurrentLevel	= TheActor->GetLevel() == CurrentWorld->PersistentLevel;
 				const bool bManuallyMarkedKeep	= KeepAnnotation.Get(TheActor);
-				const bool bDormant				= NetworkObjectInfo && NetDriver && NetDriver->ServerConnection && NetworkObjectInfo->DormantConnections.Contains(NetDriver->ServerConnection);
+				const bool bDormant				= TheActor->GetIsReplicated() && (TheActor->NetDormancy > DORM_Awake);
 				const bool bKeepNonOwnedActor	= TheActor->GetLocalRole() < ROLE_Authority && !bDormant && !TheActor->IsNetStartupActor();
 				const bool bForceExcludeActor	= TheActor->IsA(ALevelScriptActor::StaticClass());
 
 				// Keep if it's in the current level AND it isn't specifically excluded AND it was either marked as should keep OR we don't own this actor
 				if (bIsInCurrentLevel && !bForceExcludeActor && (bManuallyMarkedKeep || bKeepNonOwnedActor))
 				{
-					if (NetworkObjectInfo)
-					{
-						NetDriver->NotifyActorIsTraveling(TheActor);
-					}
+					NetDriver->NotifyActorIsTraveling(TheActor);
 
 					ActuallyKeptActors.Add(TheActor);
 					return true;
