Prevent shaders from compiling if we don't use them to render in a commandlet.

diff --git a/Engine/Source/Runtime/Engine/Private/ShaderCompiler/ShaderCompiler.cpp b/Engine/Source/Runtime/Engine/Private/ShaderCompiler/ShaderCompiler.cpp
--- a/Engine/Source/Runtime/Engine/Private/ShaderCompiler/ShaderCompiler.cpp
+++ b/Engine/Source/Runtime/Engine/Private/ShaderCompiler/ShaderCompiler.cpp
@@ -131,10 +131,17 @@ static bool IsShaderJobCacheDDCRemotePolicyEnabled()
 	return CVarJobCacheDDCPolicy.GetValueOnAnyThread();
 }
 
+
 bool IsShaderJobCacheDDCEnabled()
 {
+#if WITH_EDITOR
+	static const bool bForceAllowShaderCompilerJobCache = FParse::Param(FCommandLine::Get(), TEXT("forceAllowShaderCompilerJobCache"));
+#else
+	const bool bForceAllowShaderCompilerJobCache = false;
+#endif
+
 	// For now we only support the editor and not commandlets like the cooker.
-	if (GIsEditor && !IsRunningCommandlet())
+	if (GIsEditor && (!IsRunningCommandlet() || bForceAllowShaderCompilerJobCache))
 	{
 		// job cache itself must be enabled first
 		return GShaderCompilerJobCache && CVarJobCacheDDC.GetValueOnAnyThread();
