Expose (read only) cluster tree of HISM to help debug workflows

diff --git a/Engine/Source/Runtime/Engine/Classes/Components/HierarchicalInstancedStaticMeshComponent.h b/Engine/Source/Runtime/Engine/Classes/Components/HierarchicalInstancedStaticMeshComponent.h
--- a/Engine/Source/Runtime/Engine/Classes/Components/HierarchicalInstancedStaticMeshComponent.h
+++ b/Engine/Source/Runtime/Engine/Classes/Components/HierarchicalInstancedStaticMeshComponent.h
@@ -261,6 +261,7 @@ public:
 	void AcceptPrebuiltTree(TArray<FInstancedStaticMeshInstanceData>& InInstanceData, TArray<FClusterNode>& InClusterTree, int32 InOcclusionLayerNumNodes, int32 InNumBuiltRenderInstances);
 	bool IsAsyncBuilding() const { return bIsAsyncBuilding; }
 	bool IsTreeFullyBuilt() const { return !bIsOutOfDate; }
+	void GetTree(TArray<FClusterNode>& OutClusterTree);
 
 	/** Heuristic for the number of leaves in the tree **/
 	int32 DesiredInstancesPerLeaf();
diff --git a/Engine/Source/Runtime/Engine/Private/HierarchicalInstancedStaticMesh.cpp b/Engine/Source/Runtime/Engine/Private/HierarchicalInstancedStaticMesh.cpp
--- a/Engine/Source/Runtime/Engine/Private/HierarchicalInstancedStaticMesh.cpp
+++ b/Engine/Source/Runtime/Engine/Private/HierarchicalInstancedStaticMesh.cpp
@@ -3458,6 +3458,12 @@ void UHierarchicalInstancedStaticMeshComponent::GetOverlappingBoxTransforms(cons
 	}
 }
 
+void UHierarchicalInstancedStaticMeshComponent::GetTree(TArray<FClusterNode>& OutClusterTree)
+{
+	OutClusterTree = *ClusterTreePtr;
+}
+
+
 void UHierarchicalInstancedStaticMeshComponent::GetNavigationPerInstanceTransforms(const FBox& AreaBox, TArray<FTransform>& InstanceData) const
 {
 	if (IsTreeFullyBuilt())
