Add a EHardwareDeviceSupportedFeatures and EHardwareDevicePrimaryType enums to hardware device identifiers.

These will allow hardware device identifiers to set flags for themselves to identify some common traits that they have, as well as check what kind of support they have. These are config flags, so if a game wants to change what they support for a given input device they can easily do so, as well as add in some custom hardware types if desired.

diff --git a/Engine/Config/Windows/WindowsInput.ini b/Engine/Config/Windows/WindowsInput.ini
--- a/Engine/Config/Windows/WindowsInput.ini
+++ b/Engine/Config/Windows/WindowsInput.ini
@@ -1,2 +1,2 @@
 [InputPlatformSettings_Windows InputPlatformSettings]
-+HardwareDevices=(InputClassName="XInputInterface",HardwareDeviceIdentifier="XInputController")
\ No newline at end of file
++HardwareDevices=(InputClassName="XInputInterface",HardwareDeviceIdentifier="XInputController",PrimaryDeviceType=Gamepad,SupportedFeaturesMask=260)
\ No newline at end of file
diff --git a/Engine/Plugins/Runtime/Steam/SteamController/Config/Input.ini b/Engine/Plugins/Runtime/Steam/SteamController/Config/Input.ini
--- a/Engine/Plugins/Runtime/Steam/SteamController/Config/Input.ini
+++ b/Engine/Plugins/Runtime/Steam/SteamController/Config/Input.ini
@@ -1,2 +1,2 @@
 [InputPlatformSettings_Windows InputPlatformSettings]
-+HardwareDevices=(InputClassName="SteamController",HardwareDeviceIdentifier="SteamController")
\ No newline at end of file
++HardwareDevices=(InputClassName="SteamController",HardwareDeviceIdentifier="SteamController",PrimaryDeviceType=Gamepad,SupportedFeaturesMask=260)
\ No newline at end of file
diff --git a/Engine/Plugins/Runtime/Windows/WinDualShock/Config/Input.ini b/Engine/Plugins/Runtime/Windows/WinDualShock/Config/Input.ini
--- a/Engine/Plugins/Runtime/Windows/WinDualShock/Config/Input.ini
+++ b/Engine/Plugins/Runtime/Windows/WinDualShock/Config/Input.ini
@@ -1,3 +1,3 @@
 [InputPlatformSettings_Windows InputPlatformSettings]
-+HardwareDevices=(InputClassName="FWinDualShock",HardwareDeviceIdentifier="DualSense")
-+HardwareDevices=(InputClassName="FWinDualShock",HardwareDeviceIdentifier="DualShock4")
++HardwareDevices=(InputClassName="FWinDualShock",HardwareDeviceIdentifier="DualSense",PrimaryDeviceType=Gamepad,SupportedFeaturesMask=1996)
++HardwareDevices=(InputClassName="FWinDualShock",HardwareDeviceIdentifier="DualShock4",PrimaryDeviceType=Gamepad,SupportedFeaturesMask=332)
diff --git a/Engine/Source/Runtime/Engine/Classes/GameFramework/InputSettings.h b/Engine/Source/Runtime/Engine/Classes/GameFramework/InputSettings.h
--- a/Engine/Source/Runtime/Engine/Classes/GameFramework/InputSettings.h
+++ b/Engine/Source/Runtime/Engine/Classes/GameFramework/InputSettings.h
@@ -293,6 +293,97 @@ private:
 	void AddInternationalConsoleKey();
 };
 
+
+/**
+ * A bitmask of supported features that a hardware device has. 
+ */
+UENUM(BlueprintType, Meta = (Bitflags, UseEnumValuesAsMaskValuesInEditor = "true"))
+namespace EHardwareDeviceSupportedFeatures
+{
+	enum Type : int32
+	{
+		/** A device that has not specified the type  */
+		Unspecified			= 0x00000000,
+
+		/** This device can support basic key presses */
+		Keypress			= 0x00000001,
+
+		/** This device can handle basic pointer behavior, such as a mouse */
+		Pointer				= 0x00000002,
+
+		/** This device has basic gamepad support */
+		Gamepad				= 0x00000004,
+
+		/** This device supports touch in some capactiy (tablet, controller with a touch pad, etc) */
+		Touch				= 0x00000008,
+
+		/** Does this device have a camera on it that we can access? */
+		Camera				= 0x00000010,
+
+		/** Can this device track motion in a 3D space? (VR controllers, headset, etc) */
+		MotionTracking		= 0x00000020,
+
+		/** This hardware supports setting a light color (such as an LED light bar) */
+		Lights				= 0x00000040,
+
+		/** Does this device have trigger haptics available? */
+		TriggerHaptics		= 0x00000080,
+
+		/** Flagged true if this device supports force feedback */
+		ForceFeedback		= 0x00000100,
+
+		/** Does this device support vibrations sourced from an audio file? */
+		AudioBasedVibrations= 0x00000200,
+
+		/** This device can track acceleration in the users physical space */
+		Acceleration		= 0x00000400,
+
+		/** This is a virtual device simulating input, not a physical device */
+		Virtual				= 0x00000800,
+
+		/** This device has a microphone on it that you can get audio from */
+		Microphone			= 0x00001000,
+
+		/** This device can track the orientation in world space, such as a gyroscope */
+		Orientation			= 0x00002000,
+
+		/** Some custom flags that can be used in your game if you have custom hardware! */
+		CustomA				= 0x01000000,
+		CustomB				= 0x02000000,
+		CustomC				= 0x04000000,
+		CustomD				= 0x08000000,
+
+		/** A flag for ALL supported device flags */
+		All = 0x7FFFFFFF UMETA(Hidden),
+	};
+}
+
+ENUM_CLASS_FLAGS(EHardwareDeviceSupportedFeatures::Type);
+
+/**
+ * What is the primary use of an input device type? 
+ * Each hardware device can only be one primary type.
+ */
+UENUM(BlueprintType)
+enum class EHardwareDevicePrimaryType : uint8
+{
+	Unspecified,
+	KeyboardAndMouse,
+	Gamepad,
+	Touch,
+	MotionTracking,
+	RacingWheel,
+	FlightStick,
+	Camera,
+	Instrument,
+
+	// Some custom devices that can be used for your game specific hardware if desired
+	CustomTypeA,
+	CustomTypeB,
+	CustomTypeC,
+	CustomTypeD,
+};
+
 /**
 * An identifier that can be used to determine what input devices are available based on the FInputDeviceScope.
 * These mappings should match a FInputDeviceScope that is used by an IInputDevice
@@ -303,7 +394,12 @@ struct ENGINE_API FHardwareDeviceIdentifier
 	GENERATED_BODY()
 
 	FHardwareDeviceIdentifier();
-	FHardwareDeviceIdentifier(const FName InClassName, const FName InHardwareDeviceIdentifier);
+	
+	FHardwareDeviceIdentifier(
+		const FName InClassName,
+		const FName InHardwareDeviceIdentifier,
+		EHardwareDevicePrimaryType InPrimaryType = EHardwareDevicePrimaryType::Unspecified,
+		EHardwareDeviceSupportedFeatures::Type Flags = EHardwareDeviceSupportedFeatures::Unspecified);
 	
 	/** 
 	* The name of the Input Class that uses this hardware device.
@@ -319,6 +415,20 @@ struct ENGINE_API FHardwareDeviceIdentifier
 	UPROPERTY(EditDefaultsOnly, BlueprintReadOnly, Category = "Hardware")
 	FName HardwareDeviceIdentifier;
 
+	/** The generic type that this hardware identifies as. This can be used to easily determine behaviors  */
+	UPROPERTY(EditDefaultsOnly, BlueprintReadOnly, Category = "Hardware")
+	EHardwareDevicePrimaryType PrimaryDeviceType;
+	
+	/** Flags that represent this hardware device's traits */
+	UPROPERTY(EditDefaultsOnly, BlueprintReadOnly, Category = "Hardware", meta=(Bitmask, BitmaskEnum="/Script/Engine.EHardwareDeviceSupportedFeatures"))
+	int32 SupportedFeaturesMask;
+
+	/** Returns true if this hardware device has ANY of the given supported feature flags */
+	bool HasAnySupportedFeatures(const EHardwareDeviceSupportedFeatures::Type FlagsToCheck) const;
+
+	/** Returns true if this hardware device has ALL of the given supported feature flags */
+	bool HasAllSupportedFeatures(const EHardwareDeviceSupportedFeatures::Type FlagsToCheck) const;
+
 	/** Returns true if this hardware device Identifier has valid names */
 	bool IsValid() const;
 
@@ -334,10 +444,7 @@ struct ENGINE_API FHardwareDeviceIdentifier
 	/** Hardware device ID that represents a keyboard and mouse. This is what will be set when an Input Event's FKey is not a gamepad key. */
 	static FHardwareDeviceIdentifier DefaultKeyboardAndMouse;
 
-	bool operator==(const FHardwareDeviceIdentifier& Other) const
-	{
-		return Other.InputClassName == InputClassName && Other.HardwareDeviceIdentifier == HardwareDeviceIdentifier;
-	}
+	bool operator==(const FHardwareDeviceIdentifier& Other) const;
 
 	ENGINE_API friend uint32 GetTypeHash(const FHardwareDeviceIdentifier& InDevice);
 	ENGINE_API friend FArchive& operator<<(FArchive& Ar, FHardwareDeviceIdentifier& InDevice);
diff --git a/Engine/Source/Runtime/Engine/Private/GameFramework/InputDeviceSubsystem.cpp b/Engine/Source/Runtime/Engine/Private/GameFramework/InputDeviceSubsystem.cpp
--- a/Engine/Source/Runtime/Engine/Private/GameFramework/InputDeviceSubsystem.cpp
+++ b/Engine/Source/Runtime/Engine/Private/GameFramework/InputDeviceSubsystem.cpp
@@ -42,7 +42,17 @@ class FInputDeviceSubsystemProcessor : public IInputProcessor
 			if (const FInputDeviceScope* Scope = FInputDeviceScope::GetCurrent())
 			{
 				// TODO: Refactor FInputDeviceScope to use FName's instead of a FString for HardwareDeviceIdentifier
-				SubSystem->SetMostRecentlyUsedHardwareDevice(InDeviceId, { Scope->InputDeviceName, FName(*Scope->HardwareDeviceIdentifier) });
+				// Look up a hardware device ID from the config that matches this identifier
+				const FName DeviceScopeName(*Scope->HardwareDeviceIdentifier);
+				if (const FHardwareDeviceIdentifier* Hardware = UInputPlatformSettings::Get()->GetHardwareDeviceForClassName(DeviceScopeName))
+				{
+					SubSystem->SetMostRecentlyUsedHardwareDevice(InDeviceId, *Hardware);
+				}
+				// If there isn't one specified in the config file, we can just use the device name and scope name to make one.
+				else
+				{
+					SubSystem->SetMostRecentlyUsedHardwareDevice(InDeviceId, { Scope->InputDeviceName, DeviceScopeName });
+				}
 			}
 			// If there isn't a recent input device scope, then we can check if the key was from a keyboard and mouse
 			else if (!Key.IsGamepadKey())
diff --git a/Engine/Source/Runtime/Engine/Private/UserInterface/InputSettings.cpp b/Engine/Source/Runtime/Engine/Private/UserInterface/InputSettings.cpp
--- a/Engine/Source/Runtime/Engine/Private/UserInterface/InputSettings.cpp
+++ b/Engine/Source/Runtime/Engine/Private/UserInterface/InputSettings.cpp
@@ -506,8 +506,21 @@ void UInputSettings::SetDefaultInputComponentClass(TSubclassOf<UInputComponent>
 /////////////////////////////////////////////////////////////
 // FHardwareDeviceIdentifier
 
-FHardwareDeviceIdentifier FHardwareDeviceIdentifier::Invalid = { NAME_None, NAME_None };
-FHardwareDeviceIdentifier FHardwareDeviceIdentifier::DefaultKeyboardAndMouse = { TEXT("DefaultKeyboardAndMouse"), TEXT("KBM") };
+FHardwareDeviceIdentifier FHardwareDeviceIdentifier::Invalid =
+	{
+		/* .InputClassName= */ NAME_None,
+		/* .HardwareDeviceIdentifier= */ NAME_None,
+		/* .PrimaryDeviceType= */ EHardwareDevicePrimaryType::Unspecified,
+		/* .SupportedFeaturesMask= */ EHardwareDeviceSupportedFeatures::Type::Unspecified
+	};
+
+FHardwareDeviceIdentifier FHardwareDeviceIdentifier::DefaultKeyboardAndMouse = 
+	{
+		/* .InputClassName= */ TEXT("DefaultKeyboardAndMouse"),
+		/* .HardwareDeviceIdentifier= */ TEXT("KBM"),
+		/* .PrimaryDeviceType= */ EHardwareDevicePrimaryType::KeyboardAndMouse,
+		/* .SupportedFeaturesMask= */ EHardwareDeviceSupportedFeatures::Type::Keypress | EHardwareDeviceSupportedFeatures::Type::Pointer
+	};
 
 // Default to the invalid hardware device identifier
 FHardwareDeviceIdentifier::FHardwareDeviceIdentifier()
@@ -517,9 +530,11 @@ FHardwareDeviceIdentifier::FHardwareDeviceIdentifier()
 	
 }
 
-FHardwareDeviceIdentifier::FHardwareDeviceIdentifier(const FName InClassName, const FName InHardwareDeviceIdentifier)
+FHardwareDeviceIdentifier::FHardwareDeviceIdentifier(const FName InClassName, const FName InHardwareDeviceIdentifier, EHardwareDevicePrimaryType InPrimaryType, EHardwareDeviceSupportedFeatures::Type Flags)
 	: InputClassName(InClassName)
 	, HardwareDeviceIdentifier(InHardwareDeviceIdentifier)
+	, PrimaryDeviceType(InPrimaryType)
+	, SupportedFeaturesMask(Flags)
 {
 
 }
@@ -533,10 +548,22 @@ FArchive& operator<<(FArchive& Ar, FHardwareDeviceIdentifier& InDevice)
 {
 	Ar << InDevice.InputClassName;
 	Ar << InDevice.HardwareDeviceIdentifier;
+	Ar << InDevice.PrimaryDeviceType;
+	Ar << InDevice.SupportedFeaturesMask;
 	
 	return Ar;
 }
 
+bool FHardwareDeviceIdentifier::HasAnySupportedFeatures(const EHardwareDeviceSupportedFeatures::Type FlagsToCheck) const
+{
+	return (SupportedFeaturesMask & static_cast<int32>(FlagsToCheck)) != 0;
+}
+
+bool FHardwareDeviceIdentifier::HasAllSupportedFeatures(const EHardwareDeviceSupportedFeatures::Type FlagsToCheck) const
+{
+	return ((SupportedFeaturesMask & FlagsToCheck) == FlagsToCheck);
+}
+
 bool FHardwareDeviceIdentifier::IsValid() const
 {
 	return InputClassName.IsValid() && HardwareDeviceIdentifier.IsValid();
@@ -547,6 +574,15 @@ FString FHardwareDeviceIdentifier::ToString() const
 	return FString::Printf(TEXT("%s::%s"), *InputClassName.ToString(), *HardwareDeviceIdentifier.ToString());
 }
 
+bool FHardwareDeviceIdentifier::operator==(const FHardwareDeviceIdentifier& Other) const
+{
+	return
+		Other.InputClassName == InputClassName &&
+		Other.HardwareDeviceIdentifier == HardwareDeviceIdentifier &&
+		Other.SupportedFeaturesMask == SupportedFeaturesMask &&
+		Other.PrimaryDeviceType == PrimaryDeviceType;	
+}
+
 //////////////////////////////////////////////////////////////////
 // UInputPlatformSettings
 
