Fix race condition by invoking FlushAsyncLoad before disabling RHIThread when r.RHISetGPUCaptureOptions is set to 0

#rnx

diff --git a/Engine/Source/Runtime/Engine/Private/UnrealEngine.cpp b/Engine/Source/Runtime/Engine/Private/UnrealEngine.cpp
--- a/Engine/Source/Runtime/Engine/Private/UnrealEngine.cpp
+++ b/Engine/Source/Runtime/Engine/Private/UnrealEngine.cpp
@@ -18444,6 +18444,33 @@ void UEngine::SetPriorityAndAffinityOnGameThread()
 	::SetPriorityAndAffinityOnGameThread();
 }
 
+// Flush async loading before disabling rhi thread when executing r.RHISetGPUCaptureOptions console command to avoid a race condition
+// Handling of this console command is part of unrealengine translation unit instead of RHI ones because
+// FlushAsyncLoad shouldn't be invoked from RHI cpp files
+static void BaseRHISetGPUCaptureOptions(const TArray<FString>& Args, UWorld* World)
+{
+	if (Args.Num() > 0)
+	{
+		// Make sure there isnt any loading asset in flight that would possibly need rendering thread in post init
+		FlushAsyncLoading();
+
+		const bool bEnabled = Args[0].ToBool();
+		GDynamicRHI->EnableIdealGPUCaptureOptions(bEnabled);
+	}
+	else
+	{
+		UE_LOG(LogRHI, Display, TEXT("Usage: r.RHISetGPUCaptureOptions 0 or r.RHISetGPUCaptureOptions 1"));
+	}
+}
+
+static FAutoConsoleCommandWithWorldAndArgs GBaseRHISetGPUCaptureOptions(
+	TEXT("r.RHISetGPUCaptureOptions"),
+	TEXT("Utility function to change multiple CVARs useful when profiling or debugging GPU rendering. Setting to 1 or 0 will guarantee all options are in the appropriate state.\n")
+	TEXT("r.rhithread.enable, r.rhicmdbypass, r.showmaterialdrawevents, toggledrawevents\n")
+	TEXT("Platform RHI's may implement more feature toggles."),
+	FConsoleCommandWithWorldAndArgsDelegate::CreateStatic(&BaseRHISetGPUCaptureOptions)
+);
+
 #if !UE_BUILD_SHIPPING
 
 static void PakFileTest(const TArray<FString>& Args)
diff --git a/Engine/Source/Runtime/RHI/Private/DynamicRHI.cpp b/Engine/Source/Runtime/RHI/Private/DynamicRHI.cpp
--- a/Engine/Source/Runtime/RHI/Private/DynamicRHI.cpp
+++ b/Engine/Source/Runtime/RHI/Private/DynamicRHI.cpp
@@ -434,28 +434,6 @@ void RHIExit()
 	FRHICommandListImmediate::CleanupGraphEvents();
 }
 
-
-static void BaseRHISetGPUCaptureOptions(const TArray<FString>& Args, UWorld* World)
-{
-	if (Args.Num() > 0)
-	{
-		const bool bEnabled = Args[0].ToBool();
-		GDynamicRHI->EnableIdealGPUCaptureOptions(bEnabled);
-	}
-	else
-	{
-		UE_LOG(LogRHI, Display, TEXT("Usage: r.RHISetGPUCaptureOptions 0 or r.RHISetGPUCaptureOptions 1"));
-	}
-}
-
-static FAutoConsoleCommandWithWorldAndArgs GBaseRHISetGPUCaptureOptions(
-	TEXT("r.RHISetGPUCaptureOptions"),
-	TEXT("Utility function to change multiple CVARs useful when profiling or debugging GPU rendering. Setting to 1 or 0 will guarantee all options are in the appropriate state.\n")
-	TEXT("r.rhithread.enable, r.rhicmdbypass, r.showmaterialdrawevents, toggledrawevents\n")
-	TEXT("Platform RHI's may implement more feature toggles."),
-	FConsoleCommandWithWorldAndArgsDelegate::CreateStatic(&BaseRHISetGPUCaptureOptions)
-	);
-
 // Default fallback; will not work for non-8-bit surfaces and it's extremely slow.
 void FDynamicRHI::RHIReadSurfaceData(FRHITexture* Texture, FIntRect Rect, TArray<FLinearColor>& OutData, FReadSurfaceDataFlags InFlags)
 {
