Fixed asan for exec command tests: Make sure that the memory buffer has a null termination character before building a FString from it.

diff --git a/Engine/Source/Runtime/Engine/Private/Tests/ExecCommandsTests.cpp b/Engine/Source/Runtime/Engine/Private/Tests/ExecCommandsTests.cpp
--- a/Engine/Source/Runtime/Engine/Private/Tests/ExecCommandsTests.cpp
+++ b/Engine/Source/Runtime/Engine/Private/Tests/ExecCommandsTests.cpp
@@ -253,7 +253,12 @@ bool FCommonCommandsExecCommandTest::RunTest(const FString& Parameters)
 	FMemoryWriter MemoryWriter(Buffer);
 	MemoryOutput.Dump(MemoryWriter);
 
-	FString ConsoleCommandsDump(reinterpret_cast<ANSICHAR*>(Buffer.GetData()), Buffer.Num());
+	if (static_cast<ANSICHAR>(Buffer.Last()) != '\0')
+	{
+		Buffer.Add('\0');
+	}
+
+	FString ConsoleCommandsDump(reinterpret_cast<ANSICHAR*>(Buffer.GetData()));
 
 	TArray<FString> CommandList;
 	ConsoleCommandsDump.ParseIntoArrayLines(CommandList);
