World Partition: detect level instance circular references and avoid crashing.

diff --git a/Engine/Plugins/Editor/DataValidation/Source/DataValidation/Private/WorldPartitionChangelistValidator.cpp b/Engine/Plugins/Editor/DataValidation/Source/DataValidation/Private/WorldPartitionChangelistValidator.cpp
--- a/Engine/Plugins/Editor/DataValidation/Source/DataValidation/Private/WorldPartitionChangelistValidator.cpp
+++ b/Engine/Plugins/Editor/DataValidation/Source/DataValidation/Private/WorldPartitionChangelistValidator.cpp
@@ -430,6 +430,12 @@ void UWorldPartitionChangelistValidator::OnLevelInstanceInvalidWorldAsset(const
 		case ELevelInstanceInvalidReason::WorldAssetHasInvalidContainer:
 			// We cannot treat that error as a validation error as it's possible to validate changelists without loading the world
 			break;
+		case ELevelInstanceInvalidReason::CirculalReference:
+			CurrentError = FText::Format(LOCTEXT("DataValidation.Changelist.WorldPartition.LevelInstanceCircularReference", "Level instance {0} has a circular reference {1}."),
+				FText::FromString(GetFullActorName(ActorDescView)),
+				FText::FromName(WorldAsset));
+			Errors->Add(CurrentError);
+			break;
 		};
 	}
 }
diff --git a/Engine/Source/Runtime/Engine/Private/WorldPartition/ErrorHandling/WorldPartitionStreamingGenerationLogErrorHandler.cpp b/Engine/Source/Runtime/Engine/Private/WorldPartition/ErrorHandling/WorldPartitionStreamingGenerationLogErrorHandler.cpp
--- a/Engine/Source/Runtime/Engine/Private/WorldPartition/ErrorHandling/WorldPartitionStreamingGenerationLogErrorHandler.cpp
+++ b/Engine/Source/Runtime/Engine/Private/WorldPartition/ErrorHandling/WorldPartitionStreamingGenerationLogErrorHandler.cpp
@@ -85,6 +85,9 @@ void FStreamingGenerationLogErrorHandler::OnLevelInstanceInvalidWorldAsset(const
 	case ELevelInstanceInvalidReason::WorldAssetHasInvalidContainer:
 		UE_ASSET_LOG_ACTORDESCVIEW(LogWorldPartition, Log, ActorDescView, TEXT("Actor %s world asset %s has an invalid container"), *ActorName, *WorldAsset.ToString());
 		break;
+	case ELevelInstanceInvalidReason::CirculalReference:
+		UE_ASSET_LOG_ACTORDESCVIEW(LogWorldPartition, Log, ActorDescView, TEXT("Actor %s world asset %s has a circular reference"), *ActorName, *WorldAsset.ToString());
+		break;
 	};
 }
 #endif
diff --git a/Engine/Source/Runtime/Engine/Private/WorldPartition/ErrorHandling/WorldPartitionStreamingGenerationTokenizedMessageErrorHandler.cpp b/Engine/Source/Runtime/Engine/Private/WorldPartition/ErrorHandling/WorldPartitionStreamingGenerationTokenizedMessageErrorHandler.cpp
--- a/Engine/Source/Runtime/Engine/Private/WorldPartition/ErrorHandling/WorldPartitionStreamingGenerationTokenizedMessageErrorHandler.cpp
+++ b/Engine/Source/Runtime/Engine/Private/WorldPartition/ErrorHandling/WorldPartitionStreamingGenerationTokenizedMessageErrorHandler.cpp
@@ -179,6 +179,10 @@ void ITokenizedMessageErrorHandler::OnLevelInstanceInvalidWorldAsset(const FWorl
 		MessageSeverity = EMessageSeverity::Error;
 		ReasonText = LOCTEXT("TokenMessage_WorldPartition_WorldAssetHasInvalidContainer", "has an invalid container");
 		break;
+	case ELevelInstanceInvalidReason::CirculalReference:
+		MessageSeverity = EMessageSeverity::Error;
+		ReasonText = LOCTEXT("TokenMessage_WorldPartition_WorldAssetHasCircularReference", "has a circular reference");
+		break;
 	};
 
 	const FString WorldAssetStr = WorldAsset.ToString();
diff --git a/Engine/Source/Runtime/Engine/Private/WorldPartition/WorldPartitionStreamingGeneration.cpp b/Engine/Source/Runtime/Engine/Private/WorldPartition/WorldPartitionStreamingGeneration.cpp
--- a/Engine/Source/Runtime/Engine/Private/WorldPartition/WorldPartitionStreamingGeneration.cpp
+++ b/Engine/Source/Runtime/Engine/Private/WorldPartition/WorldPartitionStreamingGeneration.cpp
@@ -453,6 +453,14 @@ class FWorldPartitionStreamingGenerator
 				continue;
 			}
 
+			if (ContainerInstancesStack.Contains(SubContainerInstance.Container->ContainerPackageName))
+			{
+				ErrorHandler->OnLevelInstanceInvalidWorldAsset(ContainerInstanceView, ContainerInstanceView.GetLevelPackage(), IStreamingGenerationErrorHandler::ELevelInstanceInvalidReason::CirculalReference);
+				continue;
+			}
+
+			ContainerInstancesStack.Add(SubContainerInstance.Container->ContainerPackageName);
+
 			check(SubContainerInstance.Container);
 			FContainerInstanceDescriptor SubContainerInstanceDescriptor;
 
@@ -460,8 +468,6 @@ class FWorldPartitionStreamingGenerator
 			SubContainerInstanceDescriptor.ID = SubContainerInstance.GetID();
 			SubContainerInstanceDescriptor.Container = SubContainerInstance.Container;
 			SubContainerInstanceDescriptor.Transform = SubContainerInstance.Transform * InContainerInstanceDescriptor.Transform;
-
-
 			SubContainerInstanceDescriptor.ParentID = InContainerInstanceDescriptor.ID;
 			SubContainerInstanceDescriptor.OwnerName = *ContainerInstanceView.GetActorLabelOrName().ToString();
 
@@ -476,6 +482,8 @@ class FWorldPartitionStreamingGenerator
 			SubContainerInstanceDescriptor.RuntimeDataLayers = InheritedRuntimeDataLayers;
 
 			CreateActorDescriptorViewsRecursive(SubContainerInstanceDescriptor);
+
+			verify(ContainerInstancesStack.Remove(SubContainerInstance.Container->ContainerPackageName));
 		}
 							
 		if (!InContainerInstanceDescriptor.ID.IsMainContainer())
@@ -999,6 +1007,9 @@ private:
 
 	/** List of containers participating in this streaming generation step */
 	TMap<FGuid, const UActorDescContainer*> ActorGuidsToContainerMap;
+
+	/** List of current container instances on the stack to detect circular references */
+	TSet<FName> ContainerInstancesStack;
 };
 
 bool UWorldPartition::GenerateStreaming(TArray<FString>* OutPackagesToGenerate)
diff --git a/Engine/Source/Runtime/Engine/Public/WorldPartition/ErrorHandling/WorldPartitionStreamingGenerationErrorHandler.h b/Engine/Source/Runtime/Engine/Public/WorldPartition/ErrorHandling/WorldPartitionStreamingGenerationErrorHandler.h
--- a/Engine/Source/Runtime/Engine/Public/WorldPartition/ErrorHandling/WorldPartitionStreamingGenerationErrorHandler.h
+++ b/Engine/Source/Runtime/Engine/Public/WorldPartition/ErrorHandling/WorldPartitionStreamingGenerationErrorHandler.h
@@ -79,7 +79,8 @@ public:
 		WorldAssetNotFound,
 		WorldAssetNotUsingExternalActors,
 		WorldAssetImcompatiblePartitioned,
-		WorldAssetHasInvalidContainer
+		WorldAssetHasInvalidContainer,
+		CirculalReference
 	};
 
 	virtual void OnLevelInstanceInvalidWorldAsset(const FWorldPartitionActorDescView& ActorDescView, FName WorldAsset, ELevelInstanceInvalidReason Reason) = 0;
