Fix bad pragma macro

diff --git a/Engine/Source/Runtime/Engine/Private/NetDriver.cpp b/Engine/Source/Runtime/Engine/Private/NetDriver.cpp
--- a/Engine/Source/Runtime/Engine/Private/NetDriver.cpp
+++ b/Engine/Source/Runtime/Engine/Private/NetDriver.cpp
@@ -860,7 +860,7 @@ bool UNetDriver::IsPendingNetUpdate(const AActor* InActor) const
 		{
 			return NetActor->bPendingNetUpdate || (NetActor->NextUpdateTime < World->GetTimeSeconds());
 		}
-		PRAGMA_ENABLE_BUFFER_OVERRUN_WARNING
+		PRAGMA_ENABLE_DEPRECATION_WARNINGS
 	}
 		
 	return false;
