Update instanced static mesh debug to only report data on the world the command has been ran on

diff --git a/Engine/Source/Runtime/Engine/Private/InstancedStaticMesh.cpp b/Engine/Source/Runtime/Engine/Private/InstancedStaticMesh.cpp
--- a/Engine/Source/Runtime/Engine/Private/InstancedStaticMesh.cpp
+++ b/Engine/Source/Runtime/Engine/Private/InstancedStaticMesh.cpp
@@ -155,7 +155,8 @@ class FISMExecHelper : public FSelfRegisteringExec
 			for(TObjectIterator<UInstancedStaticMeshComponent> It; It; ++It)
 			{
 				UInstancedStaticMeshComponent* ISMComponent = *It;
-				if((!ISMComponent->GetWorld() || !ISMComponent->GetWorld()->IsGameWorld()))
+				const UWorld* ComponentWorld = ISMComponent->GetWorld();
+				if (ComponentWorld!=InWorld || !ComponentWorld || !ComponentWorld->IsGameWorld())
 				{
 					continue;
 				}
@@ -218,7 +219,8 @@ class FISMExecHelper : public FSelfRegisteringExec
 			for (TObjectIterator<UInstancedStaticMeshComponent> It; It; ++It)
 			{
 				UInstancedStaticMeshComponent* ISMComponent = *It;
-				if ((!ISMComponent->GetWorld() || !ISMComponent->GetWorld()->IsGameWorld()))
+				const UWorld* ComponentWorld = ISMComponent->GetWorld();
+				if (ComponentWorld != InWorld || !ComponentWorld || !ComponentWorld->IsGameWorld())
 				{
 					continue;
 				}
