PR #10211: GetInputAnalogStickState overload fix (Contributed by [REDACTED_PERSON])

diff --git a/Engine/Source/Runtime/Engine/Private/PlayerController.cpp b/Engine/Source/Runtime/Engine/Private/PlayerController.cpp
--- a/Engine/Source/Runtime/Engine/Private/PlayerController.cpp
+++ b/Engine/Source/Runtime/Engine/Private/PlayerController.cpp
@@ -5706,7 +5706,7 @@ void APlayerController::GetInputAnalogStickState(EControllerAnalogStick::Type Wh
 void APlayerController::GetInputAnalogStickState(EControllerAnalogStick::Type WhichStick, double& StickX, double& StickY) const
 {
 	float DX, DY;
-	GetInputMouseDelta(DX, DY);
+	GetInputAnalogStickState(WhichStick, DX, DY);
 	StickX = DX;
 	StickY = DY;
 }
