SkeletalMeshComponent : Extracted some LOD bone generation functions, so these can be used from plugins.
Updated function ComputeRequiredBones to use the extracted functions (so the same code is used both at the engine and plugins).

diff --git a/Engine/Source/Runtime/Engine/Classes/Components/SkeletalMeshComponent.h b/Engine/Source/Runtime/Engine/Classes/Components/SkeletalMeshComponent.h
--- a/Engine/Source/Runtime/Engine/Classes/Components/SkeletalMeshComponent.h
+++ b/Engine/Source/Runtime/Engine/Classes/Components/SkeletalMeshComponent.h
@@ -1637,6 +1637,31 @@ public:
 	*/
 	void ComputeRequiredBones(TArray<FBoneIndexType>& OutRequiredBones, TArray<FBoneIndexType>& OutFillComponentSpaceTransformsRequiredBones, int32 LODIndex, bool bIgnorePhysicsAsset) const;
 
+	// Get the required virtual bones from the SkeletalMesh Reference Skeleton
+	static void GetRequiredVirtualBones(const USkeletalMesh* SkeletalMesh, TArray<FBoneIndexType>& OutRequiredBones);
+
+	// Get the bones used by the physics asset (if any), as these have to be updated even if LODed out
+	static void GetPhysicsRequiredBones(const USkeletalMesh* SkeletalMesh, const UPhysicsAsset* PhysicsAsset, TArray<FBoneIndexType>& OutRequiredBones);
+
+	// Removes the bones expicitly hidden or hiden by parent
+	static void ExcludeHiddenBones(const USkeletalMeshComponent* SkeletalMeshComponent, const USkeletalMesh* SkeletalMesh, TArray<FBoneIndexType>& OutRequiredBones);
+
+	// Get the bones used for mirroring in the skeletal mesh asset (if any)
+	static void GetMirroringRequiredBones(const USkeletalMesh* SkeletalMesh, TArray<FBoneIndexType>& OutRequiredBones);
+
+	// Get the bones that are used by sockets
+	// OutRequiredBones : the bones used by sockets that are set as always animate
+	// NeededBonesForFillComponentSpaceTransforms : the rest of bones used by sockets (that are not set to always animate)
+	static void GetSocketRequiredBones(const USkeletalMesh* SkeletalMesh, TArray<FBoneIndexType>& OutRequiredBones, TArray<FBoneIndexType>& NeededBonesForFillComponentSpaceTransforms);
+
+	// Get the bones required for shadow shapes
+	static void GetShadowShapeRequiredBones(const USkeletalMeshComponent* SkeletalMeshComponent, TArray<FBoneIndexType>& OutRequiredBones);
+
+	/** Takes sorted array Base and then adds any elements from sorted array Insert which is missing from it, preserving order.
+	 * this assumes both arrays are sorted and contain unique bone indices.
+	 */
+	static void MergeInBoneIndexArrays(TArray<FBoneIndexType>& BaseArray, const TArray<FBoneIndexType>& InsertArray);
+
 	/**
 	* Recalculates the AnimCurveUids array in RequiredBone of this SkeletalMeshComponent based on current required bone set
 	* Is called when Skeleton->IsRequiredCurvesUpToDate() = false
diff --git a/Engine/Source/Runtime/Engine/Private/Components/SkeletalMeshComponent.cpp b/Engine/Source/Runtime/Engine/Private/Components/SkeletalMeshComponent.cpp
--- a/Engine/Source/Runtime/Engine/Private/Components/SkeletalMeshComponent.cpp
+++ b/Engine/Source/Runtime/Engine/Private/Components/SkeletalMeshComponent.cpp
@@ -1669,7 +1669,7 @@ static void IntersectBoneIndexArrays(TArray<FBoneIndexType>& Output, const TArra
 
 /** Takes sorted array Base and then adds any elements from sorted array Insert which is missing from it, preserving order.
  * this assumes both arrays are sorted and contain unique bone indices. */
-static void MergeInBoneIndexArrays(TArray<FBoneIndexType>& BaseArray, const TArray<FBoneIndexType>& InsertArray)
+/*static*/ void USkeletalMeshComponent::MergeInBoneIndexArrays(TArray<FBoneIndexType>& BaseArray, const TArray<FBoneIndexType>& InsertArray)
 {
 	// Then we merge them into the array of required bones.
 	int32 BaseBonePos = 0;
@@ -1787,33 +1787,14 @@ void USkeletalMeshComponent::ComputeRequiredBones(TArray<FBoneIndexType>& OutReq
 	OutRequiredBones = LODData.RequiredBones;
 
 	// Add virtual bones
-	MergeInBoneIndexArrays(OutRequiredBones, SkelMesh->GetRefSkeleton().GetRequiredVirtualBones());
+	GetRequiredVirtualBones(SkelMesh, OutRequiredBones);
 
 	const UPhysicsAsset* const PhysicsAsset = GetPhysicsAsset();
 	// If we have a PhysicsAsset, we also need to make sure that all the bones used by it are always updated, as its used
 	// by line checks etc. We might also want to kick in the physics, which means having valid bone transforms.
 	if (!bIgnorePhysicsAsset && PhysicsAsset)
 	{
-		TArray<FBoneIndexType> PhysAssetBones;
-		PhysAssetBones.Reserve(PhysicsAsset->SkeletalBodySetups.Num());
-		for (int32 i = 0; i<PhysicsAsset->SkeletalBodySetups.Num(); i++)
-		{
-			if (!ensure(PhysicsAsset->SkeletalBodySetups[i]))
-			{
-				continue;
-			}
-			int32 PhysBoneIndex = SkelMesh->GetRefSkeleton().FindBoneIndex(PhysicsAsset->SkeletalBodySetups[i]->BoneName);
-			if (PhysBoneIndex != INDEX_NONE)
-			{
-				PhysAssetBones.Add(PhysBoneIndex);
-			}
-		}
-
-		// Then sort array of required bones in hierarchy order
-		PhysAssetBones.Sort();
-
-		// Make sure all of these are in RequiredBones.
-		MergeInBoneIndexArrays(OutRequiredBones, PhysAssetBones);
+		GetPhysicsRequiredBones(SkelMesh, PhysicsAsset, OutRequiredBones);
 	}
 
 	// TODO - Make sure that bones with per-poly collision are also always updated.
@@ -1821,14 +1802,85 @@ void USkeletalMeshComponent::ComputeRequiredBones(TArray<FBoneIndexType>& OutReq
 	// Purge invisible bones and their children
 	// this has to be done before mirror table check/physics body checks
 	// mirror table/phys body ones has to be calculated
-	const TArray<uint8>& EditableBoneVisibilityStates = GetEditableBoneVisibilityStates();
-	if (ShouldUpdateBoneVisibility() && EditableBoneVisibilityStates.Num() > 0)
+	ExcludeHiddenBones(this, SkelMesh, OutRequiredBones);
+
+	// Add in any bones that may be required when mirroring.
+	PRAGMA_DISABLE_DEPRECATION_WARNINGS
+	if (SkelMesh->SkelMirrorTable.Num() > 0 &&
+		SkelMesh->SkelMirrorTable.Num() == BoneSpaceTransforms.Num())
 	{
-		check(EditableBoneVisibilityStates.Num() == GetNumComponentSpaceTransforms());
-		
-		if (ensureMsgf(EditableBoneVisibilityStates.Num() >= OutRequiredBones.Num(), 
-			TEXT("Skeletal Mesh asset '%s' has incorrect BoneVisibilityStates. # of BoneVisibilityStatese (%d), # of OutRequiredBones (%d)"), 
-			*SkelMesh->GetName(), EditableBoneVisibilityStates.Num(), OutRequiredBones.Num()))
+		GetMirroringRequiredBones(SkelMesh, OutRequiredBones);
+	}
+	PRAGMA_ENABLE_DEPRECATION_WARNINGS
+
+	// Get socket bones set to animate and bones required to fill the component space base transforms
+	TArray<FBoneIndexType> NeededBonesForFillComponentSpaceTransforms;
+	GetSocketRequiredBones(SkelMesh, OutRequiredBones, NeededBonesForFillComponentSpaceTransforms);
+
+	// Gather any bones referenced by shadow shapes
+	GetShadowShapeRequiredBones(this, OutRequiredBones);
+
+	// Ensure that we have a complete hierarchy down to those bones.
+	FAnimationRuntime::EnsureParentsPresent(OutRequiredBones, SkelMesh->GetRefSkeleton());
+
+	OutFillComponentSpaceTransformsRequiredBones.Reset(OutRequiredBones.Num() + NeededBonesForFillComponentSpaceTransforms.Num());
+	OutFillComponentSpaceTransformsRequiredBones = OutRequiredBones;
+
+	NeededBonesForFillComponentSpaceTransforms.Sort();
+	MergeInBoneIndexArrays(OutFillComponentSpaceTransformsRequiredBones, NeededBonesForFillComponentSpaceTransforms);
+	FAnimationRuntime::EnsureParentsPresent(OutFillComponentSpaceTransformsRequiredBones, SkelMesh->GetRefSkeleton());
+}
+
+/*static*/ void USkeletalMeshComponent::GetRequiredVirtualBones(const USkeletalMesh* SkeletalMesh, TArray<FBoneIndexType>& OutRequiredBones)
+{
+	check (SkeletalMesh != nullptr);
+
+	// Add virtual bones
+	MergeInBoneIndexArrays(OutRequiredBones, SkeletalMesh->GetRefSkeleton().GetRequiredVirtualBones());
+}
+
+/*static*/ void USkeletalMeshComponent::GetPhysicsRequiredBones(const USkeletalMesh* SkeletalMesh, const UPhysicsAsset* PhysicsAsset, TArray<FBoneIndexType>& OutRequiredBones)
+{
+	check(SkeletalMesh != nullptr);
+	check(PhysicsAsset != nullptr);
+
+	// If we have a PhysicsAsset, we also need to make sure that all the bones used by it are always updated, as its used
+	// by line checks etc. We might also want to kick in the physics, which means having valid bone transforms.
+	TArray<FBoneIndexType> PhysAssetBones;
+	PhysAssetBones.Reserve(PhysicsAsset->SkeletalBodySetups.Num());
+	for (int32 i = 0; i < PhysicsAsset->SkeletalBodySetups.Num(); i++)
+	{
+		if (!ensure(PhysicsAsset->SkeletalBodySetups[i]))
+		{
+			continue;
+		}
+		int32 PhysBoneIndex = SkeletalMesh->GetRefSkeleton().FindBoneIndex(PhysicsAsset->SkeletalBodySetups[i]->BoneName);
+		if (PhysBoneIndex != INDEX_NONE)
+		{
+			PhysAssetBones.Add(PhysBoneIndex);
+		}
+	}
+
+	// Then sort array of required bones in hierarchy order
+	PhysAssetBones.Sort();
+
+	// Make sure all of these are in RequiredBones.
+	MergeInBoneIndexArrays(OutRequiredBones, PhysAssetBones);
+}
+
+/*static*/ void USkeletalMeshComponent::ExcludeHiddenBones(const USkeletalMeshComponent* SkeletalMeshComponent, const USkeletalMesh* SkeletalMesh, TArray<FBoneIndexType>& OutRequiredBones)
+{
+	check(SkeletalMeshComponent != nullptr);
+	check(SkeletalMesh != nullptr);
+
+	const TArray<uint8>& EditableBoneVisibilityStates = SkeletalMeshComponent->GetEditableBoneVisibilityStates();
+	if (SkeletalMeshComponent->ShouldUpdateBoneVisibility() && EditableBoneVisibilityStates.Num() > 0)
+	{
+		check(EditableBoneVisibilityStates.Num() == SkeletalMeshComponent->GetNumComponentSpaceTransforms());
+
+		if (ensureMsgf(EditableBoneVisibilityStates.Num() >= OutRequiredBones.Num(),
+			TEXT("Skeletal Mesh asset '%s' has incorrect BoneVisibilityStates. # of BoneVisibilityStatese (%d), # of OutRequiredBones (%d)"),
+			*SkeletalMesh->GetName(), EditableBoneVisibilityStates.Num(), OutRequiredBones.Num()))
 		{
 			int32 VisibleBoneWriteIndex = 0;
 			for (int32 i = 0; i < OutRequiredBones.Num(); ++i)
@@ -1836,8 +1888,8 @@ void USkeletalMeshComponent::ComputeRequiredBones(TArray<FBoneIndexType>& OutReq
 				FBoneIndexType CurBoneIndex = OutRequiredBones[i];
 
 				// Current bone visible?
-				if (EditableBoneVisibilityStates[CurBoneIndex] == BVS_Visible || 
-					VisibilityBasedAnimTickOption == EVisibilityBasedAnimTickOption::AlwaysTickPoseAndRefreshBones)
+				if (EditableBoneVisibilityStates[CurBoneIndex] == BVS_Visible ||
+					SkeletalMeshComponent->VisibilityBasedAnimTickOption == EVisibilityBasedAnimTickOption::AlwaysTickPoseAndRefreshBones)
 				{
 					OutRequiredBones[VisibleBoneWriteIndex++] = CurBoneIndex;
 				}
@@ -1851,60 +1903,66 @@ void USkeletalMeshComponent::ComputeRequiredBones(TArray<FBoneIndexType>& OutReq
 			}
 		}
 	}
+}
+
+/*static*/ void USkeletalMeshComponent::GetMirroringRequiredBones(const USkeletalMesh* SkeletalMesh, TArray<FBoneIndexType>& OutRequiredBones)
+{
+	check(SkeletalMesh != nullptr);
 
 	// Add in any bones that may be required when mirroring.
 	// JTODO: This is only required if there are mirroring nodes in the tree, but hard to know...
 	PRAGMA_DISABLE_DEPRECATION_WARNINGS
-	if (SkelMesh->SkelMirrorTable.Num() > 0 &&
-		SkelMesh->SkelMirrorTable.Num() == BoneSpaceTransforms.Num())
-	{
-		TArray<FBoneIndexType> MirroredDesiredBones;
-		MirroredDesiredBones.AddUninitialized(RequiredBones.Num());
+	TArray<FBoneIndexType> MirroredDesiredBones;
+	MirroredDesiredBones.AddUninitialized(OutRequiredBones.Num());
 
-		const TArray<struct FBoneMirrorInfo>& SkelMirrorTable = SkelMesh->GetSkelMirrorTable();
-		// Look up each bone in the mirroring table.
-		for (int32 i = 0; i<OutRequiredBones.Num(); i++)
-		{
-			MirroredDesiredBones[i] = SkelMirrorTable[OutRequiredBones[i]].SourceIndex;
-		}
+	const TArray<struct FBoneMirrorInfo>& SkelMirrorTable = SkeletalMesh->GetSkelMirrorTable();
+	// Look up each bone in the mirroring table.
+	for (int32 i = 0; i<OutRequiredBones.Num(); i++)
+	{
+		MirroredDesiredBones[i] = SkelMirrorTable[OutRequiredBones[i]].SourceIndex;
+	}
 
-		// Sort to ensure strictly increasing order.
-		MirroredDesiredBones.Sort();
+	// Sort to ensure strictly increasing order.
+	MirroredDesiredBones.Sort();
 
-		// Make sure all of these are in OutRequiredBones, and 
-		MergeInBoneIndexArrays(OutRequiredBones, MirroredDesiredBones);
-	}
+	// Make sure all of these are in OutRequiredBones, and 
+	MergeInBoneIndexArrays(OutRequiredBones, MirroredDesiredBones);
 	PRAGMA_ENABLE_DEPRECATION_WARNINGS
-	TArray<FBoneIndexType> NeededBonesForFillComponentSpaceTransforms;
+}
+
+/*static*/ void USkeletalMeshComponent::GetSocketRequiredBones(const USkeletalMesh* SkeletalMesh, TArray<FBoneIndexType>& OutRequiredBones, TArray<FBoneIndexType>& NeededBonesForFillComponentSpaceTransforms)
+{
+	check(SkeletalMesh != nullptr);
+
+	TArray<FBoneIndexType> ForceAnimatedSocketBones;
+	TArray<USkeletalMeshSocket*> ActiveSocketList = SkeletalMesh->GetActiveSocketList();
+	ForceAnimatedSocketBones.Reserve(ActiveSocketList.Num());
+	for (const USkeletalMeshSocket* Socket : ActiveSocketList)
 	{
-		TArray<FBoneIndexType> ForceAnimatedSocketBones;
-		TArray<USkeletalMeshSocket*> ActiveSocketList = SkelMesh->GetActiveSocketList();
-		ForceAnimatedSocketBones.Reserve(ActiveSocketList.Num());
-		for (const USkeletalMeshSocket* Socket : ActiveSocketList)
+		int32 BoneIndex = SkeletalMesh->GetRefSkeleton().FindBoneIndex(Socket->BoneName);
+		if (BoneIndex != INDEX_NONE)
 		{
-			int32 BoneIndex = SkelMesh->GetRefSkeleton().FindBoneIndex(Socket->BoneName);
-			if (BoneIndex != INDEX_NONE)
+			if (Socket->bForceAlwaysAnimated)
 			{
-				if (Socket->bForceAlwaysAnimated)
-				{
-					ForceAnimatedSocketBones.AddUnique(BoneIndex);
-				}
-				else
-				{
-					NeededBonesForFillComponentSpaceTransforms.AddUnique(BoneIndex);
-				}
+				ForceAnimatedSocketBones.AddUnique(BoneIndex);
+			}
+			else
+			{
+				NeededBonesForFillComponentSpaceTransforms.AddUnique(BoneIndex);
 			}
 		}
+	}
 
-		// Then sort array of required bones in hierarchy order
-		ForceAnimatedSocketBones.Sort();
+	// Then sort array of required bones in hierarchy order
+	ForceAnimatedSocketBones.Sort();
 
-		// Make sure all of these are in OutRequiredBones.
-		MergeInBoneIndexArrays(OutRequiredBones, ForceAnimatedSocketBones);
-	}
+	// Make sure all of these are in OutRequiredBones.
+	MergeInBoneIndexArrays(OutRequiredBones, ForceAnimatedSocketBones);
+}
 
-	// Gather any bones referenced by shadow shapes
-	if (FSkeletalMeshSceneProxy* SkeletalMeshProxy = (FSkeletalMeshSceneProxy*)SceneProxy)
+/*static*/ void USkeletalMeshComponent::GetShadowShapeRequiredBones(const USkeletalMeshComponent* SkeletalMeshComponent, TArray<FBoneIndexType>& OutRequiredBones)
+{
+	if (FSkeletalMeshSceneProxy* SkeletalMeshProxy = (FSkeletalMeshSceneProxy*)SkeletalMeshComponent->SceneProxy)
 	{
 		const TArray<FBoneIndexType>& ShadowShapeBones = SkeletalMeshProxy->GetSortedShadowBoneIndices();
 
@@ -1914,18 +1972,9 @@ void USkeletalMeshComponent::ComputeRequiredBones(TArray<FBoneIndexType>& OutReq
 			MergeInBoneIndexArrays(OutRequiredBones, ShadowShapeBones);
 		}
 	}
-
-	// Ensure that we have a complete hierarchy down to those bones.
-	FAnimationRuntime::EnsureParentsPresent(OutRequiredBones, SkelMesh->GetRefSkeleton());
-
-	OutFillComponentSpaceTransformsRequiredBones.Reset(OutRequiredBones.Num() + NeededBonesForFillComponentSpaceTransforms.Num());
-	OutFillComponentSpaceTransformsRequiredBones = OutRequiredBones;
-
-	NeededBonesForFillComponentSpaceTransforms.Sort();
-	MergeInBoneIndexArrays(OutFillComponentSpaceTransformsRequiredBones, NeededBonesForFillComponentSpaceTransforms);
-	FAnimationRuntime::EnsureParentsPresent(OutFillComponentSpaceTransformsRequiredBones, SkelMesh->GetRefSkeleton());
 }
 
+
 void USkeletalMeshComponent::RecalcRequiredBones(int32 LODIndex)
 {
 	if (!GetSkeletalMeshAsset())
