Sequencer crashes when trying to use 'Import AnimSequence Into This Section' menu item
#fix remap skeleton indices when populating a retargeting base pose between compatible skeletons

diff --git a/Engine/Source/Runtime/Engine/Private/Animation/AnimSequence.cpp b/Engine/Source/Runtime/Engine/Private/Animation/AnimSequence.cpp
--- a/Engine/Source/Runtime/Engine/Private/Animation/AnimSequence.cpp
+++ b/Engine/Source/Runtime/Engine/Private/Animation/AnimSequence.cpp
@@ -44,6 +44,8 @@
 #include "ProfilingDebugging/CookStats.h"
 #include "Animation/AnimationPoseData.h"
 #include "ITimeManagementModule.h"
+#include "SkeletonRemappingRegistry.h"
+#include "Animation/SkeletonRemapping.h"
 
 LLM_DEFINE_TAG(SequenceData);
 
@@ -1497,17 +1499,17 @@ void UAnimSequence::GetBonePose(FAnimationPoseData& OutAnimationPoseData, const
 		if (bDisableRetargeting)
 		{
 			TArray<FTransform> const& AuthoredOnRefSkeleton = GetRetargetTransforms();
-			const int32 NumRawBones = RequiredBones.GetSkeletonAsset()->GetReferenceSkeleton().GetRawBoneNum();
+			// Map from this sequence its Skeleton to target
+			const FSkeletonRemapping& SkeletonRemapping = UE::Anim::FSkeletonRemappingRegistry::Get().GetRemapping(GetSkeleton(), RequiredBones.GetSkeletonAsset());
 			for (FCompactPoseBoneIndex PoseBoneIndex : OutPose.ForEachBoneIndex())
 			{
-				const FSkeletonPoseBoneIndex SkeletonBoneIndex = RequiredBones.GetSkeletonPoseIndexFromCompactPoseIndex(PoseBoneIndex);
-				// Pose bone index should always exist in Skeleton
-				checkSlow(SkeletonBoneIndex != INDEX_NONE);
+				int32 SkeletonBoneIndex = RequiredBones.GetSkeletonPoseIndexFromCompactPoseIndex(PoseBoneIndex).GetInt();
+				SkeletonBoneIndex = SkeletonRemapping.IsValid() ? SkeletonRemapping.GetSourceSkeletonBoneIndex(SkeletonBoneIndex) : SkeletonBoneIndex;
 
 				// Virtual bones are part of the retarget transform pose, so if the pose has not been updated (recently) there might be a mismatch
-				if (SkeletonBoneIndex < NumRawBones || AuthoredOnRefSkeleton.IsValidIndex(SkeletonBoneIndex.GetInt()))
+				if (SkeletonBoneIndex != INDEX_NONE && AuthoredOnRefSkeleton.IsValidIndex(SkeletonBoneIndex))
 				{
-					OutPose[PoseBoneIndex] = AuthoredOnRefSkeleton[SkeletonBoneIndex.GetInt()];
+					OutPose[PoseBoneIndex] = AuthoredOnRefSkeleton[SkeletonBoneIndex];
 				}
 			}
 		}
