World Partition: fix crash when disabling world partition on a world.

diff --git a/Engine/Source/Editor/WorldPartitionEditor/Private/WorldPartition/Customizations/WorldPartitionDetailsCustomization.cpp b/Engine/Source/Editor/WorldPartitionEditor/Private/WorldPartition/Customizations/WorldPartitionDetailsCustomization.cpp
--- a/Engine/Source/Editor/WorldPartitionEditor/Private/WorldPartition/Customizations/WorldPartitionDetailsCustomization.cpp
+++ b/Engine/Source/Editor/WorldPartitionEditor/Private/WorldPartition/Customizations/WorldPartitionDetailsCustomization.cpp
@@ -96,7 +96,7 @@ void FWorldPartitionDetails::CustomizeDetails(IDetailLayoutBuilder& InDetailBuil
 					.HAlign(HAlign_Center)
 					.IsEnabled_Lambda([this]()
 					{
-						return !WorldPartition->IsStreamingEnabled();
+						return WorldPartition.IsValid() && !WorldPartition->IsStreamingEnabled();
 					})
 					.OnClicked_Lambda([this]()
 					{
diff --git a/Engine/Source/Runtime/Engine/Private/WorldPartition/WorldPartition.cpp b/Engine/Source/Runtime/Engine/Private/WorldPartition/WorldPartition.cpp
--- a/Engine/Source/Runtime/Engine/Private/WorldPartition/WorldPartition.cpp
+++ b/Engine/Source/Runtime/Engine/Private/WorldPartition/WorldPartition.cpp
@@ -977,6 +977,8 @@ bool UWorldPartition::RemoveWorldPartition(AWorldSettings* WorldSettings)
 	{
 		if (!WorldPartition->IsStreamingEnabled())
 		{
+			FWorldPartitionLoadingContext::FNull LoadingContext;
+
 			WorldSettings->Modify();
 			
 			ULevel* PersistentLevel = WorldSettings->GetLevel();
diff --git a/Engine/Source/Runtime/Engine/Private/WorldPartition/WorldPartitionActorLoaderInterface.cpp b/Engine/Source/Runtime/Engine/Private/WorldPartition/WorldPartitionActorLoaderInterface.cpp
--- a/Engine/Source/Runtime/Engine/Private/WorldPartition/WorldPartitionActorLoaderInterface.cpp
+++ b/Engine/Source/Runtime/Engine/Private/WorldPartition/WorldPartitionActorLoaderInterface.cpp
@@ -101,7 +101,6 @@ void IWorldPartitionActorLoaderInterface::ILoaderAdapter::RegisterDelegates()
 {
 	ActorLoaderInterfaceRefreshState.AddRaw(this, &IWorldPartitionActorLoaderInterface::ILoaderAdapter::OnRefreshLoadedState);
 
-	check(World->GetWorldPartition());
 	if (UWorldPartition* WorldPartition = World->GetWorldPartition())
 	{
 		WorldPartition->OnActorDescContainerRegistered.AddRaw(this, &IWorldPartitionActorLoaderInterface::ILoaderAdapter::ILoaderAdapter::OnActorDescContainerInitialize);
@@ -113,7 +112,6 @@ void IWorldPartitionActorLoaderInterface::ILoaderAdapter::UnregisterDelegates()
 {
 	ActorLoaderInterfaceRefreshState.RemoveAll(this);
 
-	check(World->GetWorldPartition());
 	if (UWorldPartition* WorldPartition = World->GetWorldPartition())
 	{
 		WorldPartition->OnActorDescContainerRegistered.RemoveAll(this);
diff --git a/Engine/Source/Runtime/Engine/Public/WorldPartition/WorldPartitionHandle.h b/Engine/Source/Runtime/Engine/Public/WorldPartition/WorldPartitionHandle.h
--- a/Engine/Source/Runtime/Engine/Public/WorldPartition/WorldPartitionHandle.h
+++ b/Engine/Source/Runtime/Engine/Public/WorldPartition/WorldPartitionHandle.h
@@ -28,8 +28,8 @@ public:
 		virtual ~IContext();
 
 	private:
-		virtual void RegisterActor(FWorldPartitionActorDesc* ActorDesc) =0;
-		virtual void UnregisterActor(FWorldPartitionActorDesc* ActorDesc) =0;
+		virtual void RegisterActor(FWorldPartitionActorDesc* ActorDesc) = 0;
+		virtual void UnregisterActor(FWorldPartitionActorDesc* ActorDesc) = 0;
 	};
 
 	/**
@@ -71,6 +71,16 @@ public:
 		int32 NumUnregistrations = 0;
 	};
 
+	/**
+	 * Null loading context, which will ignore all loading/unloading commands.
+	 */
+	class ENGINE_API FNull : public IContext
+	{
+	private:
+		virtual void RegisterActor(FWorldPartitionActorDesc* ActorDesc) override {}
+		virtual void UnregisterActor(FWorldPartitionActorDesc* ActorDesc) override {}
+	};
+
 	static void LoadAndRegisterActor(FWorldPartitionActorDesc* ActorDesc);
 	static void UnloadAndUnregisterActor(FWorldPartitionActorDesc* ActorDesc);
 
