Incremental build fix for header include-path

diff --git a/Engine/Source/Runtime/Engine/Private/Animation/AnimSequence.cpp b/Engine/Source/Runtime/Engine/Private/Animation/AnimSequence.cpp
--- a/Engine/Source/Runtime/Engine/Private/Animation/AnimSequence.cpp
+++ b/Engine/Source/Runtime/Engine/Private/Animation/AnimSequence.cpp
@@ -44,7 +44,7 @@
 #include "ProfilingDebugging/CookStats.h"
 #include "Animation/AnimationPoseData.h"
 #include "ITimeManagementModule.h"
-#include "SkeletonRemappingRegistry.h"
+#include "Animation/SkeletonRemappingRegistry.h"
 #include "Animation/SkeletonRemapping.h"
 
 LLM_DEFINE_TAG(SequenceData);
