Character model breaks when cooking a throwable item
#fix ensure we verify smartname curve UIDs for montages in non-editor builds

diff --git a/Engine/Source/Runtime/Engine/Private/Animation/AnimCompositeBase.cpp b/Engine/Source/Runtime/Engine/Private/Animation/AnimCompositeBase.cpp
--- a/Engine/Source/Runtime/Engine/Private/Animation/AnimCompositeBase.cpp
+++ b/Engine/Source/Runtime/Engine/Private/Animation/AnimCompositeBase.cpp
@@ -780,6 +780,24 @@ void UAnimCompositeBase::PostLoad()
 {
 	Super::PostLoad();
 
+#if !WITH_EDITOR
+	if (USkeleton* MySkeleton = GetSkeleton())
+	{
+		if (FLinkerLoad* SkeletonLinker = MySkeleton->GetLinker())
+		{
+			SkeletonLinker->Preload(MySkeleton);
+		}
+		MySkeleton->ConditionalPostLoad();
+
+		PRAGMA_DISABLE_DEPRECATION_WARNINGS
+		for (FFloatCurve& Curve : RawCurveData.FloatCurves)
+		{
+			MySkeleton->VerifySmartName(USkeleton::AnimCurveMappingName, Curve.Name);
+		}	
+		PRAGMA_ENABLE_DEPRECATION_WARNINGS
+	}
+#endif
+
 	InvalidateRecursiveAsset();
 }
 
