[materials] change logic for skipping compilation due to ODSC to only check for requires synchronous compilation, not requires complete. required-complete special engine materials will still only get compiled when rendered, so the previous condition means the "missing cached shadermap" is printed every time these materials are loaded if they are never rendered.

diff --git a/Engine/Source/Runtime/Engine/Private/Materials/MaterialShared.cpp b/Engine/Source/Runtime/Engine/Private/Materials/MaterialShared.cpp
--- a/Engine/Source/Runtime/Engine/Private/Materials/MaterialShared.cpp
+++ b/Engine/Source/Runtime/Engine/Private/Materials/MaterialShared.cpp
@@ -2753,7 +2753,7 @@ bool FMaterial::CacheShaders(const FMaterialShaderMapId& ShaderMapId, EShaderPla
 		}
 		else
 		{
-			const bool bSkipCompilationForODSC = !IsRequiredComplete() && (GShaderCompilingManager->IsShaderCompilationSkipped() || IsShaderJobCacheDDCEnabled());
+			const bool bSkipCompilationForODSC = !RequiresSynchronousCompilation() && (GShaderCompilingManager->IsShaderCompilationSkipped() || IsShaderJobCacheDDCEnabled());
 			// If we aren't actually compiling shaders don't print the debug message that we are compiling shaders.
 			if (!bSkipCompilationForODSC)
 			{
@@ -3129,7 +3129,7 @@ bool FMaterial::BeginCompileShaderMap(
 
 		// we can ignore requests for synch compilation if we are compiling for a different platform than we're running, or we're a commandlet that doesn't render (e.g. cooker)
 		const bool bCanIgnoreSynchronousRequirement = (TargetPlatform && !TargetPlatform->IsRunningPlatform()) || (IsRunningCommandlet() && !IsAllowCommandletRendering());
-		const bool bSkipCompilationForODSC = !IsRequiredComplete() && GShaderCompilingManager->IsShaderCompilationSkipped();
+		const bool bSkipCompilationForODSC = !RequiresSynchronousCompilation() && GShaderCompilingManager->IsShaderCompilationSkipped();
 		if (bSkipCompilationForODSC)
 		{
 			// Force compilation off.
