Nondeterministic curve codec hash generation by from now on always serialize out the InstanceGUID on an animation compression codec

#localization none
#tests consecutive cooks + client run

diff --git a/Engine/Source/Runtime/Core/Public/UObject/FortniteReleaseBranchCustomObjectVersion.h b/Engine/Source/Runtime/Core/Public/UObject/FortniteReleaseBranchCustomObjectVersion.h
--- a/Engine/Source/Runtime/Core/Public/UObject/FortniteReleaseBranchCustomObjectVersion.h
+++ b/Engine/Source/Runtime/Core/Public/UObject/FortniteReleaseBranchCustomObjectVersion.h
@@ -30,6 +30,9 @@ struct CORE_API FFortniteReleaseBranchCustomObjectVersion
 		// Remove any cooked collision data from nanite landscape / editor spline meshes since collisions are not needed there :
 		RemoveUselessLandscapeMeshesCookedCollisionData,
 
+		// Serialize out UAnimCurveCompressionCodec::InstanceGUID to maintain deterministic DDC key generation in cooked-editor
+		SerializeAnimCurveCompressionCodecGuidOnCook,
+
 		// -----<new versions can be added above this line>-------------------------------------------------
 		VersionPlusOne,
 		LatestVersion = VersionPlusOne - 1
diff --git a/Engine/Source/Runtime/Engine/Classes/Animation/AnimCurveCompressionCodec.h b/Engine/Source/Runtime/Engine/Classes/Animation/AnimCurveCompressionCodec.h
--- a/Engine/Source/Runtime/Engine/Classes/Animation/AnimCurveCompressionCodec.h
+++ b/Engine/Source/Runtime/Engine/Classes/Animation/AnimCurveCompressionCodec.h
@@ -34,22 +34,19 @@ class ENGINE_API UAnimCurveCompressionCodec : public UObject
 {
 	GENERATED_UCLASS_BODY()
 
-#if WITH_EDITORONLY_DATA
 	/** A GUID that is unique to this codec instance. After creation, it never changes. */
 	FGuid InstanceGuid;
-#endif
 
 	/** Allow us to convert DDC serialized path back into codec object */
 	virtual UAnimCurveCompressionCodec* GetCodec(const FString& Path) { return this; }
 
 	//////////////////////////////////////////////////////////////////////////
 
+	virtual void Serialize(FArchive& Ar) override;
 #if WITH_EDITORONLY_DATA
 	// UObject overrides
 	virtual void PostInitProperties() override;
 	virtual void PostDuplicate(bool bDuplicateForPIE) override;
-	virtual void Serialize(FArchive& Ar) override;
-
 	/** Returns whether or not we can use this codec to compress. */
 	virtual bool IsCodecValid() const { return true; }
 
diff --git a/Engine/Source/Runtime/Engine/Private/Animation/AnimCurveCompressionCodec.cpp b/Engine/Source/Runtime/Engine/Private/Animation/AnimCurveCompressionCodec.cpp
--- a/Engine/Source/Runtime/Engine/Private/Animation/AnimCurveCompressionCodec.cpp
+++ b/Engine/Source/Runtime/Engine/Private/Animation/AnimCurveCompressionCodec.cpp
@@ -1,7 +1,9 @@
 // Copyright Epic Games, Inc. All Rights Reserved.
 
 #include "Animation/AnimCurveCompressionCodec.h"
+#include "UObject/FortniteReleaseBranchCustomObjectVersion.h"
 #include "UObject/Package.h"
+#include "Interfaces/ITargetPlatform.h"
 
 #include UE_INLINE_GENERATED_CPP_BY_NAME(AnimCurveCompressionCodec)
 
@@ -10,6 +12,19 @@ UAnimCurveCompressionCodec::UAnimCurveCompressionCodec(const FObjectInitializer&
 {
 }
 
+void UAnimCurveCompressionCodec::Serialize(FArchive& Ar)
+{
+	Super::Serialize(Ar);
+
+	Ar.UsingCustomVersion(FFortniteReleaseBranchCustomObjectVersion::GUID);
+#if !WITH_EDITOR
+	if (Ar.CustomVer(FFortniteReleaseBranchCustomObjectVersion::GUID) >= FFortniteReleaseBranchCustomObjectVersion::SerializeAnimCurveCompressionCodecGuidOnCook)
+#endif // !WITH_EDITOR
+	{
+		Ar << InstanceGuid;
+	}
+}
+
 #if WITH_EDITORONLY_DATA
 void UAnimCurveCompressionCodec::PostInitProperties()
 {
@@ -28,16 +43,6 @@ void UAnimCurveCompressionCodec::PostDuplicate(bool bDuplicateForPIE)
 	InstanceGuid = FGuid::NewGuid();
 }
 
-void UAnimCurveCompressionCodec::Serialize(FArchive& Ar)
-{
-	Super::Serialize(Ar);
-
-	if (!Ar.IsCooking() && !(Ar.IsLoading() && GetOutermost()->bIsCookedForEditor))
-	{
-		Ar << InstanceGuid;
-	}
-}
-
 void UAnimCurveCompressionCodec::PopulateDDCKey(FArchive& Ar)
 {
 	Ar << InstanceGuid;
