Fix for uninitzlied struct properties

#rnx

diff --git a/Engine/Source/Runtime/Engine/Private/UserInterface/InputSettings.cpp b/Engine/Source/Runtime/Engine/Private/UserInterface/InputSettings.cpp
--- a/Engine/Source/Runtime/Engine/Private/UserInterface/InputSettings.cpp
+++ b/Engine/Source/Runtime/Engine/Private/UserInterface/InputSettings.cpp
@@ -526,6 +526,8 @@ FHardwareDeviceIdentifier FHardwareDeviceIdentifier::DefaultKeyboardAndMouse =
 FHardwareDeviceIdentifier::FHardwareDeviceIdentifier()
 	: InputClassName(Invalid.InputClassName)
 	, HardwareDeviceIdentifier(Invalid.HardwareDeviceIdentifier)
+	, PrimaryDeviceType(EHardwareDevicePrimaryType::Unspecified)
+	, SupportedFeaturesMask(EHardwareDeviceSupportedFeatures::Type::Unspecified)
 {
 	
 }
