Minor improvements for DynamicMeshBuilder
- FPooledDynamicMeshVertexBuffer should keep track of its ERHIFeatureLevel.
- FPooledDynamicMeshVertexBuffer should create views *after* uploading its data. This will prevent unnecessary rename handling in some RHI backends.

diff --git a/Engine/Source/Runtime/Engine/Private/DynamicMeshBuilder.cpp b/Engine/Source/Runtime/Engine/Private/DynamicMeshBuilder.cpp
--- a/Engine/Source/Runtime/Engine/Private/DynamicMeshBuilder.cpp
+++ b/Engine/Source/Runtime/Engine/Private/DynamicMeshBuilder.cpp
@@ -395,10 +395,11 @@ public:
 
 	TArray<FDynamicMeshVertex> Vertices;
 
-	FPooledDynamicMeshVertexBuffer(uint32 InNumTexCoords, uint32 InLightmapCoordinateIndex, bool InUse16bitTexCoord, FDynamicMeshBufferAllocator& InDynamicMeshBufferAllocator)
+	FPooledDynamicMeshVertexBuffer(ERHIFeatureLevel::Type InFeatureLevel, uint32 InNumTexCoords, uint32 InLightmapCoordinateIndex, bool InUse16bitTexCoord, FDynamicMeshBufferAllocator& InDynamicMeshBufferAllocator)
 		: NumTexCoords(InNumTexCoords)
 		, LightmapCoordinateIndex(InLightmapCoordinateIndex)
 		, Use16bitTexCoord(InUse16bitTexCoord)
+		, FeatureLevel(InFeatureLevel)
 		, DynamicMeshBufferAllocator(InDynamicMeshBufferAllocator)
 	{
 		check(NumTexCoords > 0 && NumTexCoords <= MAX_STATIC_TEXCOORDS);
@@ -424,15 +425,6 @@ public:
 		TexCoordBuffer.VertexBufferRHI = DynamicMeshBufferAllocator.AllocVertexBuffer(TextureStride, NumTexCoords * Vertices.Num());
 		ColorBuffer.VertexBufferRHI    = DynamicMeshBufferAllocator.AllocVertexBuffer(sizeof(FColor), Vertices.Num());
 
-		if (RHISupportsManualVertexFetch(GMaxRHIShaderPlatform))
-		{
-			TangentBufferSRV = RHICreateShaderResourceView(TangentBuffer.VertexBufferRHI, 4, PF_R8G8B8A8_SNORM);
-			TexCoordBufferSRV = RHICreateShaderResourceView(TexCoordBuffer.VertexBufferRHI, TextureStride, TextureFormat);
-			ColorBufferSRV = RHICreateShaderResourceView(ColorBuffer.VertexBufferRHI, 4, PF_R8G8B8A8);
-
-			PositionBufferSRV = RHICreateShaderResourceView(PositionBuffer.VertexBufferRHI, sizeof(float), PF_R32_FLOAT);
-		}
-
 		void* TexCoordBufferData = RHILockBuffer(TexCoordBuffer.VertexBufferRHI, 0, NumTexCoords * TextureStride * Vertices.Num(), RLM_WriteOnly);
 		FVector2f* TexCoordBufferData32 = !Use16bitTexCoord ? static_cast<FVector2f*>(TexCoordBufferData) : nullptr;
 		FVector2DHalf* TexCoordBufferData16 = Use16bitTexCoord ? static_cast<FVector2DHalf*>(TexCoordBufferData) : nullptr;
@@ -472,6 +464,15 @@ public:
 		RHIUnlockBuffer(TangentBuffer.VertexBufferRHI);
 		RHIUnlockBuffer(TexCoordBuffer.VertexBufferRHI);
 		RHIUnlockBuffer(ColorBuffer.VertexBufferRHI);
+
+		if (RHISupportsManualVertexFetch(GetFeatureLevelShaderPlatform(FeatureLevel)))
+		{
+			TangentBufferSRV = RHICreateShaderResourceView(TangentBuffer.VertexBufferRHI, 4, PF_R8G8B8A8_SNORM);
+			TexCoordBufferSRV = RHICreateShaderResourceView(TexCoordBuffer.VertexBufferRHI, TextureStride, TextureFormat);
+			ColorBufferSRV = RHICreateShaderResourceView(ColorBuffer.VertexBufferRHI, 4, PF_R8G8B8A8);
+
+			PositionBufferSRV = RHICreateShaderResourceView(PositionBuffer.VertexBufferRHI, sizeof(float), PF_R32_FLOAT);
+		}
 	}
 
 	void InitResource() override
@@ -530,6 +531,7 @@ private:
 	const uint32 NumTexCoords;
 	const uint32 LightmapCoordinateIndex;
 	const bool Use16bitTexCoord;
+	const ERHIFeatureLevel::Type FeatureLevel;
 	FDynamicMeshBufferAllocator& DynamicMeshBufferAllocator;
 };
 
@@ -664,7 +666,7 @@ FDynamicMeshBuilder::FDynamicMeshBuilder(ERHIFeatureLevel::Type InFeatureLevel,
 		InDynamicMeshBufferAllocator = &DefaultDynamicMeshBufferAllocator;
 	}
 
-	VertexBuffer = new FPooledDynamicMeshVertexBuffer(InNumTexCoords, InLightmapCoordinateIndex, InUse16bitTexCoord, *InDynamicMeshBufferAllocator);
+	VertexBuffer = new FPooledDynamicMeshVertexBuffer(InFeatureLevel, InNumTexCoords, InLightmapCoordinateIndex, InUse16bitTexCoord, *InDynamicMeshBufferAllocator);
 	IndexBuffer = new FPooledDynamicMeshIndexBuffer(*InDynamicMeshBufferAllocator);
 }
 
