- Set bRenderInDepthPass = ShouldRenderInDepthPass() for SkeletalMeshRender to match other components

#rnx

diff --git a/Engine/Source/Runtime/Engine/Private/SkeletalMesh.cpp b/Engine/Source/Runtime/Engine/Private/SkeletalMesh.cpp
--- a/Engine/Source/Runtime/Engine/Private/SkeletalMesh.cpp
+++ b/Engine/Source/Runtime/Engine/Private/SkeletalMesh.cpp
@@ -6759,6 +6759,7 @@ FPrimitiveViewRelevance FSkeletalMeshSceneProxy::GetViewRelevance(const FSceneVi
 	Result.bDynamicRelevance = !Result.bStaticRelevance;
 	Result.bRenderCustomDepth = ShouldRenderCustomDepth();
 	Result.bRenderInMainPass = ShouldRenderInMainPass();
+	Result.bRenderInDepthPass = ShouldRenderInDepthPass();
 	Result.bUsesLightingChannels = GetLightingChannelMask() != GetDefaultLightingChannelMask();
 	Result.bTranslucentSelfShadow = bCastVolumetricTranslucentShadow;
 	
