blueprint library function for finding the closest physics object should only filter based on whether the particle is enabled or not

diff --git a/Engine/Source/Runtime/Engine/Private/PhysicsEngine/PhysicsObjectBlueprintLibrary.cpp b/Engine/Source/Runtime/Engine/Private/PhysicsEngine/PhysicsObjectBlueprintLibrary.cpp
--- a/Engine/Source/Runtime/Engine/Private/PhysicsEngine/PhysicsObjectBlueprintLibrary.cpp
+++ b/Engine/Source/Runtime/Engine/Private/PhysicsEngine/PhysicsObjectBlueprintLibrary.cpp
@@ -16,8 +16,8 @@ FClosestPhysicsObjectResult UPhysicsObjectBlueprintLibrary::GetClosestPhysicsObj
 	PhysicsObjects = PhysicsObjects.FilterByPredicate(
 		// Maybe allow users to customize the predicate?
 		[&Interface](Chaos::FPhysicsObject* Object) {
-			// We need to filter so we only get the leaf bodies. Mainly in consideration of the geometry collection.
-			return !Interface->HasChildren(Object);
+			// Filter for enabled particles. Other systems should be responsible for attaching the wheel to child particles if the geometry collection breaks (for example).
+			return !Interface->AreAllDisabled({ &Object, 1 });
 		}
 	);
 	return Interface->GetClosestPhysicsBodyFromLocation(PhysicsObjects, WorldLocation);
