remove unnecessary code dupes of GetRunningTargetPlatform

diff --git a/Engine/Source/Developer/TargetPlatform/Public/Interfaces/ITargetPlatformManagerModule.h b/Engine/Source/Developer/TargetPlatform/Public/Interfaces/ITargetPlatformManagerModule.h
--- a/Engine/Source/Developer/TargetPlatform/Public/Interfaces/ITargetPlatformManagerModule.h
+++ b/Engine/Source/Developer/TargetPlatform/Public/Interfaces/ITargetPlatformManagerModule.h
@@ -125,7 +125,7 @@ public:
 	/**
 	 * Returns the target platform that is currently running.
 	 *
-	 * Note: This method is not available on consoles.
+	 * Note: This method only returns a Target Platform when WITH_EDITOR.
 	 *
 	 * @return Running target platform.
 	 */
diff --git a/Engine/Source/Runtime/Engine/Private/Animation/AnimStreamable.cpp b/Engine/Source/Runtime/Engine/Private/Animation/AnimStreamable.cpp
--- a/Engine/Source/Runtime/Engine/Private/Animation/AnimStreamable.cpp
+++ b/Engine/Source/Runtime/Engine/Private/Animation/AnimStreamable.cpp
@@ -543,9 +543,10 @@ void UAnimStreamable::RequestCompressedData(const ITargetPlatform* Platform)
 	if (!Platform)
 	{
 		Platform = TPM->GetRunningTargetPlatform();
+		check( Platform != nullptr );
 	}
 
-	const bool bIsRunningPlatform = (Platform == TPM->GetRunningTargetPlatform());
+	const bool bIsRunningPlatform = Platform->IsRunningPlatform();
 
 	if (bIsRunningPlatform)
 	{
diff --git a/Engine/Source/Runtime/Engine/Private/AudioDerivedData.cpp b/Engine/Source/Runtime/Engine/Private/AudioDerivedData.cpp
--- a/Engine/Source/Runtime/Engine/Private/AudioDerivedData.cpp
+++ b/Engine/Source/Runtime/Engine/Private/AudioDerivedData.cpp
@@ -298,28 +298,6 @@ static void GetStreamedAudioDerivedDataKey(
 	GetStreamedAudioDerivedDataKeyFromSuffix(KeySuffix, OutKey);
 }
 
-static ITargetPlatform* GetRunningTargetPlatform(ITargetPlatformManagerModule* TPM)
-{
-	ITargetPlatform* CurrentPlatform = NULL;
-	const TArray<ITargetPlatform*>& Platforms = TPM->GetActiveTargetPlatforms();
-
-	check(Platforms.Num());
-
-	CurrentPlatform = Platforms[0];
-
-	for (int32 Index = 1; Index < Platforms.Num(); Index++)
-	{
-		if (Platforms[Index]->IsRunningPlatform())
-		{
-			CurrentPlatform = Platforms[Index];
-			break;
-		}
-	}
-
-	check(CurrentPlatform != NULL);
-	return CurrentPlatform;
-}
-
 /**
  * Gets Wave format for a SoundWave on the current running platform
  * @param SoundWave - The SoundWave to get format for.
@@ -330,8 +308,9 @@ static FName GetWaveFormatForRunningPlatform(USoundWave& SoundWave)
 	ITargetPlatformManagerModule* TPM = GetTargetPlatformManager();
 	if (TPM)
 	{
-		ITargetPlatform* CurrentPlatform = GetRunningTargetPlatform(TPM);
+		ITargetPlatform* CurrentPlatform = TPM->GetRunningTargetPlatform();
 
+		check( CurrentPlatform != nullptr );
 
 		return CurrentPlatform->GetWaveFormat(&SoundWave);
 	}
diff --git a/Engine/Source/Runtime/Engine/Private/TextureDerivedData.cpp b/Engine/Source/Runtime/Engine/Private/TextureDerivedData.cpp
--- a/Engine/Source/Runtime/Engine/Private/TextureDerivedData.cpp
+++ b/Engine/Source/Runtime/Engine/Private/TextureDerivedData.cpp
@@ -994,6 +994,7 @@ static void GetTextureBuildSettings(
 	// A ULightMapVirtualTexture2D with multiple layers saved in MapBuildData could be loaded with the r.VirtualTexture disabled, it will generate DDC before we decide to invalidate the light map data, to skip the ensure failure let it generate VT DDC anyway.
 	const bool bForVirtualTextureStreamingBuild = ULightMapVirtualTexture2D::StaticClass() == Texture.GetClass();
 	const bool bVirtualTextureStreaming = bForVirtualTextureStreamingBuild || (CVarVirtualTexturesEnabled->GetValueOnAnyThread() && bPlatformSupportsVirtualTextureStreaming && Texture.VirtualTextureStreaming);
+
 	const FIntPoint SourceSize = Texture.Source.GetLogicalSize();
 
 	OutBuildSettings.MipGenSettings = MipGenSettings;
@@ -1163,21 +1164,7 @@ static void GetBuildSettingsForRunningPlatform(
 	ITargetPlatformManagerModule* TPM = GetTargetPlatformManager();
 	if (TPM)
 	{
-		ITargetPlatform* TargetPlatform = NULL;
-		const TArray<ITargetPlatform*>& Platforms = TPM->GetActiveTargetPlatforms();
-
-		check(Platforms.Num());
-
-		TargetPlatform = Platforms[0];
-
-		for (int32 Index = 1; Index < Platforms.Num(); Index++)
-		{
-			if (Platforms[Index]->IsRunningPlatform())
-			{
-				TargetPlatform = Platforms[Index];
-				break;
-			}
-		}
+		ITargetPlatform* TargetPlatform = TPM->GetRunningTargetPlatform();
 
 		check(TargetPlatform != NULL);
 
