PKG_FilterEditorOnly packages might still be cooked, so test PKG_Cooked instead

#rnx

diff --git a/Engine/Source/Runtime/Engine/Private/BlueprintGeneratedClass.cpp b/Engine/Source/Runtime/Engine/Private/BlueprintGeneratedClass.cpp
--- a/Engine/Source/Runtime/Engine/Private/BlueprintGeneratedClass.cpp
+++ b/Engine/Source/Runtime/Engine/Private/BlueprintGeneratedClass.cpp
@@ -306,7 +306,7 @@ UClass* UBlueprintGeneratedClass::GetAuthoritativeClass()
  	if (nullptr == ClassGeneratedBy) // to track UE-11597 and UE-11595
  	{
 		// If this is a cooked blueprint, the generatedby class will have been discarded so we'll just have to assume we're authoritative!
-		if (bCooked || RootPackageHasAnyFlags(PKG_FilterEditorOnly))
+		if (bCooked || RootPackageHasAnyFlags(PKG_Cooked))
 		{ 
 			return this;
 		}
