Set FSoundAttenuationSettings struct instance in class USoundAttenuation to be ReadWrite instead of ReadOnly to be able to set it in a blueprint graph

#rnx

diff --git a/Engine/Source/Runtime/Engine/Classes/Sound/SoundAttenuation.h b/Engine/Source/Runtime/Engine/Classes/Sound/SoundAttenuation.h
--- a/Engine/Source/Runtime/Engine/Classes/Sound/SoundAttenuation.h
+++ b/Engine/Source/Runtime/Engine/Classes/Sound/SoundAttenuation.h
@@ -512,7 +512,7 @@ class USoundAttenuation : public UObject
 {
 	GENERATED_UCLASS_BODY()
 
-	UPROPERTY(EditAnywhere, BlueprintReadOnly, Category=Settings, meta = (CustomizeProperty))
+	UPROPERTY(EditAnywhere, BlueprintReadWrite, Category=Settings, meta = (CustomizeProperty))
 	FSoundAttenuationSettings Attenuation;
 };
 
