Add IsUnattended to the Editor SessionSummary event

* Gauntlet, the automated test suite, force shuts the editor instead of allowing it a graceful shutdown.  This results in a large number of sessions reported with AbnormalShutdown.
* To allow for these sessions to be filtered out, adding an IsUnattended field to the SessionSummary. This will be true with applications run with `-unattended` as a command line argument.

Testing:
* Setup a stream filter to detect new session summary events with IsUnattended
* Started and then cleanly shut the editor, saw event come into the data lake

diff --git a/Engine/Source/Runtime/Engine/Private/Analytics/EditorAnalyticsSessionSummary.cpp b/Engine/Source/Runtime/Engine/Private/Analytics/EditorAnalyticsSessionSummary.cpp
--- a/Engine/Source/Runtime/Engine/Private/Analytics/EditorAnalyticsSessionSummary.cpp
+++ b/Engine/Source/Runtime/Engine/Private/Analytics/EditorAnalyticsSessionSummary.cpp
@@ -21,6 +21,7 @@ namespace EditorAnalyticsProperties
 	static const TAnalyticsProperty<int32>   Idle30Min              = TEXT("30MinIdle");
 	static const TAnalyticsProperty<bool>    IsInPIE                = TEXT("IsInPIE");
 	static const TAnalyticsProperty<bool>    IsInVRMode             = TEXT("IsInVRMode");
+	static const TAnalyticsProperty<bool>    IsUnattended           = TEXT("IsUnattended");
 	static const TAnalyticsProperty<uint32>  DirtyPackageCount      = TEXT("DirtyPackageCount");
 
 	static const TAnalyticsProperty<uint32>  TotalStallCount        = TEXT("TotalStallCount");
@@ -49,6 +50,7 @@ FEditorAnalyticsSessionSummary::FEditorAnalyticsSessionSummary(TSharedPtr<IAnaly
 	EditorAnalyticsProperties::Idle30Min.Set(GetStore(), 0);
 	EditorAnalyticsProperties::IsInVRMode.Set(GetStore(), IVREditorModule::Get().IsVREditorModeActive());
 	EditorAnalyticsProperties::IsInPIE.Set(GetStore(), FPlayWorldCommandCallbacks::IsInPIE());
+	EditorAnalyticsProperties::IsUnattended.Set(GetStore(), FApp::IsUnattended());
 	EditorAnalyticsProperties::DirtyPackageCount.Set(GetStore(), 0);
 
 	EditorAnalyticsProperties::TopStallName.Set(GetStore(), TEXT(""), /*Capacity*/128);
