Fix seamless travel crash introduced by 2456814

diff --git a/Engine/Source/Runtime/Engine/Private/World.cpp b/Engine/Source/Runtime/Engine/Private/World.cpp
--- a/Engine/Source/Runtime/Engine/Private/World.cpp
+++ b/Engine/Source/Runtime/Engine/Private/World.cpp
@@ -7171,7 +7171,10 @@ UWorld* FSeamlessTravelHandler::Tick()
 				// Keep if it's in the current level AND it isn't specifically excluded AND it was either marked as should keep OR we don't own this actor
 				if (bIsInCurrentLevel && !bForceExcludeActor && (bManuallyMarkedKeep || bKeepNonOwnedActor))
 				{
-					NetDriver->NotifyActorIsTraveling(TheActor);
+					if (NetDriver != nullptr)
+					{
+						NetDriver->NotifyActorIsTraveling(TheActor);
+					}
 
 					ActuallyKeptActors.Add(TheActor);
 					return true;
