Fix vertex buffers with software rendering in content workers.

diff --git a/Engine/Source/Runtime/Engine/Private/Rendering/ColorVertexBuffer.cpp b/Engine/Source/Runtime/Engine/Private/Rendering/ColorVertexBuffer.cpp
--- a/Engine/Source/Runtime/Engine/Private/Rendering/ColorVertexBuffer.cpp
+++ b/Engine/Source/Runtime/Engine/Private/Rendering/ColorVertexBuffer.cpp
@@ -246,6 +246,11 @@ void FColorVertexBuffer::Serialize( FArchive& Ar, bool bNeedsCPUAccess )
 				}
 			}
 		}
+		if (StripFlags.IsDataStrippedForServer())
+		{
+			// if we stripped all the other stuff and decided not to serialize it in probably need to strip the NumVertices Too
+			NumVertices = Stride = 0;
+		}
 	}
 }
 
