Fix incorrect EV conversion for spot light.

diff --git a/Engine/Source/Runtime/Engine/Private/Components/SpotLightComponent.cpp b/Engine/Source/Runtime/Engine/Private/Components/SpotLightComponent.cpp
--- a/Engine/Source/Runtime/Engine/Private/Components/SpotLightComponent.cpp
+++ b/Engine/Source/Runtime/Engine/Private/Components/SpotLightComponent.cpp
@@ -247,7 +247,7 @@ float USpotLightComponent::ComputeLightBrightness() const
 		}
 		else if (IntensityUnits == ELightUnits::EV)
 		{
-			LightBrightness = EV100ToLuminance(LightBrightness * (100.f * 100.f));
+			LightBrightness = EV100ToLuminance(LightBrightness) * (100.f * 100.f);
 		}
 		else
 		{
