Cloth - Clear ClothCollisionSource cached data when extracting collisions for cloth if the Source Component mesh has changed.

#rnx

diff --git a/Engine/Source/Runtime/Engine/Private/SkeletalMeshComponentPhysics.cpp b/Engine/Source/Runtime/Engine/Private/SkeletalMeshComponentPhysics.cpp
--- a/Engine/Source/Runtime/Engine/Private/SkeletalMeshComponentPhysics.cpp
+++ b/Engine/Source/Runtime/Engine/Private/SkeletalMeshComponentPhysics.cpp
@@ -2865,6 +2865,10 @@ void USkeletalMeshComponent::ExtractCollisionsForCloth(
 		// Init cache on first copy
 		if(!ClothCollisionSource.bCached || ClothCollisionSource.CachedSkeletalMesh.Get() != SourceComponent->GetSkeletalMeshAsset())
 		{
+			// Clear previous cached data
+			ClothCollisionSource.CachedSpheres.Reset();
+			ClothCollisionSource.CachedSphereConnections.Reset();
+
 			for(const USkeletalBodySetup* SkeletalBodySetup : PhysicsAsset->SkeletalBodySetups)
 			{
 				// Cache bones
