Fix r.HDR.EnableHDROutput which could get stuck on windows platforms: this CVar now means that the project wants to enable HDR, and won't be changed in case no display supports HDR. We rely on IsHDREnabled / HDRGetMetaData in the rest of the engine which checks if the display we're rendering to supports HDR.
Fix UGameUserSettings::GetCurrentHDRDisplayNits not updating r.HDR.Display.MaxLuminance when it was changing

diff --git a/Engine/Source/Runtime/Engine/Private/UnrealEngine.cpp b/Engine/Source/Runtime/Engine/Private/UnrealEngine.cpp
--- a/Engine/Source/Runtime/Engine/Private/UnrealEngine.cpp
+++ b/Engine/Source/Runtime/Engine/Private/UnrealEngine.cpp
@@ -911,42 +911,35 @@ void ENGINE_API HDRSettingChangedSinkCallback()
 		return;
 	}
 
-	bool bIsHDREnabled = CVarHDROutputEnabled->GetValueOnAnyThread() != 0;
-	if (bIsHDREnabled && !IsHDREnabled())
-	{
-		UE_LOG(LogConsoleResponse, Display, TEXT("Tried to enable HDR display output but unsupported, forcing off."));
-		bIsHDREnabled = false;
-	}
-
 	static float GHDRMinLuminnanceLog10 = CVarHDRMinLuminanceLog10->GetValueOnAnyThread();
 	static float GHDRMidLuminnance = CVarHDRMidLuminance->GetValueOnAnyThread();
 	static int32 GHDRMaxLuminnance = CVarHDRMaxLuminance->GetValueOnAnyThread();
 	static float GHDRSceneColorMultiplier = CVarHDRSceneColorMultiplier->GetValueOnAnyThread();
+	static int32 GDisplayNitLevel = 1000;
 
-	if(bIsHDREnabled != GRHIIsHDREnabled || !FMath::IsNearlyEqual(GHDRMinLuminnanceLog10, CVarHDRMinLuminanceLog10->GetValueOnAnyThread()) || !FMath::IsNearlyEqual(GHDRMidLuminnance, CVarHDRMidLuminance->GetValueOnAnyThread()) ||
-									        (GHDRMaxLuminnance != CVarHDRMaxLuminance->GetValueOnAnyThread()) || !FMath::IsNearlyEqual(GHDRSceneColorMultiplier, CVarHDRSceneColorMultiplier->GetValueOnAnyThread()))
+	bool bIsHDREnabled = IsHDREnabled();
+	int32 DisplayNitLevel = (GEngine && GEngine->GetGameUserSettings()) ? GEngine->GetGameUserSettings()->GetCurrentHDRDisplayNits() : 0;
+
+	// We'll naively fall back to 1000 if DisplayNits is 0. This can happen also if HDR is disabled from GameUserSettings
+	if (DisplayNitLevel == 0)
 	{
-		// We'll naively fall back to 1000 if DisplayNits is 0
-		uint32 DisplayNitLevel = 0;
-		
-		if (GEngine && GEngine->GetGameUserSettings())
-		{
-			DisplayNitLevel = GEngine->GetGameUserSettings()->GetCurrentHDRDisplayNits();
-		}
+		DisplayNitLevel = 1000;
+	}
 
-		if(DisplayNitLevel == 0)
-		{
-			DisplayNitLevel = 1000;
-		}
-		
+	if(bIsHDREnabled != GRHIIsHDREnabled || !FMath::IsNearlyEqual(GHDRMinLuminnanceLog10, CVarHDRMinLuminanceLog10->GetValueOnAnyThread())
+		|| !FMath::IsNearlyEqual(GHDRMidLuminnance, CVarHDRMidLuminance->GetValueOnAnyThread())
+		|| (GHDRMaxLuminnance != CVarHDRMaxLuminance->GetValueOnAnyThread()) || (DisplayNitLevel != GDisplayNitLevel)
+		|| !FMath::IsNearlyEqual(GHDRSceneColorMultiplier, CVarHDRSceneColorMultiplier->GetValueOnAnyThread()))
+	{
 		HDRConfigureCVars(bIsHDREnabled, DisplayNitLevel, false);
 
 		// Now set the HDR setting.
-		GRHIIsHDREnabled = IsHDREnabled();
+		GRHIIsHDREnabled = bIsHDREnabled;
 		GHDRMinLuminnanceLog10 = CVarHDRMinLuminanceLog10->GetValueOnAnyThread();
 		GHDRMidLuminnance = CVarHDRMidLuminance->GetValueOnAnyThread();
 		GHDRMaxLuminnance = CVarHDRMaxLuminance->GetValueOnAnyThread();
 		GHDRSceneColorMultiplier = CVarHDRSceneColorMultiplier->GetValueOnAnyThread();
+		GDisplayNitLevel = DisplayNitLevel;
 	}
 }
 
diff --git a/Engine/Source/Runtime/RenderCore/Private/RenderCore.cpp b/Engine/Source/Runtime/RenderCore/Private/RenderCore.cpp
--- a/Engine/Source/Runtime/RenderCore/Private/RenderCore.cpp
+++ b/Engine/Source/Runtime/RenderCore/Private/RenderCore.cpp
@@ -521,15 +521,15 @@ RENDERCORE_API void HDRGetMetaData(EDisplayOutputFormat& OutDisplayOutputFormat,
 	OutDisplayOutputFormat = HDRMetaData.DisplayOutputFormat;
 	OutDisplayColorGamut = HDRMetaData.DisplayColorGamut;
 	OutbHDRSupported = HDRMetaData.bHDRSupported;
-	if (!IsHDREnabled() || OSWindow == nullptr)
+	if (!IsHDREnabled())
 	{
 		return;
 	}
 
 	FDisplayInformationArray DisplayList;
 	RHIGetDisplaysInformation(DisplayList);
-	// In case we have 1 display or less, the CVars that were setup do represent the state of the displays
-	if (DisplayList.Num() <= 1)
+	// In case we have no display registered, use the CVars
+	if (DisplayList.Num() == 0)
 	{
 		return;
 	}
@@ -555,7 +555,7 @@ RENDERCORE_API void HDRGetMetaData(EDisplayOutputFormat& OutDisplayOutputFormat,
 	if (OutbHDRSupported)
 	{
 		FPlatformMisc::ChooseHDRDeviceAndColorGamut(GRHIVendorId, CVarHDRDisplayMaxLuminance.GetValueOnAnyThread(), OutDisplayOutputFormat, OutDisplayColorGamut);
-	}
+}
 
 }
 
@@ -577,7 +577,7 @@ RENDERCORE_API void HDRConfigureCVars(bool bIsHDREnabled, uint32 DisplayNits, bo
 {
 	if (bIsHDREnabled && !GRHISupportsHDROutput)
 	{
-		UE_LOG(LogRendererCore, Warning, TEXT("Trying to enable HDR but it is not supported by the RHI: IsHDREnabled will return false"));
+		// In case we request HDR but no display supports it, we still need to setup Gamut / OETF properly for the tonemapper to work
 		bIsHDREnabled = false;
 	}
 
@@ -594,12 +594,10 @@ RENDERCORE_API void HDRConfigureCVars(bool bIsHDREnabled, uint32 DisplayNits, bo
 	//CVarHDRDisplayMaxLuminance is ECVF_SetByCode as it's only a mean of communicating the information from UGameUserSettings to the rest of the engine
 	if (bIsHDREnabled)
 	{
-		CVarHDROutputEnabled->Set(1, bFromGameSettings ? ECVF_SetByGameSetting : ECVF_SetByCode);
 		CVarHDRDisplayMaxLuminance->Set((int32)DisplayNits, ECVF_SetByCode);
 	}
 	else
 	{
-		CVarHDROutputEnabled->Set(0, bFromGameSettings ? ECVF_SetByGameSetting : ECVF_SetByCode);
 		CVarHDRDisplayMaxLuminance->Set(0, ECVF_SetByCode);
 	}
 
