Do not clear transaction buffer on on transactionable worlds

#rnx

diff --git a/Engine/Source/Runtime/Engine/Private/WorldPartition/WorldPartitionActorLoaderInterface.cpp b/Engine/Source/Runtime/Engine/Private/WorldPartition/WorldPartitionActorLoaderInterface.cpp
--- a/Engine/Source/Runtime/Engine/Private/WorldPartition/WorldPartitionActorLoaderInterface.cpp
+++ b/Engine/Source/Runtime/Engine/Private/WorldPartition/WorldPartitionActorLoaderInterface.cpp
@@ -257,7 +257,7 @@ void IWorldPartitionActorLoaderInterface::ILoaderAdapter::PostLoadedStateChanged
 
 		if (NumUnloads)
 		{
-			if (!GIsTransacting)
+			if (!GIsTransacting && World->HasAnyFlags(RF_Transactional))
 			{
 				GEditor->ResetTransaction(LOCTEXT("UnloadingEditorActorResetTrans", "Editor Actors Unloaded"));
 			}
