Mark UserDefinedEnums postload as thread-safe
  - This change only affects the editor and will only have an effect when ALT is activated

#rnx

diff --git a/Engine/Source/Runtime/Engine/Classes/Engine/UserDefinedEnum.h b/Engine/Source/Runtime/Engine/Classes/Engine/UserDefinedEnum.h
--- a/Engine/Source/Runtime/Engine/Classes/Engine/UserDefinedEnum.h
+++ b/Engine/Source/Runtime/Engine/Classes/Engine/UserDefinedEnum.h
@@ -74,6 +74,7 @@ public:
 	virtual bool Rename(const TCHAR* NewName = NULL, UObject* NewOuter = NULL, ERenameFlags Flags = REN_None) override;
 	virtual void PostDuplicate(bool bDuplicateForPIE) override;
 	virtual void PostLoad() override;
+	virtual bool IsPostLoadThreadSafe() const override;
 	virtual void PostEditUndo() override;
 	virtual void PostEditChangeProperty(FPropertyChangedEvent& PropertyChangedEvent);
 	virtual void GetAssetRegistryTags(TArray<FAssetRegistryTag>& OutTags) const override;
diff --git a/Engine/Source/Runtime/Engine/Private/UserDefinedEnum.cpp b/Engine/Source/Runtime/Engine/Private/UserDefinedEnum.cpp
--- a/Engine/Source/Runtime/Engine/Private/UserDefinedEnum.cpp
+++ b/Engine/Source/Runtime/Engine/Private/UserDefinedEnum.cpp
@@ -80,6 +80,11 @@ void UUserDefinedEnum::PostDuplicate(bool bDuplicateForPIE)
 	}
 }
 
+bool UUserDefinedEnum::IsPostLoadThreadSafe() const
+{
+	return true;
+}
+
 void UUserDefinedEnum::PostLoad()
 {
 	Super::PostLoad();
