Fix crashes in Material duplication logic.

diff --git a/Engine/Source/Runtime/Engine/Private/Materials/MaterialInterface.cpp b/Engine/Source/Runtime/Engine/Private/Materials/MaterialInterface.cpp
--- a/Engine/Source/Runtime/Engine/Private/Materials/MaterialInterface.cpp
+++ b/Engine/Source/Runtime/Engine/Private/Materials/MaterialInterface.cpp
@@ -706,13 +706,18 @@ void UMaterialInterface::PostDuplicate(bool bDuplicateForPIE)
 #if WITH_EDITORONLY_DATA
 	// Duplication will create a new editor only data but this MI already links a populated data. This makes sure that the
 	// duplicate instance EOD has the correct name (necessary for running the editor on cooked data).
+	UMaterialInterfaceEditorOnlyData* EditorOnly = GetEditorOnlyData();
 	FString EditorOnlyDataName = MaterialInterface::GetEditorOnlyDataName(*GetName());
 	UObject* EditorOnlyExisting = StaticFindObject(/*Class=*/ nullptr, this, *EditorOnlyDataName, true);
-	if (EditorOnlyExisting)
+	if (EditorOnlyExisting && EditorOnlyExisting != EditorOnly)
 	{
 		EditorOnlyExisting->Rename(*(EditorOnlyDataName + "_Unused"), nullptr, REN_NonTransactional | REN_DontCreateRedirectors);
+		EditorOnlyExisting->MarkAsGarbage();
+	}
+	if (EditorOnly)
+	{
+		GetEditorOnlyData()->Rename(*EditorOnlyDataName, nullptr, REN_NonTransactional | REN_DontCreateRedirectors);
 	}
-	GetEditorOnlyData()->Rename(*EditorOnlyDataName, nullptr, REN_NonTransactional | REN_DontCreateRedirectors);
 #endif // WITH_EDITORONLY_DATA
 }
 
@@ -1530,15 +1535,19 @@ void UMaterialInterface::PreSave(FObjectPreSaveContext ObjectSaveContext)
 	// the editor running on cooked data.
 	UMaterialInterfaceEditorOnlyData* EditorOnly = GetEditorOnlyData();
 	FString EditorOnlyDataName = MaterialInterface::GetEditorOnlyDataName(*GetName());
-	if (EditorOnly->GetName() != EditorOnlyDataName)
+	if (EditorOnly && EditorOnly->GetName() != EditorOnlyDataName)
 	{
 		UE_LOG(LogMaterial, Display, TEXT("MaterialInterface %s has a incorrectly name EditorOnlyData '%s'. This may cause issues when running the editor on cooked data. Trying to rename it to the correct name '%s'."), *GetName(), *EditorOnly->GetName(), *EditorOnlyDataName);
 		UObject* EditorOnlyExisting = StaticFindObject(/*Class=*/ nullptr, this, *EditorOnlyDataName, true);
-		if (EditorOnlyExisting)
+		if (EditorOnlyExisting && EditorOnlyExisting != EditorOnly)
 		{
 			EditorOnlyExisting->Rename(*(EditorOnlyDataName + "_Unused"), nullptr, REN_NonTransactional | REN_DontCreateRedirectors);
+			EditorOnlyExisting->MarkAsGarbage();
+		}
+		if (EditorOnly)
+		{
+			EditorOnly->Rename(*EditorOnlyDataName, nullptr, REN_NonTransactional | REN_DontCreateRedirectors);
 		}
-		EditorOnly->Rename(*EditorOnlyDataName, nullptr, REN_NonTransactional | REN_DontCreateRedirectors);
 	}
 #endif // WITH_EDITORONLY_DATA
 
