ChildActorComponent - Add api to custom configure name of child actor

diff --git a/Engine/Source/Runtime/Engine/Classes/Components/ChildActorComponent.h b/Engine/Source/Runtime/Engine/Classes/Components/ChildActorComponent.h
--- a/Engine/Source/Runtime/Engine/Classes/Components/ChildActorComponent.h
+++ b/Engine/Source/Runtime/Engine/Classes/Components/ChildActorComponent.h
@@ -142,6 +142,9 @@ private:
 	/** Flag indicating that when the component is registered that the child actor should be recreated */
 	uint8 bNeedsRecreate:1;
 
+	/** When true, does not modify ChildActorName during spawn such as removing _UAID_ */
+	uint8 bChildActorNameIsExact:1;
+
 #if WITH_EDITOR
 	virtual void SetPackageExternal(bool bExternal, bool bShouldDirty) override;
 #endif
@@ -186,6 +189,10 @@ public:
 	AActor* GetChildActorTemplate() const { return ChildActorTemplate; }
 
 	FName GetChildActorName() const { return ChildActorName; }
+	void SetChildActorName(const FName InName);
+
+	/** When true, does not modify ChildActorName during spawn such as removing _UAID_ */
+	void SetChildActorNameIsExact(bool bInExact);
 
 	/** Kill any currently present child actor */
 	void DestroyChildActor();
diff --git a/Engine/Source/Runtime/Engine/Private/Components/ChildActorComponent.cpp b/Engine/Source/Runtime/Engine/Private/Components/ChildActorComponent.cpp
--- a/Engine/Source/Runtime/Engine/Private/Components/ChildActorComponent.cpp
+++ b/Engine/Source/Runtime/Engine/Private/Components/ChildActorComponent.cpp
@@ -32,6 +32,8 @@ static FAutoConsoleVariableRef CVarExperimentalAllowPerInstanceChildActorPropert
 UChildActorComponent::UChildActorComponent(const FObjectInitializer& ObjectInitializer)
 	: Super(ObjectInitializer)
 	, ActorOuter(nullptr)
+	, bNeedsRecreate(false)
+	, bChildActorNameIsExact(false)
 {
 	bAllowReregistration = false;
 
@@ -751,7 +753,13 @@ void UChildActorComponent::CreateChildActor(TFunction<void(AActor*)> CustomizerF
 				Params.bAllowDuringConstructionScript = true;
 				Params.OverrideLevel = (MyOwner ? MyOwner->GetLevel() : nullptr);
 				Params.Name = ChildActorName;
-				Params.NameMode = FActorSpawnParameters::ESpawnActorNameMode::Requested;
+
+				if (!bChildActorNameIsExact)
+				{
+					// Note: Requested will remove of _UAID_ from a name
+					Params.NameMode = FActorSpawnParameters::ESpawnActorNameMode::Requested;
+				}
+
 				Params.OverrideParentComponent = this;
 
 #if WITH_EDITOR
@@ -839,6 +847,16 @@ void UChildActorComponent::CreateChildActor(TFunction<void(AActor*)> CustomizerF
 	}
 }
 
+void UChildActorComponent::SetChildActorName(const FName InName)
+{
+	ChildActorName = InName;
+}
+
+void UChildActorComponent::SetChildActorNameIsExact(bool bInExact)
+{
+	bChildActorNameIsExact = bInExact;
+}
+
 void UChildActorComponent::DestroyChildActor()
 {
 	UWorld* LocalWorld = GetWorld();
