Call RecalcRequiredBones on persistent shared linked anim instances when they are reused to fix issues with LODs and required bones mismatches.

diff --git a/Engine/Source/Runtime/Engine/Private/Animation/AnimSubsystem_SharedLinkedAnimLayers.cpp b/Engine/Source/Runtime/Engine/Private/Animation/AnimSubsystem_SharedLinkedAnimLayers.cpp
--- a/Engine/Source/Runtime/Engine/Private/Animation/AnimSubsystem_SharedLinkedAnimLayers.cpp
+++ b/Engine/Source/Runtime/Engine/Private/Animation/AnimSubsystem_SharedLinkedAnimLayers.cpp
@@ -96,6 +96,8 @@ UAnimInstance* FLinkedAnimLayerClassData::FindOrAddInstanceForLinking(UAnimInsta
 			if (LayerInstanceData.IsPersistent() && LayerInstanceData.GetLinkedFunctions().Num() == 0)
 			{
 				USkeletalMeshComponent* Mesh = OwningInstance->GetSkelMeshComponent();
+				// Make sure the bones to update are up to date with LOD changes / bone visibility / cosmetics, etc.
+				LayerInstanceData.Instance->RecalcRequiredBones();
 				check(!Mesh->GetLinkedAnimInstances().Contains(LayerInstanceData.Instance));
 				Mesh->GetLinkedAnimInstances().Add(LayerInstanceData.Instance);
 			}
