Level Instance: WorldPartitionActorFilter support (experimental)
- Experimental Flag hides Filter UI
- DataLayerAsset can enable filtering (will show up in Filter Outliner)
- DataLayerInstance can provide default filter value (default value when actor doesn't override)
- Level Instance Actor can have its own recursive filter
- Filtering Editor/Cook/PIE works
- Undo/Redo/Copy/Paste/ResetToDefault
- PackedLevelActors not yet supported

#rnx

diff --git a/Engine/Plugins/Editor/DataValidation/Source/DataValidation/Private/WorldPartitionChangelistValidator.cpp b/Engine/Plugins/Editor/DataValidation/Source/DataValidation/Private/WorldPartitionChangelistValidator.cpp
--- a/Engine/Plugins/Editor/DataValidation/Source/DataValidation/Private/WorldPartitionChangelistValidator.cpp
+++ b/Engine/Plugins/Editor/DataValidation/Source/DataValidation/Private/WorldPartitionChangelistValidator.cpp
@@ -440,5 +440,10 @@ void UWorldPartitionChangelistValidator::OnLevelInstanceInvalidWorldAsset(const
 	}
 }
 
+void UWorldPartitionChangelistValidator::OnActorFilterFailed(const FWorldPartitionActorDescView& ActorDescView)
+{
+	// Not a validation error
+}
+
 #undef LOCTEXT_NAMESPACE
 
diff --git a/Engine/Plugins/Editor/DataValidation/Source/DataValidation/Private/WorldPartitionChangelistValidator.h b/Engine/Plugins/Editor/DataValidation/Source/DataValidation/Private/WorldPartitionChangelistValidator.h
--- a/Engine/Plugins/Editor/DataValidation/Source/DataValidation/Private/WorldPartitionChangelistValidator.h
+++ b/Engine/Plugins/Editor/DataValidation/Source/DataValidation/Private/WorldPartitionChangelistValidator.h
@@ -46,4 +46,5 @@ protected:
 	virtual void OnDataLayerAssetConflict(const UDataLayerInstanceWithAsset* DataLayerInstance, const UDataLayerInstanceWithAsset* ConflictingDataLayerInstance) override;
 	virtual void OnActorNeedsResave(const FWorldPartitionActorDescView& ActorDescView) override;
 	virtual void OnLevelInstanceInvalidWorldAsset(const FWorldPartitionActorDescView& ActorDescView, FName WorldAsset, ELevelInstanceInvalidReason Reason) override;
+	virtual void OnActorFilterFailed(const FWorldPartitionActorDescView& ActorDescView) override;
 };
diff --git a/Engine/Source/Editor/LevelInstanceEditor/LevelInstanceEditor.Build.cs b/Engine/Source/Editor/LevelInstanceEditor/LevelInstanceEditor.Build.cs
--- a/Engine/Source/Editor/LevelInstanceEditor/LevelInstanceEditor.Build.cs
+++ b/Engine/Source/Editor/LevelInstanceEditor/LevelInstanceEditor.Build.cs
@@ -32,7 +32,8 @@ public class LevelInstanceEditor : ModuleRules
 				"MessageLog",
 				"EditorWidgets",
 				"DeveloperSettings",
-				"SceneOutliner"
+				"SceneOutliner",
+				"WorldPartitionEditor"
 			}
 		);
 
diff --git a/Engine/Source/Editor/LevelInstanceEditor/Private/LevelInstanceEditorModule.cpp b/Engine/Source/Editor/LevelInstanceEditor/Private/LevelInstanceEditorModule.cpp
--- a/Engine/Source/Editor/LevelInstanceEditor/Private/LevelInstanceEditorModule.cpp
+++ b/Engine/Source/Editor/LevelInstanceEditor/Private/LevelInstanceEditorModule.cpp
@@ -2,6 +2,7 @@
 #include "LevelInstanceEditorModule.h"
 #include "LevelInstanceActorDetails.h"
 #include "LevelInstancePivotDetails.h"
+#include "LevelInstanceFilterPropertyTypeCustomization.h"
 #include "LevelInstance/LevelInstanceSubsystem.h"
 #include "LevelInstance/LevelInstanceInterface.h"
 #include "LevelInstance/LevelInstanceActor.h"
@@ -41,6 +42,7 @@
 #include "MessageLogModule.h"
 #include "Settings/EditorExperimentalSettings.h"
 #include "WorldPartition/WorldPartitionConverter.h"
+#include "WorldPartition/WorldPartitionActorLoaderInterface.h"
 
 IMPLEMENT_MODULE( FLevelInstanceEditorModule, LevelInstanceEditor );
 
@@ -669,7 +671,8 @@ void FLevelInstanceEditorModule::StartupModule()
 
 	FPropertyEditorModule& PropertyModule = FModuleManager::LoadModuleChecked<FPropertyEditorModule>("PropertyEditor");
 	PropertyModule.RegisterCustomClassLayout("LevelInstance", FOnGetDetailCustomizationInstance::CreateStatic(&FLevelInstanceActorDetails::MakeInstance));
-	PropertyModule.RegisterCustomClassLayout("LevelInstancePivot", FOnGetDetailCustomizationInstance::CreateStatic(&FLevelInstancePivotDetails::MakeInstance));
+	PropertyModule.RegisterCustomClassLayout("LevelInstancePivot", FOnGetDetailCustomizationInstance::CreateStatic(&FLevelInstancePivotDetails::MakeInstance));		
+	PropertyModule.RegisterCustomPropertyTypeLayout("WorldPartitionActorFilter", FOnGetPropertyTypeCustomizationInstance::CreateStatic(&FLevelInstanceFilterPropertyTypeCustomization::MakeInstance), MakeShared<FLevelInstancePropertyTypeIdentifier>());
 	PropertyModule.NotifyCustomizationModuleChanged();
 
 	// GEditor needs to be set before this module is loaded
@@ -678,6 +681,20 @@ void FLevelInstanceEditorModule::StartupModule()
 	
 	EditorLevelUtils::CanMoveActorToLevelDelegate.AddRaw(this, &FLevelInstanceEditorModule::CanMoveActorToLevel);
 
+	// Register actor descriptor loading filter
+	IWorldPartitionActorLoaderInterface::RegisterActorDescFilter([](UWorld* World, const FWorldPartitionHandle& ActorHandle) 
+	{ 
+		if (UWorld* OwningWorld = World->PersistentLevel->GetWorld())
+		{
+			if (ULevelInstanceSubsystem* LevelInstanceSubsystem = OwningWorld->GetSubsystem<ULevelInstanceSubsystem>())
+			{
+				return LevelInstanceSubsystem->PassLevelInstanceFilter(World, ActorHandle);
+			}
+		}
+
+		return true;
+	});
+
 	FMessageLogModule& MessageLogModule = FModuleManager::LoadModuleChecked<FMessageLogModule>("MessageLog");
 	FMessageLogInitializationOptions InitOptions;
 	InitOptions.bShowFilters = true;
diff --git a/Engine/Source/Editor/LevelInstanceEditor/Private/LevelInstanceFilterPropertyTypeCustomization.cpp b/Engine/Source/Editor/LevelInstanceEditor/Private/LevelInstanceFilterPropertyTypeCustomization.cpp
new file mode 100644
--- /dev/null
+++ b/Engine/Source/Editor/LevelInstanceEditor/Private/LevelInstanceFilterPropertyTypeCustomization.cpp
@@ -0,0 +1,79 @@
+// Copyright Epic Games, Inc. All Rights Reserved.
+
+#include "LevelInstanceFilterPropertyTypeCustomization.h"
+#include "LevelInstance/LevelInstanceInterface.h"
+#include "LevelInstance/LevelInstanceSubsystem.h"
+#include "LevelInstance/LevelInstanceComponent.h"
+#include "WorldPartition/Filter/WorldPartitionActorFilter.h"
+#include "Engine/World.h"
+#include "GameFramework/Actor.h"
+
+#define LOCTEXT_NAMESPACE "WorldPartitionActorFilter"
+
+TSharedPtr<FWorldPartitionActorFilterMode::FFilter> FLevelInstanceFilterPropertyTypeCustomization::CreateModeFilter(TArray<UObject*> OuterObjects)
+{
+	LevelInstances.Empty();
+	FString ActorLabel;
+
+	// Find Selected Level Instances with matching WorldAssetPackage
+	FString WorldAssetPackage;
+	UWorld* World = nullptr;
+	for (UObject* OuterObject : OuterObjects)
+	{
+		if (AActor* OuterActor = OuterObject->GetTypedOuter<AActor>())
+		{
+			if (ILevelInstanceInterface* LevelInstanceInterface = Cast<ILevelInstanceInterface>(OuterActor))
+			{
+				LevelInstances.Add(LevelInstanceInterface);
+				ActorLabel = OuterActor->GetActorLabel();
+
+				if (WorldAssetPackage.IsEmpty())
+				{
+					World = OuterActor->GetWorld();
+					WorldAssetPackage = LevelInstanceInterface->GetWorldAssetPackage();
+				}
+				else if (WorldAssetPackage != LevelInstanceInterface->GetWorldAssetPackage())
+				{
+					// Multiple WorldAssetPackages not supported
+					return nullptr;
+				}
+			}
+		}
+	}
+
+	// No matching WorldAssetPackage for selected Level Instances
+	// or Actor is a template and doesn't have world
+	if (WorldAssetPackage.IsEmpty() || !World)
+	{
+		return nullptr;
+	}
+
+	ULevelInstanceSubsystem* LevelInstanceSubsystem = World->GetSubsystem<ULevelInstanceSubsystem>();
+
+	// Get the Default Filter for selected WorldAssetPackage
+	TSharedPtr<FWorldPartitionActorFilter> Filter = MakeShared<FWorldPartitionActorFilter>(LevelInstanceSubsystem->GetLevelInstanceFilter(WorldAssetPackage));
+
+	// Set its name based on single/multi selection (root node name in outliner)
+	Filter->DisplayName = LevelInstances.Num() == 1 ? ActorLabel : TEXT("(Multiple Actors)");
+
+	// Gather Filters for selected level instances
+	TArray<const FWorldPartitionActorFilter*> SelectedFilters;
+	Algo::Transform(LevelInstances, SelectedFilters, [](ILevelInstanceInterface* LevelInstanceInterface) { return &LevelInstanceInterface->GetFilter(); });
+
+	// Create Mode Filter which holds the final values for the filter
+	return MakeShared<FWorldPartitionActorFilterMode::FFilter>(Filter, SelectedFilters);
+}
+
+void FLevelInstanceFilterPropertyTypeCustomization::ApplyFilter(const FWorldPartitionActorFilterMode& Mode)
+{
+	const FScopedTransaction Transaction(LOCTEXT("WorldPartitionActorFilterApply_Transaction", "Apply Level Instance Filter"));
+	for (ILevelInstanceInterface* LevelInstance : LevelInstances)
+	{
+		ULevelInstanceComponent* Component = LevelInstance->GetLevelInstanceComponent();
+		FWorldPartitionActorFilter ComponentFilter = Component->GetFilter();
+		Mode.Apply(ComponentFilter);
+		Component->SetFilter(ComponentFilter);
+	}
+}
+
+#undef LOCTEXT_NAMESPACE
\ No newline at end of file
diff --git a/Engine/Source/Editor/LevelInstanceEditor/Private/LevelInstanceFilterPropertyTypeCustomization.h b/Engine/Source/Editor/LevelInstanceEditor/Private/LevelInstanceFilterPropertyTypeCustomization.h
new file mode 100644
--- /dev/null
+++ b/Engine/Source/Editor/LevelInstanceEditor/Private/LevelInstanceFilterPropertyTypeCustomization.h
@@ -0,0 +1,36 @@
+// Copyright Epic Games, Inc. All Rights Reserved.
+
+#pragma once
+
+#include "IPropertyTypeCustomization.h"
+#include "Internationalization/Text.h"
+#include "PropertyEditorModule.h"
+#include "PropertyHandle.h"
+#include "Templates/SharedPointer.h"
+#include "WorldPartition/Filter/WorldPartitionActorFilterPropertyTypeCustomization.h"
+#include "WorldPartition/Filter/WorldPartitionActorFilterMode.h"
+
+class ILevelInstanceInterface;
+
+class FLevelInstancePropertyTypeIdentifier : public IPropertyTypeIdentifier
+{
+	virtual bool IsPropertyTypeCustomized(const IPropertyHandle& PropertyHandle) const override
+	{
+		return PropertyHandle.HasMetaData(TEXT("LevelInstanceFilter"));
+	}
+};
+
+// Registered (FLevelInstanceEditorModule::StartupModule) Property Customization for properties of type FWorldPartitionActorFilter for Level Instances
+struct FLevelInstanceFilterPropertyTypeCustomization : public FWorldPartitionActorFilterPropertyTypeCustomization
+{
+public:
+	static TSharedRef<IPropertyTypeCustomization> MakeInstance()
+	{
+		return MakeShareable(new FLevelInstanceFilterPropertyTypeCustomization);
+	}
+private:
+	virtual TSharedPtr<FWorldPartitionActorFilterMode::FFilter> CreateModeFilter(TArray<UObject*> OuterObjects) override;
+	virtual void ApplyFilter(const FWorldPartitionActorFilterMode& Mode) override;
+
+	TArray<ILevelInstanceInterface*> LevelInstances;
+};
diff --git a/Engine/Source/Editor/UnrealEd/Classes/Settings/EditorExperimentalSettings.h b/Engine/Source/Editor/UnrealEd/Classes/Settings/EditorExperimentalSettings.h
--- a/Engine/Source/Editor/UnrealEd/Classes/Settings/EditorExperimentalSettings.h
+++ b/Engine/Source/Editor/UnrealEd/Classes/Settings/EditorExperimentalSettings.h
@@ -140,6 +140,9 @@ public:
 	UPROPERTY(EditAnywhere, config, Category = Level)
 	bool bLevelInstance;
 
+	UPROPERTY(EditAnywhere, config, Category = WorldPartition)
+	bool bEnableWorldPartitionActorFilters;
+
 	/**
 	 * Returns an event delegate that is executed when a setting has changed.
 	 *
diff --git a/Engine/Source/Editor/WorldPartitionEditor/Private/WorldPartition/Filter/WorldPartitionActorFilterColumn.cpp b/Engine/Source/Editor/WorldPartitionEditor/Private/WorldPartition/Filter/WorldPartitionActorFilterColumn.cpp
new file mode 100644
--- /dev/null
+++ b/Engine/Source/Editor/WorldPartitionEditor/Private/WorldPartition/Filter/WorldPartitionActorFilterColumn.cpp
@@ -0,0 +1,144 @@
+// Copyright Epic Games, Inc. All Rights Reserved.
+
+#include "WorldPartition/Filter/WorldPartitionActorFilterColumn.h"
+
+#include "Widgets/Images/SImage.h"
+#include "Widgets/Input/SButton.h"
+#include "Widgets/Input/SCheckBox.h"
+#include "Widgets/Views/STreeView.h"
+#include "Widgets/SNullWidget.h"
+#include "ScopedTransaction.h"
+#include "Editor.h"
+#include "Styling/AppStyle.h"
+#include "WorldPartition/Filter/WorldPartitionActorFilterTreeItems.h"
+#include "WorldPartition/Filter/WorldPartitionActorFilterMode.h"
+#include "WorldPartition/Filter/WorldPartitionActorFilter.h"
+
+#define LOCTEXT_NAMESPACE "WorldPartitionActorFilter"
+
+FName FWorldPartitionActorFilterColumn::ID("WorldPartitionActorFilter");
+
+FName FWorldPartitionActorFilterColumn::GetID()
+{
+	return ID;
+}
+
+SHeaderRow::FColumn::FArguments FWorldPartitionActorFilterColumn::ConstructHeaderRowColumn()
+{
+	return SHeaderRow::Column(GetColumnID())
+		.FixedWidth(48.f)
+		.HAlignHeader(HAlign_Center)
+		.VAlignHeader(VAlign_Center)
+		.HAlignCell(HAlign_Center)
+		.VAlignCell(VAlign_Center)
+		.DefaultTooltip(LOCTEXT("WorldPartitionActorFilterColumnToolTip", "Filter State"))
+		[
+			SNew(STextBlock)
+			.Margin(FMargin(0.0f))
+			.Text(LOCTEXT("WorldPartitionActorFilterColumnName", "Filter"))
+		];
+}
+
+const TSharedRef<SWidget> FWorldPartitionActorFilterColumn::ConstructRowWidget(FSceneOutlinerTreeItemRef TreeItem, const STableRow<FSceneOutlinerTreeItemPtr>& Row)
+{
+	TSharedPtr<ISceneOutliner> SceneOutliner = TreeItem->WeakSceneOutliner.Pin();
+	const FWorldPartitionActorFilterMode* Mode = static_cast<const FWorldPartitionActorFilterMode*>(SceneOutliner->GetMode());
+
+	if (FWorldPartitionActorFilterDataLayerItem* DataLayerItem = TreeItem->CastTo<FWorldPartitionActorFilterDataLayerItem>())
+	{
+		auto GetOverrideCheckedState = [DataLayerItem, Mode]() -> ECheckBoxState
+		{
+			if (const FWorldPartitionActorFilter* Filter = DataLayerItem->GetFilter())
+			{
+				const FWorldPartitionActorFilterMode::FFilter::FDataLayerFilters& DataLayerFilters = Mode->FindChecked(Filter);
+				const FWorldPartitionActorFilterMode::FFilter::FDataLayerFilter& DataLayerFilter = DataLayerFilters.FindChecked(DataLayerItem->GetAssetPath());
+
+				return DataLayerFilter.bOverride.IsSet() ? (DataLayerFilter.bOverride.GetValue() ? ECheckBoxState::Checked : ECheckBoxState::Unchecked) : ECheckBoxState::Undetermined;
+			}
+
+			return ECheckBoxState::Undetermined;
+		};
+
+		auto OnOverrideCheckedStateChanged = [DataLayerItem, Mode](ECheckBoxState NewState)
+		{
+			if (const FWorldPartitionActorFilter* Filter = DataLayerItem->GetFilter())
+			{
+				FWorldPartitionActorFilterMode::FFilter::FDataLayerFilters& DataLayerFilters = Mode->FindChecked(Filter);
+				FWorldPartitionActorFilterMode::FFilter::FDataLayerFilter& DataLayerFilter = DataLayerFilters.FindChecked(DataLayerItem->GetAssetPath());
+				DataLayerFilter.bOverride = NewState == ECheckBoxState::Checked ? true : false;
+				// If we are removing override set value to default
+				if (!DataLayerFilter.bOverride.GetValue())
+				{
+					DataLayerFilter.bIncluded = DataLayerItem->GetDefaultValue();
+				}
+			}
+		};
+
+		auto GetCheckedState = [DataLayerItem, Mode]() -> ECheckBoxState
+		{
+			if (const FWorldPartitionActorFilter* Filter = DataLayerItem->GetFilter())
+			{
+				const FWorldPartitionActorFilterMode::FFilter::FDataLayerFilters& DataLayerFilters = Mode->FindChecked(Filter);
+				const FWorldPartitionActorFilterMode::FFilter::FDataLayerFilter& DataLayerFilter = DataLayerFilters.FindChecked(DataLayerItem->GetAssetPath());
+
+				if (!DataLayerFilter.bOverride.IsSet())
+				{
+					return ECheckBoxState::Undetermined;
+				}
+
+				return DataLayerFilter.bIncluded.IsSet() ? (DataLayerFilter.bIncluded.GetValue() ? ECheckBoxState::Checked : ECheckBoxState::Unchecked) : ECheckBoxState::Undetermined;
+			}
+
+			return ECheckBoxState::Undetermined;
+		};
+
+		auto OnCheckedStateChanged = [DataLayerItem, Mode](ECheckBoxState NewState)
+		{
+			if (const FWorldPartitionActorFilter* Filter = DataLayerItem->GetFilter())
+			{
+				FWorldPartitionActorFilterMode::FFilter::FDataLayerFilters& DataLayerFilters = Mode->FindChecked(Filter);
+				FWorldPartitionActorFilterMode::FFilter::FDataLayerFilter& DataLayerFilter = DataLayerFilters.FindChecked(DataLayerItem->GetAssetPath());
+				DataLayerFilter.bIncluded = NewState == ECheckBoxState::Checked ? true : false;
+			}
+		};
+
+		auto IsCheckedStateEnabled = [DataLayerItem, Mode]() -> bool
+		{
+			if (const FWorldPartitionActorFilter* Filter = DataLayerItem->GetFilter())
+			{
+				FWorldPartitionActorFilterMode::FFilter::FDataLayerFilters& DataLayerFilters = Mode->FindChecked(Filter);
+				const FWorldPartitionActorFilterMode::FFilter::FDataLayerFilter& DataLayerFilter = DataLayerFilters.FindChecked(DataLayerItem->GetAssetPath());
+
+				return DataLayerFilter.bOverride.IsSet() && DataLayerFilter.bOverride.GetValue();
+			}
+			return false;
+		};
+
+		return SNew(SHorizontalBox)
+				+ SHorizontalBox::Slot()
+					.Padding(0, 0, 0, 0)
+					.AutoWidth()
+					.VAlign(VAlign_Center)
+					[
+						SNew(SCheckBox)
+						.IsChecked_Lambda(GetOverrideCheckedState)
+						.OnCheckStateChanged_Lambda(OnOverrideCheckedStateChanged)
+						.ToolTipText(LOCTEXT("WorldPartitionActorFilterOverrideToolTip", "Override default"))
+					]
+				+SHorizontalBox::Slot()
+					.Padding(0, 0, 0, 0)
+					.AutoWidth()
+					.VAlign(VAlign_Center)
+					[
+						SNew(SCheckBox)
+						.IsChecked_Lambda(GetCheckedState)
+						.OnCheckStateChanged_Lambda(OnCheckedStateChanged)
+						.IsEnabled_Lambda(IsCheckedStateEnabled)
+						.ToolTipText(LOCTEXT("WorldPartitionActorFilterToolTip", "Included/Excluded"))
+					];
+	}
+	
+	return SNullWidget::NullWidget;
+}
+
+#undef LOCTEXT_NAMESPACE
\ No newline at end of file
diff --git a/Engine/Source/Editor/WorldPartitionEditor/Private/WorldPartition/Filter/WorldPartitionActorFilterHierarchy.cpp b/Engine/Source/Editor/WorldPartitionEditor/Private/WorldPartition/Filter/WorldPartitionActorFilterHierarchy.cpp
new file mode 100644
--- /dev/null
+++ b/Engine/Source/Editor/WorldPartitionEditor/Private/WorldPartition/Filter/WorldPartitionActorFilterHierarchy.cpp
@@ -0,0 +1,96 @@
+// Copyright Epic Games, Inc. All Rights Reserved.
+
+#include "WorldPartition/Filter/WorldPartitionActorFilterHierarchy.h"
+#include "WorldPartition/Filter/WorldPartitionActorFilterMode.h"
+#include "WorldPartition/Filter/WorldPartitionActorFilterTreeItems.h"
+
+FWorldPartitionActorFilterHierarchy::FWorldPartitionActorFilterHierarchy(FWorldPartitionActorFilterMode* Mode)
+	:ISceneOutlinerHierarchy(Mode)
+{
+	
+}
+
+FWorldPartitionActorFilterHierarchy::~FWorldPartitionActorFilterHierarchy()
+{
+	
+}
+
+TUniquePtr<FWorldPartitionActorFilterHierarchy> FWorldPartitionActorFilterHierarchy::Create(FWorldPartitionActorFilterMode* Mode)
+{
+	return TUniquePtr<FWorldPartitionActorFilterHierarchy>(new FWorldPartitionActorFilterHierarchy(Mode));
+}
+
+void FWorldPartitionActorFilterHierarchy::CreateItems(TArray<FSceneOutlinerTreeItemPtr>& OutItems) const
+{
+	FWorldPartitionActorFilterMode* FilterMode = static_cast<FWorldPartitionActorFilterMode*>(Mode);
+	check(FilterMode);
+	const FWorldPartitionActorFilter* Filter = FilterMode->GetFilter();
+	check(Filter);
+	if (FSceneOutlinerTreeItemPtr FilterTreeItem = Mode->CreateItemFor<FWorldPartitionActorFilterItem>(FWorldPartitionActorFilterItem::FTreeItemData(Filter), true))
+	{
+		OutItems.Add(FilterTreeItem);
+		CreateChildren(FilterTreeItem, OutItems, true);
+	}
+}
+
+void FWorldPartitionActorFilterHierarchy::CreateChildren(const FSceneOutlinerTreeItemPtr& Item, TArray<FSceneOutlinerTreeItemPtr>& OutItems) const
+{
+	CreateChildren(Item, OutItems, false);
+}
+
+void FWorldPartitionActorFilterHierarchy::CreateChildren(const FSceneOutlinerTreeItemPtr& Item, TArray<FSceneOutlinerTreeItemPtr>& OutItems, bool bRecursive) const
+{
+	if (FWorldPartitionActorFilterItem* FilterItem = Item->CastTo<FWorldPartitionActorFilterItem>())
+	{
+		for (auto& [AssetPath, DataLayerFilter] : FilterItem->GetFilter()->DataLayerFilters)
+		{
+			if (FSceneOutlinerTreeItemPtr FilterTreeDataLayerItem = Mode->CreateItemFor<FWorldPartitionActorFilterDataLayerItem>(FWorldPartitionActorFilterDataLayerItem::FTreeItemData(FilterItem->GetFilter(), AssetPath), true))
+			{
+				OutItems.Add(FilterTreeDataLayerItem);
+			}
+		}
+
+		for (auto& [ActorGuid, ChildFilter] : FilterItem->GetFilter()->GetChildFilters())
+		{
+			if (FSceneOutlinerTreeItemPtr FilterTreeItem = Mode->CreateItemFor<FWorldPartitionActorFilterItem>(FWorldPartitionActorFilterItem::FTreeItemData(ChildFilter), true))
+			{
+				OutItems.Add(FilterTreeItem);
+				if (bRecursive)
+				{
+					CreateChildren(FilterTreeItem, OutItems, true);
+				}
+			}
+		}
+	}
+}
+
+FSceneOutlinerTreeItemPtr FWorldPartitionActorFilterHierarchy::FindOrCreateParentItem(const ISceneOutlinerTreeItem& Item, const TMap<FSceneOutlinerTreeItemID, FSceneOutlinerTreeItemPtr>& Items, bool bCreate /* = false */)
+{
+	const FWorldPartitionActorFilter* ParentFilter = nullptr;
+
+	if (const FWorldPartitionActorFilterItem* FilterItem = Item.CastTo<FWorldPartitionActorFilterItem>())
+	{
+		ParentFilter = FilterItem->GetFilter()->GetParentFilter();
+		
+	}
+	else if (const FWorldPartitionActorFilterDataLayerItem* DataLayerItem = Item.CastTo<FWorldPartitionActorFilterDataLayerItem>())
+	{
+		ParentFilter = DataLayerItem->GetFilter();
+	}
+
+	if (!ParentFilter)
+	{
+		return nullptr;
+	}
+
+	if (const FSceneOutlinerTreeItemPtr* ParentItem = Items.Find(FWorldPartitionActorFilterItem::ComputeTreeItemID(FWorldPartitionActorFilterItem::FTreeItemData(ParentFilter))))
+	{
+		return *ParentItem;
+	}
+	else if (bCreate)
+	{
+		return Mode->CreateItemFor<FWorldPartitionActorFilterItem>(FWorldPartitionActorFilterItem::FTreeItemData(ParentFilter));
+	}
+		
+	return nullptr;
+}
diff --git a/Engine/Source/Editor/WorldPartitionEditor/Private/WorldPartition/Filter/WorldPartitionActorFilterMode.cpp b/Engine/Source/Editor/WorldPartitionEditor/Private/WorldPartition/Filter/WorldPartitionActorFilterMode.cpp
new file mode 100644
--- /dev/null
+++ b/Engine/Source/Editor/WorldPartitionEditor/Private/WorldPartition/Filter/WorldPartitionActorFilterMode.cpp
@@ -0,0 +1,185 @@
+// Copyright Epic Games, Inc. All Rights Reserved.
+
+#include "WorldPartition/Filter/WorldPartitionActorFilterMode.h"
+#include "WorldPartition/Filter/WorldPartitionActorFilterHierarchy.h"
+#include "WorldPartition/Filter/WorldPartitionActorFilterTreeItems.h"
+#include "Algo/Transform.h"
+
+#define LOCTEXT_NAMESPACE "WorldPartitionActorFilterMode"
+
+FWorldPartitionActorFilterMode::FWorldPartitionActorFilterMode(SSceneOutliner* InSceneOutliner, TSharedPtr<FWorldPartitionActorFilterMode::FFilter> InFilter)
+	: ISceneOutlinerMode(InSceneOutliner)
+	, Filter(InFilter)
+{
+}
+
+void FWorldPartitionActorFilterMode::Rebuild()
+{
+	Hierarchy = CreateHierarchy();
+}
+
+int32 FWorldPartitionActorFilterMode::GetTypeSortPriority(const ISceneOutlinerTreeItem& Item) const
+{
+	if (Item.IsA<FWorldPartitionActorFilterDataLayerItem>())
+	{
+		return 0;
+	}
+	
+	return 1;
+}
+
+TUniquePtr<ISceneOutlinerHierarchy> FWorldPartitionActorFilterMode::CreateHierarchy()
+{
+	TUniquePtr<FWorldPartitionActorFilterHierarchy> NewHierarchy = FWorldPartitionActorFilterHierarchy::Create(this);
+	return NewHierarchy;
+}
+
+FWorldPartitionActorFilterMode::FFilter::FFilter(TSharedPtr<FWorldPartitionActorFilter> InLevelFilter, const TArray<const FWorldPartitionActorFilter*>& InSelectedFilters)
+	: LevelFilter(InLevelFilter)
+{
+	Initialize(LevelFilter.Get());
+	for (int i = 0; i < InSelectedFilters.Num(); ++i)
+	{
+		if (i == 0)
+		{
+			Override(InSelectedFilters[i]);
+		}
+		else
+		{
+			Merge(InSelectedFilters[i]);
+		}
+	}
+}
+
+void FWorldPartitionActorFilterMode::FFilter::Initialize(const FWorldPartitionActorFilter* InFilter)
+{
+	TFunction<void(const FWorldPartitionActorFilter*)> InitializeInternal = [this, &InitializeInternal](const FWorldPartitionActorFilter* InFilterInternal)
+	{
+		FFilter::FDataLayerFilters& DataLayerFilters = FilterValues.Add(InFilterInternal);
+		for (auto& [AssetPath, DataLayerFilter] : InFilterInternal->DataLayerFilters)
+		{
+			// Initialize new merged filter to not overriden and overriden value false
+			DataLayerFilters.Add(AssetPath, FFilter::FDataLayerFilter(false, DataLayerFilter.bIncluded));
+		}
+
+		for (auto& [ActorGuid, WorldPartitionActorFilter] : InFilterInternal->GetChildFilters())
+		{
+			InitializeInternal(WorldPartitionActorFilter);
+		}
+	};
+	InitializeInternal(InFilter);
+}
+
+void FWorldPartitionActorFilterMode::FFilter::Override(const FWorldPartitionActorFilter* Other)
+{
+	TFunction<void(const FWorldPartitionActorFilter*, const FWorldPartitionActorFilter*)> OverrideInternal = [this, &OverrideInternal](const FWorldPartitionActorFilter* InValue, const FWorldPartitionActorFilter* InOverride)
+	{
+		if (FFilter::FDataLayerFilters* DataLayerFilters = FilterValues.Find(InValue))
+		{
+			for (auto& [AssetPath, DataLayerFilter] : InOverride->DataLayerFilters)
+			{
+				if (FFilter::FDataLayerFilter* Found = DataLayerFilters->Find(AssetPath))
+				{
+					Found->bOverride = true;
+					Found->bIncluded = DataLayerFilter.bIncluded;
+				}
+			}
+		}
+
+		for (auto& [ActorGuid, WorldPartitionActorFilter] : InOverride->GetChildFilters())
+		{
+			OverrideInternal(InValue->GetChildFilters().FindChecked(ActorGuid), WorldPartitionActorFilter);
+		}
+	};
+	OverrideInternal(LevelFilter.Get(), Other);
+}
+
+void FWorldPartitionActorFilterMode::FFilter::Merge(const FWorldPartitionActorFilter* Other)
+{
+	TFunction<void(const FWorldPartitionActorFilter*, const FWorldPartitionActorFilter*)> MergeInternal = [this, &MergeInternal](const FWorldPartitionActorFilter* InValue, const FWorldPartitionActorFilter* InMerge)
+	{
+		if (FFilter::FDataLayerFilters* DataLayerFilters = FilterValues.Find(InValue))
+		{
+			for (auto& [AssetPath, DataLayerFilter] : *DataLayerFilters)
+			{
+				const FWorldPartitionActorFilter::FDataLayerFilter* Found = InMerge ? InMerge->DataLayerFilters.Find(AssetPath) : nullptr;
+	
+				if (DataLayerFilter.bOverride.IsSet())
+				{
+					if (DataLayerFilter.bOverride.GetValue() != !!Found)
+					{
+						DataLayerFilter.bOverride.Reset();
+						DataLayerFilter.bIncluded.Reset();
+					}
+					else if (DataLayerFilter.bOverride.GetValue() && DataLayerFilter.bIncluded.IsSet() && Found && DataLayerFilter.bIncluded.GetValue() != Found->bIncluded)
+					{
+						DataLayerFilter.bIncluded.Reset();
+					}
+				}
+			}
+		}
+
+		for (auto& [ActorGuid, WorldPartitionActorFilter] : InValue->GetChildFilters())
+		{
+			FWorldPartitionActorFilter*const* MergeChild = InMerge ? InMerge->GetChildFilters().Find(ActorGuid) : nullptr;
+			MergeInternal(WorldPartitionActorFilter, MergeChild ? *MergeChild : nullptr );
+		}
+	};
+	MergeInternal(LevelFilter.Get(), Other);
+}
+
+void FWorldPartitionActorFilterMode::Apply(FWorldPartitionActorFilter& InOutResult) const
+{
+	Filter->Apply(&InOutResult);
+}
+
+void FWorldPartitionActorFilterMode::FFilter::Apply(FWorldPartitionActorFilter* InOutResult) const
+{
+	TFunction<bool(const FWorldPartitionActorFilter*, FWorldPartitionActorFilter*)> ApplyInternal = [this, &ApplyInternal](const FWorldPartitionActorFilter* InValue, FWorldPartitionActorFilter* InResult)
+	{
+		const FFilter::FDataLayerFilters& MergedDataLayerFilter = FilterValues.FindChecked(InValue);
+		for (auto& [AssetPath, DataLayerFilter] : InValue->DataLayerFilters)
+		{
+			const auto& Found = MergedDataLayerFilter.FindChecked(AssetPath);
+			if (Found.bOverride.IsSet())
+			{
+				if (Found.bOverride.GetValue())
+				{
+					if (Found.bIncluded.IsSet())
+					{
+						InResult->DataLayerFilters.FindOrAdd(AssetPath).bIncluded = Found.bIncluded.GetValue();
+					}
+				}
+				else
+				{
+					InResult->DataLayerFilters.Remove(AssetPath);
+				}
+			}
+		}
+
+		for (auto& [ActorGuid, WorldPartitionActorFilter] : InValue->GetChildFilters())
+		{
+			if (FWorldPartitionActorFilter* const* ChildResultFilter = InResult->GetChildFilters().Find(ActorGuid))
+			{
+				if (!ApplyInternal(WorldPartitionActorFilter, *ChildResultFilter))
+				{
+					InResult->RemoveChildFilter(ActorGuid);
+				}
+			}
+			else
+			{
+				FWorldPartitionActorFilter ChildResult;
+				if (ApplyInternal(WorldPartitionActorFilter, &ChildResult))
+				{
+					InResult->AddChildFilter(ActorGuid, new FWorldPartitionActorFilter(MoveTemp(ChildResult)));
+				}
+			}
+		}
+
+		return InResult->DataLayerFilters.Num() > 0 || InResult->GetChildFilters().Num();
+	};
+
+	ApplyInternal(LevelFilter.Get(), InOutResult);
+}
+
+#undef LOCTEXT_NAMESPACE
\ No newline at end of file
diff --git a/Engine/Source/Editor/WorldPartitionEditor/Private/WorldPartition/Filter/WorldPartitionActorFilterPropertyTypeCustomization.cpp b/Engine/Source/Editor/WorldPartitionEditor/Private/WorldPartition/Filter/WorldPartitionActorFilterPropertyTypeCustomization.cpp
new file mode 100644
--- /dev/null
+++ b/Engine/Source/Editor/WorldPartitionEditor/Private/WorldPartition/Filter/WorldPartitionActorFilterPropertyTypeCustomization.cpp
@@ -0,0 +1,135 @@
+// Copyright Epic Games, Inc. All Rights Reserved.
+
+#include "WorldPartition/Filter/WorldPartitionActorFilterPropertyTypeCustomization.h"
+
+#include "IPropertyUtilities.h"
+#include "IPropertyTypeCustomization.h"
+#include "Containers/Array.h"
+#include "CoreTypes.h"
+#include "DetailWidgetRow.h"
+#include "Internationalization/Internationalization.h"
+#include "LevelInstance/LevelInstanceSubsystem.h"
+#include "LevelInstance/LevelInstanceInterface.h"
+#include "WorldPartition/Filter/WorldPartitionActorFilter.h"
+#include "LevelInstance/LevelInstanceComponent.h"
+#include "Modules/ModuleManager.h"
+#include "ISceneOutlinerMode.h"
+#include "ISceneOutliner.h"
+#include "SceneOutlinerModule.h"
+#include "WorldPartition/Filter/WorldPartitionActorFilterMode.h"
+#include "WorldPartition/Filter/WorldPartitionActorFilterColumn.h"
+#include "Settings/EditorExperimentalSettings.h"
+
+#define LOCTEXT_NAMESPACE "WorldPartitionActorFilter"
+
+FWorldPartitionActorFilterPropertyTypeCustomization::~FWorldPartitionActorFilterPropertyTypeCustomization()
+{
+	FWorldPartitionActorFilter::GetOnWorldPartitionActorFilterChanged().Remove(WorldPartitionActorFilterChangedHandle);
+}
+
+void FWorldPartitionActorFilterPropertyTypeCustomization::CustomizeHeader(TSharedRef<IPropertyHandle> StructPropertyHandle, FDetailWidgetRow& HeaderRow, IPropertyTypeCustomizationUtils& StructCustomizationUtils)
+{
+	if (!GetDefault<UEditorExperimentalSettings>()->bEnableWorldPartitionActorFilters)
+	{
+		StructPropertyHandle->MarkHiddenByCustomization();
+		return;
+	}
+
+	TArray<UObject*> OuterObjects;
+	StructPropertyHandle->GetOuterObjects(OuterObjects);
+		
+	// Create Mode Filter which holds the final values for the filter
+	TSharedPtr<FWorldPartitionActorFilterMode::FFilter> ModeFilter = CreateModeFilter(OuterObjects);
+	if (!ModeFilter)
+	{
+		return;
+	}
+	
+	FCreateSceneOutlinerMode ModeFactory = FCreateSceneOutlinerMode::CreateLambda([=](SSceneOutliner* Outliner)
+	{
+		return new FWorldPartitionActorFilterMode(Outliner, ModeFilter);
+	});
+
+	FSceneOutlinerInitializationOptions InitOptions;
+	{
+		InitOptions.bShowHeaderRow = true;
+		InitOptions.bFocusSearchBoxWhenOpened = true;
+		InitOptions.bShowCreateNewFolder = false;
+		InitOptions.ModeFactory = ModeFactory;
+	}
+
+	InitOptions.ColumnMap.Add(FWorldPartitionActorFilterColumn::ID, FSceneOutlinerColumnInfo(ESceneOutlinerColumnVisibility::Visible, 0, FCreateSceneOutlinerColumn::CreateLambda([](ISceneOutliner& InSceneOutliner) { return MakeShareable(new FWorldPartitionActorFilterColumn(InSceneOutliner)); }), false, TOptional<float>(), LOCTEXT("WorldPartitionActorFilterColumn","Filter")));
+	InitOptions.ColumnMap.Add(FSceneOutlinerBuiltInColumnTypes::Label(), FSceneOutlinerColumnInfo(ESceneOutlinerColumnVisibility::Visible, 1, FCreateSceneOutlinerColumn(), false, TOptional<float>(), FSceneOutlinerBuiltInColumnTypes::Label_Localized()));
+
+	// Actor selector to allow the user to choose a parent actor
+	FSceneOutlinerModule& SceneOutlinerModule = FModuleManager::LoadModuleChecked<FSceneOutlinerModule>("SceneOutliner");
+
+	// When value changes refresh details + filter actors
+	const FSimpleDelegate OnValueChanged = FSimpleDelegate::CreateLambda([PropertyUtils = StructCustomizationUtils.GetPropertyUtilities()]()
+	{
+		PropertyUtils->ForceRefresh();
+		FWorldPartitionActorFilter::RequestFilterRefresh(true);
+	});
+
+	// Listen to outside changes to Filters to refresh details panel (Undo/Redo)
+	WorldPartitionActorFilterChangedHandle = FWorldPartitionActorFilter::GetOnWorldPartitionActorFilterChanged().AddLambda([PropertyUtils = StructCustomizationUtils.GetPropertyUtilities()]()
+	{
+		PropertyUtils->ForceRefresh();
+	});
+
+	SceneOutliner = SceneOutlinerModule.CreateSceneOutliner(InitOptions);
+
+	auto OnApplyFilter = [OnValueChanged, this]()
+	{
+		const FWorldPartitionActorFilterMode* Mode = static_cast<const FWorldPartitionActorFilterMode*>(SceneOutliner->GetMode());
+		ApplyFilter(*Mode);
+		OnValueChanged.Execute();
+		return FReply::Handled();
+	};
+
+	HeaderRow
+	.NameContent()
+	[
+		StructPropertyHandle->CreatePropertyNameWidget()
+	]
+	.ValueContent()
+	[
+		SNew(SHorizontalBox)
+		+SHorizontalBox::Slot()
+		[
+			SNew(SComboButton)
+			.ContentPadding(2.f)
+			.HasDownArrow(false)
+			.ButtonContent()
+			[
+				SNew(STextBlock)
+				.Text(LOCTEXT("FilterText", "Select Filters..."))
+				.Font(FAppStyle::GetFontStyle(TEXT("PropertyWindow.NormalFont")))
+			]
+			.MenuContent()
+			[
+				SNew(SBox)
+				.MaxDesiredHeight(400.0f)
+				.WidthOverride(300.0f)
+				[
+					SceneOutliner.ToSharedRef()
+				]
+			]
+		]
+		+SHorizontalBox::Slot()
+		.AutoWidth()
+		[
+			SNew(SButton)
+			.VAlign(VAlign_Center)
+			.HAlign(HAlign_Center)
+			.OnClicked_Lambda(OnApplyFilter)
+			.Text(LOCTEXT("ApplyFilter", "Apply"))
+			.ToolTipText(LOCTEXT("ApplyFilterToolTip", "Apply Filter"))
+		]
+	];
+		
+	// This gets called when hitting ResetToDefault
+	StructPropertyHandle->SetOnPropertyValueChanged(OnValueChanged);
+}
+
+#undef LOCTEXT_NAMESPACE
\ No newline at end of file
diff --git a/Engine/Source/Editor/WorldPartitionEditor/Private/WorldPartition/Filter/WorldPartitionActorFilterTreeItems.cpp b/Engine/Source/Editor/WorldPartitionEditor/Private/WorldPartition/Filter/WorldPartitionActorFilterTreeItems.cpp
new file mode 100644
--- /dev/null
+++ b/Engine/Source/Editor/WorldPartitionEditor/Private/WorldPartition/Filter/WorldPartitionActorFilterTreeItems.cpp
@@ -0,0 +1,177 @@
+// Copyright Epic Games, Inc. All Rights Reserved.
+
+#include "WorldPartition/Filter/WorldPartitionActorFilterTreeItems.h"
+#include "WorldPartition/Filter/WorldPartitionActorFilter.h"
+#include "LevelInstance/LevelInstanceActor.h"
+#include "WorldPartition/DataLayer/DataLayerAsset.h"
+#include "Widgets/Text/STextBlock.h"
+#include "Widgets/SBoxPanel.h"
+#include "ISceneOutliner.h"
+#include "SceneOutlinerPublicTypes.h"
+#include "Styling/SlateIconFinder.h"
+
+#define LOCTEXT_NAMESPACE "WorldPartitionActorFilter"
+
+const FSceneOutlinerTreeItemType FWorldPartitionActorFilterItemBase::Type(&ISceneOutlinerTreeItem::Type);
+const FSceneOutlinerTreeItemType FWorldPartitionActorFilterItem::Type(&FWorldPartitionActorFilterItemBase::Type);
+const FSceneOutlinerTreeItemType FWorldPartitionActorFilterDataLayerItem::Type(&FWorldPartitionActorFilterItemBase::Type);
+
+struct SWorldPartitionActorFilterLabel : FSceneOutlinerCommonLabelData, public SCompoundWidget
+{
+	SLATE_BEGIN_ARGS(SWorldPartitionActorFilterLabel) {}
+	SLATE_END_ARGS()
+
+	void Construct(const FArguments& InArgs, ISceneOutlinerTreeItem& TreeItem, ISceneOutliner& SceneOutliner, const STableRow<FSceneOutlinerTreeItemPtr>& InRow)
+	{
+		TreeItemPtr = StaticCastSharedRef<ISceneOutlinerTreeItem>(TreeItem.AsShared());
+		WeakSceneOutliner = StaticCastSharedRef<ISceneOutliner>(SceneOutliner.AsShared());
+				
+		HighlightText = SceneOutliner.GetFilterHighlightText();
+								
+		ChildSlot
+		[
+			SNew(SHorizontalBox)
+			+ SHorizontalBox::Slot()
+			.AutoWidth()
+			.VAlign(VAlign_Center)
+			.Padding(FSceneOutlinerDefaultTreeItemMetrics::IconPadding())
+			[
+				SNew(SBox)
+				.WidthOverride(FSceneOutlinerDefaultTreeItemMetrics::IconSize())
+				.HeightOverride(FSceneOutlinerDefaultTreeItemMetrics::IconSize())
+				[
+					SNew(SImage)
+					.Image(this, &SWorldPartitionActorFilterLabel::GetIcon)
+				]
+			]
+			+SHorizontalBox::Slot()
+			.FillWidth(1.0f)
+			.VAlign(VAlign_Center)
+			.Padding(0.0f, 2.0f)
+			[
+				SNew(SHorizontalBox)
+				+ SHorizontalBox::Slot()
+				[
+					SNew(STextBlock)
+					.Text(this, &SWorldPartitionActorFilterLabel::GetDisplayText)
+					.ColorAndOpacity(this, &SWorldPartitionActorFilterLabel::GetForegroundColor)
+				]
+			]
+		];
+	}
+
+private:
+	TWeakPtr<ISceneOutlinerTreeItem> TreeItemPtr;
+	TAttribute<FText> HighlightText;
+
+
+	FText GetDisplayText() const
+	{
+		auto Item = TreeItemPtr.Pin();
+		return Item.IsValid() ? FText::FromString(Item->GetDisplayString()) : FText();
+	}
+	
+	const FSlateBrush* GetIcon() const
+	{
+		if (FSceneOutlinerTreeItemPtr PinnedPtr = TreeItemPtr.Pin(); PinnedPtr.IsValid())
+		{
+			if (WeakSceneOutliner.IsValid())
+			{
+				if (FWorldPartitionActorFilterItemBase* WorldPartitionActorFilterItem = PinnedPtr->CastTo<FWorldPartitionActorFilterItemBase>())
+				{
+					if (UClass* BrushClass = WorldPartitionActorFilterItem->GetIconClass())
+					{
+						const FSlateBrush* CachedBrush = WeakSceneOutliner.Pin()->GetCachedIconForClass(BrushClass->GetFName());
+						if (CachedBrush != nullptr)
+						{
+							return CachedBrush;
+						}
+						else
+						{
+							const FSlateBrush* FoundSlateBrush = FSlateIconFinder::FindIconBrushForClass(BrushClass);
+							WeakSceneOutliner.Pin()->CacheIconForClass(BrushClass->GetFName(), const_cast<FSlateBrush*>(FoundSlateBrush));
+							return FoundSlateBrush;
+						}
+					}
+				}
+			}
+		}
+		
+		return nullptr;
+	}
+		
+	FSlateColor GetForegroundColor() const
+	{
+		if (auto BaseColor = FSceneOutlinerCommonLabelData::GetForegroundColor(*TreeItemPtr.Pin()))
+		{
+			return BaseColor.GetValue();
+		}
+
+		return FSlateColor::UseForeground();
+	}
+};
+
+TSharedRef<SWidget> FWorldPartitionActorFilterItemBase::GenerateLabelWidget(ISceneOutliner& Outliner, const STableRow<FSceneOutlinerTreeItemPtr>& InRow)
+{
+	return SNew(SWorldPartitionActorFilterLabel, *this, Outliner, InRow);
+}
+
+FWorldPartitionActorFilterItem::FWorldPartitionActorFilterItem(const FWorldPartitionActorFilterItem::FTreeItemData& InData)
+	: FWorldPartitionActorFilterItemBase(FWorldPartitionActorFilterItem::Type)
+	, Data(InData)
+{
+}
+
+FSceneOutlinerTreeItemID FWorldPartitionActorFilterItem::GetID() const
+{
+	return ComputeTreeItemID(Data);
+}
+
+FString FWorldPartitionActorFilterItem::GetDisplayString() const
+{
+	return Data.Filter->DisplayName;
+}
+
+UClass* FWorldPartitionActorFilterItem::GetIconClass() const
+{
+	return ALevelInstance::StaticClass();
+}
+
+FSceneOutlinerTreeItemID FWorldPartitionActorFilterItem::ComputeTreeItemID(const FWorldPartitionActorFilterItem::FTreeItemData& InData)
+{
+	return FSceneOutlinerTreeItemID(GetTypeHash(InData.Filter));
+}
+
+FWorldPartitionActorFilterDataLayerItem::FWorldPartitionActorFilterDataLayerItem(const FWorldPartitionActorFilterDataLayerItem::FTreeItemData& InData)
+	: FWorldPartitionActorFilterItemBase(FWorldPartitionActorFilterDataLayerItem::Type)
+	, Data(InData)
+{
+}
+
+FSceneOutlinerTreeItemID FWorldPartitionActorFilterDataLayerItem::GetID() const
+{
+	return ComputeTreeItemID(Data);
+}
+
+FString FWorldPartitionActorFilterDataLayerItem::GetDisplayString() const
+{
+	return Data.Filter->DataLayerFilters[Data.AssetPath].DisplayName;
+}
+
+UClass* FWorldPartitionActorFilterDataLayerItem::GetIconClass() const
+{
+	return UDataLayerAsset::StaticClass();
+}
+
+FSceneOutlinerTreeItemID FWorldPartitionActorFilterDataLayerItem::ComputeTreeItemID(const FWorldPartitionActorFilterDataLayerItem::FTreeItemData& InData)
+{
+	return FSceneOutlinerTreeItemID(GetTypeHash(InData.Filter) ^ GetTypeHash(InData.AssetPath));
+}
+
+bool FWorldPartitionActorFilterDataLayerItem::GetDefaultValue() const
+{
+	check(IsValid());
+	return Data.Filter->DataLayerFilters.FindChecked(Data.AssetPath).bIncluded;
+}
+
+#undef LOCTEXT_NAMESPACE
\ No newline at end of file
diff --git a/Engine/Source/Editor/WorldPartitionEditor/Public/WorldPartition/Filter/WorldPartitionActorFilterColumn.h b/Engine/Source/Editor/WorldPartitionEditor/Public/WorldPartition/Filter/WorldPartitionActorFilterColumn.h
new file mode 100644
--- /dev/null
+++ b/Engine/Source/Editor/WorldPartitionEditor/Public/WorldPartition/Filter/WorldPartitionActorFilterColumn.h
@@ -0,0 +1,32 @@
+// Copyright Epic Games, Inc. All Rights Reserved.
+
+#pragma once
+
+#include "CoreMinimal.h"
+#include "SceneOutlinerFwd.h"
+#include "ISceneOutliner.h"
+#include "ISceneOutlinerColumn.h"
+#include "Widgets/Views/SHeaderRow.h"
+
+template<typename ItemType> class STableRow;
+class FWorldPartitionActorFilterColumn : public ISceneOutlinerColumn
+{
+public:
+	FWorldPartitionActorFilterColumn(ISceneOutliner& SceneOutliner) : WeakSceneOutliner(StaticCastSharedRef<ISceneOutliner>(SceneOutliner.AsShared())) {}
+	virtual ~FWorldPartitionActorFilterColumn() {}
+	static FName GetID();
+
+	//////////////////////////////////////////////////////////////////////////
+	// Begin ISceneOutlinerColumn Implementation
+	virtual FName GetColumnID() override { return GetID(); }
+	virtual SHeaderRow::FColumn::FArguments ConstructHeaderRowColumn() override;
+	virtual const TSharedRef<SWidget> ConstructRowWidget(FSceneOutlinerTreeItemRef TreeItem, const STableRow<FSceneOutlinerTreeItemPtr>& Row) override;
+	virtual bool SupportsSorting() const override { return false; }
+	// End ISceneOutlinerColumn Implementation
+	//////////////////////////////////////////////////////////////////////////
+
+	static FName ID;
+private:
+	TWeakPtr<ISceneOutliner> WeakSceneOutliner;
+};
+
diff --git a/Engine/Source/Editor/WorldPartitionEditor/Public/WorldPartition/Filter/WorldPartitionActorFilterHierarchy.h b/Engine/Source/Editor/WorldPartitionEditor/Public/WorldPartition/Filter/WorldPartitionActorFilterHierarchy.h
new file mode 100644
--- /dev/null
+++ b/Engine/Source/Editor/WorldPartitionEditor/Public/WorldPartition/Filter/WorldPartitionActorFilterHierarchy.h
@@ -0,0 +1,26 @@
+// Copyright Epic Games, Inc. All Rights Reserved.
+
+#pragma once
+
+#include "CoreMinimal.h"
+#include "ISceneOutlinerHierarchy.h"
+
+class FWorldPartitionActorFilterMode;
+
+class FWorldPartitionActorFilterHierarchy : public ISceneOutlinerHierarchy
+{
+public:
+	static TUniquePtr<FWorldPartitionActorFilterHierarchy> Create(FWorldPartitionActorFilterMode* Mode);
+
+	virtual ~FWorldPartitionActorFilterHierarchy();
+
+	//~ Begin ISceneOutlinerHierarchy interface
+	virtual void CreateItems(TArray<FSceneOutlinerTreeItemPtr>& OutItems) const override;
+	virtual void CreateChildren(const FSceneOutlinerTreeItemPtr& Item, TArray<FSceneOutlinerTreeItemPtr>& OutItems) const override;
+	virtual FSceneOutlinerTreeItemPtr FindOrCreateParentItem(const ISceneOutlinerTreeItem& Item, const TMap<FSceneOutlinerTreeItemID, FSceneOutlinerTreeItemPtr>& Items, bool bCreate = false) override;
+	//~ End ISceneOutlinerHierarchy interface
+
+private:
+	virtual void CreateChildren(const FSceneOutlinerTreeItemPtr& Item, TArray<FSceneOutlinerTreeItemPtr>& OutItems, bool bRecursive) const;
+	FWorldPartitionActorFilterHierarchy(FWorldPartitionActorFilterMode* Mode);
+};
\ No newline at end of file
diff --git a/Engine/Source/Editor/WorldPartitionEditor/Public/WorldPartition/Filter/WorldPartitionActorFilterMode.h b/Engine/Source/Editor/WorldPartitionEditor/Public/WorldPartition/Filter/WorldPartitionActorFilterMode.h
new file mode 100644
--- /dev/null
+++ b/Engine/Source/Editor/WorldPartitionEditor/Public/WorldPartition/Filter/WorldPartitionActorFilterMode.h
@@ -0,0 +1,73 @@
+// Copyright Epic Games, Inc. All Rights Reserved.
+
+#pragma once
+
+#include "CoreMinimal.h"
+#include "ISceneOutlinerMode.h"
+#include "WorldPartition/Filter/WorldPartitionActorFilter.h"
+#include "Containers/Map.h"
+#include "Templates/Tuple.h"
+
+class SSceneOutliner;
+
+class WORLDPARTITIONEDITOR_API FWorldPartitionActorFilterMode : public ISceneOutlinerMode
+{
+public:
+	// struct holding multi selection values that might be null if they differ between actors
+	class WORLDPARTITIONEDITOR_API FFilter
+	{
+	public:
+		FFilter(TSharedPtr<FWorldPartitionActorFilter> LevelFilter, const TArray<const FWorldPartitionActorFilter*>& SelectedFilters);
+		
+		void Apply(FWorldPartitionActorFilter* Result) const;
+				
+		struct FDataLayerFilter
+		{
+			FDataLayerFilter(){}
+			FDataLayerFilter(bool bInOverride, bool bInIncluded) : bOverride(bInOverride), bIncluded(bInIncluded) {}
+
+			TOptional<bool> bOverride;
+			TOptional<bool> bIncluded;
+		};
+
+		using FDataLayerFilters = TMap<FSoftObjectPath, FDataLayerFilter>;
+	private:
+		void Initialize(const FWorldPartitionActorFilter* Filter);
+		void Override(const FWorldPartitionActorFilter* Other);
+		void Merge(const FWorldPartitionActorFilter* Other);
+				
+		// Unmodified Reference Level Filter
+		TSharedPtr<FWorldPartitionActorFilter> LevelFilter;
+		// Current values for Filter based Selection		
+		TMap<const FWorldPartitionActorFilter*, FDataLayerFilters> FilterValues;
+
+		friend class FWorldPartitionActorFilterMode;
+	};
+
+	FWorldPartitionActorFilterMode(SSceneOutliner* InSceneOutliner, TSharedPtr<FFilter> InFilter);
+	
+	//~ Begin ISceneOutlinerMode interface
+	virtual void Rebuild() override;
+	virtual int32 GetTypeSortPriority(const ISceneOutlinerTreeItem& Item) const override;
+	//~ Begin ISceneOutlinerMode interface
+		
+	void Apply(FWorldPartitionActorFilter& InOutResult) const;
+	
+	FFilter::FDataLayerFilters& FindChecked(const FWorldPartitionActorFilter* InFilter) const
+	{
+		return Filter->FilterValues.FindChecked(InFilter);
+	}
+
+	const FWorldPartitionActorFilter* GetFilter() const
+	{
+		return Filter->LevelFilter.Get();
+	}
+
+protected:
+	//~ Begin ISceneOutlinerMode interface
+	virtual TUniquePtr<ISceneOutlinerHierarchy> CreateHierarchy() override;
+	//~ End ISceneOutlinerMode interface
+		
+private:
+	TSharedPtr<FFilter> Filter;
+};
\ No newline at end of file
diff --git a/Engine/Source/Editor/WorldPartitionEditor/Public/WorldPartition/Filter/WorldPartitionActorFilterPropertyTypeCustomization.h b/Engine/Source/Editor/WorldPartitionEditor/Public/WorldPartition/Filter/WorldPartitionActorFilterPropertyTypeCustomization.h
new file mode 100644
--- /dev/null
+++ b/Engine/Source/Editor/WorldPartitionEditor/Public/WorldPartition/Filter/WorldPartitionActorFilterPropertyTypeCustomization.h
@@ -0,0 +1,31 @@
+// Copyright Epic Games, Inc. All Rights Reserved.
+
+#pragma once
+
+#include "IPropertyTypeCustomization.h"
+#include "Internationalization/Text.h"
+#include "PropertyEditorModule.h"
+#include "Templates/SharedPointer.h"
+#include "WorldPartition/Filter/WorldPartitionActorFilterMode.h"
+
+class FDetailWidgetRow;
+class IDetailChildrenBuilder;
+class IPropertyHandle;
+class ISceneOutliner;
+
+// Registered (FLevelInstanceEditorModule::StartupModule) Property Customization for properties of type FWorldPartitionActorFilter
+struct WORLDPARTITIONEDITOR_API FWorldPartitionActorFilterPropertyTypeCustomization : public IPropertyTypeCustomization
+{
+public:
+	~FWorldPartitionActorFilterPropertyTypeCustomization();
+
+	virtual void CustomizeHeader(TSharedRef<IPropertyHandle> StructPropertyHandle, FDetailWidgetRow& HeaderRow, IPropertyTypeCustomizationUtils& StructCustomizationUtils) override;
+	virtual void CustomizeChildren(TSharedRef<IPropertyHandle> StructPropertyHandle, IDetailChildrenBuilder& StructBuilder, IPropertyTypeCustomizationUtils& StructCustomizationUtils) override {}
+
+private:
+	virtual TSharedPtr<FWorldPartitionActorFilterMode::FFilter> CreateModeFilter(TArray<UObject*> OuterObjects) = 0;
+	virtual void ApplyFilter(const FWorldPartitionActorFilterMode& Mode) = 0;
+
+	FDelegateHandle WorldPartitionActorFilterChangedHandle;
+	TSharedPtr<ISceneOutliner> SceneOutliner;
+};
diff --git a/Engine/Source/Editor/WorldPartitionEditor/Public/WorldPartition/Filter/WorldPartitionActorFilterTreeItems.h b/Engine/Source/Editor/WorldPartitionEditor/Public/WorldPartition/Filter/WorldPartitionActorFilterTreeItems.h
new file mode 100644
--- /dev/null
+++ b/Engine/Source/Editor/WorldPartitionEditor/Public/WorldPartition/Filter/WorldPartitionActorFilterTreeItems.h
@@ -0,0 +1,84 @@
+// Copyright Epic Games, Inc. All Rights Reserved.
+
+#pragma once
+
+#include "CoreMinimal.h"
+#include "Misc/Attribute.h"
+#include "SceneOutlinerFwd.h"
+#include "ISceneOutlinerTreeItem.h"
+
+struct FWorldPartitionActorFilter;
+class UClass;
+
+struct FWorldPartitionActorFilterItemBase : ISceneOutlinerTreeItem
+{
+	FWorldPartitionActorFilterItemBase(FSceneOutlinerTreeItemType InType) : ISceneOutlinerTreeItem(InType) {}
+	virtual UClass* GetIconClass() const = 0;
+	virtual TSharedRef<SWidget> GenerateLabelWidget(ISceneOutliner& Outliner, const STableRow<FSceneOutlinerTreeItemPtr>& InRow) override;
+	virtual bool CanInteract() const override { return true; }
+
+	static const FSceneOutlinerTreeItemType Type;
+};
+
+struct FWorldPartitionActorFilterItem : FWorldPartitionActorFilterItemBase
+{
+public:
+	struct FTreeItemData
+	{
+		FTreeItemData(const FWorldPartitionActorFilter* InFilter)
+			: Filter(InFilter)
+		{}
+
+		const FWorldPartitionActorFilter* Filter;
+	};
+
+	FWorldPartitionActorFilterItem(const FTreeItemData& InData);
+
+	static FSceneOutlinerTreeItemID ComputeTreeItemID(const FWorldPartitionActorFilterItem::FTreeItemData& InData);
+
+	/* Begin ISceneOutlinerTreeItem Implementation */
+	virtual bool IsValid() const override { return Data.Filter != nullptr; }
+	virtual FSceneOutlinerTreeItemID GetID() const override;
+	virtual FString GetDisplayString() const override;
+	/* End ISceneOutlinerTreeItem Implementation */
+
+	virtual UClass* GetIconClass() const override;
+	const FWorldPartitionActorFilter* GetFilter() const { return Data.Filter; }
+
+	static const FSceneOutlinerTreeItemType Type;
+private:
+	FTreeItemData Data;
+};
+
+struct FWorldPartitionActorFilterDataLayerItem : FWorldPartitionActorFilterItemBase
+{
+public:
+	struct FTreeItemData
+	{
+		FTreeItemData(const FWorldPartitionActorFilter* InFilter, const FSoftObjectPath& InAssetPath)
+			: Filter(InFilter), AssetPath(InAssetPath)
+		{}
+
+		const FWorldPartitionActorFilter* Filter;
+		const FSoftObjectPath AssetPath;
+	};
+
+	FWorldPartitionActorFilterDataLayerItem(const FTreeItemData& InData);
+
+	static FSceneOutlinerTreeItemID ComputeTreeItemID(const FWorldPartitionActorFilterDataLayerItem::FTreeItemData& InData);
+
+	/* Begin ISceneOutlinerTreeItem Implementation */
+	virtual bool IsValid() const override { return Data.Filter != nullptr; }
+	virtual FSceneOutlinerTreeItemID GetID() const override;
+	virtual FString GetDisplayString() const override;
+	/* End ISceneOutlinerTreeItem Implementation */
+	
+	virtual UClass* GetIconClass() const override;
+	const FWorldPartitionActorFilter* GetFilter() const { return Data.Filter; }
+	const FSoftObjectPath& GetAssetPath() const { return Data.AssetPath; }
+	bool GetDefaultValue() const;
+
+	static const FSceneOutlinerTreeItemType Type;
+private:
+	FTreeItemData Data;
+};
\ No newline at end of file
diff --git a/Engine/Source/Runtime/Core/Public/UObject/FortniteMainBranchObjectVersion.h b/Engine/Source/Runtime/Core/Public/UObject/FortniteMainBranchObjectVersion.h
--- a/Engine/Source/Runtime/Core/Public/UObject/FortniteMainBranchObjectVersion.h
+++ b/Engine/Source/Runtime/Core/Public/UObject/FortniteMainBranchObjectVersion.h
@@ -271,6 +271,9 @@ struct CORE_API FFortniteMainBranchObjectVersion
 		// Dynamic cast nodes use an enumerated pure node state to include a value for the default setting
 		DynamicCastNodesUsePureStateEnum,
 
+		// Add FWorldPartitionActorFilter to FLevelInstanceActorDesc/FDataLayerInstanceDesc
+		WorldPartitionActorFilter,
+
 		// -----<new versions can be added above this line>-------------------------------------------------
 		VersionPlusOne,
 		LatestVersion = VersionPlusOne - 1
diff --git a/Engine/Source/Runtime/Engine/Private/LevelInstance/LevelInstanceActor.cpp b/Engine/Source/Runtime/Engine/Private/LevelInstance/LevelInstanceActor.cpp
--- a/Engine/Source/Runtime/Engine/Private/LevelInstance/LevelInstanceActor.cpp
+++ b/Engine/Source/Runtime/Engine/Private/LevelInstance/LevelInstanceActor.cpp
@@ -21,7 +21,7 @@ ALevelInstance::ALevelInstance()
 	: LevelInstanceActorGuid(this)
 	, LevelInstanceActorImpl(this)
 {
-	RootComponent = CreateDefaultSubobject<ULevelInstanceComponent>(TEXT("Root"));
+	RootComponent = LevelInstanceComponent = CreateDefaultSubobject<ULevelInstanceComponent>(TEXT("Root"));
 	RootComponent->Mobility = EComponentMobility::Static;
 
 #if WITH_EDITORONLY_DATA
@@ -150,6 +150,7 @@ void ALevelInstance::PostLoad()
 {
 	Super::PostLoad();
 
+	LevelInstanceComponent = Cast<ULevelInstanceComponent>(RootComponent);
 	OnLevelInstanceActorPostLoad.Broadcast(this);
 }
 
diff --git a/Engine/Source/Runtime/Engine/Private/LevelInstance/LevelInstanceComponent.cpp b/Engine/Source/Runtime/Engine/Private/LevelInstance/LevelInstanceComponent.cpp
--- a/Engine/Source/Runtime/Engine/Private/LevelInstance/LevelInstanceComponent.cpp
+++ b/Engine/Source/Runtime/Engine/Private/LevelInstance/LevelInstanceComponent.cpp
@@ -46,17 +46,39 @@ void ULevelInstanceComponent::OnRegister()
 #endif //WITH_EDITORONLY_DATA
 }
 
+void ULevelInstanceComponent::SetFilter(const FWorldPartitionActorFilter& InFilter)
+{
+	if (Filter != InFilter)
+	{
+		Modify();
+		Filter = InFilter;
+		FWorldPartitionActorFilter::GetOnWorldPartitionActorFilterChanged().Broadcast();
+	}
+}
+
 bool ULevelInstanceComponent::ShouldShowSpriteComponent() const
 {
 	return GetOwner() && GetOwner()->GetLevel() && (GetOwner()->GetLevel()->IsPersistentLevel() || !GetOwner()->GetLevel()->IsInstancedLevel());
 }
 
+void ULevelInstanceComponent::PreEditUndo()
+{
+	CachedFilter = Filter;
+}
+
 void ULevelInstanceComponent::PostEditUndo()
 {
 	Super::PostEditUndo();
 		
 	UpdateComponentToWorld();
 	UpdateEditorInstanceActor();
+
+	if (Filter != CachedFilter)
+	{
+		FWorldPartitionActorFilter::RequestFilterRefresh(false);
+		FWorldPartitionActorFilter::GetOnWorldPartitionActorFilterChanged().Broadcast();
+	}
+	CachedFilter = FWorldPartitionActorFilter();
 }
 
 void ULevelInstanceComponent::PostEditChangeProperty(FPropertyChangedEvent& PropertyChangedEvent)
diff --git a/Engine/Source/Runtime/Engine/Private/LevelInstance/LevelInstanceInterface.cpp b/Engine/Source/Runtime/Engine/Private/LevelInstance/LevelInstanceInterface.cpp
--- a/Engine/Source/Runtime/Engine/Private/LevelInstance/LevelInstanceInterface.cpp
+++ b/Engine/Source/Runtime/Engine/Private/LevelInstance/LevelInstanceInterface.cpp
@@ -237,6 +237,17 @@ bool ILevelInstanceInterface::MoveActorsTo(const TArray<AActor*>& ActorsToMove)
 	return false;
 }
 
+const FWorldPartitionActorFilter& ILevelInstanceInterface::GetFilter() const
+{
+	if (ULevelInstanceComponent* LevelInstanceComponent = GetLevelInstanceComponent())
+	{
+		return LevelInstanceComponent->GetFilter();
+	}
+
+	static FWorldPartitionActorFilter NoFilter;
+	return NoFilter;
+}
+
 #endif
 
 
diff --git a/Engine/Source/Runtime/Engine/Private/LevelInstance/LevelInstanceSubsystem.cpp b/Engine/Source/Runtime/Engine/Private/LevelInstance/LevelInstanceSubsystem.cpp
--- a/Engine/Source/Runtime/Engine/Private/LevelInstance/LevelInstanceSubsystem.cpp
+++ b/Engine/Source/Runtime/Engine/Private/LevelInstance/LevelInstanceSubsystem.cpp
@@ -19,10 +19,12 @@
 #include "LevelInstance/LevelInstanceEditorInstanceActor.h"
 #include "LevelInstance/LevelInstanceEditorObject.h"
 #include "LevelInstance/LevelInstanceEditorPivotActor.h"
+#include "WorldPartition/LevelInstance/LevelInstanceActorDesc.h"
 #include "WorldPartition/WorldPartition.h"
 #include "WorldPartition/WorldPartitionSubsystem.h"
 #include "WorldPartition/DataLayer/DataLayerManager.h"
 #include "WorldPartition/DataLayer/DataLayerInstanceWithAsset.h"
+#include "WorldPartition/DataLayer/WorldDataLayersActorDesc.h"
 #include "WorldPartition/WorldPartitionMiniMap.h"
 #include "Misc/ScopedSlowTask.h"
 #include "Misc/Paths.h"
@@ -2471,6 +2473,141 @@ bool ULevelInstanceSubsystem::CanUseWorldAsset(const ILevelInstanceInterface* Le
 	return true;
 }
 
+FWorldPartitionActorFilter ULevelInstanceSubsystem::GetLevelInstanceFilter(const FString& LevelPackage) const
+{
+	TSet<FString> VisitedPackages;
+	return GetLevelInstanceFilterInternal(LevelPackage, VisitedPackages);
+}
+
+FWorldPartitionActorFilter ULevelInstanceSubsystem::GetLevelInstanceFilterInternal(const FString& LevelPackage, TSet<FString>& VisitedPackages) const
+{
+	if (VisitedPackages.Contains(LevelPackage))
+	{
+		return FWorldPartitionActorFilter(LevelPackage);
+	}
+
+	VisitedPackages.Add(LevelPackage);
+		
+	// Most of the time if the LevelInstance Actor already exists this will return an existing Container but when loading a new LevelInstance (Content Browser Drag&Drop, Create LI) 
+	// This will make sure Container exists.
+	UActorDescContainer* LevelContainer = ActorDescContainerInstanceManager.RegisterContainer(*LevelPackage, GetWorld());
+	check(LevelContainer);
+	ON_SCOPE_EXIT{ ActorDescContainerInstanceManager.UnregisterContainer(LevelContainer); };
+
+	// Lazy create filter for now
+	TArray<const FLevelInstanceActorDesc*> LevelInstanceActorDescs;
+	const FWorldDataLayersActorDesc* WorldDataLayersActorDesc = nullptr;
+
+	for (FActorDescList::TConstIterator<> ActorDescIt(LevelContainer); ActorDescIt; ++ActorDescIt)
+	{
+		if (ActorDescIt->GetActorNativeClass()->IsChildOf<AWorldDataLayers>())
+		{
+			check(!WorldDataLayersActorDesc);
+			WorldDataLayersActorDesc = static_cast<const FWorldDataLayersActorDesc*>(*ActorDescIt);
+		}
+		else if (ActorDescIt->GetActorNativeClass()->ImplementsInterface(ULevelInstanceInterface::StaticClass()))
+		{
+			LevelInstanceActorDescs.Add(static_cast<const FLevelInstanceActorDesc*>(*ActorDescIt));
+		}
+	}
+	
+	FWorldPartitionActorFilter Filter(LevelPackage);
+
+	if (WorldDataLayersActorDesc)
+	{
+		for (const FDataLayerInstanceDesc& DataLayerInstanceDesc : WorldDataLayersActorDesc->GetDataLayerInstances())
+		{
+			// For now consider all DataLayerInstances using Assets as filters that are included by default
+			if (DataLayerInstanceDesc.IsUsingAsset() && DataLayerInstanceDesc.GetAsset() && DataLayerInstanceDesc.GetAsset()->SupportsActorFilters())
+			{
+				Filter.DataLayerFilters.Add(FSoftObjectPath(DataLayerInstanceDesc.GetAssetPath().ToString()), FWorldPartitionActorFilter::FDataLayerFilter(DataLayerInstanceDesc.GetShortName(), DataLayerInstanceDesc.IsIncludedInActorFilterDefault()));
+			}
+		}
+	}
+
+	for (const FLevelInstanceActorDesc* LevelInstanceActorDesc : LevelInstanceActorDescs)
+	{
+		TSet<FString> VisitedPackagesCopy(VisitedPackages);
+
+		// Get World Default Filter
+		FWorldPartitionActorFilter* ChildFilter = new FWorldPartitionActorFilter(GetLevelInstanceFilterInternal(LevelInstanceActorDesc->GetLevelPackage().ToString(), VisitedPackagesCopy));
+		ChildFilter->DisplayName = LevelInstanceActorDesc->GetActorLabelOrName().ToString();
+
+		// Apply LevelInstanceActorDesc Filter to Default
+		ChildFilter->Override(LevelInstanceActorDesc->GetFilter());
+
+		Filter.AddChildFilter(LevelInstanceActorDesc->GetGuid(), ChildFilter);
+	}
+
+	return Filter;
+}
+
+bool ULevelInstanceSubsystem::PassLevelInstanceFilter(UWorld* World, const FWorldPartitionHandle& Actor) const
+{
+	if (Actor->IsUsingDataLayerAsset() && Actor->GetDataLayers().Num() > 0)
+	{
+		// Make sure Actor is using a filtered Data Layer Instance
+		TSet<FSoftObjectPath> FilteredDataLayerAssets;
+		if (UDataLayerManager* DataLayerManager = UDataLayerManager::GetDataLayerManager(World))
+		{
+			bool bFoundFilteredDataLayer = false;
+			for (FName AssetPath : Actor->GetDataLayers())
+			{
+				FSoftObjectPath SoftAssetPath(AssetPath.ToString());
+				UDataLayerAsset* Asset = Cast<UDataLayerAsset>(SoftAssetPath.TryLoad());
+				if (const UDataLayerInstanceWithAsset* DataLayerInstance = Cast<UDataLayerInstanceWithAsset>(DataLayerManager->GetDataLayerInstanceFromAsset(Asset)); DataLayerInstance && DataLayerInstance->SupportsActorFilters())
+				{
+					FilteredDataLayerAssets.Add(SoftAssetPath);
+				}
+			}
+		}
+
+		// Actor is not using a filtered Data Layer Instance
+		if (FilteredDataLayerAssets.IsEmpty())
+		{
+			return true;
+		}
+
+		// Build Filter
+		if (const ILevelInstanceInterface* LevelInstanceInterface = GetOwningLevelInstance(World->PersistentLevel))
+		{
+			TArray<FGuid> AncestorGuids;
+			ForEachLevelInstanceAncestorsAndSelf(Cast<AActor>(LevelInstanceInterface), [&AncestorGuids, &LevelInstanceInterface](const ILevelInstanceInterface* Ancestor)
+			{
+				AncestorGuids.Insert(Cast<AActor>(Ancestor)->GetActorGuid(), 0);
+				LevelInstanceInterface = Ancestor;
+				return true;
+			});
+
+			check(LevelInstanceInterface);
+
+			FWorldPartitionActorFilter ActorFilter = GetLevelInstanceFilter(LevelInstanceInterface->GetWorldAssetPackage());
+			ActorFilter.Override(LevelInstanceInterface->GetFilter());
+
+			FWorldPartitionActorFilter* FoundFilter = &ActorFilter;
+			for (int32 AncestorIndex = 1; AncestorIndex < AncestorGuids.Num(); ++AncestorIndex)
+			{
+				FoundFilter = FoundFilter->GetChildFilters().FindChecked(AncestorGuids[AncestorIndex]);
+			}
+
+			for (const FSoftObjectPath& AssetPath : FilteredDataLayerAssets)
+			{
+				if (FWorldPartitionActorFilter::FDataLayerFilter* DataLayerFilter = FoundFilter->DataLayerFilters.Find(AssetPath))
+				{
+					if (DataLayerFilter->bIncluded)
+					{
+						return true;
+					}
+				}				
+			}
+
+			return false;
+		}
+	}
+
+	return true;
+}
+
 #endif
 
 #undef LOCTEXT_NAMESPACE
diff --git a/Engine/Source/Runtime/Engine/Private/WorldPartition/DataLayer/DataLayerInstanceWithAsset.cpp b/Engine/Source/Runtime/Engine/Private/WorldPartition/DataLayer/DataLayerInstanceWithAsset.cpp
--- a/Engine/Source/Runtime/Engine/Private/WorldPartition/DataLayer/DataLayerInstanceWithAsset.cpp
+++ b/Engine/Source/Runtime/Engine/Private/WorldPartition/DataLayer/DataLayerInstanceWithAsset.cpp
@@ -12,6 +12,9 @@
 
 UDataLayerInstanceWithAsset::UDataLayerInstanceWithAsset(const FObjectInitializer& ObjectInitializer)
 	: Super(ObjectInitializer)
+#if WITH_EDITORONLY_DATA
+	, bIsIncludedInActorFilterDefault(true)
+#endif
 {
 
 }
@@ -124,6 +127,31 @@ void UDataLayerInstanceWithAsset::PostEditChangeProperty(FPropertyChangedEvent&
 	}
 }
 
+bool UDataLayerInstanceWithAsset::CanEditChange(const FProperty* InProperty) const
+{
+	if (!Super::CanEditChange(InProperty))
+	{
+		return false;
+	}
+
+	if (InProperty && (InProperty->GetFName() == GET_MEMBER_NAME_CHECKED(UDataLayerInstanceWithAsset, bIsIncludedInActorFilterDefault)))
+	{
+		return SupportsActorFilters();
+	}
+
+	return true;
+}
+
+bool UDataLayerInstanceWithAsset::SupportsActorFilters() const
+{
+	return DataLayerAsset && DataLayerAsset->SupportsActorFilters();
+}
+
+bool UDataLayerInstanceWithAsset::IsIncludedInActorFilterDefault() const
+{
+	return bIsIncludedInActorFilterDefault;
+}
+
 void UDataLayerInstanceWithAsset::PreEditUndo()
 {
 	Super::PreEditUndo();
diff --git a/Engine/Source/Runtime/Engine/Private/WorldPartition/DataLayer/WorldDataLayersActorDesc.cpp b/Engine/Source/Runtime/Engine/Private/WorldPartition/DataLayer/WorldDataLayersActorDesc.cpp
--- a/Engine/Source/Runtime/Engine/Private/WorldPartition/DataLayer/WorldDataLayersActorDesc.cpp
+++ b/Engine/Source/Runtime/Engine/Private/WorldPartition/DataLayer/WorldDataLayersActorDesc.cpp
@@ -9,6 +9,7 @@
 #include "WorldPartition/DataLayer/DataLayerInstanceWithAsset.h"
 #include "WorldPartition/WorldPartitionActorContainerID.h"
 #include "UObject/FortniteNCBranchObjectVersion.h"
+#include "UObject/FortniteMainBranchObjectVersion.h"
 
 FArchive& operator<<(FArchive& Ar, FDataLayerInstanceDesc& Desc)
 {
@@ -40,6 +41,11 @@ FArchive& operator<<(FArchive& Ar, FDataLayerInstanceDesc& Desc)
 		FWorldPartitionActorDescUtils::FixupRedirectedAssetPath(Desc.AssetPath);
 	}
 
+	if (Ar.CustomVer(FFortniteMainBranchObjectVersion::GUID) >= FFortniteMainBranchObjectVersion::WorldPartitionActorFilter)
+	{
+		Ar << Desc.bIsIncludedInActorFilterDefault;
+	}
+
 	return Ar;
 }
 
@@ -50,7 +56,8 @@ bool operator == (const FDataLayerInstanceDesc& Lhs, const FDataLayerInstanceDes
 		(Lhs.bIsUsingAsset == Rhs.bIsUsingAsset) &&
 		(Lhs.AssetPath == Rhs.AssetPath) &&
 		(Lhs.bIsRuntime == Rhs.bIsRuntime) &&
-		(Lhs.ShortName == Rhs.ShortName);
+		(Lhs.ShortName == Rhs.ShortName) &&
+		(Lhs.bIsIncludedInActorFilterDefault == Rhs.bIsIncludedInActorFilterDefault);
 }
 
 bool operator < (const FDataLayerInstanceDesc& Lhs, const FDataLayerInstanceDesc& Rhs)
@@ -65,6 +72,10 @@ bool operator < (const FDataLayerInstanceDesc& Lhs, const FDataLayerInstanceDesc
 				{
 					if (Lhs.bIsRuntime == Rhs.bIsRuntime)
 					{
+						if (Lhs.ShortName == Rhs.ShortName)
+						{
+							return (int)Lhs.bIsIncludedInActorFilterDefault < (int)Rhs.bIsIncludedInActorFilterDefault;
+						}
 						return Lhs.ShortName < Rhs.ShortName;
 					}
 					return (int)Lhs.bIsRuntime < (int)Rhs.bIsRuntime;
@@ -115,6 +126,7 @@ FString FDataLayerInstanceDesc::GetShortName() const
 FDataLayerInstanceDesc::FDataLayerInstanceDesc()
 : bIsUsingAsset(false)
 , bIsRuntime(false)
+, bIsIncludedInActorFilterDefault(true)
 {
 }
 
@@ -129,6 +141,8 @@ void FDataLayerInstanceDesc::Init(UDataLayerInstance* InDataLayerInstance)
 		if (const UDataLayerAsset* DataLayerAsset = DataLayerWithAsset->GetAsset())
 		{
 			AssetPath = FName(DataLayerAsset->GetPathName());
+			bIsIncludedInActorFilterDefault = DataLayerWithAsset->IsIncludedInActorFilterDefault();
+			
 		}
 	}
 	ShortName = InDataLayerInstance->GetDataLayerShortName();
@@ -167,6 +181,7 @@ bool FWorldDataLayersActorDesc::Equals(const FWorldPartitionActorDesc* Other) co
 void FWorldDataLayersActorDesc::Serialize(FArchive& Ar)
 {
 	Ar.UsingCustomVersion(FFortniteNCBranchObjectVersion::GUID);
+	Ar.UsingCustomVersion(FFortniteMainBranchObjectVersion::GUID);
 
 	FWorldPartitionActorDesc::Serialize(Ar);
 
diff --git a/Engine/Source/Runtime/Engine/Private/WorldPartition/ErrorHandling/WorldPartitionStreamingGenerationLogErrorHandler.cpp b/Engine/Source/Runtime/Engine/Private/WorldPartition/ErrorHandling/WorldPartitionStreamingGenerationLogErrorHandler.cpp
--- a/Engine/Source/Runtime/Engine/Private/WorldPartition/ErrorHandling/WorldPartitionStreamingGenerationLogErrorHandler.cpp
+++ b/Engine/Source/Runtime/Engine/Private/WorldPartition/ErrorHandling/WorldPartitionStreamingGenerationLogErrorHandler.cpp
@@ -90,4 +90,10 @@ void FStreamingGenerationLogErrorHandler::OnLevelInstanceInvalidWorldAsset(const
 		break;
 	};
 }
+
+void FStreamingGenerationLogErrorHandler::OnActorFilterFailed(const FWorldPartitionActorDescView& ActorDescView)
+{
+	UE_ASSET_LOG_ACTORDESCVIEW(LogWorldPartition, Log, ActorDescView, TEXT("Actor %s will not be filtered out because it is referenced/referencing other Actor(s) not part of the filter"), *GetActorName(ActorDescView));
+}
+
 #endif
diff --git a/Engine/Source/Runtime/Engine/Private/WorldPartition/ErrorHandling/WorldPartitionStreamingGenerationProxyErrorHandler.cpp b/Engine/Source/Runtime/Engine/Private/WorldPartition/ErrorHandling/WorldPartitionStreamingGenerationProxyErrorHandler.cpp
--- a/Engine/Source/Runtime/Engine/Private/WorldPartition/ErrorHandling/WorldPartitionStreamingGenerationProxyErrorHandler.cpp
+++ b/Engine/Source/Runtime/Engine/Private/WorldPartition/ErrorHandling/WorldPartitionStreamingGenerationProxyErrorHandler.cpp
@@ -62,4 +62,9 @@ void FStreamingGenerationProxyErrorHandler::OnLevelInstanceInvalidWorldAsset(con
 {
 	InnerErrorHandler->OnLevelInstanceInvalidWorldAsset(ActorDescView, WorldAsset, Reason);
 }
+
+void FStreamingGenerationProxyErrorHandler::OnActorFilterFailed(const FWorldPartitionActorDescView& ActorDescView)
+{
+	InnerErrorHandler->OnActorFilterFailed(ActorDescView);
+}
 #endif
diff --git a/Engine/Source/Runtime/Engine/Private/WorldPartition/ErrorHandling/WorldPartitionStreamingGenerationTokenizedMessageErrorHandler.cpp b/Engine/Source/Runtime/Engine/Private/WorldPartition/ErrorHandling/WorldPartitionStreamingGenerationTokenizedMessageErrorHandler.cpp
--- a/Engine/Source/Runtime/Engine/Private/WorldPartition/ErrorHandling/WorldPartitionStreamingGenerationTokenizedMessageErrorHandler.cpp
+++ b/Engine/Source/Runtime/Engine/Private/WorldPartition/ErrorHandling/WorldPartitionStreamingGenerationTokenizedMessageErrorHandler.cpp
@@ -198,6 +198,16 @@ void ITokenizedMessageErrorHandler::OnLevelInstanceInvalidWorldAsset(const FWorl
 	HandleTokenizedMessage(MoveTemp(Message));
 }
 
+void ITokenizedMessageErrorHandler::OnActorFilterFailed(const FWorldPartitionActorDescView& ActorDescView)
+{
+	TSharedRef<FTokenizedMessage> Message = FTokenizedMessage::Create(EMessageSeverity::Warning);
+	Message->AddToken(FTextToken::Create(LOCTEXT("TokenMessage_WorldPartition_ActorFilterFailed", "Actor will not be filtered out because it is referenced/referencing other Actor(s) not part of the filter")))
+		->AddToken(FActorToken::Create(ActorDescView.GetActorSoftPath().ToString(), ActorDescView.GetGuid(), FText::FromString(GetActorName(ActorDescView))))
+		->AddToken(FMapErrorToken::Create(TEXT("WorldPartition_ActorFilterFailed_CheckForErrors")));
+
+	HandleTokenizedMessage(MoveTemp(Message));
+}
+
 #undef LOCTEXT_NAMESPACE
 
 #endif
diff --git a/Engine/Source/Runtime/Engine/Private/WorldPartition/Filter/WorldPartitionActorFilter.cpp b/Engine/Source/Runtime/Engine/Private/WorldPartition/Filter/WorldPartitionActorFilter.cpp
new file mode 100644
--- /dev/null
+++ b/Engine/Source/Runtime/Engine/Private/WorldPartition/Filter/WorldPartitionActorFilter.cpp
@@ -0,0 +1,382 @@
+// Copyright Epic Games, Inc. All Rights Reserved.
+
+#include "WorldPartition/Filter/WorldPartitionActorFilter.h"
+#include "Containers/Ticker.h"
+#include "WorldPartition/WorldPartitionActorLoaderInterface.h"
+
+#if WITH_EDITORONLY_DATA
+FWorldPartitionActorFilter::FOnWorldPartitionActorFilterChanged FWorldPartitionActorFilter::OnWorldPartitionActorFilterChanged;
+#endif
+
+FWorldPartitionActorFilter::~FWorldPartitionActorFilter()
+{
+#if WITH_EDITORONLY_DATA
+	ClearChildFilters();
+#endif
+}
+
+FWorldPartitionActorFilter::FWorldPartitionActorFilter(const FWorldPartitionActorFilter& Other)
+{
+	*this = Other;
+}
+FWorldPartitionActorFilter::FWorldPartitionActorFilter(FWorldPartitionActorFilter&& Other)
+{
+	*this = MoveTemp(Other);
+}
+
+FWorldPartitionActorFilter& FWorldPartitionActorFilter::operator=(const FWorldPartitionActorFilter& Other)
+{
+#if WITH_EDITORONLY_DATA
+	Parent = nullptr;
+	DisplayName = Other.DisplayName;
+	DataLayerFilters = Other.DataLayerFilters;
+	ClearChildFilters();
+	for (const auto& [ActorGuid, FilterPtr] : Other.ChildFilters)
+	{
+		AddChildFilter(ActorGuid, new FWorldPartitionActorFilter(*FilterPtr));
+	}
+#endif
+	return *this;
+}
+
+FWorldPartitionActorFilter& FWorldPartitionActorFilter::operator=(FWorldPartitionActorFilter&& Other)
+{
+#if WITH_EDITORONLY_DATA
+	Parent = nullptr;
+	DisplayName = MoveTemp(Other.DisplayName);
+	DataLayerFilters = MoveTemp(Other.DataLayerFilters);
+	ClearChildFilters();
+	ChildFilters = MoveTemp(Other.ChildFilters);
+	for (const auto& [ActorGuid, FilterPtr] : ChildFilters)
+	{
+		FilterPtr->Parent = this;
+	}
+#endif
+	return *this;
+}
+
+#if WITH_EDITORONLY_DATA
+void FWorldPartitionActorFilter::RequestFilterRefresh(bool bIsFromUserChange)
+{
+	if (GUndo || GIsTransacting)
+	{
+		FTSTicker::GetCoreTicker().AddTicker(FTickerDelegate::CreateLambda([=](float DeltaTime)
+		{
+			if (GUndo || GIsTransacting)
+			{
+				// Tick again we are still inside a transaction...
+				return true;
+			}
+			
+			IWorldPartitionActorLoaderInterface::RefreshLoadedState(bIsFromUserChange);
+			return false;
+		}));
+	}
+	else
+	{
+		IWorldPartitionActorLoaderInterface::RefreshLoadedState(bIsFromUserChange);
+	}
+}
+
+void FWorldPartitionActorFilter::AddChildFilter(const FGuid& InGuid, FWorldPartitionActorFilter* InChildFilter)
+{
+	InChildFilter->Parent = this;
+	ChildFilters.Add(InGuid, InChildFilter);
+}
+
+void FWorldPartitionActorFilter::RemoveChildFilter(const FGuid& InGuid)
+{
+	FWorldPartitionActorFilter* ChildFilter = nullptr;
+	if (ChildFilters.RemoveAndCopyValue(InGuid, ChildFilter))
+	{
+		delete ChildFilter;
+	}
+}
+
+void FWorldPartitionActorFilter::ClearChildFilters()
+{
+	for (auto& [ActorGuid, ChildFilter] : ChildFilters)
+	{
+		delete ChildFilter;
+	}
+	ChildFilters.Empty();
+}
+
+void FWorldPartitionActorFilter::Override(const FWorldPartitionActorFilter& Other)
+{
+	for (const auto& [AssetPath, OtherDataLayerFilter] : Other.DataLayerFilters)
+	{
+		if (FDataLayerFilter* ExistingFilter = DataLayerFilters.Find(AssetPath))
+		{
+			ExistingFilter->bIncluded = OtherDataLayerFilter.bIncluded;
+		}
+	}
+	
+	for (const auto& [ActorGuid, OtherFilter] : Other.GetChildFilters())
+	{
+		if (FWorldPartitionActorFilter** ExistingFilter = ChildFilters.Find(ActorGuid))
+		{
+			(*ExistingFilter)->Override(*OtherFilter);
+		}
+	}
+}
+
+bool FWorldPartitionActorFilter::operator==(const FWorldPartitionActorFilter& Other) const
+{
+	if (DataLayerFilters.Num() != Other.DataLayerFilters.Num() || ChildFilters.Num() != Other.ChildFilters.Num())
+	{
+		return false;
+	}
+
+	for (const auto& [AssetPath, DataLayerFilter] : DataLayerFilters)
+	{
+		if (const FDataLayerFilter* OtherDataLayerFilter = Other.DataLayerFilters.Find(AssetPath))
+		{
+			if (DataLayerFilter.bIncluded != OtherDataLayerFilter->bIncluded)
+			{
+				return false;
+			}
+		}
+		else
+		{
+			return false;
+		}
+	}
+
+	for (const auto& [ActorGuid, WorldPartitionActorFilter] : ChildFilters)
+	{
+		if (const FWorldPartitionActorFilter*const* OtherChildFilter = Other.GetChildFilters().Find(ActorGuid))
+		{
+			if (*WorldPartitionActorFilter != **OtherChildFilter)
+			{
+				return false;
+			}
+		}
+		else
+		{
+			return false;
+		}
+	}
+
+	return true;
+}
+
+bool FWorldPartitionActorFilter::Serialize(FArchive& Ar)
+{
+	uint32 DataLayerFilterCount = DataLayerFilters.Num();
+	Ar << DataLayerFilterCount;
+
+	if (Ar.IsLoading())
+	{
+		DataLayerFilters.Empty(DataLayerFilterCount);
+
+		for (uint32 i = 0; i < DataLayerFilterCount; ++i)
+		{
+			FString AssetPath;
+			Ar << AssetPath;
+			bool bIncluded;
+			Ar << bIncluded;
+
+			DataLayerFilters.Add(FSoftObjectPath(AssetPath)).bIncluded = bIncluded;
+		}
+	}
+	else
+	{
+		for (auto& [AssetPath, DataLayerFilter] : DataLayerFilters)
+		{
+			FString AssetPathStr = AssetPath.ToString();
+			Ar << AssetPathStr;
+			Ar << DataLayerFilter.bIncluded;
+		}
+	}
+
+	uint32 ChildFilterCount = ChildFilters.Num();
+	Ar << ChildFilterCount;
+
+	if (Ar.IsLoading())
+	{
+		// with undo/redo serialization we can be serializing over existing childs
+		ClearChildFilters();
+
+		for (uint32 i = 0; i < ChildFilterCount; ++i)
+		{
+			FGuid ActorGuid;
+			Ar << ActorGuid;
+			FWorldPartitionActorFilter* ChildFilter = new FWorldPartitionActorFilter();
+			Ar << *ChildFilter;
+			AddChildFilter(ActorGuid, ChildFilter);
+		}
+	}
+	else
+	{
+		for (auto& [ActorGuid, ChildFilter] : ChildFilters)
+		{
+			Ar << ActorGuid;
+			Ar << *ChildFilter;
+		}
+	}
+
+	return true;
+}
+
+FArchive& operator<<(FArchive& Ar, FWorldPartitionActorFilter& Filter)
+{
+	Filter.Serialize(Ar);
+	return Ar;
+}
+
+bool FWorldPartitionActorFilter::ExportTextItem(FString& ValueStr, FWorldPartitionActorFilter const& DefaultValue, UObject* ParentObject, int32 PortFlags, UObject* ExportRootScope) const
+{
+	const uint32 OriginalLen = ValueStr.Len();
+	const FSoftObjectPath DefaultPath;
+
+	ValueStr += TEXT("(");
+	bool bFirst = true;
+	
+	for (const auto& [AssetPath, DataLayerFilter] : DataLayerFilters)
+	{
+		if (!bFirst)
+		{
+			ValueStr += TEXT(",");
+		}
+		bFirst = false;
+
+		if (!AssetPath.ExportTextItem(ValueStr, DefaultPath, ParentObject, PortFlags, ExportRootScope))
+		{
+			ValueStr.LeftInline(OriginalLen);
+			return false;
+		}
+
+		ValueStr += TEXT(";");
+		ValueStr += DataLayerFilter.bIncluded ? TEXT("1") : TEXT("0");
+	}
+	ValueStr += TEXT(")(");
+
+	bFirst = true;
+	const FGuid DefaultGuid;
+	for (const auto& [ActorGuid, WorldPartitionActorFilter] : ChildFilters)
+	{
+		if (!bFirst)
+		{
+			ValueStr += TEXT(",");
+		}
+		bFirst = false;
+				
+		if (!ActorGuid.ExportTextItem(ValueStr, DefaultGuid, ParentObject, PortFlags, ExportRootScope))
+		{
+			ValueStr.LeftChopInline(OriginalLen);
+			return false;
+		}
+		ValueStr += TEXT(";");
+
+		if (!WorldPartitionActorFilter->ExportTextItem(ValueStr, DefaultValue, ParentObject, PortFlags, ExportRootScope))
+		{
+			ValueStr.LeftInline(OriginalLen);
+			return false;
+		}
+	}
+
+	ValueStr += TEXT(")");
+	return true;
+}
+
+bool FWorldPartitionActorFilter::ImportTextItem(const TCHAR*& Buffer, int32 PortFlags, UObject* ParentObject, FOutputDevice* ErrorText)
+{
+	FWorldPartitionActorFilter ImportedFilter;
+	if (*Buffer != '(')
+	{
+		ErrorText->Logf(ELogVerbosity::Warning, TEXT("FWorldPartitionActorFilter: Missing opening \'(\' while importing property values."));
+
+		// Parse error
+		return false;
+	}
+	Buffer++;
+
+	while (*Buffer != ')')
+	{
+		FSoftObjectPath AssetPath;
+		if (!AssetPath.ImportTextItem(Buffer, PortFlags, ParentObject, ErrorText))
+		{
+			return false;
+		}
+
+		if (*Buffer != ';')
+		{
+			ErrorText->Logf(ELogVerbosity::Warning, TEXT("FWorldPartitionActorFilter: Missing \';\' while importing property values."));
+			return false;
+		}
+		Buffer++;
+		bool bValue;
+		if (*Buffer == '0' || *Buffer == '1')
+		{
+			bValue = *Buffer == '1';
+		}
+		else
+		{
+			ErrorText->Logf(ELogVerbosity::Warning, TEXT("FWorldPartitionActorFilter: Missing \'0\' or \'1\' while importing property values."));
+			return false;
+		}
+		Buffer++;
+
+		if (*Buffer == ',')
+		{
+			Buffer++;
+		}
+		else if(*Buffer != ')')
+		{
+			ErrorText->Logf(ELogVerbosity::Warning, TEXT("FWorldPartitionActorFilter: Missing \')\' while importing property values."));
+			return false;
+		}
+		ImportedFilter.DataLayerFilters.Add(AssetPath).bIncluded = bValue;
+	}
+	Buffer++;
+	
+	if (*Buffer != '(')
+	{
+		ErrorText->Logf(ELogVerbosity::Warning, TEXT("FWorldPartitionActorFilter: Missing opening \'(\' while importing property values."));
+
+		// Parse error
+		return false;
+	}
+	Buffer++;
+
+	while (*Buffer != ')')
+	{
+		FGuid ActorGuid;
+		if (!ActorGuid.ImportTextItem(Buffer, PortFlags, ParentObject, ErrorText))
+		{
+			return false;
+		}
+
+		if (*Buffer != ';')
+		{
+			ErrorText->Logf(ELogVerbosity::Warning, TEXT("FWorldPartitionActorFilter: Missing \';\' while importing property values."));
+			return false;
+		}
+		Buffer++;
+
+		FWorldPartitionActorFilter ChildFilter;
+		if (!ChildFilter.ImportTextItem(Buffer, PortFlags, ParentObject, ErrorText))
+		{
+			return false;
+		}
+
+		if (*Buffer == ',')
+		{
+			Buffer++;
+		}
+		else if (*Buffer != ')')
+		{
+			ErrorText->Logf(ELogVerbosity::Warning, TEXT("FWorldPartitionActorFilter: Missing \')\' while importing property values."));
+			return false;
+		}
+
+		ImportedFilter.AddChildFilter(ActorGuid, new FWorldPartitionActorFilter(MoveTemp(ChildFilter)));
+	}
+	Buffer++;
+
+	*this = MoveTemp(ImportedFilter);
+	return true;
+}
+
+#endif 
\ No newline at end of file
diff --git a/Engine/Source/Runtime/Engine/Private/WorldPartition/LevelInstance/LevelInstanceActorDesc.cpp b/Engine/Source/Runtime/Engine/Private/WorldPartition/LevelInstance/LevelInstanceActorDesc.cpp
--- a/Engine/Source/Runtime/Engine/Private/WorldPartition/LevelInstance/LevelInstanceActorDesc.cpp
+++ b/Engine/Source/Runtime/Engine/Private/WorldPartition/LevelInstance/LevelInstanceActorDesc.cpp
@@ -8,11 +8,13 @@
 #include "Engine/Level.h"
 #include "LevelInstance/LevelInstanceInterface.h"
 #include "LevelInstance/LevelInstanceSubsystem.h"
+#include "LevelInstance/LevelInstanceComponent.h"
 #include "Misc/PackageName.h"
 #include "WorldPartition/ActorDescContainer.h"
 #include "WorldPartition/ErrorHandling/WorldPartitionStreamingGenerationErrorHandler.h"
 #include "UObject/UE5ReleaseStreamObjectVersion.h"
 #include "UObject/UE5MainStreamObjectVersion.h"
+#include "UObject/FortniteMainBranchObjectVersion.h"
 
 static int32 GLevelInstanceDebugForceLevelStreaming = 0;
 static FAutoConsoleVariableRef CVarForceLevelStreaming(
@@ -40,6 +42,7 @@ void FLevelInstanceActorDesc::Init(const AActor* InActor)
 	LevelPackage = *LevelInstance->GetWorldAssetPackage();
 	LevelInstanceTransform = InActor->GetActorTransform();
 	DesiredRuntimeBehavior = LevelInstance->GetDesiredRuntimeBehavior();
+	Filter = LevelInstance->GetFilter();
 }
 
 void FLevelInstanceActorDesc::Init(const FWorldPartitionActorDescInitData& DescData)
@@ -159,7 +162,7 @@ bool FLevelInstanceActorDesc::IsContainerInstance() const
 	return ULevelInstanceSubsystem::CanUsePackage(LevelPackage);
 }
 
-bool FLevelInstanceActorDesc::GetContainerInstance(FContainerInstance& OutContainerInstance) const
+bool FLevelInstanceActorDesc::GetContainerInstance(FContainerInstance& OutContainerInstance, bool bInBuildFilter) const
 {
 	if (LevelInstanceContainer.IsValid())
 	{
@@ -170,6 +173,71 @@ bool FLevelInstanceActorDesc::GetContainerInstance(FContainerInstance& OutContai
 		FTransform LevelInstancePivotOffsetTransform = FTransform(ULevel::GetLevelInstancePivotOffsetFromPackage(LevelInstanceContainer->GetContainerPackage()));
 		OutContainerInstance.Transform = LevelInstancePivotOffsetTransform * LevelInstanceTransform;
 
+		if (bInBuildFilter)
+		{
+			// Fill Container Instance Filter
+			ULevelInstanceSubsystem* LevelInstanceSubsystem = UWorld::GetSubsystem<ULevelInstanceSubsystem>(LevelInstanceContainerWorldContext.Get());
+			check(LevelInstanceSubsystem);
+
+			FWorldPartitionActorFilter ContainerFilter = LevelInstanceSubsystem->GetLevelInstanceFilter(LevelInstanceContainer->GetContainerPackage().ToString());
+			ContainerFilter.Override(Filter);
+
+			// Flatten Filter to FActorContainerID map
+			TMap<FActorContainerID, TSet<FSoftObjectPath>> FilteredOutDataLayersPerContainer;
+			TFunction<void(const FActorContainerID&, const FWorldPartitionActorFilter&)> ProcessFilter = [&FilteredOutDataLayersPerContainer, &ProcessFilter](const FActorContainerID& InContainerID, const FWorldPartitionActorFilter& InContainerFilter)
+			{
+				check(!FilteredOutDataLayersPerContainer.Contains(InContainerID));
+				TSet<FSoftObjectPath>& Filtered = FilteredOutDataLayersPerContainer.Add(InContainerID);
+
+				for (auto& [AssetPath, DataLayerFilter] : InContainerFilter.DataLayerFilters)
+				{
+					if (!DataLayerFilter.bIncluded)
+					{
+						Filtered.Add(AssetPath);
+					}
+				}
+
+				for (auto& [ActorGuid, WorldPartitionActorFilter] : InContainerFilter.GetChildFilters())
+				{
+					ProcessFilter(FActorContainerID(InContainerID, ActorGuid), *WorldPartitionActorFilter);
+				}
+			};
+
+			ProcessFilter(OutContainerInstance.GetID(), ContainerFilter);
+
+			TFunction<void(const FActorContainerID&, const UActorDescContainer*)> ProcessContainers = [&FilteredOutDataLayersPerContainer, &OutContainerInstance, &ProcessContainers](const FActorContainerID& InContainerID, const UActorDescContainer* InContainer)
+			{
+				const TSet<FSoftObjectPath>& Filtered = FilteredOutDataLayersPerContainer.FindChecked(InContainerID);
+				for (FActorDescList::TConstIterator<> ActorDescIt(InContainer); ActorDescIt; ++ActorDescIt)
+				{
+					if (ActorDescIt->GetDataLayers().Num() > 0 && ActorDescIt->IsUsingDataLayerAsset())
+					{
+						for (FName DataLayerName : ActorDescIt->GetDataLayers())
+						{
+							FSoftObjectPath DataLayerAsset(DataLayerName.ToString());
+							if (Filtered.Contains(DataLayerAsset))
+							{
+								OutContainerInstance.FilteredActors.FindOrAdd(InContainerID).Add(ActorDescIt->GetGuid());
+								break;
+							}
+						}
+					}
+
+					if (ActorDescIt->GetActorNativeClass()->ImplementsInterface(ULevelInstanceInterface::StaticClass()))
+					{
+						const FLevelInstanceActorDesc* LevelInstanceActorDesc = static_cast<const FLevelInstanceActorDesc*>(*ActorDescIt);
+						FWorldPartitionActorDesc::FContainerInstance ChildContainerInstance;
+						const bool bBuildFilter = false;
+						if (LevelInstanceActorDesc->GetContainerInstance(ChildContainerInstance, bBuildFilter))
+						{
+							ProcessContainers(FActorContainerID(InContainerID, ActorDescIt->GetGuid()), ChildContainerInstance.Container);
+						}
+					}
+				}
+			};
+
+			ProcessContainers(OutContainerInstance.GetID(), OutContainerInstance.Container);
+		}
 		return true;
 	}
 
@@ -221,6 +289,7 @@ void FLevelInstanceActorDesc::Serialize(FArchive& Ar)
 	FWorldPartitionActorDesc::Serialize(Ar);
 	Ar.UsingCustomVersion(FUE5ReleaseStreamObjectVersion::GUID);
 	Ar.UsingCustomVersion(FUE5MainStreamObjectVersion::GUID);
+	Ar.UsingCustomVersion(FFortniteMainBranchObjectVersion::GUID);
 
 	Ar << LevelPackage;
 
@@ -264,6 +333,11 @@ void FLevelInstanceActorDesc::Serialize(FArchive& Ar)
 			}
 		}
 	}
+
+	if (Ar.CustomVer(FFortniteMainBranchObjectVersion::GUID) >= FFortniteMainBranchObjectVersion::WorldPartitionActorFilter)
+	{
+		Ar << Filter;
+	}
 }
 
 #endif
diff --git a/Engine/Source/Runtime/Engine/Private/WorldPartition/WorldPartitionActorDescView.cpp b/Engine/Source/Runtime/Engine/Private/WorldPartition/WorldPartitionActorDescView.cpp
--- a/Engine/Source/Runtime/Engine/Private/WorldPartition/WorldPartitionActorDescView.cpp
+++ b/Engine/Source/Runtime/Engine/Private/WorldPartition/WorldPartitionActorDescView.cpp
@@ -184,9 +184,9 @@ FName FWorldPartitionActorDescView::GetLevelPackage() const
 	return ActorDesc->GetLevelPackage();
 }
 
-bool FWorldPartitionActorDescView::GetContainerInstance(FWorldPartitionActorDesc::FContainerInstance& OutContainerInstance) const
+bool FWorldPartitionActorDescView::GetContainerInstance(FWorldPartitionActorDesc::FContainerInstance& OutContainerInstance, bool bInBuildFilter) const
 {
-	return ActorDesc->GetContainerInstance(OutContainerInstance);
+	return ActorDesc->GetContainerInstance(OutContainerInstance, bInBuildFilter);
 }
 
 void FWorldPartitionActorDescView::CheckForErrors(IStreamingGenerationErrorHandler* ErrorHandler) const
diff --git a/Engine/Source/Runtime/Engine/Private/WorldPartition/WorldPartitionStreamingGeneration.cpp b/Engine/Source/Runtime/Engine/Private/WorldPartition/WorldPartitionStreamingGeneration.cpp
--- a/Engine/Source/Runtime/Engine/Private/WorldPartition/WorldPartitionStreamingGeneration.cpp
+++ b/Engine/Source/Runtime/Engine/Private/WorldPartition/WorldPartitionStreamingGeneration.cpp
@@ -106,6 +106,12 @@ class FWorldPartitionStreamingGenerator
 				}
 			}
 
+			TMap<FActorContainerID, TSet<FGuid>> FilteredOutActorsPerContainer;
+			for (const auto& [ContainerID, ContainerInstanceDescriptor] : StreamingGenerator->ContainerInstanceDescriptorsMap)
+			{
+				FilteredOutActorsPerContainer.Append(ContainerInstanceDescriptor.FilteredActors);
+			}
+
 			ActorSetInstances.Empty();
 			for (const auto& [ContainerID, ContainerInstanceDescriptor] : StreamingGenerator->ContainerInstanceDescriptorsMap)
 			{
@@ -117,6 +123,9 @@ class FWorldPartitionStreamingGenerator
 					const FActorSet& ActorSet = *ActorSetPtr;
 					const FWorldPartitionActorDescView& ReferenceActorDescView = ActorSetContainer.ActorDescViewMap->FindByGuidChecked(ActorSet.Actors[0]);
 
+					TSet<FGuid> FilteredActors;
+					TSet<FGuid>* FilteredOutActors = FilteredOutActorsPerContainer.Find(ContainerID);
+
 					// Validate assumptions
 					for (const FGuid& ActorGuid : ActorSet.Actors)
 					{
@@ -124,6 +133,25 @@ class FWorldPartitionStreamingGenerator
 						check(ActorDescView.GetRuntimeGrid() == ReferenceActorDescView.GetRuntimeGrid());
 						check(ActorDescView.GetIsSpatiallyLoaded() == ReferenceActorDescView.GetIsSpatiallyLoaded());
 						check(ActorDescView.GetContentBundleGuid() == ReferenceActorDescView.GetContentBundleGuid());
+
+						if (FilteredOutActors && FilteredOutActors->Contains(ActorGuid))
+						{
+							FilteredActors.Add(ActorGuid);
+						}
+					}
+
+					// We can only filter the complete set because of dependencies between actors of the same set
+					if (FilteredActors.Num() == ActorSet.Actors.Num())
+					{
+						continue;
+					}
+					else if (FilteredActors.Num() > 0)
+					{
+						// Report failed filtering
+						for (const FGuid& FilteredActorGuid : FilteredActors)
+						{
+							StreamingGenerator->ErrorHandler->OnActorFilterFailed(ActorSetContainer.ActorDescViewMap->FindByGuidChecked(FilteredActorGuid));
+						}
 					}
 
 					FActorSetInstance& ActorSetInstance = ActorSetInstances.Emplace_GetRef();
@@ -235,6 +263,7 @@ class FWorldPartitionStreamingGenerator
 		FString OwnerName;
 		FActorContainerID ID;
 		FActorContainerID ParentID;
+		TMap<FActorContainerID, TSet<FGuid>> FilteredActors;
 	};
 
 	void ResolveRuntimeSpatiallyLoaded(FWorldPartitionActorDescView& ActorDescView)
@@ -447,7 +476,8 @@ class FWorldPartitionStreamingGenerator
 		{
 			const FGuid ActorGuid = ContainerInstanceView.GetGuid();
 			FWorldPartitionActorDesc::FContainerInstance SubContainerInstance(FActorContainerID(InContainerInstanceDescriptor.ID, ActorGuid));
-			if (!ContainerInstanceView.GetContainerInstance(SubContainerInstance))
+			bool bBuildFilter = InContainerInstanceDescriptor.ID.IsMainContainer(); // Build filter on top level child containers
+			if (!ContainerInstanceView.GetContainerInstance(SubContainerInstance, bBuildFilter))
 			{
 				ErrorHandler->OnLevelInstanceInvalidWorldAsset(ContainerInstanceView, ContainerInstanceView.GetLevelPackage(), IStreamingGenerationErrorHandler::ELevelInstanceInvalidReason::WorldAssetHasInvalidContainer);
 				continue;
@@ -468,6 +498,7 @@ class FWorldPartitionStreamingGenerator
 			SubContainerInstanceDescriptor.ID = SubContainerInstance.GetID();
 			SubContainerInstanceDescriptor.Container = SubContainerInstance.Container;
 			SubContainerInstanceDescriptor.Transform = SubContainerInstance.Transform * InContainerInstanceDescriptor.Transform;
+			SubContainerInstanceDescriptor.FilteredActors = SubContainerInstance.FilteredActors;
 			SubContainerInstanceDescriptor.ParentID = InContainerInstanceDescriptor.ID;
 			SubContainerInstanceDescriptor.OwnerName = *ContainerInstanceView.GetActorLabelOrName().ToString();
 
diff --git a/Engine/Source/Runtime/Engine/Public/LevelInstance/LevelInstanceActor.h b/Engine/Source/Runtime/Engine/Public/LevelInstance/LevelInstanceActor.h
--- a/Engine/Source/Runtime/Engine/Public/LevelInstance/LevelInstanceActor.h
+++ b/Engine/Source/Runtime/Engine/Public/LevelInstance/LevelInstanceActor.h
@@ -27,6 +27,8 @@ protected:
 	UPROPERTY(EditAnywhere, Category = Level, meta = (NoCreate, DisplayName="Level"))
 	TSoftObjectPtr<UWorld> WorldAsset;
 #endif
+	UPROPERTY(VisibleAnywhere, Category = Default)
+	TObjectPtr<ULevelInstanceComponent> LevelInstanceComponent;
 
 	UPROPERTY()
 	TSoftObjectPtr<UWorld> CookedWorldAsset;
diff --git a/Engine/Source/Runtime/Engine/Public/LevelInstance/LevelInstanceComponent.h b/Engine/Source/Runtime/Engine/Public/LevelInstance/LevelInstanceComponent.h
--- a/Engine/Source/Runtime/Engine/Public/LevelInstance/LevelInstanceComponent.h
+++ b/Engine/Source/Runtime/Engine/Public/LevelInstance/LevelInstanceComponent.h
@@ -4,6 +4,7 @@
 
 #include "CoreMinimal.h"
 #include "Components/SceneComponent.h"
+#include "WorldPartition/Filter/WorldPartitionActorFilter.h"
 
 #include "LevelInstanceComponent.generated.h"
 
@@ -26,6 +27,7 @@ public:
 #if WITH_EDITOR
 	// Those are the methods that need overriding to be able to properly update the AttachComponent
 	virtual void OnRegister() override;
+	virtual void PreEditUndo() override;
 	virtual void PostEditUndo() override;
 	virtual void PostEditChangeProperty(FPropertyChangedEvent& PropertyChangedEvent) override;
 	virtual void OnUpdateTransform(EUpdateTransformFlags UpdateTransformFlags, ETeleportType Teleport) override;
@@ -33,9 +35,19 @@ public:
 	void UpdateEditorInstanceActor();
 	void OnEdit();
 	void OnCommit();
+
+	const FWorldPartitionActorFilter& GetFilter() const { return Filter; }
+	void SetFilter(const FWorldPartitionActorFilter& InFilter);
 private:
 	bool ShouldShowSpriteComponent() const;
 
 	TWeakObjectPtr<AActor> CachedEditorInstanceActorPtr;
 #endif
+
+#if WITH_EDITORONLY_DATA
+	UPROPERTY(EditAnywhere, Category = Filter, meta=(LevelInstanceFilter))
+	FWorldPartitionActorFilter Filter;
+
+	FWorldPartitionActorFilter CachedFilter;
+#endif
 };
\ No newline at end of file
diff --git a/Engine/Source/Runtime/Engine/Public/LevelInstance/LevelInstanceInterface.h b/Engine/Source/Runtime/Engine/Public/LevelInstance/LevelInstanceInterface.h
--- a/Engine/Source/Runtime/Engine/Public/LevelInstance/LevelInstanceInterface.h
+++ b/Engine/Source/Runtime/Engine/Public/LevelInstance/LevelInstanceInterface.h
@@ -10,6 +10,7 @@
 
 class ULevelInstanceComponent;
 class ULevelStreamingLevelInstance;
+struct FWorldPartitionActorFilter;
 
 UINTERFACE(BlueprintType, meta = (CannotImplementInterfaceInBlueprint))
 class ENGINE_API ULevelInstanceInterface : public UInterface
@@ -114,5 +115,7 @@ class ENGINE_API ILevelInstanceInterface
 	virtual bool SetCurrent();
 	
 	virtual bool MoveActorsTo(const TArray<AActor*>& ActorsToMove);
+
+	virtual const FWorldPartitionActorFilter& GetFilter() const;
 #endif
 };
\ No newline at end of file
diff --git a/Engine/Source/Runtime/Engine/Public/LevelInstance/LevelInstanceSubsystem.h b/Engine/Source/Runtime/Engine/Public/LevelInstance/LevelInstanceSubsystem.h
--- a/Engine/Source/Runtime/Engine/Public/LevelInstance/LevelInstanceSubsystem.h
+++ b/Engine/Source/Runtime/Engine/Public/LevelInstance/LevelInstanceSubsystem.h
@@ -8,6 +8,7 @@
 #include "UObject/GCObject.h"
 #include "Folder.h"
 #include "LevelInstance/LevelInstanceTypes.h"
+#include "WorldPartition/Filter/WorldPartitionActorFilter.h"
 #include "WorldPartition/WorldPartitionHandle.h"
 
 #include "LevelInstanceSubsystem.generated.h"
@@ -125,6 +126,9 @@ public:
 	FLevelInstanceChanged& OnLevelInstanceChanged() { return LevelInstanceChangedEvent; }
 
 	static void ResetLoadersForWorldAsset(const FString& WorldAsset);
+
+	FWorldPartitionActorFilter GetLevelInstanceFilter(const FString& LevelPackage) const;
+	bool PassLevelInstanceFilter(UWorld* World, const FWorldPartitionHandle& Actor) const;
 #endif
 
 private:
@@ -137,6 +141,8 @@ private:
 	void RegisterLoadedLevelStreamingLevelInstance(ULevelStreamingLevelInstance* LevelStreaming);
 
 #if WITH_EDITOR
+	FWorldPartitionActorFilter GetLevelInstanceFilterInternal(const FString& LevelPackage, TSet<FString>& VisitedPackages) const;
+
 	void ResetLoadersForWorldAssetInternal(const FString& WorldAsset);
 	void OnAssetsPreDelete(const TArray<UObject*>& Objects);
 
@@ -257,7 +263,7 @@ private:
 
 	TMap<FLevelInstanceID, int32> ChildEdits;
 
-	FActorDescContainerInstanceManager ActorDescContainerInstanceManager;
+	mutable FActorDescContainerInstanceManager ActorDescContainerInstanceManager;
 
 	FWorldPartitionReference CurrentEditLevelInstanceActor;
 #endif
diff --git a/Engine/Source/Runtime/Engine/Public/WorldPartition/DataLayer/DataLayerAsset.h b/Engine/Source/Runtime/Engine/Public/WorldPartition/DataLayer/DataLayerAsset.h
--- a/Engine/Source/Runtime/Engine/Public/WorldPartition/DataLayer/DataLayerAsset.h
+++ b/Engine/Source/Runtime/Engine/Public/WorldPartition/DataLayer/DataLayerAsset.h
@@ -37,11 +37,15 @@ public:
 	UFUNCTION(Category = "Data Layer|Runtime", BlueprintCallable)
 	FColor GetDebugColor() const { return DebugColor; }
 
+	bool SupportsActorFilters() const { return bSupportsActorFilters; }
 private:
 	/** Whether the Data Layer affects actor runtime loading */
 	UPROPERTY(Category = "Data Layer", EditAnywhere)
 	EDataLayerType DataLayerType;
 
+	UPROPERTY(Category = "Data Layer|Actor Filter", EditAnywhere)
+	bool bSupportsActorFilters;
+		
 	UPROPERTY(Category = "Data Layer|Runtime", EditAnywhere)
 	FColor DebugColor;
 };
\ No newline at end of file
diff --git a/Engine/Source/Runtime/Engine/Public/WorldPartition/DataLayer/DataLayerInstanceWithAsset.h b/Engine/Source/Runtime/Engine/Public/WorldPartition/DataLayer/DataLayerInstanceWithAsset.h
--- a/Engine/Source/Runtime/Engine/Public/WorldPartition/DataLayer/DataLayerInstanceWithAsset.h
+++ b/Engine/Source/Runtime/Engine/Public/WorldPartition/DataLayer/DataLayerInstanceWithAsset.h
@@ -18,6 +18,7 @@ public:
 	static FName MakeName(const UDataLayerAsset* DeprecatedDataLayer);
 	void OnCreated(const UDataLayerAsset* Asset);
 
+	virtual bool CanEditChange(const FProperty* InProperty) const;
 	virtual void PreEditUndo() override;
 	virtual void PostEditUndo() override;
 	virtual bool IsLocked() const override;
@@ -27,6 +28,9 @@ public:
 
 	virtual bool Validate(IStreamingGenerationErrorHandler* ErrorHandler) const override;
 	virtual void PostEditChangeProperty(FPropertyChangedEvent& PropertyChangedEvent) override;
+
+	bool SupportsActorFilters() const;
+	bool IsIncludedInActorFilterDefault() const;
 #endif
 
 	const UDataLayerAsset* GetAsset() const { return DataLayerAsset; }
@@ -50,6 +54,11 @@ private:
 	UPROPERTY(Category = "Data Layer", EditAnywhere)
 	TObjectPtr<const UDataLayerAsset> DataLayerAsset;
 
+#if WITH_EDITORONLY_DATA
+	UPROPERTY(Category = "Data Layer|Actor Filter", EditAnywhere, meta = (DisplayName = "Is Included", ToolTip = "Whether actors assigned to this DataLayer are included by default when used in a filter"))
+	bool bIsIncludedInActorFilterDefault;
+#endif
+
 #if WITH_EDITOR
 	// Used to compare state pre/post undo
 	TObjectPtr<const UDataLayerAsset> CachedDataLayerAsset;
diff --git a/Engine/Source/Runtime/Engine/Public/WorldPartition/DataLayer/WorldDataLayersActorDesc.h b/Engine/Source/Runtime/Engine/Public/WorldPartition/DataLayer/WorldDataLayersActorDesc.h
--- a/Engine/Source/Runtime/Engine/Public/WorldPartition/DataLayer/WorldDataLayersActorDesc.h
+++ b/Engine/Source/Runtime/Engine/Public/WorldPartition/DataLayer/WorldDataLayersActorDesc.h
@@ -23,6 +23,7 @@ public:
 	class UDataLayerAsset* GetAsset() const;
 	EDataLayerType GetDataLayerType() const;
 	FString GetShortName() const;
+	bool IsIncludedInActorFilterDefault() const { return bIsIncludedInActorFilterDefault; }
 
 private:
 	// DataLayerInstance Name
@@ -43,6 +44,9 @@ private:
 	// Label
 	FString ShortName;
 	//~ End UDeprecatedDataLayerInstance
+
+	// Returns if data layer should be included by default in FWorldPartitionActorFilter.
+	bool bIsIncludedInActorFilterDefault;
 };
 
 /**
diff --git a/Engine/Source/Runtime/Engine/Public/WorldPartition/ErrorHandling/WorldPartitionStreamingGenerationErrorHandler.h b/Engine/Source/Runtime/Engine/Public/WorldPartition/ErrorHandling/WorldPartitionStreamingGenerationErrorHandler.h
--- a/Engine/Source/Runtime/Engine/Public/WorldPartition/ErrorHandling/WorldPartitionStreamingGenerationErrorHandler.h
+++ b/Engine/Source/Runtime/Engine/Public/WorldPartition/ErrorHandling/WorldPartitionStreamingGenerationErrorHandler.h
@@ -85,6 +85,8 @@ public:
 
 	virtual void OnLevelInstanceInvalidWorldAsset(const FWorldPartitionActorDescView& ActorDescView, FName WorldAsset, ELevelInstanceInvalidReason Reason) = 0;
 
+	virtual void OnActorFilterFailed(const FWorldPartitionActorDescView& ActorDescView) = 0;
+
 	// Helpers
 	static FString GetActorName(const FWorldPartitionActorDescView& ActorDescView);
 	static FString GetFullActorName(const FWorldPartitionActorDescView& ActorDescView);
diff --git a/Engine/Source/Runtime/Engine/Public/WorldPartition/ErrorHandling/WorldPartitionStreamingGenerationLogErrorHandler.h b/Engine/Source/Runtime/Engine/Public/WorldPartition/ErrorHandling/WorldPartitionStreamingGenerationLogErrorHandler.h
--- a/Engine/Source/Runtime/Engine/Public/WorldPartition/ErrorHandling/WorldPartitionStreamingGenerationLogErrorHandler.h
+++ b/Engine/Source/Runtime/Engine/Public/WorldPartition/ErrorHandling/WorldPartitionStreamingGenerationLogErrorHandler.h
@@ -19,5 +19,6 @@ public:
 	virtual void OnDataLayerAssetConflict(const UDataLayerInstanceWithAsset* DataLayerInstance, const UDataLayerInstanceWithAsset* ConflictingDataLayerInstance) override;
 	virtual void OnActorNeedsResave(const FWorldPartitionActorDescView& ActorDescView) override;
 	virtual void OnLevelInstanceInvalidWorldAsset(const FWorldPartitionActorDescView& ActorDescView, FName WorldAsset, ELevelInstanceInvalidReason Reason) override;
+	virtual void OnActorFilterFailed(const FWorldPartitionActorDescView& ActorDescView) override;
 };
 #endif
diff --git a/Engine/Source/Runtime/Engine/Public/WorldPartition/ErrorHandling/WorldPartitionStreamingGenerationNullErrorHandler.h b/Engine/Source/Runtime/Engine/Public/WorldPartition/ErrorHandling/WorldPartitionStreamingGenerationNullErrorHandler.h
--- a/Engine/Source/Runtime/Engine/Public/WorldPartition/ErrorHandling/WorldPartitionStreamingGenerationNullErrorHandler.h
+++ b/Engine/Source/Runtime/Engine/Public/WorldPartition/ErrorHandling/WorldPartitionStreamingGenerationNullErrorHandler.h
@@ -19,5 +19,6 @@ public:
 	virtual void OnDataLayerAssetConflict(const UDataLayerInstanceWithAsset* DataLayerInstance, const UDataLayerInstanceWithAsset* ConflictingDataLayerInstance) override {};
 	virtual void OnActorNeedsResave(const FWorldPartitionActorDescView& ActorDescView) override {};
 	virtual void OnLevelInstanceInvalidWorldAsset(const FWorldPartitionActorDescView& ActorDescView, FName WorldAsset, ELevelInstanceInvalidReason Reason) override {};
+	virtual void OnActorFilterFailed(const FWorldPartitionActorDescView& ActorDescView) override {};
 };
 #endif
diff --git a/Engine/Source/Runtime/Engine/Public/WorldPartition/ErrorHandling/WorldPartitionStreamingGenerationProxyErrorHandler.h b/Engine/Source/Runtime/Engine/Public/WorldPartition/ErrorHandling/WorldPartitionStreamingGenerationProxyErrorHandler.h
--- a/Engine/Source/Runtime/Engine/Public/WorldPartition/ErrorHandling/WorldPartitionStreamingGenerationProxyErrorHandler.h
+++ b/Engine/Source/Runtime/Engine/Public/WorldPartition/ErrorHandling/WorldPartitionStreamingGenerationProxyErrorHandler.h
@@ -19,10 +19,11 @@ public:
 	virtual void OnInvalidReferenceLevelScriptDataLayers(const FWorldPartitionActorDescView& ActorDescView) override;
 	virtual void OnInvalidReferenceRuntimeGrid(const FWorldPartitionActorDescView& ActorDescView, const FWorldPartitionActorDescView& ReferenceActorDescView) override;
 	virtual void OnInvalidReferenceDataLayerAsset(const UDataLayerInstanceWithAsset* DataLayerInstance) override;
-	virtual void OnDataLayerHierarchyTypeMismatch(const UDataLayerInstance* DataLayerInstance, const UDataLayerInstance* Parent) override;;
-	virtual void OnDataLayerAssetConflict(const UDataLayerInstanceWithAsset* DataLayerInstance, const UDataLayerInstanceWithAsset* ConflictingDataLayerInstance) override;;
-	virtual void OnActorNeedsResave(const FWorldPartitionActorDescView& ActorDescView) override;;
-	virtual void OnLevelInstanceInvalidWorldAsset(const FWorldPartitionActorDescView& ActorDescView, FName WorldAsset, ELevelInstanceInvalidReason Reason) override;;
+	virtual void OnDataLayerHierarchyTypeMismatch(const UDataLayerInstance* DataLayerInstance, const UDataLayerInstance* Parent) override;
+	virtual void OnDataLayerAssetConflict(const UDataLayerInstanceWithAsset* DataLayerInstance, const UDataLayerInstanceWithAsset* ConflictingDataLayerInstance) override;
+	virtual void OnActorNeedsResave(const FWorldPartitionActorDescView& ActorDescView) override;
+	virtual void OnLevelInstanceInvalidWorldAsset(const FWorldPartitionActorDescView& ActorDescView, FName WorldAsset, ELevelInstanceInvalidReason Reason) override;
+	virtual void OnActorFilterFailed(const FWorldPartitionActorDescView& ActorDescView) override;
 
 protected:
 	IStreamingGenerationErrorHandler* InnerErrorHandler;
diff --git a/Engine/Source/Runtime/Engine/Public/WorldPartition/ErrorHandling/WorldPartitionStreamingGenerationTokenizedMessageErrorHandler.h b/Engine/Source/Runtime/Engine/Public/WorldPartition/ErrorHandling/WorldPartitionStreamingGenerationTokenizedMessageErrorHandler.h
--- a/Engine/Source/Runtime/Engine/Public/WorldPartition/ErrorHandling/WorldPartitionStreamingGenerationTokenizedMessageErrorHandler.h
+++ b/Engine/Source/Runtime/Engine/Public/WorldPartition/ErrorHandling/WorldPartitionStreamingGenerationTokenizedMessageErrorHandler.h
@@ -20,6 +20,7 @@ class ENGINE_API ITokenizedMessageErrorHandler : public IStreamingGenerationErro
 	virtual void OnDataLayerAssetConflict(const UDataLayerInstanceWithAsset* DataLayerInstance, const UDataLayerInstanceWithAsset* ConflictingDataLayerInstance) override;
 	virtual void OnActorNeedsResave(const FWorldPartitionActorDescView& ActorDescView) override;
 	virtual void OnLevelInstanceInvalidWorldAsset(const FWorldPartitionActorDescView& ActorDescView, FName WorldAsset, ELevelInstanceInvalidReason Reason) override;
+	virtual void OnActorFilterFailed(const FWorldPartitionActorDescView& ActorDescView) override;
 
 protected:
 	virtual void HandleTokenizedMessage(TSharedRef<FTokenizedMessage>&& ErrorMessage) = 0;
diff --git a/Engine/Source/Runtime/Engine/Public/WorldPartition/Filter/WorldPartitionActorFilter.h b/Engine/Source/Runtime/Engine/Public/WorldPartition/Filter/WorldPartitionActorFilter.h
new file mode 100644
--- /dev/null
+++ b/Engine/Source/Runtime/Engine/Public/WorldPartition/Filter/WorldPartitionActorFilter.h
@@ -0,0 +1,100 @@
+// Copyright Epic Games, Inc. All Rights Reserved.
+
+#pragma once
+
+#include "CoreMinimal.h"
+#include "Containers/Map.h"
+#include "Misc/Guid.h"
+#include "Misc/Optional.h"
+#include "UObject/SoftObjectPath.h"
+#include "WorldPartitionActorFilter.generated.h"
+
+USTRUCT()
+struct ENGINE_API FWorldPartitionActorFilter
+{
+	GENERATED_USTRUCT_BODY()
+
+	FWorldPartitionActorFilter() {}
+	~FWorldPartitionActorFilter();
+
+	FWorldPartitionActorFilter(const FWorldPartitionActorFilter& Other);
+	FWorldPartitionActorFilter(FWorldPartitionActorFilter&& Other);
+
+	FWorldPartitionActorFilter& operator=(const FWorldPartitionActorFilter& Other);
+	FWorldPartitionActorFilter& operator=(FWorldPartitionActorFilter&& Other);
+
+#if WITH_EDITORONLY_DATA
+	FWorldPartitionActorFilter(const FString& InDisplayName) : DisplayName(InDisplayName) {}
+	
+	void AddChildFilter(const FGuid& InGuid, FWorldPartitionActorFilter* InChildFilter);
+	void RemoveChildFilter(const FGuid& InGuid);
+	void ClearChildFilters();
+	
+	void Override(const FWorldPartitionActorFilter& Other);
+
+	const TMap<FGuid, FWorldPartitionActorFilter*>& GetChildFilters() const { return ChildFilters; }
+	FWorldPartitionActorFilter* GetParentFilter() const { return Parent; }
+
+	struct FDataLayerFilter
+	{
+		FDataLayerFilter() {}
+
+		FDataLayerFilter(const FString& InDisplayName, bool bInIncluded)
+			: bIncluded(bInIncluded)
+			, DisplayName(InDisplayName) {}
+
+		// True if DataLayer actors should be included
+		bool bIncluded;
+
+		// Transient
+		FString DisplayName;
+
+	};
+
+	// Transient
+	FString DisplayName;
+	// List of DataLayer Assets to Include or Exclude, missing DataLayer Assets will use default behavior
+	TMap<FSoftObjectPath, FDataLayerFilter> DataLayerFilters;
+
+	friend FArchive& operator<<(FArchive& Ar, FWorldPartitionActorFilter& Filter);
+	bool Serialize(FArchive& Ar);
+
+	// Operators.
+	bool operator==(const FWorldPartitionActorFilter& Other) const;
+
+	bool operator!=(const FWorldPartitionActorFilter& Other) const
+	{
+		return !(*this == Other);
+	}
+
+	// Needed for Copy/Paste/ResetToDefault
+	bool ExportTextItem(FString& ValueStr, FWorldPartitionActorFilter const& DefaultValue, UObject* Parent, int32 PortFlags, UObject* ExportRootScope) const;
+	bool ImportTextItem(const TCHAR*& Buffer, int32 PortFlags, UObject* Parent, FOutputDevice* ErrorText);
+
+	static void RequestFilterRefresh(bool bIsFromUserChange);
+
+	DECLARE_MULTICAST_DELEGATE(FOnWorldPartitionActorFilterChanged);
+	static FOnWorldPartitionActorFilterChanged& GetOnWorldPartitionActorFilterChanged() { return OnWorldPartitionActorFilterChanged; }
+private:
+	// Transient
+	FWorldPartitionActorFilter* Parent = nullptr;
+	// Map of FWorldPartitionActorFilters per Child Level Instance, recursive
+	TMap<FGuid, FWorldPartitionActorFilter*> ChildFilters;
+	// Static Event for when some filter changes
+	static FOnWorldPartitionActorFilterChanged OnWorldPartitionActorFilterChanged;
+#endif
+};
+
+#if WITH_EDITORONLY_DATA
+template<> struct TStructOpsTypeTraits<FWorldPartitionActorFilter> : public TStructOpsTypeTraitsBase2<FWorldPartitionActorFilter>
+{
+	enum
+	{
+		WithSerializer = true,
+		WithIdenticalViaEquality = true,
+		WithImportTextItem = true,
+		WithExportTextItem = true
+	};
+};
+#endif
+
diff --git a/Engine/Source/Runtime/Engine/Public/WorldPartition/LevelInstance/LevelInstanceActorDesc.h b/Engine/Source/Runtime/Engine/Public/WorldPartition/LevelInstance/LevelInstanceActorDesc.h
--- a/Engine/Source/Runtime/Engine/Public/WorldPartition/LevelInstance/LevelInstanceActorDesc.h
+++ b/Engine/Source/Runtime/Engine/Public/WorldPartition/LevelInstance/LevelInstanceActorDesc.h
@@ -9,6 +9,7 @@
 #include "UObject/ObjectPtr.h"
 #include "Containers/Map.h"
 #include "WorldPartition/WorldPartitionActorDesc.h"
+#include "WorldPartition/Filter/WorldPartitionActorFilter.h"
 
 class ULevelInstanceSubsystem;
 class UActorDescContainer;
@@ -26,8 +27,9 @@ public:
 
 	virtual bool IsContainerInstance() const override;
 	virtual FName GetLevelPackage() const override { return LevelPackage; }
-	virtual bool GetContainerInstance(FContainerInstance& OutContainerInstance) const override;
+	virtual bool GetContainerInstance(FContainerInstance& OutContainerInstance, bool bInBuildFilter) const override;
 	virtual void CheckForErrors(IStreamingGenerationErrorHandler* ErrorHandler) const override;
+	const FWorldPartitionActorFilter& GetFilter() const { return Filter; }
 
 protected:
 	virtual void Init(const AActor* InActor) override;
@@ -44,6 +46,8 @@ protected:
 	TWeakObjectPtr<UActorDescContainer> LevelInstanceContainer;
 	TWeakObjectPtr<UWorld> LevelInstanceContainerWorldContext;
 
+	FWorldPartitionActorFilter Filter;
+
 private:
 	void RegisterContainerInstance(UWorld* InWorldContext);
 	void UnregisterContainerInstance();
diff --git a/Engine/Source/Runtime/Engine/Public/WorldPartition/WorldPartitionActorDesc.h b/Engine/Source/Runtime/Engine/Public/WorldPartition/WorldPartitionActorDesc.h
--- a/Engine/Source/Runtime/Engine/Public/WorldPartition/WorldPartitionActorDesc.h
+++ b/Engine/Source/Runtime/Engine/Public/WorldPartition/WorldPartitionActorDesc.h
@@ -84,7 +84,7 @@ public:
 		const UActorDescContainer* Container = nullptr;
 		FTransform Transform = FTransform::Identity;
 		EContainerClusterMode ClusterMode;
-
+		TMap<FActorContainerID, TSet<FGuid>> FilteredActors;
 	private:
 		FActorContainerID ContainerID;
 	};
@@ -153,7 +153,7 @@ PRAGMA_ENABLE_DEPRECATION_WARNINGS
 
 	virtual bool IsContainerInstance() const { return false; }
 	virtual FName GetLevelPackage() const { return NAME_None; }
-	virtual bool GetContainerInstance(FContainerInstance& OutContainerInstance) const { return false; }
+	virtual bool GetContainerInstance(FContainerInstance& OutContainerInstance, bool bInBuildFilter = false) const { return false; }
 
 	FGuid GetContentBundleGuid() const;
 
diff --git a/Engine/Source/Runtime/Engine/Public/WorldPartition/WorldPartitionActorDescView.h b/Engine/Source/Runtime/Engine/Public/WorldPartition/WorldPartitionActorDescView.h
--- a/Engine/Source/Runtime/Engine/Public/WorldPartition/WorldPartitionActorDescView.h
+++ b/Engine/Source/Runtime/Engine/Public/WorldPartition/WorldPartitionActorDescView.h
@@ -65,7 +65,7 @@ PRAGMA_ENABLE_DEPRECATION_WARNINGS
 
 	bool IsContainerInstance() const;
 	FName GetLevelPackage() const;
-	bool GetContainerInstance(FWorldPartitionActorDesc::FContainerInstance& OutContainerInstance) const;
+	bool GetContainerInstance(FWorldPartitionActorDesc::FContainerInstance& OutContainerInstance, bool bInBuildFilter) const;
 
 	void CheckForErrors(IStreamingGenerationErrorHandler* ErrorHandler) const;
 
