Candidate fix for incremental build error (Issue 273140) related to my change to EditorOnlyData naming bug fixing.

diff --git a/Engine/Source/Runtime/Engine/Private/Materials/MaterialInterface.cpp b/Engine/Source/Runtime/Engine/Private/Materials/MaterialInterface.cpp
--- a/Engine/Source/Runtime/Engine/Private/Materials/MaterialInterface.cpp
+++ b/Engine/Source/Runtime/Engine/Private/Materials/MaterialInterface.cpp
@@ -1542,7 +1542,6 @@ void UMaterialInterface::PreSave(FObjectPreSaveContext ObjectSaveContext)
 		if (EditorOnlyExisting && EditorOnlyExisting != EditorOnly)
 		{
 			EditorOnlyExisting->Rename(*(EditorOnlyDataName + "_Unused"), nullptr, REN_NonTransactional | REN_DontCreateRedirectors);
-			EditorOnlyExisting->MarkAsGarbage();
 		}
 		if (EditorOnly)
 		{
