World Partition - HLOD: When performing HLOD warmup, make sure ULevelStreaming::CanMakeVisible() / ULevelStreaming::CanMakeInvisible() are always called, regardless of the warmup progress.
* Avoiding the calls (unintentionally) was messing up the client/server visibility transactions.

diff --git a/Engine/Source/Runtime/Engine/Classes/Engine/LevelStreaming.h b/Engine/Source/Runtime/Engine/Classes/Engine/LevelStreaming.h
--- a/Engine/Source/Runtime/Engine/Classes/Engine/LevelStreaming.h
+++ b/Engine/Source/Runtime/Engine/Classes/Engine/LevelStreaming.h
@@ -352,10 +352,17 @@ protected:
 	/** Updates the current state of the streaming level and notifies any observers. */
 	void SetCurrentState(ELevelStreamingState NewState);
 
+	/** 
+	 * Prepare the transition to the new state.
+	 * @return true if the state change can be performed, false otherwise.
+	 **/
+	virtual bool RequestVisibilityChange(bool bVisible);
+
 	/** Returns whether the streaming level can make visible (can call AddToWorld). */
-	virtual bool CanMakeVisible();
+	bool CanMakeVisible();
+
 	/** Returns whether the streaming level can make invisible (can call RemoveFromWorld). */
-	virtual bool CanMakeInvisible();
+	bool CanMakeInvisible();
 
 private:
 
diff --git a/Engine/Source/Runtime/Engine/Private/LevelStreaming.cpp b/Engine/Source/Runtime/Engine/Private/LevelStreaming.cpp
--- a/Engine/Source/Runtime/Engine/Private/LevelStreaming.cpp
+++ b/Engine/Source/Runtime/Engine/Private/LevelStreaming.cpp
@@ -751,8 +751,15 @@ bool ULevelStreaming::ShouldWaitForServerAckBeforeChangingVisibilityState(ENetLe
 	return false;
 };
 
+bool ULevelStreaming::RequestVisibilityChange(bool bVisible)
+{
+	return true;
+}
+
 bool ULevelStreaming::CanMakeInvisible()
 {
+	const bool bCanMakeInvisible = RequestVisibilityChange(false);
+
 	if (ShouldClientUseMakingInvisibleTransactionRequest())
 	{
 		if (ShouldWaitForServerAckBeforeChangingVisibilityState(ENetLevelVisibilityRequest::MakingInvisible, bShouldBeVisible))
@@ -760,11 +767,14 @@ bool ULevelStreaming::CanMakeInvisible()
 			return false;
 		}
 	}
-	return true;
+
+	return bCanMakeInvisible;
 }
 
 bool ULevelStreaming::CanMakeVisible()
 {
+	const bool bCanMakeVisible = RequestVisibilityChange(true);
+
 	if (ShouldClientUseMakingVisibleTransactionRequest())
 	{
 		if (ShouldWaitForServerAckBeforeChangingVisibilityState(ENetLevelVisibilityRequest::MakingVisible, true))
@@ -780,7 +790,8 @@ bool ULevelStreaming::CanMakeVisible()
 			return false;
 		}
 	}
-	return true;
+
+	return bCanMakeVisible;
 }
 
 void ULevelStreaming::UpdateNetVisibilityTransactionState(bool bInShouldBeVisible, FNetLevelVisibilityTransactionId TransactionId)
diff --git a/Engine/Source/Runtime/Engine/Private/WorldPartition/WorldPartitionLevelStreamingDynamic.cpp b/Engine/Source/Runtime/Engine/Private/WorldPartition/WorldPartitionLevelStreamingDynamic.cpp
--- a/Engine/Source/Runtime/Engine/Private/WorldPartition/WorldPartitionLevelStreamingDynamic.cpp
+++ b/Engine/Source/Runtime/Engine/Private/WorldPartition/WorldPartitionLevelStreamingDynamic.cpp
@@ -667,16 +667,11 @@ bool UWorldPartitionLevelStreamingDynamic::CanChangeVisibility(bool bMakeVisible
 	return true;
 }
 
-bool UWorldPartitionLevelStreamingDynamic::CanMakeVisible()
+bool UWorldPartitionLevelStreamingDynamic::RequestVisibilityChange(bool bVisible)
 {
-	const bool bCanMakeVisible = CanChangeVisibility(true);
-	return bCanMakeVisible && Super::CanMakeVisible();
-}
-
-bool UWorldPartitionLevelStreamingDynamic::CanMakeInvisible()
-{
-	const bool bCanMakeInvisible = CanChangeVisibility(false);
-	return bCanMakeInvisible && Super::CanMakeInvisible();
+	const bool bSuperAllowsVisibilityChange = Super::RequestVisibilityChange(bVisible);
+	const bool bAllowVisibilityChange = CanChangeVisibility(bVisible);
+	return bSuperAllowsVisibilityChange && bAllowVisibilityChange;
 }
 
 bool UWorldPartitionLevelStreamingDynamic::ShouldBlockOnUnload() const
diff --git a/Engine/Source/Runtime/Engine/Public/WorldPartition/WorldPartitionLevelStreamingDynamic.h b/Engine/Source/Runtime/Engine/Public/WorldPartition/WorldPartitionLevelStreamingDynamic.h
--- a/Engine/Source/Runtime/Engine/Public/WorldPartition/WorldPartitionLevelStreamingDynamic.h
+++ b/Engine/Source/Runtime/Engine/Public/WorldPartition/WorldPartitionLevelStreamingDynamic.h
@@ -37,8 +37,7 @@ class ENGINE_API UWorldPartitionLevelStreamingDynamic : public ULevelStreamingDy
 	virtual bool ShouldBlockOnUnload() const override;
 	virtual bool ShouldRequireFullVisibilityToRender() const override { return true; }
 
-	virtual bool CanMakeVisible() override;
-	virtual bool CanMakeInvisible() override;
+	virtual bool RequestVisibilityChange(bool bVisible) override;
 
 #if !WITH_EDITOR
 	virtual void PostLoad();
