Implemented support for bilinear filtering of SparseVolumeTextures. Also added support for texture addressing modes (wrap/clamp/mirror). The addressing mode can be selected on the material sample node or on the SVT asset itself; similar to a regular texture.

#rnx

diff --git a/Engine/Plugins/FX/Niagara/Shaders/Private/NiagaraDataInterfaceSparseVolumeTextureTemplate.ush b/Engine/Plugins/FX/Niagara/Shaders/Private/NiagaraDataInterfaceSparseVolumeTextureTemplate.ush
--- a/Engine/Plugins/FX/Niagara/Shaders/Private/NiagaraDataInterfaceSparseVolumeTextureTemplate.ush
+++ b/Engine/Plugins/FX/Niagara/Shaders/Private/NiagaraDataInterfaceSparseVolumeTextureTemplate.ush
@@ -2,6 +2,7 @@
 
 #include "/Engine/Private/SparseVolumeTexture/SparseVolumeTextureCommon.ush"
 
+SamplerState	{ParameterName}_TileDataTextureSampler;
 Texture3D<uint>	{ParameterName}_PageTableTexture;
 Texture3D		{ParameterName}_PhysicalTileDataATexture;
 Texture3D		{ParameterName}_PhysicalTileDataBTexture;
@@ -22,9 +23,9 @@ void LoadSparseVolumeTexture_{ParameterName}(in int TexelX, in int TexelY, in in
 void SampleSparseVolumeTexture_{ParameterName}(in float3 UVW, in int MipLevel, out float4 OutAttributesA, out float4 OutAttributesB)
 {
 	const FSparseVolumeTextureUniforms Uniforms = SparseVolumeTextureUnpackUniforms({ParameterName}_PackedUniforms0, {ParameterName}_PackedUniforms1);
-	const float3 VoxelCoord = SparseVolumeTextureSamplePageTable({ParameterName}_PageTableTexture, Uniforms, UVW, MipLevel);
-	OutAttributesA = {ParameterName}_PhysicalTileDataATexture.Load(int4(int3(VoxelCoord), 0));
-	OutAttributesB = {ParameterName}_PhysicalTileDataBTexture.Load(int4(int3(VoxelCoord), 0));
+	const float3 VoxelUVW = SparseVolumeTextureSamplePageTable({ParameterName}_PageTableTexture, Uniforms, UVW, SVTADDRESSMODE_WRAP, SVTADDRESSMODE_WRAP, SVTADDRESSMODE_WRAP, MipLevel);
+	OutAttributesA = SparseVolumeTextureSamplePhysicalTileData({ParameterName}_PhysicalTileDataATexture, {ParameterName}_PhysicalTileDataBTexture, {ParameterName}_TileDataTextureSampler, VoxelUVW, 0);
+	OutAttributesB = SparseVolumeTextureSamplePhysicalTileData({ParameterName}_PhysicalTileDataATexture, {ParameterName}_PhysicalTileDataBTexture, {ParameterName}_TileDataTextureSampler, VoxelUVW, 1);
 }
 
 void GetTextureDimensions_{ParameterName}(int MipLevel, out float3 OutValue)
diff --git a/Engine/Plugins/FX/Niagara/Source/Niagara/Classes/NiagaraDataInterfaceSparseVolumeTexture.h b/Engine/Plugins/FX/Niagara/Source/Niagara/Classes/NiagaraDataInterfaceSparseVolumeTexture.h
--- a/Engine/Plugins/FX/Niagara/Source/Niagara/Classes/NiagaraDataInterfaceSparseVolumeTexture.h
+++ b/Engine/Plugins/FX/Niagara/Source/Niagara/Classes/NiagaraDataInterfaceSparseVolumeTexture.h
@@ -16,6 +16,7 @@ class NIAGARA_API UNiagaraDataInterfaceSparseVolumeTexture : public UNiagaraData
 	GENERATED_UCLASS_BODY()
 
 	BEGIN_SHADER_PARAMETER_STRUCT(FShaderParameters, )
+		SHADER_PARAMETER_SAMPLER(SamplerState, TileDataTextureSampler)
 		SHADER_PARAMETER_TEXTURE(Texture3D<uint>, PageTableTexture)
 		SHADER_PARAMETER_TEXTURE(Texture3D, PhysicalTileDataATexture)
 		SHADER_PARAMETER_TEXTURE(Texture3D, PhysicalTileDataBTexture)
diff --git a/Engine/Plugins/FX/Niagara/Source/Niagara/Private/NiagaraDataInterfaceSparseVolumeTexture.cpp b/Engine/Plugins/FX/Niagara/Source/Niagara/Private/NiagaraDataInterfaceSparseVolumeTexture.cpp
--- a/Engine/Plugins/FX/Niagara/Source/Niagara/Private/NiagaraDataInterfaceSparseVolumeTexture.cpp
+++ b/Engine/Plugins/FX/Niagara/Source/Niagara/Private/NiagaraDataInterfaceSparseVolumeTexture.cpp
@@ -308,6 +308,7 @@ void UNiagaraDataInterfaceSparseVolumeTexture::SetShaderParameters(const FNiagar
 	FNDISparseVolumeTextureInstanceData_RenderThread* RTInstanceData = TextureProxy.InstanceData_RT.Find(Context.GetSystemInstanceID());
 
 	FShaderParameters* Parameters = Context.GetParameterNestedStruct<FShaderParameters>();
+	Parameters->TileDataTextureSampler = TStaticSamplerState<SF_Bilinear, AM_Clamp, AM_Clamp, AM_Clamp>::GetRHI();
 	if (RTInstanceData && RTInstanceData->PageTableTextureRHI.IsValid() && (RTInstanceData->PhysicalTileDataATextureRHI.IsValid() || RTInstanceData->PhysicalTileDataBTextureRHI.IsValid()))
 	{
 		Parameters->PageTableTexture = RTInstanceData->PageTableTextureRHI;
diff --git a/Engine/Shaders/Private/SparseVolumeTexture/SparseVolumeTextureCommon.ush b/Engine/Shaders/Private/SparseVolumeTexture/SparseVolumeTextureCommon.ush
--- a/Engine/Shaders/Private/SparseVolumeTexture/SparseVolumeTextureCommon.ush
+++ b/Engine/Shaders/Private/SparseVolumeTexture/SparseVolumeTextureCommon.ush
@@ -2,41 +2,85 @@
 
 #pragma once
 
+#define SVTADDRESSMODE_CLAMP 0u
+#define SVTADDRESSMODE_WRAP 1u
+#define SVTADDRESSMODE_MIRROR 2u
+
+float SparseVolumeTextureApplyAddressModeMirror(float v)
+{
+	float t = frac(v * 0.5f) * 2.0f;
+	return 1.0f - abs(t - 1.0f);
+}
+
+float SparseVolumeTextureApplyAddressMode(float v, uint AddressMode)
+{
+	// For CLAMP address mode, can't clamp to 1.0f, otherwise 'int(UVW * VolumePageResolution)' might overflow page table bounds by 1
+	// Instead, clamp to slightly below 1, this ensures that when rounded down to int, above value will be at most 'PageTableResolution - 1'
+	// The actual texel we clamp to doesn't matter too much for sampling physical texture, since we have borders around the physical pages
+	// Just need to make sure we don't clamp too far and chop off valid texels at the edge of texture
+	const float MaxTextureSize = 65536.0f;
+
+	switch (AddressMode)
+	{
+	case SVTADDRESSMODE_WRAP: return frac(v);
+	case SVTADDRESSMODE_MIRROR: return SparseVolumeTextureApplyAddressModeMirror(v);
+	default: return clamp(v, 0.0f, 1.0f - (1.0f / MaxTextureSize));
+	}
+}
+
+float3 SparseVolumeTextureApplyAddressMode(float3 UVW, uint AddressU, uint AddressV, uint AddressW)
+{
+	return float3(
+		SparseVolumeTextureApplyAddressMode(UVW.x, AddressU),
+		SparseVolumeTextureApplyAddressMode(UVW.y, AddressV),
+		SparseVolumeTextureApplyAddressMode(UVW.z, AddressW));
+}
+
 struct FSparseVolumeTextureUniforms
 {
 	float3			VolumePageResolution;
 	float3			PageTableOffset;
+	float3			TileDataTexelSize;
 	int             MinMipLevel;
 	int             MaxMipLevel;
 	float			TileSize;
+	float			BorderSize;
 };
 
 FSparseVolumeTextureUniforms SparseVolumeTextureUnpackUniforms(const uint4 Packed0, const uint4 Packed1)
 {
 	FSparseVolumeTextureUniforms Result;
 	Result.VolumePageResolution = asfloat(Packed0.xyz);
-	Result.PageTableOffset = asfloat(Packed1.xyz);
-	Result.MinMipLevel = int(Packed0.w & 0xFFu);
-	Result.MaxMipLevel = int((Packed0.w >> 8u) & 0xFFu);
-	Result.TileSize = asfloat(Packed1.w);
+	Result.PageTableOffset.x = float(Packed0.w & 0x7FFu);
+	Result.PageTableOffset.y = float((Packed0.w >> 11) & 0x7FF);
+	Result.PageTableOffset.z = float((Packed0.w >> 22) & 0x3FF);
+	Result.TileDataTexelSize = asfloat(Packed1.xyz);
+	Result.MinMipLevel = int(Packed1.w & 0xFFu);
+	Result.MaxMipLevel = int((Packed1.w >> 8u) & 0xFFu);
+	Result.TileSize = float((Packed1.w >> 16u) & 0xFFu);
+	Result.BorderSize = float((Packed1.w >> 24u) & 0xFFu);
 
 	return Result;
 }
 
-float3 SparseVolumeTextureGetVoxelCoord(const uint PackedPhysicalTileCoord, const float TileSize, const float3 PageTableCoord)
+float3 SparseVolumeTextureGetVoxelCoord(const uint PackedPhysicalTileCoord, const float TileSize, const float BorderSize, const float3 PageTableCoord)
 {
 	const int3 PhysicalTileCoord = int3(
 		PackedPhysicalTileCoord & 0x7FF,
 		(PackedPhysicalTileCoord >> 11) & 0x7FF,
 		(PackedPhysicalTileCoord >> 22) & 0x3FF);
 
-	const float3 VoxelCoord = (float3(PhysicalTileCoord) + frac(PageTableCoord)) * TileSize;
+	const float PaddedTileSize = BorderSize * 2.0f + TileSize;
+	const float3 VoxelCoord = float3(PhysicalTileCoord) * PaddedTileSize + (frac(PageTableCoord) * TileSize + BorderSize);
 	return VoxelCoord;
 }
 
-float3 SparseVolumeTextureSamplePageTable(Texture3D<uint> PageTable, const FSparseVolumeTextureUniforms Uniforms, const float3 UVW, int MipLevel = 0)
+float3 SparseVolumeTextureSamplePageTable(Texture3D<uint> PageTable, FSparseVolumeTextureUniforms Uniforms, float3 UVW, uint AddressU, uint AddressV, uint AddressW, int MipLevel = 0)
 {
+	// Apply address mode to UVW and clamp mip level to resident levels
+	UVW = SparseVolumeTextureApplyAddressMode(UVW, AddressU, AddressV, AddressW);
 	MipLevel = clamp(MipLevel, Uniforms.MinMipLevel, Uniforms.MaxMipLevel);
+
 	const float RcpMipLevelFactor = rcp(float(1u << (uint)MipLevel));
 	const float3 VolumePageCoord = UVW * Uniforms.VolumePageResolution;
 	const float3 MipPageTableOffset = floor(Uniforms.PageTableOffset * RcpMipLevelFactor);
@@ -44,11 +88,12 @@ float3 SparseVolumeTextureSamplePageTable(Texture3D<uint> PageTable, const FSpar
 
 	const uint PackedPhysicalTileCoord = PageTable.Load(int4(int3(PageTableCoord), (MipLevel - Uniforms.MinMipLevel))).x;
 
-	const float3 VoxelCoord = SparseVolumeTextureGetVoxelCoord(PackedPhysicalTileCoord, Uniforms.TileSize, PageTableCoord);
-	return VoxelCoord;
+	const float3 VoxelCoord = SparseVolumeTextureGetVoxelCoord(PackedPhysicalTileCoord, Uniforms.TileSize, Uniforms.BorderSize, PageTableCoord);
+	const float3 VoxelUVW = VoxelCoord * Uniforms.TileDataTexelSize;
+	return VoxelUVW;
 }
 
-int3 SparseVolumeTextureLoadPageTable(Texture3D<uint> PageTable, const FSparseVolumeTextureUniforms Uniforms, const int3 TexelCoord, int MipLevel = 0)
+int3 SparseVolumeTextureLoadPageTable(Texture3D<uint> PageTable, FSparseVolumeTextureUniforms Uniforms, int3 TexelCoord, int MipLevel = 0)
 {
 	if (MipLevel < Uniforms.MinMipLevel || MipLevel > Uniforms.MaxMipLevel)
 	{
@@ -61,16 +106,16 @@ int3 SparseVolumeTextureLoadPageTable(Texture3D<uint> PageTable, const FSparseVo
 
 	const uint PackedPhysicalTileCoord = PageTable.Load(int4(int3(PageTableCoord), (MipLevel - Uniforms.MinMipLevel))).x;
 
-	const float3 VoxelCoord = SparseVolumeTextureGetVoxelCoord(PackedPhysicalTileCoord, Uniforms.TileSize, PageTableCoord);
+	const float3 VoxelCoord = SparseVolumeTextureGetVoxelCoord(PackedPhysicalTileCoord, Uniforms.TileSize, Uniforms.BorderSize, PageTableCoord);
 	return int3(VoxelCoord);
 }
 
-float4 SparseVolumeTextureSamplePhysicalTileData(Texture3D PhysicalTileDataA, Texture3D PhysicalTileDataB, const float3 VoxelCoord, const int PhysicalTileDataIndex)
+float4 SparseVolumeTextureSamplePhysicalTileData(Texture3D PhysicalTileDataA, Texture3D PhysicalTileDataB, SamplerState TileDataSampler, float3 VoxelUVW, int PhysicalTileDataIndex)
 {
 	switch (PhysicalTileDataIndex)
 	{
-	case 0: return PhysicalTileDataA.Load(int4(int3(VoxelCoord), 0));
-	case 1: return PhysicalTileDataB.Load(int4(int3(VoxelCoord), 0));
+	case 0: return PhysicalTileDataA.SampleLevel(TileDataSampler, VoxelUVW, 0.0f);
+	case 1: return PhysicalTileDataB.SampleLevel(TileDataSampler, VoxelUVW, 0.0f);
 	default: return 0.0f;
 	}
 }
diff --git a/Engine/Shaders/Private/SparseVolumeTexture/VisualizeSparseVolumeTexture.usf b/Engine/Shaders/Private/SparseVolumeTexture/VisualizeSparseVolumeTexture.usf
--- a/Engine/Shaders/Private/SparseVolumeTexture/VisualizeSparseVolumeTexture.usf
+++ b/Engine/Shaders/Private/SparseVolumeTexture/VisualizeSparseVolumeTexture.usf
@@ -41,6 +41,7 @@ bool Slabs(float3 p0, float3 p1, float3 rayOrigin, float3 invRaydir, out float o
 	return maxtmin <= mintmax;
 }
 
+SamplerState TileDataTextureSampler;
 Texture3D<uint> SparseVolumeTexturePageTable;
 Texture3D<float4> SparseVolumeTextureA;
 Texture3D<float4> SparseVolumeTextureB;
@@ -138,11 +139,12 @@ void VisualizeSparseVolumeTexturePS(
 		while (P.x >= 0 && P.y >= 0 && P.z >= 0 && P.x <= VolumeResolution.x && P.y <= VolumeResolution.y && P.z <= VolumeResolution.z)
 		{
 			const float3 UVW = P * RcpVolumeResolution;
-			const float3 VoxelCoord = SparseVolumeTextureSamplePageTable(SparseVolumeTexturePageTable, SVTUniforms, UVW, MipLevel);
-			if (any(VoxelCoord >= SVTUniforms.TileSize)) // skip on null tile
+			const float3 VoxelUVW = SparseVolumeTextureSamplePageTable(SparseVolumeTexturePageTable, SVTUniforms, UVW, SVTADDRESSMODE_CLAMP, SVTADDRESSMODE_CLAMP, SVTADDRESSMODE_CLAMP, MipLevel);
+			const float3 UVWTileSize = (SVTUniforms.TileSize + 2.0f * SVTUniforms.BorderSize) * SVTUniforms.TileDataTexelSize;
+			if (any(VoxelUVW >= UVWTileSize)) // skip on null tile
 			{
 				const int PhysicalTileDataIndex = ComponentToVisualize < 4 ? 0 : 1;
-				const float4 VoxelData = SparseVolumeTextureSamplePhysicalTileData(SparseVolumeTextureA, SparseVolumeTextureB, VoxelCoord, PhysicalTileDataIndex);
+				const float4 VoxelData = SparseVolumeTextureSamplePhysicalTileData(SparseVolumeTextureA, SparseVolumeTextureB, TileDataTextureSampler, VoxelUVW, PhysicalTileDataIndex);
 
 				float Density = 0.0f;
 				if (ComponentToVisualize < 4)
diff --git a/Engine/Source/Runtime/Engine/Classes/Materials/MaterialExpressionSparseVolumeTextureSample.h b/Engine/Source/Runtime/Engine/Classes/Materials/MaterialExpressionSparseVolumeTextureSample.h
--- a/Engine/Source/Runtime/Engine/Classes/Materials/MaterialExpressionSparseVolumeTextureSample.h
+++ b/Engine/Source/Runtime/Engine/Classes/Materials/MaterialExpressionSparseVolumeTextureSample.h
@@ -29,6 +29,15 @@ class ENGINE_API UMaterialExpressionSparseVolumeTextureSample : public UMaterial
 	UPROPERTY(meta = (RequiredInput = "false", ToolTip = "Defaults to 0 if not specified"))
 	FExpressionInput MipLevel;
 
+	/**
+	 * Controls where the sampler for this texture lookup will come from.
+	 * Choose 'from texture asset' to make use of the USparseVolumeTexture addressing settings,
+	 * Otherwise use one of the global samplers, which will not consume a sampler slot.
+	 * This allows materials to use more than 16 unique textures on SM5 platforms.
+	 */
+	UPROPERTY(EditAnywhere, Category = MaterialExpressionTextureSample, Meta = (ShowAsInputPin = "Advanced"))
+	TEnumAsByte<ESamplerSourceMode> SamplerSource;
+
 protected:
 
 #if WITH_EDITOR
diff --git a/Engine/Source/Runtime/Engine/Classes/SparseVolumeTexture/SparseVolumeTexture.h b/Engine/Source/Runtime/Engine/Classes/SparseVolumeTexture/SparseVolumeTexture.h
--- a/Engine/Source/Runtime/Engine/Classes/SparseVolumeTexture/SparseVolumeTexture.h
+++ b/Engine/Source/Runtime/Engine/Classes/SparseVolumeTexture/SparseVolumeTexture.h
@@ -10,13 +10,16 @@
 #include "Containers/Array.h"
 #include "UnrealClient.h"
 #include "UObject/ObjectSaveContext.h"
+#include "Engine/TextureDefines.h"
 
 #include "SparseVolumeTexture.generated.h"
 
 namespace UE { namespace Shader	{ enum class EValueType : uint8; } }
 namespace UE { namespace DerivedData { class FRequestOwner; } }
 
-#define SPARSE_VOLUME_TILE_RES	16
+#define SPARSE_VOLUME_TILE_RES 16
+#define SPARSE_VOLUME_TILE_BORDER 1
+#define SPARSE_VOLUME_TILE_RES_PADDED (SPARSE_VOLUME_TILE_RES + 2 * SPARSE_VOLUME_TILE_BORDER)
 
 struct FSparseVolumeTextureData;
 class FSparseVolumeTextureSceneProxy;
@@ -53,7 +56,7 @@ struct ENGINE_API FSparseVolumeTextureFrame
 
 	FSparseVolumeTextureFrame();
 	virtual ~FSparseVolumeTextureFrame();
-	bool BuildDerivedData(FSparseVolumeTextureData* OutMippedTextureData);
+	bool BuildDerivedData(const FIntVector3& VolumeResolution, TextureAddress AddressX, TextureAddress AddressY, TextureAddress AddressZ, FSparseVolumeTextureData* OutMippedTextureData);
 	void Serialize(FArchive& Ar, class UStreamableSparseVolumeTexture* Owner, int32 FrameIndex);
 };
 
@@ -89,8 +92,15 @@ public:
 
 	UFUNCTION(BlueprintCallable, Category = "Texture")
 	virtual int32 GetNumFrames() const { return 0; }
+
+	UFUNCTION(BlueprintCallable, Category = "Texture")
 	virtual int32 GetNumMipLevels() const { return 0; }
+
 	virtual FIntVector GetVolumeResolution() const { return FIntVector(); }
+	virtual TextureAddress GetTextureAddressX() const { return TA_Wrap; }
+	virtual TextureAddress GetTextureAddressY() const { return TA_Wrap; }
+	virtual TextureAddress GetTextureAddressZ() const { return TA_Wrap; }
+
 	virtual const FSparseVolumeTextureSceneProxy* GetSparseVolumeTextureSceneProxy() const { return nullptr; }
 
 	/** Getter for the shader uniform parameters with index as ESparseVolumeTextureShaderUniform. */
@@ -116,12 +126,24 @@ class UStreamableSparseVolumeTexture : public USparseVolumeTexture
 
 public:
 
-	UPROPERTY(VisibleAnywhere, Category = "Rendering")
+	UPROPERTY(VisibleAnywhere, Category = "Texture")
 	FIntVector VolumeResolution;
 
-	UPROPERTY(VisibleAnywhere, Category = "Rendering")
+	UPROPERTY(VisibleAnywhere, Category = "Texture")
 	int32 NumMipLevels;
 
+	/** The addressing mode to use for the X axis.								*/
+	UPROPERTY(EditAnywhere, BlueprintReadWrite, Category = "Texture", meta = (DisplayName = "X-axis Tiling Method"), AssetRegistrySearchable, AdvancedDisplay)
+	TEnumAsByte<enum TextureAddress> AddressX;
+
+	/** The addressing mode to use for the Y axis.								*/
+	UPROPERTY(EditAnywhere, BlueprintReadWrite, Category = "Texture", meta = (DisplayName = "Y-axis Tiling Method"), AssetRegistrySearchable, AdvancedDisplay)
+	TEnumAsByte<enum TextureAddress> AddressY;
+
+	/** The addressing mode to use for the Z axis.								*/
+	UPROPERTY(EditAnywhere, BlueprintReadWrite, Category = "Texture", meta = (DisplayName = "Z-axis Tiling Method"), AssetRegistrySearchable, AdvancedDisplay)
+	TEnumAsByte<enum TextureAddress> AddressZ;
+
 	UStreamableSparseVolumeTexture();
 	virtual ~UStreamableSparseVolumeTexture() = default;
 
@@ -137,18 +159,33 @@ public:
 	//~ End UObject Interface.
 
 	//~ Begin USparseVolumeTexture Interface.
-	int32 GetNumFrames() const override { return Frames.Num(); }
-	int32 GetNumMipLevels() const override { return NumMipLevels; }
-	FIntVector GetVolumeResolution() const override { return VolumeResolution; };
-	const FSparseVolumeTextureSceneProxy* GetSparseVolumeTextureSceneProxy() const override { return GetStreamedFrameProxyOrFallback(0 /*FrameIndex*/, 0 /*MipLevel*/); };
+	virtual int32 GetNumFrames() const override { return Frames.Num(); }
+	virtual int32 GetNumMipLevels() const override { return NumMipLevels; }
+	virtual FIntVector GetVolumeResolution() const override { return VolumeResolution; };
+	virtual TextureAddress GetTextureAddressX() const override { return AddressX; }
+	virtual TextureAddress GetTextureAddressY() const override { return AddressY; }
+	virtual TextureAddress GetTextureAddressZ() const override { return AddressZ; }
+	virtual const FSparseVolumeTextureSceneProxy* GetSparseVolumeTextureSceneProxy() const override { return GetStreamedFrameProxyOrFallback(0 /*FrameIndex*/, 0 /*MipLevel*/); };
 	//~ End USparseVolumeTexture Interface.
 
 	const FSparseVolumeTextureSceneProxy* GetStreamedFrameProxyOrFallback(int32 FrameIndex, int32 MipLevel) const;
 	TArrayView<const FSparseVolumeTextureFrame> GetFrames() const;
 
 protected:
+
 	TArray<FSparseVolumeTextureFrame> Frames;
 
+#if WITH_EDITOR
+	enum class ENotifyMaterialsEffectOnShaders
+	{
+		Default,
+		DoesNotInvalidate
+	};
+
+	/** Notify any loaded material instances that the texture has changed. */
+	void NotifyMaterials(const ENotifyMaterialsEffectOnShaders EffectOnShaders = ENotifyMaterialsEffectOnShaders::Default);
+#endif // WITH_EDITOR
+
 	void GenerateOrLoadDDCRuntimeDataAndCreateSceneProxy();
 	void GenerateOrLoadDDCRuntimeDataForFrame(FSparseVolumeTextureFrame& Frame, UE::DerivedData::FRequestOwner& DDCRequestOwner);
 };
@@ -211,17 +248,23 @@ public:
 
 	static USparseVolumeTextureFrame* CreateFrame(USparseVolumeTexture* Texture, int32 FrameIndex, int32 MipLevel);
 
-	void Initialize(const FSparseVolumeTextureSceneProxy* InSceneProxy, const FIntVector& InVolumeResolution);
+	void Initialize(const FSparseVolumeTextureSceneProxy* InSceneProxy, const FIntVector& InVolumeResolution, TextureAddress InAddressX, TextureAddress InAddressY, TextureAddress InAddressZ);
 
 	//~ Begin USparseVolumeTexture Interface.
-	int32 GetNumFrames() const override { return 1; }
-	int32 GetNumMipLevels() const override { return 1; }
-	FIntVector GetVolumeResolution() const override { return VolumeResolution; };
-	const FSparseVolumeTextureSceneProxy* GetSparseVolumeTextureSceneProxy() const override { return SceneProxy; };
+	virtual int32 GetNumFrames() const override { return 1; }
+	virtual int32 GetNumMipLevels() const override { return 1; }
+	virtual FIntVector GetVolumeResolution() const override { return VolumeResolution; };
+	virtual TextureAddress GetTextureAddressX() const override { return AddressX; }
+	virtual TextureAddress GetTextureAddressY() const override { return AddressY; }
+	virtual TextureAddress GetTextureAddressZ() const override { return AddressZ; }
+	virtual const FSparseVolumeTextureSceneProxy* GetSparseVolumeTextureSceneProxy() const override { return SceneProxy; };
 	//~ End USparseVolumeTexture Interface.
 
 private:
 	FIntVector3 VolumeResolution;
+	TEnumAsByte<enum TextureAddress> AddressX;
+	TEnumAsByte<enum TextureAddress> AddressY;
+	TEnumAsByte<enum TextureAddress> AddressZ;
 	const FSparseVolumeTextureSceneProxy* SceneProxy;
 };
 
diff --git a/Engine/Source/Runtime/Engine/Private/Materials/HLSLMaterialTranslator.cpp b/Engine/Source/Runtime/Engine/Private/Materials/HLSLMaterialTranslator.cpp
--- a/Engine/Source/Runtime/Engine/Private/Materials/HLSLMaterialTranslator.cpp
+++ b/Engine/Source/Runtime/Engine/Private/Materials/HLSLMaterialTranslator.cpp
@@ -12634,7 +12634,18 @@ int32 FHLSLMaterialTranslator::SparseVolumeTextureUniformParameter(FName Paramet
 	return AddUniformExpression(new FMaterialUniformExpressionSparseVolumeTextureUniform(ParameterInfo, TextureIndex, VectorIndex), GetMaterialValueType(Type), TEXT(""));
 }
 
-int32 FHLSLMaterialTranslator::SparseVolumeTextureSamplePageTable(int32 SparseVolumeTextureIndex, int32 UVWIndex, int32 MipLevelIndex)
+static const TCHAR* GetSparseVolumeTextureAddressMode(TextureAddress Address)
+{
+	switch (Address)
+	{
+	case TA_Wrap: return TEXT("SVTADDRESSMODE_WRAP");
+	case TA_Clamp: return TEXT("SVTADDRESSMODE_CLAMP");
+	case TA_Mirror: return TEXT("SVTADDRESSMODE_MIRROR");
+	default: checkNoEntry(); return nullptr;
+	}
+}
+
+int32 FHLSLMaterialTranslator::SparseVolumeTextureSamplePageTable(int32 SparseVolumeTextureIndex, int32 UVWIndex, int32 MipLevelIndex, ESamplerSourceMode SamplerSource)
 {
 	if (SparseVolumeTextureIndex == INDEX_NONE || UVWIndex == INDEX_NONE || MipLevelIndex == INDEX_NONE)
 	{
@@ -12667,9 +12678,48 @@ int32 FHLSLMaterialTranslator::SparseVolumeTextureSamplePageTable(int32 SparseVo
 	int32 SVTReferenceIndex = UniformTextureExpressions[(uint32)EMaterialTextureParameterType::SparseVolume].Find(TextureUniformExpression);
 	check(UniformTextureExpressions[(uint32)EMaterialTextureParameterType::SparseVolume].IsValidIndex(SVTReferenceIndex));
 
+	const USparseVolumeTexture* SVTexture = Cast<USparseVolumeTexture>(Material->GetReferencedTextures()[TextureUniformExpression->GetTextureIndex()]);
+
+	// StaticAddress mode at time of compile
+	// Similar to VirtualTextures, this may not be 100% correct, if SamplerSource is set to 'SSM_FromTextureAsset', as texture parameter may change the address mode in a derived instance
+	TextureAddress StaticAddressX = TA_Wrap;
+	TextureAddress StaticAddressY = TA_Wrap;
+	TextureAddress StaticAddressZ = TA_Wrap;
+	switch (SamplerSource)
+	{
+	case SSM_FromTextureAsset:
+		if (SVTexture)
+		{
+			StaticAddressX = SVTexture->GetTextureAddressX();
+			StaticAddressY = SVTexture->GetTextureAddressY();
+			StaticAddressZ = SVTexture->GetTextureAddressZ();
+		}
+		break;
+	case SSM_Wrap_WorldGroupSettings:
+		StaticAddressX = TA_Wrap;
+		StaticAddressY = TA_Wrap;
+		StaticAddressZ = TA_Wrap;
+		break;
+	case SSM_Clamp_WorldGroupSettings: // fallthrough
+	case SSM_TerrainWeightmapGroupSettings:
+		StaticAddressX = TA_Clamp;
+		StaticAddressY = TA_Clamp;
+		StaticAddressZ = TA_Clamp;
+		break;
+	default:
+		checkNoEntry();
+		break;
+	}
+
 	AddEstimatedTextureSample();
-	FString SampleCode = FString::Printf(TEXT("SparseVolumeTextureSamplePageTable(Material.SparseVolumeTexturePageTable_%d, %s, %s, %s)"),
-		SVTReferenceIndex , *GetParameterCode(SparseVolumeTextureIndex), *GetParameterCode(UVWAsFloat3Index), *GetParameterCode(MipLevelAsFloatIndex));
+	FString SampleCode = FString::Printf(TEXT("SparseVolumeTextureSamplePageTable(Material.SparseVolumeTexturePageTable_%d, %s, %s, %s, %s, %s, %s)"),
+		SVTReferenceIndex, 
+		*GetParameterCode(SparseVolumeTextureIndex), 
+		*GetParameterCode(UVWAsFloat3Index), 
+		GetSparseVolumeTextureAddressMode(StaticAddressX), 
+		GetSparseVolumeTextureAddressMode(StaticAddressY), 
+		GetSparseVolumeTextureAddressMode(StaticAddressZ), 
+		*GetParameterCode(MipLevelAsFloatIndex));
 	return AddCodeChunk(MCT_Float3, *SampleCode);
 }
 
@@ -12707,8 +12757,8 @@ int32 FHLSLMaterialTranslator::SparseVolumeTextureSamplePhysicalTileData(int32 S
 	check(UniformTextureExpressions[(uint32)EMaterialTextureParameterType::SparseVolume].IsValidIndex(SVTReferenceIndex));
 
 	AddEstimatedTextureSample();
-	FString SampleCode = FString::Printf(TEXT("SparseVolumeTextureSamplePhysicalTileData(Material.SparseVolumeTexturePhysicalA_%d, Material.SparseVolumeTexturePhysicalB_%d, %s, %s)"),
-		SVTReferenceIndex, SVTReferenceIndex, *GetParameterCode(VoxelCoordAsFloat3Index), *GetParameterCode(IndexAsFloatIndex));
+	FString SampleCode = FString::Printf(TEXT("SparseVolumeTextureSamplePhysicalTileData(Material.SparseVolumeTexturePhysicalA_%d, Material.SparseVolumeTexturePhysicalB_%d, Material.SparseVolumeTexturePhysical_%dSampler, %s, %s)"),
+		SVTReferenceIndex, SVTReferenceIndex, SVTReferenceIndex, *GetParameterCode(VoxelCoordAsFloat3Index), *GetParameterCode(IndexAsFloatIndex));
 	return AddCodeChunk(MCT_Float4, *SampleCode);
 }
 
diff --git a/Engine/Source/Runtime/Engine/Private/Materials/HLSLMaterialTranslator.h b/Engine/Source/Runtime/Engine/Private/Materials/HLSLMaterialTranslator.h
--- a/Engine/Source/Runtime/Engine/Private/Materials/HLSLMaterialTranslator.h
+++ b/Engine/Source/Runtime/Engine/Private/Materials/HLSLMaterialTranslator.h
@@ -928,7 +928,7 @@ protected:
 	virtual int32 SparseVolumeTextureParameter(FName ParameterName, USparseVolumeTexture* InDefaultTexture, int32& TextureReferenceIndex, EMaterialSamplerType SamplerType) override;
 	virtual int32 SparseVolumeTextureUniform(int32 TextureIndex, int32 VectorIndex, UE::Shader::EValueType Type) override;
 	virtual int32 SparseVolumeTextureUniformParameter(FName ParameterName, int32 TextureIndex, int32 VectorIndex, UE::Shader::EValueType Type) override;
-	virtual int32 SparseVolumeTextureSamplePageTable(int32 SparseVolumeTextureIndex, int32 UVWIndex, int32 MipLevelIndex) override;
+	virtual int32 SparseVolumeTextureSamplePageTable(int32 SparseVolumeTextureIndex, int32 UVWIndex, int32 MipLevelIndex, ESamplerSourceMode SamplerSource) override;
 	virtual int32 SparseVolumeTextureSamplePhysicalTileData(int32 SparseVolumeTextureIndex, int32 VoxelCoordIndex, int32 PhysicalTileDataIdxIndex) override;
 
 	virtual UObject* GetReferencedTexture(int32 Index);
diff --git a/Engine/Source/Runtime/Engine/Private/Materials/MaterialExpressions.cpp b/Engine/Source/Runtime/Engine/Private/Materials/MaterialExpressions.cpp
--- a/Engine/Source/Runtime/Engine/Private/Materials/MaterialExpressions.cpp
+++ b/Engine/Source/Runtime/Engine/Private/Materials/MaterialExpressions.cpp
@@ -26229,7 +26229,7 @@ int32 UMaterialExpressionSparseVolumeTextureSample::Compile(class FMaterialCompi
 		
 		// Sample the first mip
 		int32 MipLevel0Index = MipLevelExpression ? Compiler->Floor(MipLevelInputIndex) : Compiler->Constant(0.0f);
-		int32 VoxelCoordMip0Index = Compiler->SparseVolumeTextureSamplePageTable(SparseVolumeTextureIndex, UVWIndex, MipLevel0Index);
+		int32 VoxelCoordMip0Index = Compiler->SparseVolumeTextureSamplePageTable(SparseVolumeTextureIndex, UVWIndex, MipLevel0Index, SamplerSource);
 		int32 Mip0SampleIndex = Compiler->SparseVolumeTextureSamplePhysicalTileData(SparseVolumeTextureIndex, VoxelCoordMip0Index, PhysicalTileDataIdxIndex);
 
 		if (MipLevelExpression)
@@ -26237,7 +26237,7 @@ int32 UMaterialExpressionSparseVolumeTextureSample::Compile(class FMaterialCompi
 			// Sample the second mip
 			// SVT_TODO: Try to optimize out this second sample if LerpAlpha == 0. Might need to do that in HLSL.
 			int32 MipLevel1Index = Compiler->Ceil(MipLevelInputIndex);
-			int32 VoxelCoordMip1Index = Compiler->SparseVolumeTextureSamplePageTable(SparseVolumeTextureIndex, UVWIndex, MipLevel1Index);
+			int32 VoxelCoordMip1Index = Compiler->SparseVolumeTextureSamplePageTable(SparseVolumeTextureIndex, UVWIndex, MipLevel1Index, SamplerSource);
 			int32 Mip1SampleIndex = Compiler->SparseVolumeTextureSamplePhysicalTileData(SparseVolumeTextureIndex, VoxelCoordMip1Index, PhysicalTileDataIdxIndex);
 
 			// Lerp
diff --git a/Engine/Source/Runtime/Engine/Private/Materials/MaterialUniformExpressions.cpp b/Engine/Source/Runtime/Engine/Private/Materials/MaterialUniformExpressions.cpp
--- a/Engine/Source/Runtime/Engine/Private/Materials/MaterialUniformExpressions.cpp
+++ b/Engine/Source/Runtime/Engine/Private/Materials/MaterialUniformExpressions.cpp
@@ -1287,7 +1287,7 @@ void FUniformExpressionSet::FillUniformBuffer(const FMaterialRenderContext& Mate
 					FTextureRHIRef TileDataBTexture = SVTextureProxy->GetPhysicalTileDataBTextureRHI();
 					*ResourceTableTexturePhysicalBPtr = TileDataBTexture ? TileDataBTexture : GBlackVolumeTexture->TextureRHI;
 
-					*ResourceTablePhysicalSamplerPtr = TStaticSamplerState<SF_Point, AM_Clamp, AM_Clamp, AM_Clamp, 0, 8>::GetRHI();
+					*ResourceTablePhysicalSamplerPtr = TStaticSamplerState<SF_Bilinear, AM_Clamp, AM_Clamp, AM_Clamp>::GetRHI();
 				}
 			}
 		}
diff --git a/Engine/Source/Runtime/Engine/Private/SparseVolumeTexture/SparseVolumeTexture.cpp b/Engine/Source/Runtime/Engine/Private/SparseVolumeTexture/SparseVolumeTexture.cpp
--- a/Engine/Source/Runtime/Engine/Private/SparseVolumeTexture/SparseVolumeTexture.cpp
+++ b/Engine/Source/Runtime/Engine/Private/SparseVolumeTexture/SparseVolumeTexture.cpp
@@ -6,6 +6,9 @@
 
 #include "SparseVolumeTexture/SparseVolumeTexture.h"
 
+#include "Materials/Material.h"
+#include "MaterialShared.h"
+#include "UObject/UObjectIterator.h"
 #include "Misc/SecureHash.h"
 #include "EngineUtils.h"
 #include "Shader/ShaderTypes.h"
@@ -74,7 +77,7 @@ FSparseVolumeTextureFrame::~FSparseVolumeTextureFrame()
 {
 }
 
-bool FSparseVolumeTextureFrame::BuildDerivedData(FSparseVolumeTextureData* OutMippedTextureData)
+bool FSparseVolumeTextureFrame::BuildDerivedData(const FIntVector3& VolumeResolution, TextureAddress AddressX, TextureAddress AddressY, TextureAddress AddressZ, FSparseVolumeTextureData* OutMippedTextureData)
 {
 #if WITH_EDITORONLY_DATA
 	// Check if the virtualized bulk data payload is available now
@@ -85,7 +88,13 @@ bool FSparseVolumeTextureFrame::BuildDerivedData(FSparseVolumeTextureData* OutMi
 		FSparseVolumeTextureData TextureData;
 		TextureData.Serialize(RawDataArchiveReader);
 
-		TextureData.GenerateMipMaps();
+		FSparseVolumeTextureDataAddressingInfo AddressingInfo{};
+		AddressingInfo.VolumeResolution = VolumeResolution;
+		AddressingInfo.AddressX = AddressX;
+		AddressingInfo.AddressY = AddressY;
+		AddressingInfo.AddressZ = AddressZ;
+
+		TextureData.BuildDerivedData(AddressingInfo);
 
 		*OutMippedTextureData = MoveTemp(TextureData);
 
@@ -137,7 +146,7 @@ void FSparseVolumeTextureFrame::Serialize(FArchive& Ar, UStreamableSparseVolumeT
 		// We are cooking the game, serialize the asset out.
 
 		FSparseVolumeTextureData DerivedData;
-		const bool bBuiltDerivedData = BuildDerivedData(&DerivedData);
+		const bool bBuiltDerivedData = BuildDerivedData(Owner->VolumeResolution, Owner->AddressX, Owner->AddressY, Owner->AddressZ, &DerivedData);
 		check(bBuiltDerivedData); // SVT_TODO: actual error handling
 
 		// Write derived data into RuntimeStreamedInData
@@ -208,6 +217,7 @@ FVector4 USparseVolumeTexture::GetUniformParameter(int32 Index) const
 void USparseVolumeTexture::GetPackedUniforms(FUintVector4& OutPacked0, FUintVector4& OutPacked1) const
 {
 	FIntVector3 PageTableOffset = FIntVector3::ZeroValue;
+	FVector3f TileDataTexelSize = FVector3f(0.0f, 0.0f, 0.0f);
 	int32 MinMipLevel = 0;
 	int32 MaxMipLevel = 0;
 	const FSparseVolumeTextureSceneProxy* Proxy = GetSparseVolumeTextureSceneProxy();
@@ -215,12 +225,14 @@ void USparseVolumeTexture::GetPackedUniforms(FUintVector4& OutPacked0, FUintVect
 	{
 		const FSparseVolumeTextureRuntimeHeader& Header = Proxy->GetHeader();
 		PageTableOffset = Header.PageTableVolumeAABBMin;
+		TileDataTexelSize.X = 1.0f / Header.TileDataVolumeResolution.X;
+		TileDataTexelSize.Y = 1.0f / Header.TileDataVolumeResolution.Y;
+		TileDataTexelSize.Z = 1.0f / Header.TileDataVolumeResolution.Z;
 		MinMipLevel = Header.LowestResidentLevel;
 		MaxMipLevel = Header.HighestResidentLevel;
 	}
 	const FIntVector3 VolumeResolution = GetVolumeResolution();
 	const FVector3f VolumePageResolution = FVector3f(VolumeResolution) / SPARSE_VOLUME_TILE_RES;
-	const FVector4 TileSize = GetUniformParameter(ESparseVolumeTexture_TileSize);
 
 	auto AsUint = [](float X)
 	{
@@ -231,11 +243,15 @@ void USparseVolumeTexture::GetPackedUniforms(FUintVector4& OutPacked0, FUintVect
 	OutPacked0.X = AsUint(VolumePageResolution.X);
 	OutPacked0.Y = AsUint(VolumePageResolution.Y);
 	OutPacked0.Z = AsUint(VolumePageResolution.Z);
-	OutPacked0.W = (uint32)((MinMipLevel & 0xFF) | ((MaxMipLevel & 0xFF) << 8));
-	OutPacked1.X = AsUint((float)PageTableOffset.X);
-	OutPacked1.Y = AsUint((float)PageTableOffset.Y);
-	OutPacked1.Z = AsUint((float)PageTableOffset.Z);
-	OutPacked1.W = AsUint((float)TileSize.X);
+	OutPacked0.W = UE::SVT::PackPageTableEntry(PageTableOffset);
+	OutPacked1.X = AsUint(TileDataTexelSize.X);
+	OutPacked1.Y = AsUint(TileDataTexelSize.Y);
+	OutPacked1.Z = AsUint(TileDataTexelSize.Z);
+	OutPacked1.W = 0;
+	OutPacked1.W |= (uint32)((MinMipLevel & 0xFF) << 0);
+	OutPacked1.W |= (uint32)((MaxMipLevel & 0xFF) << 8);
+	OutPacked1.W |= (uint32)((int32(SPARSE_VOLUME_TILE_RES) & 0xFF) << 16);
+	OutPacked1.W |= (uint32)((int32(SPARSE_VOLUME_TILE_BORDER) & 0xFF) << 24);
 }
 
 void USparseVolumeTexture::GetFrameUVScaleBias(FVector* OutScale, FVector* OutBias) const
@@ -336,6 +352,14 @@ void UStreamableSparseVolumeTexture::Serialize(FArchive& Ar)
 #if WITH_EDITOR
 void UStreamableSparseVolumeTexture::PostEditChangeProperty(FPropertyChangedEvent& PropertyChangedEvent)
 {
+	if (PropertyChangedEvent.GetPropertyName() == GET_MEMBER_NAME_CHECKED(UStreamableSparseVolumeTexture, AddressX)
+		|| PropertyChangedEvent.GetPropertyName() == GET_MEMBER_NAME_CHECKED(UStreamableSparseVolumeTexture, AddressY)
+		|| PropertyChangedEvent.GetPropertyName() == GET_MEMBER_NAME_CHECKED(UStreamableSparseVolumeTexture, AddressZ))
+	{
+		// SVT need to recompile shaders when address mode changes
+		NotifyMaterials();
+	}
+
 	Super::PostEditChangeProperty(PropertyChangedEvent);
 
 	GenerateOrLoadDDCRuntimeDataAndCreateSceneProxy();
@@ -386,6 +410,50 @@ TArrayView<const FSparseVolumeTextureFrame> UStreamableSparseVolumeTexture::GetF
 	return Frames;
 }
 
+#if WITH_EDITOR
+void UStreamableSparseVolumeTexture::NotifyMaterials(const ENotifyMaterialsEffectOnShaders EffectOnShaders)
+{
+	// Create a material update context to safely update materials.
+	{
+		FMaterialUpdateContext UpdateContext;
+
+		// Notify any material that uses this texture
+		TSet<UMaterial*> BaseMaterialsThatUseThisTexture;
+		for (TObjectIterator<UMaterialInterface> It; It; ++It)
+		{
+			UMaterialInterface* MaterialInterface = *It;
+			if (!FPlatformProperties::IsServerOnly() && MaterialInterface->GetReferencedTextures().Contains(this))
+			{
+				UpdateContext.AddMaterialInterface(MaterialInterface);
+				// This is a bit tricky. We want to make sure all materials using this texture are
+				// updated. Materials are always updated. Material instances may also have to be
+				// updated and if they have static permutations their children must be updated
+				// whether they use the texture or not! The safe thing to do is to add the instance's
+				// base material to the update context causing all materials in the tree to update.
+				BaseMaterialsThatUseThisTexture.Add(MaterialInterface->GetMaterial());
+			}
+		}
+
+		// Go ahead and update any base materials that need to be.
+		if (EffectOnShaders == ENotifyMaterialsEffectOnShaders::Default)
+		{
+			for (TSet<UMaterial*>::TConstIterator It(BaseMaterialsThatUseThisTexture); It; ++It)
+			{
+				(*It)->PostEditChange();
+			}
+		}
+		else
+		{
+			FPropertyChangedEvent EmptyPropertyUpdateStruct(nullptr);
+			for (TSet<UMaterial*>::TConstIterator It(BaseMaterialsThatUseThisTexture); It; ++It)
+			{
+				(*It)->PostEditChangePropertyInternal(EmptyPropertyUpdateStruct, UMaterial::EPostEditChangeEffectOnShaders::DoesNotInvalidate);
+			}
+		}
+	}
+}
+#endif //WITH_EDITOR
+
 void UStreamableSparseVolumeTexture::GenerateOrLoadDDCRuntimeDataAndCreateSceneProxy()
 {
 #if WITH_EDITORONLY_DATA
@@ -427,7 +495,9 @@ void UStreamableSparseVolumeTexture::GenerateOrLoadDDCRuntimeDataForFrame(FSpars
 	using namespace UE::DerivedData;
 
 	static const FString SparseVolumeTextureDDCVersion = TEXT("381AE2A9-A903-4C8F-8486-891E24D6EC70");	// Bump this if you want to ignore all cached data so far.
-	const FString DerivedDataKey = Frame.RawData.GetIdentifier().ToString() + SparseVolumeTextureDDCVersion;
+	const FString DerivedDataKey = Frame.RawData.GetIdentifier().ToString() 
+		+ FString::Format(TEXT("{0},{1},{2}"), {AddressX.GetIntValue(), AddressY.GetIntValue(), AddressZ.GetIntValue()}) 
+		+ SparseVolumeTextureDDCVersion;
 
 	const FCacheKey Key = ConvertLegacyCacheKey(DerivedDataKey);
 	const FSharedString Name = MakeStringView(GetPathName());
@@ -459,7 +529,7 @@ void UStreamableSparseVolumeTexture::GenerateOrLoadDDCRuntimeDataForFrame(FSpars
 						if (Frame.RawData.HasPayloadData())
 						{
 							FSparseVolumeTextureData TextureData;
-							bool bSuccess = Frame.BuildDerivedData(&TextureData);
+							bool bSuccess = Frame.BuildDerivedData(VolumeResolution, AddressX, AddressY, AddressZ, &TextureData);
 							ensure(bSuccess);
 
 							bSuccess = RuntimeData.Create(TextureData);
@@ -547,17 +617,20 @@ USparseVolumeTextureFrame* USparseVolumeTextureFrame::CreateFrame(USparseVolumeT
 	if (Proxy)
 	{
 		USparseVolumeTextureFrame* Frame = NewObject<USparseVolumeTextureFrame>();
-		Frame->Initialize(Proxy, Texture->GetVolumeResolution());
+		Frame->Initialize(Proxy, Texture->GetVolumeResolution(), Texture->GetTextureAddressX(), Texture->GetTextureAddressY(), Texture->GetTextureAddressZ());
 		return Frame;
 	}
 	
 	return nullptr;
 }
 
-void USparseVolumeTextureFrame::Initialize(const FSparseVolumeTextureSceneProxy* InSceneProxy, const FIntVector& InVolumeResolution)
+void USparseVolumeTextureFrame::Initialize(const FSparseVolumeTextureSceneProxy* InSceneProxy, const FIntVector& InVolumeResolution, TextureAddress InAddressX, TextureAddress InAddressY, TextureAddress InAddressZ)
 {
 	SceneProxy = InSceneProxy;
 	VolumeResolution = InVolumeResolution;
+	AddressX = InAddressX;
+	AddressY = InAddressY;
+	AddressZ = InAddressZ;
 }
 
 ////////////////////////////////////////////////////////////////////////////////////////////////
diff --git a/Engine/Source/Runtime/Engine/Private/SparseVolumeTexture/SparseVolumeTextureData.cpp b/Engine/Source/Runtime/Engine/Private/SparseVolumeTexture/SparseVolumeTextureData.cpp
--- a/Engine/Source/Runtime/Engine/Private/SparseVolumeTexture/SparseVolumeTextureData.cpp
+++ b/Engine/Source/Runtime/Engine/Private/SparseVolumeTexture/SparseVolumeTextureData.cpp
@@ -2,9 +2,13 @@
 
 #include "SparseVolumeTexture/SparseVolumeTextureData.h"
 #include "SparseVolumeTexture/SparseVolumeTextureUtility.h"
+#include "Async/ParallelFor.h"
 
 DEFINE_LOG_CATEGORY_STATIC(LogSparseVolumeTextureData, Log, All);
 
+// Enabling this ensures proper bilinear filtering between physical pages and empty pages by tagging neighboring empty pages as resident/physical. This causes more physical tiles to be generated though.
+#define SVT_CORRECT_TILE_ALLOCATION_FOR_LINEAR_FILTERING 0
+
 ////////////////////////////////////////////////////////////////////////////////////////////////
 
 void FSparseVolumeTextureDataHeader::Serialize(FArchive& Ar)
@@ -103,15 +107,44 @@ bool FSparseVolumeTextureData::Construct(const ISparseVolumeTextureDataConstruct
 			const bool bIsFallbackValue = (VoxelValue == AttributesInfo[AttributesIdx].FallbackValue[ComponentIdx]);
 			if (!bIsFallbackValue)
 			{
+				check(Coord.X >= Header.VirtualVolumeAABBMin.X && Coord.Y >= Header.VirtualVolumeAABBMin.Y && Coord.Z >= Header.VirtualVolumeAABBMin.Z);
+				check(Coord.X < Header.VirtualVolumeAABBMax.X && Coord.Y < Header.VirtualVolumeAABBMax.Y && Coord.Z < Header.VirtualVolumeAABBMax.Z);
+
+#if SVT_CORRECT_TILE_ALLOCATION_FOR_LINEAR_FILTERING
+				// Tag all pages touching the 3x3 area around this voxel in order to properly support bilinear/border voxels
+				for (int32 OffsetZ = -1; OffsetZ < 2; ++OffsetZ)
+				{
+					for (int32 OffsetY = -1; OffsetY < 2; ++OffsetY)
+					{
+						for (int32 OffsetX = -1; OffsetX < 2; ++OffsetX)
+						{
+							const FIntVector3 GridCoord = Coord + FIntVector3(OffsetX, OffsetY, OffsetZ);
+							if (GridCoord.X >= Header.VirtualVolumeAABBMin.X
+								&& GridCoord.Y >= Header.VirtualVolumeAABBMin.Y 
+								&& GridCoord.Z >= Header.VirtualVolumeAABBMin.Z
+								&& GridCoord.X < Header.VirtualVolumeAABBMax.X
+								&& GridCoord.Y < Header.VirtualVolumeAABBMax.Y
+								&& GridCoord.Z < Header.VirtualVolumeAABBMax.Z)
+							{
+								const FIntVector3 PageCoord = (GridCoord / SPARSE_VOLUME_TILE_RES) - Header.PageTableVolumeAABBMin;
+								check(PageCoord.X >= 0 && PageCoord.Y >= 0 && PageCoord.Z >= 0);
+								check(PageCoord.X < Header.PageTableVolumeResolution.X && PageCoord.Y < Header.PageTableVolumeResolution.Y && PageCoord.Z < Header.PageTableVolumeResolution.Z);
+
+								const int32 PageIndex = PageCoord.Z * (Header.PageTableVolumeResolution.Y * Header.PageTableVolumeResolution.X) + PageCoord.Y * Header.PageTableVolumeResolution.X + PageCoord.X;
+								PagesWithData[PageIndex] = true;
+							}
+						}
+					}
+				}
+#else // SVT_CORRECT_TILE_ALLOCATION_FOR_LINEAR_FILTERING
 				const FIntVector3 GridCoord = Coord;
-				check(GridCoord.X >= 0 && GridCoord.Y >= 0 && GridCoord.Z >= 0);
-				check(GridCoord.X < Header.VirtualVolumeAABBMax.X&& GridCoord.Y < Header.VirtualVolumeAABBMax.Y&& GridCoord.Z < Header.VirtualVolumeAABBMax.Z);
 				const FIntVector3 PageCoord = (GridCoord / SPARSE_VOLUME_TILE_RES) - Header.PageTableVolumeAABBMin;
 				check(PageCoord.X >= 0 && PageCoord.Y >= 0 && PageCoord.Z >= 0);
-				check(PageCoord.X < Header.PageTableVolumeResolution.X&& PageCoord.Y < Header.PageTableVolumeResolution.Y&& PageCoord.Z < Header.PageTableVolumeResolution.Z);
+				check(PageCoord.X < Header.PageTableVolumeResolution.X && PageCoord.Y < Header.PageTableVolumeResolution.Y && PageCoord.Z < Header.PageTableVolumeResolution.Z);
 
 				const int32 PageIndex = PageCoord.Z * (Header.PageTableVolumeResolution.Y * Header.PageTableVolumeResolution.X) + PageCoord.Y * Header.PageTableVolumeResolution.X + PageCoord.X;
 				PagesWithData[PageIndex] = true;
+#endif // SVT_CORRECT_TILE_ALLOCATION_FOR_LINEAR_FILTERING
 			}
 		});
 
@@ -122,11 +155,11 @@ bool FSparseVolumeTextureData::Construct(const ISparseVolumeTextureDataConstruct
 	// Go over each potential page from the source data and push allocate it if it has any data.
 	// Otherwise point to the default empty page.
 	int32 NumAllocatedPages = 0;
-	for (int32_t PageZ = 0; PageZ < Header.PageTableVolumeResolution.Z; ++PageZ)
+	for (int32 PageZ = 0; PageZ < Header.PageTableVolumeResolution.Z; ++PageZ)
 	{
-		for (int32_t PageY = 0; PageY < Header.PageTableVolumeResolution.Y; ++PageY)
+		for (int32 PageY = 0; PageY < Header.PageTableVolumeResolution.Y; ++PageY)
 		{
-			for (int32_t PageX = 0; PageX < Header.PageTableVolumeResolution.X; ++PageX)
+			for (int32 PageX = 0; PageX < Header.PageTableVolumeResolution.X; ++PageX)
 			{
 				const int32 PageIndex = PageZ * (Header.PageTableVolumeResolution.Y * Header.PageTableVolumeResolution.X) + PageY * Header.PageTableVolumeResolution.X + PageX;
 				const bool bHasAnyData = PagesWithData[PageIndex];
@@ -138,6 +171,7 @@ bool FSparseVolumeTextureData::Construct(const ISparseVolumeTextureDataConstruct
 			}
 		}
 	}
+	LinearAllocatedPages.SetNum(NumAllocatedPages);
 
 	const int32 EffectivelyAllocatedTiles = NumAllocatedPages + 1 /*null tile*/;
 
@@ -145,16 +179,16 @@ bool FSparseVolumeTextureData::Construct(const ISparseVolumeTextureDataConstruct
 	const int32 FormatSize[] = { GPixelFormats[Header.AttributesFormats[0]].BlockBytes, GPixelFormats[Header.AttributesFormats[1]].BlockBytes };
 	PageTable.SetNum(1);
 	PageTable[0].SetNumZeroed(Header.PageTableVolumeResolution.X * Header.PageTableVolumeResolution.Y * Header.PageTableVolumeResolution.Z);
-	PhysicalTileDataA.SetNumZeroed(EffectivelyAllocatedTiles * SVTNumVoxelsPerTile * FormatSize[0]);
-	PhysicalTileDataB.SetNumZeroed(EffectivelyAllocatedTiles * SVTNumVoxelsPerTile * FormatSize[1]);
+	PhysicalTileDataA.SetNumZeroed(EffectivelyAllocatedTiles * SVTNumVoxelsPerPaddedTile * FormatSize[0]);
+	PhysicalTileDataB.SetNumZeroed(EffectivelyAllocatedTiles * SVTNumVoxelsPerPaddedTile * FormatSize[1]);
 
 	// Generate page table and tile volume data by splatting the data
 	{
 		uint32 DstTileIndex = 0;
 
-		// Add an empty tile, reserve slot at coord 0
+		// Add an empty tile, reserve slot at coord 0 and fill all tiles with null tile data
 		{
-			// Fill null tile with fallback/background value
+			// Compute potentially normalized fallback values
 			FVector4f FallbackValues[] = { AttributesInfo[0].FallbackValue, AttributesInfo[1].FallbackValue };
 			for (int32 AttributesIdx = 0; AttributesIdx < 2; ++AttributesIdx)
 			{
@@ -163,21 +197,25 @@ bool FSparseVolumeTextureData::Construct(const ISparseVolumeTextureDataConstruct
 					FallbackValues[AttributesIdx] = FallbackValues[AttributesIdx] * AttributesInfo[AttributesIdx].NormalizeScale + AttributesInfo[AttributesIdx].NormalizeBias;
 				}
 			}
-			for (int32 Z = 0; Z < SPARSE_VOLUME_TILE_RES; ++Z)
-			{
-				for (int32 Y = 0; Y < SPARSE_VOLUME_TILE_RES; ++Y)
+			Header.NullTileValues[0] = FallbackValues[0];
+			Header.NullTileValues[1] = FallbackValues[1];
+
+			// Fill all tiles with fallback values since we can't guarantee that all voxels of a page will be written to in the next step.
+			ParallelFor(EffectivelyAllocatedTiles, [&FallbackValues, this](int32 TileIndex)
 				{
-					for (int32 X = 0; X < SPARSE_VOLUME_TILE_RES; ++X)
+					for (int32 Z = 0; Z < SPARSE_VOLUME_TILE_RES_PADDED; ++Z)
 					{
-						const FIntVector3 WriteCoord = FIntVector3(X, Y, Z);
-						WriteTileDataVoxel(0 /*TileIndex*/, WriteCoord, 0 /*AttributesIndex*/, FallbackValues[0]);
-						WriteTileDataVoxel(0 /*TileIndex*/, WriteCoord, 1 /*AttributesIndex*/, FallbackValues[1]);
+						for (int32 Y = 0; Y < SPARSE_VOLUME_TILE_RES_PADDED; ++Y)
+						{
+							for (int32 X = 0; X < SPARSE_VOLUME_TILE_RES_PADDED; ++X)
+							{
+								const FIntVector3 WriteCoord = FIntVector3(X - 1, Y - 1, Z - 1);
+								WriteTileDataVoxel(TileIndex, WriteCoord, 0 /*AttributesIndex*/, FallbackValues[0]);
+								WriteTileDataVoxel(TileIndex, WriteCoord, 1 /*AttributesIndex*/, FallbackValues[1]);
+							}
+						}
 					}
-				}
-			}
-
-			Header.NullTileValues[0] = FallbackValues[0];
-			Header.NullTileValues[1] = FallbackValues[1];
+				});
 
 			// PageTable is all cleared to zero, simply skip a tile
 			++DstTileIndex;
@@ -198,8 +236,8 @@ bool FSparseVolumeTextureData::Construct(const ISparseVolumeTextureDataConstruct
 			// Set the next tile to be written to
 			++DstTileIndex;
 		}
-		MipInfo.TileOffset = 0;
-		MipInfo.TileCount = EffectivelyAllocatedTiles;
+		MipInfo.TileOffset = 1; // Don't include the null tile in the list of physical tiles of mip0. The runtime SVT will create a null of its own.
+		MipInfo.TileCount = NumAllocatedPages;
 	}
 
 	// Write physical tile data
@@ -239,11 +277,10 @@ uint32 FSparseVolumeTextureData::ReadPageTable(const FIntVector3& PageTableCoord
 	{
 		return INDEX_NONE;
 	}
-	const int32 MipLevelFactor = 1 << MipLevel;
 	const FIntVector3 PageTableResolution = FIntVector3(
-		FMath::Max(1, Header.PageTableVolumeResolution.X / MipLevelFactor), 
-		FMath::Max(1, Header.PageTableVolumeResolution.Y / MipLevelFactor), 
-		FMath::Max(1, Header.PageTableVolumeResolution.Z / MipLevelFactor));
+		FMath::Max(1, Header.PageTableVolumeResolution.X >> MipLevel), 
+		FMath::Max(1, Header.PageTableVolumeResolution.Y >> MipLevel), 
+		FMath::Max(1, Header.PageTableVolumeResolution.Z >> MipLevel));
 	if (PageTableCoord.X < 0 || PageTableCoord.Y < 0 || PageTableCoord.Z < 0
 		|| PageTableCoord.X >= PageTableResolution.X
 		|| PageTableCoord.Y >= PageTableResolution.Y
@@ -263,7 +300,7 @@ uint32 FSparseVolumeTextureData::ReadPageTable(const FIntVector3& PageTableCoord
 FVector4f FSparseVolumeTextureData::ReadTileDataVoxel(int32 TileIndex, const FIntVector3& TileDataCoord, int32 AttributesIdx) const
 {
 	check(AttributesIdx >= 0 && AttributesIdx <= 1);
-	if ((AttributesIdx == 0 && PhysicalTileDataA.IsEmpty()) || (AttributesIdx == 1 && PhysicalTileDataB.IsEmpty()))
+	if (TileIndex == INDEX_NONE || (AttributesIdx == 0 && PhysicalTileDataA.IsEmpty()) || (AttributesIdx == 1 && PhysicalTileDataB.IsEmpty()))
 	{
 		return FVector4f();
 	}
@@ -278,21 +315,41 @@ FVector4f FSparseVolumeTextureData::ReadTileDataVoxel(int32 TileIndex, const FIn
 		return FVector4f();
 	}
 
-	const int32 VoxelIndex = TileIndex * UE::SVT::SVTNumVoxelsPerTile + TileDataCoord.Z * (SPARSE_VOLUME_TILE_RES * SPARSE_VOLUME_TILE_RES) + TileDataCoord.Y * SPARSE_VOLUME_TILE_RES + TileDataCoord.X;
+	const FIntVector3 CoordPadded = TileDataCoord + FIntVector3(SPARSE_VOLUME_TILE_BORDER);
+	const int32 VoxelIndex = TileIndex * UE::SVT::SVTNumVoxelsPerPaddedTile + CoordPadded.Z * (SPARSE_VOLUME_TILE_RES_PADDED * SPARSE_VOLUME_TILE_RES_PADDED) + CoordPadded.Y * SPARSE_VOLUME_TILE_RES_PADDED + CoordPadded.X;
 	const uint8* TileData = AttributesIdx == 0 ? PhysicalTileDataA.GetData() : PhysicalTileDataB.GetData();
 	const EPixelFormat Format = Header.AttributesFormats[AttributesIdx];
 	return UE::SVT::ReadVoxel(VoxelIndex, TileData, Format);
 }
 
-FVector4f FSparseVolumeTextureData::Load(const FIntVector3& VolumeCoord, int32 MipLevel, int32 AttributesIdx) const
+FVector4f FSparseVolumeTextureData::Load(const FIntVector3& VolumeCoord, int32 MipLevel, int32 AttributesIdx, const FSparseVolumeTextureDataAddressingInfo& AddressingInfo) const
 {
+	auto ApplyAddressMode = [](int32 x, int32 Width, TextureAddress Mode)
+	{
+		switch (Mode)
+		{
+		case TA_Wrap:
+			return (x % Width + Width) % Width; // Make sure it's a proper modulo for negative numbers ....
+		case TA_Clamp:
+			return FMath::Clamp(x, 0, Width - 1);
+		case TA_Mirror:
+			int32 DoubleWidth = Width + Width;
+			int32 DoubleWrap = (x % DoubleWidth + DoubleWidth) % DoubleWidth;
+			return (DoubleWrap < Width) ? DoubleWrap : (Width - 1) - (DoubleWrap - Width);
+		}
+		return x;
+	};
 	if (!Header.MipInfo.IsValidIndex(MipLevel))
 	{
 		return FVector4f();
 	}
-	const FIntVector3 PageTableCoord = (VolumeCoord / SPARSE_VOLUME_TILE_RES) - (Header.PageTableVolumeAABBMin / (1 << MipLevel));
+	const FIntVector3 AddressModeVolumeCoord = FIntVector3(
+		ApplyAddressMode(VolumeCoord.X, AddressingInfo.VolumeResolution.X, AddressingInfo.AddressX),
+		ApplyAddressMode(VolumeCoord.Y, AddressingInfo.VolumeResolution.Y, AddressingInfo.AddressY), 
+		ApplyAddressMode(VolumeCoord.Z, AddressingInfo.VolumeResolution.Z, AddressingInfo.AddressZ));
+	const FIntVector3 PageTableCoord = (AddressModeVolumeCoord / SPARSE_VOLUME_TILE_RES) - (Header.PageTableVolumeAABBMin >> MipLevel);
 	const uint32 TileIndex = ReadPageTable(PageTableCoord, MipLevel);
-	const FIntVector3 VoxelCoord = VolumeCoord % SPARSE_VOLUME_TILE_RES;
+	const FIntVector3 VoxelCoord = AddressModeVolumeCoord % SPARSE_VOLUME_TILE_RES;
 	const FVector4f Sample = ReadTileDataVoxel(TileIndex, VoxelCoord, AttributesIdx);
 	return Sample;
 }
@@ -307,21 +364,26 @@ void FSparseVolumeTextureData::WriteTileDataVoxel(int32 TileIndex, const FIntVec
 
 	// SV_TODO check for TileIndex oob
 
-	if (TileDataCoord.X < 0 || TileDataCoord.Y < 0 || TileDataCoord.Z < 0
-		|| TileDataCoord.X >= SPARSE_VOLUME_TILE_RES
-		|| TileDataCoord.Y >= SPARSE_VOLUME_TILE_RES
-		|| TileDataCoord.Z >= SPARSE_VOLUME_TILE_RES)
+	// Allow TileDataCoord to extend past SPARSE_VOLUME_TILE_RES by SPARSE_VOLUME_TILE_BORDER.
+	// This allows us to write the border voxels without changing how TileDataCoord needs to be computed for other use cases.
+	if (TileDataCoord.X < -SPARSE_VOLUME_TILE_BORDER 
+		|| TileDataCoord.Y < -SPARSE_VOLUME_TILE_BORDER 
+		|| TileDataCoord.Z < -SPARSE_VOLUME_TILE_BORDER
+		|| TileDataCoord.X >= (SPARSE_VOLUME_TILE_RES + SPARSE_VOLUME_TILE_BORDER)
+		|| TileDataCoord.Y >= (SPARSE_VOLUME_TILE_RES + SPARSE_VOLUME_TILE_BORDER)
+		|| TileDataCoord.Z >= (SPARSE_VOLUME_TILE_RES + SPARSE_VOLUME_TILE_BORDER))
 	{
 		return;
 	}
 
-	const int32 VoxelIndex = TileIndex * UE::SVT::SVTNumVoxelsPerTile + TileDataCoord.Z * (SPARSE_VOLUME_TILE_RES * SPARSE_VOLUME_TILE_RES) + TileDataCoord.Y * SPARSE_VOLUME_TILE_RES + TileDataCoord.X;
+	const FIntVector3 CoordPadded = TileDataCoord + FIntVector3(SPARSE_VOLUME_TILE_BORDER);
+	const int32 VoxelIndex = TileIndex * UE::SVT::SVTNumVoxelsPerPaddedTile + CoordPadded.Z * (SPARSE_VOLUME_TILE_RES_PADDED * SPARSE_VOLUME_TILE_RES_PADDED) + CoordPadded.Y * SPARSE_VOLUME_TILE_RES_PADDED + CoordPadded.X;
 	uint8* TileData = AttributesIdx == 0 ? PhysicalTileDataA.GetData() : PhysicalTileDataB.GetData();
 	const EPixelFormat Format = Header.AttributesFormats[AttributesIdx];
 	UE::SVT::WriteVoxel(VoxelIndex, TileData, Format, Value, DstComponent);
 }
 
-void FSparseVolumeTextureData::GenerateMipMaps(int32 NumMipLevels)
+void FSparseVolumeTextureData::GenerateMipMaps(const FSparseVolumeTextureDataAddressingInfo& AddressingInfo, int32 NumMipLevels)
 {
 	check(!Header.MipInfo.IsEmpty());
 	check(FMath::IsPowerOfTwo(Header.PageTableVolumeResolution.X) 
@@ -346,11 +408,12 @@ void FSparseVolumeTextureData::GenerateMipMaps(int32 NumMipLevels)
 		FSparseVolumeTextureMipInfo& MipInfo = Header.MipInfo[MipLevel];
 		MipInfo.TileOffset = Header.MipInfo[MipLevel - 1].TileOffset + Header.MipInfo[MipLevel - 1].TileCount;
 
-		const int32 MipLevelFactor = 1 << MipLevel;
-		const int32 ParentMipLevelFactor = 1 << (MipLevel - 1);
-		const FIntVector3 PageTableVolumeAABBMin = (Header.VirtualVolumeAABBMin / SPARSE_VOLUME_TILE_RES) / MipLevelFactor;
-		const FIntVector3 ParentPageTableVolumeAABBMin = (Header.VirtualVolumeAABBMin / SPARSE_VOLUME_TILE_RES) / ParentMipLevelFactor;
-		const FIntVector3 PageTableVolumeResolution = FIntVector3(FMath::Max(1, Header.PageTableVolumeResolution.X / MipLevelFactor), FMath::Max(1, Header.PageTableVolumeResolution.Y / MipLevelFactor), FMath::Max(1, Header.PageTableVolumeResolution.Z / MipLevelFactor));
+		const FIntVector3 PageTableVolumeAABBMin = (Header.VirtualVolumeAABBMin / SPARSE_VOLUME_TILE_RES) >> MipLevel;
+		const FIntVector3 ParentPageTableVolumeAABBMin = (Header.VirtualVolumeAABBMin / SPARSE_VOLUME_TILE_RES) >> (MipLevel - 1);
+		const FIntVector3 PageTableVolumeResolution = FIntVector3(
+			FMath::Max(1, Header.PageTableVolumeResolution.X >> MipLevel), 
+			FMath::Max(1, Header.PageTableVolumeResolution.Y >> MipLevel), 
+			FMath::Max(1, Header.PageTableVolumeResolution.Z >> MipLevel));
 		
 		check(FMath::IsPowerOfTwo(PageTableVolumeResolution.X)
 			&& FMath::IsPowerOfTwo(PageTableVolumeResolution.Y)
@@ -363,15 +426,15 @@ void FSparseVolumeTextureData::GenerateMipMaps(int32 NumMipLevels)
 		// Go over each potential page from the source data and push allocate it if it has any data.
 		// Otherwise point to the default empty page.
 		int32 NumAllocatedPages = 0;
-		for (int32_t PageZ = 0; PageZ < PageTableVolumeResolution.Z; ++PageZ)
+		for (int32 PageZ = 0; PageZ < PageTableVolumeResolution.Z; ++PageZ)
 		{
-			for (int32_t PageY = 0; PageY < PageTableVolumeResolution.Y; ++PageY)
+			for (int32 PageY = 0; PageY < PageTableVolumeResolution.Y; ++PageY)
 			{
-				for (int32_t PageX = 0; PageX < PageTableVolumeResolution.X; ++PageX)
+				for (int32 PageX = 0; PageX < PageTableVolumeResolution.X; ++PageX)
 				{
 					const FIntVector3 PageCoord = FIntVector3(PageX, PageY, PageZ);
 					bool bHasAnyData = false;
-					for (int32_t OffsetIdx = 0; OffsetIdx < 8; ++OffsetIdx)
+					for (int32 OffsetIdx = 0; OffsetIdx < 8; ++OffsetIdx)
 					{
 						const FIntVector3 Offset = FIntVector3(OffsetIdx, OffsetIdx >> 1, OffsetIdx >> 2) & 1;
 						const FIntVector3 ParentPageTableCoord = (PageTableVolumeAABBMin + PageCoord) * 2 + Offset - ParentPageTableVolumeAABBMin;
@@ -390,68 +453,135 @@ void FSparseVolumeTextureData::GenerateMipMaps(int32 NumMipLevels)
 				}
 			}
 		}
+		LinearAllocatedPages.SetNum(NumAllocatedPages);
 
 		// Initialize the SparseVolumeTexture page and tile.
 		const int32 FormatSize[] = { GPixelFormats[Header.AttributesFormats[0]].BlockBytes, GPixelFormats[Header.AttributesFormats[1]].BlockBytes };
 		PageTable.AddDefaulted();
 		PageTable[MipLevel].SetNumZeroed(PageTableVolumeResolution.X * PageTableVolumeResolution.Y * PageTableVolumeResolution.Z);
-		PhysicalTileDataA.AddZeroed(NumAllocatedPages * UE::SVT::SVTNumVoxelsPerTile * FormatSize[0]);
-		PhysicalTileDataB.AddZeroed(NumAllocatedPages * UE::SVT::SVTNumVoxelsPerTile * FormatSize[1]);
+		PhysicalTileDataA.AddZeroed(NumAllocatedPages * UE::SVT::SVTNumVoxelsPerPaddedTile * FormatSize[0]);
+		PhysicalTileDataB.AddZeroed(NumAllocatedPages * UE::SVT::SVTNumVoxelsPerPaddedTile * FormatSize[1]);
 
 		// Generate page table and tile volume data by splatting the data
 		{
-			uint32 DstTileIndex = (uint32)MipInfo.TileOffset;
-
-			for (int32 i = 0; i < NumAllocatedPages; ++i)
-			{
-				const FIntVector3 PageCoordToSplat = LinearAllocatedPages[i];
-
-				// Setup the page table entry
-				PageTable[MipLevel]
-					[
-						PageCoordToSplat.Z * PageTableVolumeResolution.X * PageTableVolumeResolution.Y +
-						PageCoordToSplat.Y * PageTableVolumeResolution.X +
-						PageCoordToSplat.X
-					] = DstTileIndex;
-
-				// Write tile data
-				const FIntVector3 ParentVolumeCoordBase = (PageTableVolumeAABBMin + PageCoordToSplat) * SPARSE_VOLUME_TILE_RES * 2;
-				for (int32 AttributesIdx = 0; AttributesIdx < 2; ++AttributesIdx)
+			ParallelFor(NumAllocatedPages, [this, &LinearAllocatedPages, &PageTableVolumeResolution, &PageTableVolumeAABBMin, &AddressingInfo, MipLevel, BaseDstTileIndex = (uint32)MipInfo.TileOffset](int32 PageIndex)
 				{
-					if ((AttributesIdx == 0 && PhysicalTileDataA.IsEmpty()) || (AttributesIdx == 1 && PhysicalTileDataB.IsEmpty()))
+					const FIntVector3 PageCoordToSplat = LinearAllocatedPages[PageIndex];
+					const uint32 DstTileIndex = BaseDstTileIndex + PageIndex;
+
+					// Setup the page table entry
+					PageTable[MipLevel]
+						[
+							PageCoordToSplat.Z * PageTableVolumeResolution.X * PageTableVolumeResolution.Y +
+							PageCoordToSplat.Y * PageTableVolumeResolution.X +
+							PageCoordToSplat.X
+						] = DstTileIndex;
+
+					// Write tile data
+					const FIntVector3 ParentVolumeCoordBase = (PageTableVolumeAABBMin + PageCoordToSplat) * SPARSE_VOLUME_TILE_RES * 2;
+					for (int32 AttributesIdx = 0; AttributesIdx < 2; ++AttributesIdx)
 					{
-						continue;
-					}
+						if ((AttributesIdx == 0 && PhysicalTileDataA.IsEmpty()) || (AttributesIdx == 1 && PhysicalTileDataB.IsEmpty()))
+						{
+							continue;
+						}
 
-					for (int32 Z = 0; Z < SPARSE_VOLUME_TILE_RES; ++Z)
-					{
-						for (int32 Y = 0; Y < SPARSE_VOLUME_TILE_RES; ++Y)
+						for (int32 Z = 0; Z < SPARSE_VOLUME_TILE_RES; ++Z)
 						{
-							for (int32 X = 0; X < SPARSE_VOLUME_TILE_RES; ++X)
+							for (int32 Y = 0; Y < SPARSE_VOLUME_TILE_RES; ++Y)
 							{
-								FVector4f DownsampledValue = FVector4f(0.0f, 0.0f, 0.0f, 0.0f);
-								for (int32_t OffsetIdx = 0; OffsetIdx < 8; ++OffsetIdx)
+								for (int32 X = 0; X < SPARSE_VOLUME_TILE_RES; ++X)
 								{
-									const FIntVector3 Offset = FIntVector3(OffsetIdx, OffsetIdx >> 1, OffsetIdx >> 2) & 1;
-									const FIntVector3 SourceCoord = ParentVolumeCoordBase + FIntVector3(X, Y, Z) * 2 + Offset;
-									const FVector4f SampleValue = Load(SourceCoord, MipLevel - 1, AttributesIdx);
-									DownsampledValue += SampleValue;
+									FVector4f DownsampledValue = FVector4f(0.0f, 0.0f, 0.0f, 0.0f);
+									for (int32 OffsetIdx = 0; OffsetIdx < 8; ++OffsetIdx)
+									{
+										const FIntVector3 Offset = FIntVector3(OffsetIdx, OffsetIdx >> 1, OffsetIdx >> 2) & 1;
+										const FIntVector3 SourceCoord = ParentVolumeCoordBase + FIntVector3(X, Y, Z) * 2 + Offset;
+										const FVector4f SampleValue = Load(SourceCoord, MipLevel - 1, AttributesIdx, AddressingInfo);
+										DownsampledValue += SampleValue;
+									}
+									DownsampledValue /= 8.0f;
+
+									const FIntVector3 WriteCoord = FIntVector3(X, Y, Z);
+									WriteTileDataVoxel(DstTileIndex, WriteCoord, AttributesIdx, DownsampledValue);
 								}
-								DownsampledValue /= 8.0f;
-
-								const FIntVector3 WriteCoord = FIntVector3(X, Y, Z);
-								WriteTileDataVoxel(DstTileIndex, WriteCoord, AttributesIdx, DownsampledValue);
 							}
 						}
 					}
+				});
+
+			MipInfo.TileCount = NumAllocatedPages;
+		}
+
+		GenerateBorderVoxels(AddressingInfo, MipLevel, LinearAllocatedPages);
+	}
+}
+
+void FSparseVolumeTextureData::GenerateBorderVoxels(const FSparseVolumeTextureDataAddressingInfo& AddressingInfo, int32 MipLevel, const TArray<FIntVector3>& PageCoords)
+{
+	const FIntVector3 PageTableVolumeResolution = FIntVector3(
+		FMath::Max(1, Header.PageTableVolumeResolution.X >> MipLevel), 
+		FMath::Max(1, Header.PageTableVolumeResolution.Y >> MipLevel), 
+		FMath::Max(1, Header.PageTableVolumeResolution.Z >> MipLevel));
+
+	ParallelFor(PageCoords.Num(), [this, MipLevel, &PageCoords, &PageTableVolumeResolution, &AddressingInfo](int32 PageIndex)
+		{
+			const FIntVector3& PageCoord = PageCoords[PageIndex];
+			const uint32 DstTileIndex = PageTable[MipLevel]
+			[
+				PageCoord.Z * PageTableVolumeResolution.X * PageTableVolumeResolution.Y +
+				PageCoord.Y * PageTableVolumeResolution.X +
+				PageCoord.X
+			];
+			const FIntVector3 PageTableOffset = (Header.PageTableVolumeAABBMin >> MipLevel);
+			for (int32 Z = -SPARSE_VOLUME_TILE_BORDER; Z < (SPARSE_VOLUME_TILE_RES + SPARSE_VOLUME_TILE_BORDER); ++Z)
+			{
+				for (int32 Y = -SPARSE_VOLUME_TILE_BORDER; Y < (SPARSE_VOLUME_TILE_RES + SPARSE_VOLUME_TILE_BORDER); ++Y)
+				{
+					for (int32 X = -SPARSE_VOLUME_TILE_BORDER; X < (SPARSE_VOLUME_TILE_RES + SPARSE_VOLUME_TILE_BORDER); ++X)
+					{
+						if (X < 0 || Y < 0 || Z < 0 || X >= SPARSE_VOLUME_TILE_RES || Y >= SPARSE_VOLUME_TILE_RES || Z >= SPARSE_VOLUME_TILE_RES)
+						{
+							const FIntVector3 LocalCoord = FIntVector3(X, Y, Z);
+							const FIntVector3 VolumeCoord = (PageTableOffset + PageCoord) * SPARSE_VOLUME_TILE_RES + LocalCoord;
+							const FVector4f Border0 = Load(VolumeCoord, MipLevel, 0, AddressingInfo);
+							const FVector4f Border1 = Load(VolumeCoord, MipLevel, 1, AddressingInfo);
+							WriteTileDataVoxel(DstTileIndex, LocalCoord, 0, Border0);
+							WriteTileDataVoxel(DstTileIndex, LocalCoord, 1, Border1);
+						}
+					}
 				}
+			}
+		});
+}
 
-				// Set the next tile to be written to
-				++DstTileIndex;
+void FSparseVolumeTextureData::BuildDerivedData(const FSparseVolumeTextureDataAddressingInfo& AddressingInfo, int32 NumMipLevels)
+{
+	// Generate border voxels of mip0
+	{
+		// Collect the page table coordinates of all non-zero pages
+		TArray<FIntVector3> PageCoords;
+		for (int32 PageZ = 0; PageZ < Header.PageTableVolumeResolution.Z; ++PageZ)
+		{
+			for (int32 PageY = 0; PageY < Header.PageTableVolumeResolution.Y; ++PageY)
+			{
+				for (int32 PageX = 0; PageX < Header.PageTableVolumeResolution.X; ++PageX)
+				{
+					const int32 PageIndex = PageZ * (Header.PageTableVolumeResolution.Y * Header.PageTableVolumeResolution.X) + PageY * Header.PageTableVolumeResolution.X + PageX;
+					const uint32 PageTableEntry = PageTable[0][PageIndex];
+					if (PageTableEntry != 0)
+					{
+						PageCoords.Add(FIntVector3(PageX, PageY, PageZ));
+					}
+				}
 			}
-			MipInfo.TileCount = NumAllocatedPages;
 		}
+
+		GenerateBorderVoxels(AddressingInfo, 0, PageCoords);
 	}
+
+	// Generate all remaining mips. Also generates border voxels.
+	GenerateMipMaps(AddressingInfo, NumMipLevels);
 }
 
 ////////////////////////////////////////////////////////////////////////////////////////////////
diff --git a/Engine/Source/Runtime/Engine/Private/SparseVolumeTexture/SparseVolumeTextureSceneProxy.cpp b/Engine/Source/Runtime/Engine/Private/SparseVolumeTexture/SparseVolumeTextureSceneProxy.cpp
--- a/Engine/Source/Runtime/Engine/Private/SparseVolumeTexture/SparseVolumeTextureSceneProxy.cpp
+++ b/Engine/Source/Runtime/Engine/Private/SparseVolumeTexture/SparseVolumeTextureSceneProxy.cpp
@@ -42,7 +42,7 @@ namespace Private
 			TileDataVolumeResolution.Z--;	// We then trim an edge to get back space.
 		}
 
-		return TileDataVolumeResolution * SPARSE_VOLUME_TILE_RES;
+		return TileDataVolumeResolution * SPARSE_VOLUME_TILE_RES_PADDED;
 	}
 } // Private
 } // SVT
@@ -60,8 +60,8 @@ bool FSparseVolumeTextureRuntime::Create(const FSparseVolumeTextureData& Texture
 	}
 
 	const int32 NumMipLevels = TextureData.Header.MipInfo.Num();
-	const int32 TileSizeBytesA = GPixelFormats[Header.AttributesFormats[0]].BlockBytes * UE::SVT::SVTNumVoxelsPerTile;
-	const int32 TileSizeBytesB = GPixelFormats[Header.AttributesFormats[1]].BlockBytes * UE::SVT::SVTNumVoxelsPerTile;
+	const int32 TileSizeBytesA = GPixelFormats[Header.AttributesFormats[0]].BlockBytes * UE::SVT::SVTNumVoxelsPerPaddedTile;
+	const int32 TileSizeBytesB = GPixelFormats[Header.AttributesFormats[1]].BlockBytes * UE::SVT::SVTNumVoxelsPerPaddedTile;
 
 	TArray<FSparseVolumeTextureTileMapping> Mappings;
 	Mappings.Empty(NumMipLevels);
@@ -183,10 +183,10 @@ bool FSparseVolumeTextureRuntime::SetTileMappings(const TArrayView<const FSparse
 	}
 
 	Header.TileDataVolumeResolution = TileDataVolumeRes;
-	const FIntVector3 TileCoordSpace = TileDataVolumeRes / SPARSE_VOLUME_TILE_RES;
-	check((TileCoordSpace.X * SPARSE_VOLUME_TILE_RES) == TileDataVolumeRes.X
-		&& (TileCoordSpace.Y * SPARSE_VOLUME_TILE_RES) == TileDataVolumeRes.Y
-		&& (TileCoordSpace.Z * SPARSE_VOLUME_TILE_RES) == TileDataVolumeRes.Z);
+	const FIntVector3 TileCoordSpace = TileDataVolumeRes / SPARSE_VOLUME_TILE_RES_PADDED;
+	check((TileCoordSpace.X * SPARSE_VOLUME_TILE_RES_PADDED) == TileDataVolumeRes.X
+		&& (TileCoordSpace.Y * SPARSE_VOLUME_TILE_RES_PADDED) == TileDataVolumeRes.Y
+		&& (TileCoordSpace.Z * SPARSE_VOLUME_TILE_RES_PADDED) == TileDataVolumeRes.Z);
 
 	// Clear page table
 	for (TArray<uint32>& PageTableMip : PageTable)
@@ -204,13 +204,13 @@ bool FSparseVolumeTextureRuntime::SetTileMappings(const TArrayView<const FSparse
 	FIntVector3 DstTileCoord = FIntVector3::ZeroValue;
 
 	// Write null tile
-	for (int32 Z = 0; Z < SPARSE_VOLUME_TILE_RES; ++Z)
+	for (int32 Z = 0; Z < SPARSE_VOLUME_TILE_RES_PADDED; ++Z)
 	{
-		for (int32 Y = 0; Y < SPARSE_VOLUME_TILE_RES; ++Y)
+		for (int32 Y = 0; Y < SPARSE_VOLUME_TILE_RES_PADDED; ++Y)
 		{
-			for (int32 X = 0; X < SPARSE_VOLUME_TILE_RES; ++X)
+			for (int32 X = 0; X < SPARSE_VOLUME_TILE_RES_PADDED; ++X)
 			{
-				const int32 VoxelIndex = Z * (SPARSE_VOLUME_TILE_RES * SPARSE_VOLUME_TILE_RES) + Y * SPARSE_VOLUME_TILE_RES + X;
+				const int32 VoxelIndex = Z * (SPARSE_VOLUME_TILE_RES_PADDED * SPARSE_VOLUME_TILE_RES_PADDED) + Y * SPARSE_VOLUME_TILE_RES_PADDED + X;
 				WriteVoxel(VoxelIndex, PhysicalTileDataPtrs[0], Header.AttributesFormats[0], Header.NullTileValues[0]);
 				WriteVoxel(VoxelIndex, PhysicalTileDataPtrs[1], Header.AttributesFormats[1], Header.NullTileValues[1]);
 			}
@@ -273,17 +273,17 @@ bool FSparseVolumeTextureRuntime::SetTileMappings(const TArrayView<const FSparse
 		// Write to tile data
 		for (int32 PhysicalTileIndex = 0; PhysicalTileIndex < Mapping.NumPhysicalTiles; ++PhysicalTileIndex)
 		{
-			for (int32 Z = 0; Z < SPARSE_VOLUME_TILE_RES; ++Z)
+			for (int32 Z = 0; Z < SPARSE_VOLUME_TILE_RES_PADDED; ++Z)
 			{
-				for (int32 Y = 0; Y < SPARSE_VOLUME_TILE_RES; ++Y)
+				for (int32 Y = 0; Y < SPARSE_VOLUME_TILE_RES_PADDED; ++Y)
 				{
-					for (int32 X = 0; X < SPARSE_VOLUME_TILE_RES; ++X)
+					for (int32 X = 0; X < SPARSE_VOLUME_TILE_RES_PADDED; ++X)
 					{
-						const int32 SrcVoxelIndex = PhysicalTileIndex * SVTNumVoxelsPerTile + Z * (SPARSE_VOLUME_TILE_RES * SPARSE_VOLUME_TILE_RES) + Y * SPARSE_VOLUME_TILE_RES + X;
+						const int32 SrcVoxelIndex = PhysicalTileIndex * SVTNumVoxelsPerPaddedTile + Z * (SPARSE_VOLUME_TILE_RES_PADDED * SPARSE_VOLUME_TILE_RES_PADDED) + Y * SPARSE_VOLUME_TILE_RES_PADDED + X;
 						const FVector4f ValueA = ReadVoxel(SrcVoxelIndex, Mapping.TileDataA, Header.AttributesFormats[0]);
 						const FVector4f ValueB = ReadVoxel(SrcVoxelIndex, Mapping.TileDataB, Header.AttributesFormats[1]);
 
-						const FIntVector3 VoxelCoord = DstTileCoord * SPARSE_VOLUME_TILE_RES + FIntVector3(X, Y, Z);
+						const FIntVector3 VoxelCoord = DstTileCoord * SPARSE_VOLUME_TILE_RES_PADDED + FIntVector3(X, Y, Z);
 						const int32 DstVoxelIndex = VoxelCoord.Z * (Header.TileDataVolumeResolution.Y * Header.TileDataVolumeResolution.X) + VoxelCoord.Y * Header.TileDataVolumeResolution.X + VoxelCoord.X;
 						WriteVoxel(DstVoxelIndex, PhysicalTileDataPtrs[0], Header.AttributesFormats[0], ValueA);
 						WriteVoxel(DstVoxelIndex, PhysicalTileDataPtrs[1], Header.AttributesFormats[1], ValueB);
diff --git a/Engine/Source/Runtime/Engine/Private/SparseVolumeTexture/SparseVolumeTextureUtility.h b/Engine/Source/Runtime/Engine/Private/SparseVolumeTexture/SparseVolumeTextureUtility.h
--- a/Engine/Source/Runtime/Engine/Private/SparseVolumeTexture/SparseVolumeTextureUtility.h
+++ b/Engine/Source/Runtime/Engine/Private/SparseVolumeTexture/SparseVolumeTextureUtility.h
@@ -13,6 +13,7 @@ namespace SVT
 	// the asset is imported on. 2048 should be a safe value that should be supported by all our platforms.
 	static constexpr int32 SVTMaxVolumeTextureDim = 2048;
 	static constexpr int32 SVTNumVoxelsPerTile = SPARSE_VOLUME_TILE_RES * SPARSE_VOLUME_TILE_RES * SPARSE_VOLUME_TILE_RES;
+	static constexpr int32 SVTNumVoxelsPerPaddedTile = SPARSE_VOLUME_TILE_RES_PADDED * SPARSE_VOLUME_TILE_RES_PADDED * SPARSE_VOLUME_TILE_RES_PADDED;
 
 	uint32 PackPageTableEntry(const FIntVector3& Coord);
 	FIntVector3 UnpackPageTableEntry(uint32 Packed);
diff --git a/Engine/Source/Runtime/Engine/Public/MaterialCompiler.h b/Engine/Source/Runtime/Engine/Public/MaterialCompiler.h
--- a/Engine/Source/Runtime/Engine/Public/MaterialCompiler.h
+++ b/Engine/Source/Runtime/Engine/Public/MaterialCompiler.h
@@ -426,9 +426,11 @@ public:
 	 * Sample a sparse volume texture page table.
 	 * @param SparseVolumeTextureIndex The code chunk index of the texture.
 	 * @param UVWIndex The UVW coordinate at which to sample the texture.
+	 * @param MipLevelIndex The mip level at which to sample the texture.
+	 * @param SamplerSource The type of sampler to use to sample the texture.
 	 * @return The code chunk index of the result of the texture sample.
 	 */
-	virtual int32 SparseVolumeTextureSamplePageTable(int32 SparseVolumeTextureIndex, int32 UVWIndex, int32 MipLevelIndex) = 0;
+	virtual int32 SparseVolumeTextureSamplePageTable(int32 SparseVolumeTextureIndex, int32 UVWIndex, int32 MipLevelIndex, ESamplerSourceMode SamplerSource) = 0;
 	/**
 	 * Sample a sparse volume texture physical tile data texture.
 	 * @param SparseVolumeTextureIndex The code chunk index of the texture.
@@ -923,7 +925,7 @@ public:
 	virtual int32 SparseVolumeTextureParameter(FName ParameterName, USparseVolumeTexture* InDefaultTexture, int32& TextureReferenceIndex, EMaterialSamplerType SamplerType)									{ return Compiler->SparseVolumeTextureParameter(ParameterName, InDefaultTexture, TextureReferenceIndex, SamplerType); }
 	virtual int32 SparseVolumeTextureUniform(int32 TextureIndex, int32 VectorIndex, UE::Shader::EValueType Type) override																					{ return Compiler->SparseVolumeTextureUniform(TextureIndex, VectorIndex, Type); }
 	virtual int32 SparseVolumeTextureUniformParameter(FName ParameterName, int32 TextureIndex, int32 VectorIndex, UE::Shader::EValueType Type) override														{ return Compiler->SparseVolumeTextureUniformParameter(ParameterName, TextureIndex, VectorIndex, Type); }
-	virtual int32 SparseVolumeTextureSamplePageTable(int32 SparseVolumeTextureIndex, int32 UVWIndex, int32 MipLevelIndex) override																			{ return Compiler->SparseVolumeTextureSamplePageTable(SparseVolumeTextureIndex, UVWIndex, MipLevelIndex); }
+	virtual int32 SparseVolumeTextureSamplePageTable(int32 SparseVolumeTextureIndex, int32 UVWIndex, int32 MipLevelIndex, ESamplerSourceMode SamplerSource) override										{ return Compiler->SparseVolumeTextureSamplePageTable(SparseVolumeTextureIndex, UVWIndex, MipLevelIndex, SamplerSource); }
 	virtual int32 SparseVolumeTextureSamplePhysicalTileData(int32 SparseVolumeTextureIndex, int32 VoxelCoordIndex, int32 PhysicalTileDataIdxIndex) override													{ return Compiler->SparseVolumeTextureSamplePhysicalTileData(SparseVolumeTextureIndex, VoxelCoordIndex, PhysicalTileDataIdxIndex); }
 
 	virtual UObject* GetReferencedTexture(int32 Index) override { return Compiler->GetReferencedTexture(Index); }
diff --git a/Engine/Source/Runtime/Engine/Public/SparseVolumeTexture/SparseVolumeTextureData.h b/Engine/Source/Runtime/Engine/Public/SparseVolumeTexture/SparseVolumeTextureData.h
--- a/Engine/Source/Runtime/Engine/Public/SparseVolumeTexture/SparseVolumeTextureData.h
+++ b/Engine/Source/Runtime/Engine/Public/SparseVolumeTexture/SparseVolumeTextureData.h
@@ -40,6 +40,14 @@ struct ENGINE_API FSparseVolumeTextureDataHeader : public FSparseVolumeTextureHe
 	void Serialize(FArchive& Ar);
 };
 
+struct ENGINE_API FSparseVolumeTextureDataAddressingInfo
+{
+	FIntVector3 VolumeResolution;
+	TextureAddress AddressX;
+	TextureAddress AddressY;
+	TextureAddress AddressZ;
+};
+
 // Holds the data for a SparseVolumeTexture that is stored on disk. It only has a single mip after importing a source asset. The mip chain is built during cooking.
 // Tiles are addressed by a flat index; unlike the runtime representation, this one stores all tiles in a 1D array and doesn't have the concept of a 3D physical tile texture.
 // The page table itself is 3D though.
@@ -57,7 +65,9 @@ struct ENGINE_API FSparseVolumeTextureData
 	bool Construct(const ISparseVolumeTextureDataConstructionAdapter& Adapter);
 	uint32 ReadPageTable(const FIntVector3& PageTableCoord, int32 MipLevel) const;
 	FVector4f ReadTileDataVoxel(int32 TileIndex, const FIntVector3& TileDataCoord, int32 AttributesIdx) const;
-	FVector4f Load(const FIntVector3& VolumeCoord, int32 MipLevel, int32 AttributesIdx) const;
+	FVector4f Load(const FIntVector3& VolumeCoord, int32 MipLevel, int32 AttributesIdx, const FSparseVolumeTextureDataAddressingInfo& AddressingInfo) const;
 	void WriteTileDataVoxel(int32 TileIndex, const FIntVector3& TileDataCoord, int32 AttributesIdx, const FVector4f& Value, int32 DstComponent = -1);
-	void GenerateMipMaps(int32 NumMipLevels = -1);
+	void GenerateMipMaps(const FSparseVolumeTextureDataAddressingInfo& AddressingInfo, int32 NumMipLevels = -1);
+	void GenerateBorderVoxels(const FSparseVolumeTextureDataAddressingInfo& AddressingInfo, int32 MipLevel, const TArray<FIntVector3>& PageCoords);
+	void BuildDerivedData(const FSparseVolumeTextureDataAddressingInfo& AddressingInfo, int32 NumMipLevels = -1);
 };
diff --git a/Engine/Source/Runtime/Renderer/Private/SparseVolumeTexture/SparseVolumeTextureViewerRendering.cpp b/Engine/Source/Runtime/Renderer/Private/SparseVolumeTexture/SparseVolumeTextureViewerRendering.cpp
--- a/Engine/Source/Runtime/Renderer/Private/SparseVolumeTexture/SparseVolumeTextureViewerRendering.cpp
+++ b/Engine/Source/Runtime/Renderer/Private/SparseVolumeTexture/SparseVolumeTextureViewerRendering.cpp
@@ -57,6 +57,7 @@ class FVisualizeSparseVolumeTexturePS : public FGlobalShader
 	BEGIN_SHADER_PARAMETER_STRUCT(FParameters, )
 		SHADER_PARAMETER_STRUCT_REF(FViewUniformShaderParameters, ViewUniformBuffer)
 		RENDER_TARGET_BINDING_SLOTS()
+		SHADER_PARAMETER_SAMPLER(SamplerState, TileDataTextureSampler)
 		SHADER_PARAMETER_TEXTURE(Texture3D<uint>, SparseVolumeTexturePageTable)
 		SHADER_PARAMETER_TEXTURE(Texture3D, SparseVolumeTextureA)
 		SHADER_PARAMETER_TEXTURE(Texture3D, SparseVolumeTextureB)
@@ -130,6 +131,7 @@ void AddSparseVolumeTextureViewerRenderPass(FRDGBuilder& GraphBuilder, FSceneRen
 			FVisualizeSparseVolumeTexturePS::FParameters* PsPassParameters = GraphBuilder.AllocParameters<FVisualizeSparseVolumeTexturePS::FParameters>();
 			PsPassParameters->ViewUniformBuffer = View.ViewUniformBuffer;
 			PsPassParameters->RenderTargets[0] = FRenderTargetBinding(SceneTextures.Color.Target, ERenderTargetLoadAction::ELoad);
+			PsPassParameters->TileDataTextureSampler = TStaticSamplerState<SF_Point, AM_Clamp, AM_Clamp, AM_Clamp>::GetRHI();
 			PsPassParameters->SparseVolumeTexturePageTable = GBlackUintVolumeTexture->TextureRHI;
 			PsPassParameters->SparseVolumeTextureA = GBlackVolumeTexture->TextureRHI;
 			PsPassParameters->SparseVolumeTextureB = GBlackVolumeTexture->TextureRHI;
