Move Handler in to UE namespace to avoid name collisions with third party code

diff --git a/Engine/Plugins/Compression/OodleNetwork/Source/Private/OodleNetworkHandlerComponent.cpp b/Engine/Plugins/Compression/OodleNetwork/Source/Private/OodleNetworkHandlerComponent.cpp
--- a/Engine/Plugins/Compression/OodleNetwork/Source/Private/OodleNetworkHandlerComponent.cpp
+++ b/Engine/Plugins/Compression/OodleNetwork/Source/Private/OodleNetworkHandlerComponent.cpp
@@ -538,7 +538,7 @@ void OodleNetworkHandlerComponent::Initialize()
 
 #endif
 
-		EOodleNetworkEnableMode EnableMode = (Handler->Mode == Handler::Mode::Server ? ServerEnableMode : ClientEnableMode);
+		EOodleNetworkEnableMode EnableMode = (Handler->Mode == UE::Handler::Mode::Server ? ServerEnableMode : ClientEnableMode);
 
 		if (EnableMode == EOodleNetworkEnableMode::AlwaysEnabled)
 		{
@@ -920,7 +920,7 @@ void OodleNetworkHandlerComponent::InitializePacketLogs()
 	// @todo #JohnB: Convert this code so that just one capture file is used for all connections, per session
 	//					(could set it up much like the dictionary sharing code)
 	//					Downside, is potential for corruption. Lots of files is a bit unwieldy, yet very stable.
-	if (bCaptureMode && Handler->Mode == Handler::Mode::Server && InPacketLog == nullptr && OutPacketLog == nullptr)
+	if (bCaptureMode && Handler->Mode == UE::Handler::Mode::Server && InPacketLog == nullptr && OutPacketLog == nullptr)
 	{
 		IPlatformFile& PlatformFile = FPlatformFileManager::Get().GetPlatformFile();
 		FString ReadOutputLogDirectory = FPaths::Combine(*GOodleSaveDir, TEXT("Server"));
@@ -1027,7 +1027,7 @@ void OodleNetworkHandlerComponent::Incoming(FIncomingPacketRef PacketRef)
 		// If the packet is not compressed, no further processing is necessary
 		if (bCompressedPacket)
 		{
-			const bool bIsServer = (Handler->Mode == Handler::Mode::Server);
+			const bool bIsServer = (Handler->Mode == UE::Handler::Mode::Server);
 
 			// Lazy-loading of dictionary when EOodleNetworkEnableMode::WhenCompressedPacketReceived is active
 			if (!bInitializedDictionaries && (bIsServer ? ServerEnableMode : ClientEnableMode) == EOodleNetworkEnableMode::WhenCompressedPacketReceived)
@@ -1123,7 +1123,7 @@ void OodleNetworkHandlerComponent::Incoming(FIncomingPacketRef PacketRef)
 						Packet.SetData(DecompressedData, DecompressedLength * 8);
 
 #if !UE_BUILD_SHIPPING || OODLE_DEV_SHIPPING
-						if (bCaptureMode && Handler->Mode == Handler::Mode::Server && InPacketLog != nullptr)
+						if (bCaptureMode && Handler->Mode == UE::Handler::Mode::Server && InPacketLog != nullptr)
 						{
 #if !UE_BUILD_SHIPPING
 							QUICK_SCOPE_CYCLE_COUNTER(STAT_Oodle_InCaptureTime);
@@ -1186,7 +1186,7 @@ void OodleNetworkHandlerComponent::Incoming(FIncomingPacketRef PacketRef)
 			}
 
 #if !UE_BUILD_SHIPPING || OODLE_DEV_SHIPPING
-			if (bCaptureMode && Handler->Mode == Handler::Mode::Server && InPacketLog != nullptr)
+			if (bCaptureMode && Handler->Mode == UE::Handler::Mode::Server && InPacketLog != nullptr)
 			{
 				uint32 SizeOfPacket = Packet.GetBytesLeft();
 
@@ -1213,7 +1213,7 @@ void OodleNetworkHandlerComponent::Outgoing(FBitWriter& Packet, FOutPacketTraits
 	if (bEnableOodle)
 	{
 #if !UE_BUILD_SHIPPING || OODLE_DEV_SHIPPING
-		if (bCaptureMode && Handler->Mode == Handler::Mode::Server && OutPacketLog != nullptr)
+		if (bCaptureMode && Handler->Mode == UE::Handler::Mode::Server && OutPacketLog != nullptr)
 		{
 			uint32 SizeOfPacket = Packet.GetNumBytes();
 
@@ -1239,7 +1239,7 @@ void OodleNetworkHandlerComponent::Outgoing(FBitWriter& Packet, FOutPacketTraits
 		static uint8 CompressedData[MAX_OODLE_BUFFER];
 
 		FOodleNetworkAnalyticsVars* AnalyticsVars = (bOodleNetworkAnalytics && NetAnalyticsData.IsValid()) ? NetAnalyticsData->GetLocalData() : nullptr;
-		const bool bIsServer = (Handler->Mode == Handler::Mode::Server);
+		const bool bIsServer = (Handler->Mode == UE::Handler::Mode::Server);
 		FOodleNetworkDictionary* CurDict = (bIsServer ? ServerDictionary.Get() : ClientDictionary.Get());
 		uint32 UncompressedBytes = Packet.GetNumBytes();
 		bool bSkipCompressionClientDisabled = false;
@@ -1441,7 +1441,7 @@ bool OodleNetworkHandlerComponent::IsCompressionActive() const
 
 	if (bEnableOodle && Handler)
 	{
-		const bool bIsServer = (Handler->Mode == Handler::Mode::Server);
+		const bool bIsServer = (Handler->Mode == UE::Handler::Mode::Server);
 		FOodleNetworkDictionary* CurDict = (bIsServer ? ServerDictionary.Get() : ClientDictionary.Get());
 
 		const bool bSkipCompression =
@@ -1475,7 +1475,7 @@ void OodleNetworkHandlerComponent::NotifyAnalyticsProvider()
 
 		if (Handler->GetProvider().IsValid() && Aggregator.IsValid())
 		{
-			const bool bIsServer = (Handler->Mode == Handler::Mode::Server);
+			const bool bIsServer = (Handler->Mode == UE::Handler::Mode::Server);
 
 			if (bIsServer)
 			{
@@ -1587,7 +1587,7 @@ static bool OodleExec(UWorld* InWorld, const TCHAR* Cmd, FOutputDevice& Ar)
 			bool bSendServerCommand = FModuleManager::Get().IsModuleLoaded(TEXT("NetcodeUnitTest"));
 
 			bSendServerCommand = bSendServerCommand && OodleComponentList.Num() > 0 &&
-									OodleComponentList[0]->Handler->Mode == Handler::Mode::Client;
+				OodleComponentList[0]->Handler->Mode == UE::Handler::Mode::Client;
 
 			if (bSendServerCommand)
 			{
@@ -1689,7 +1689,7 @@ static bool OodleExec(UWorld* InWorld, const TCHAR* Cmd, FOutputDevice& Ar)
 				bool bSendServerCommand = FModuleManager::Get().IsModuleLoaded(TEXT("NetcodeUnitTest"));
 
 				bSendServerCommand = bSendServerCommand && OodleComponentList.Num() > 0 &&
-										OodleComponentList[0]->Handler->Mode == Handler::Mode::Client;
+										OodleComponentList[0]->Handler->Mode == UE::Handler::Mode::Client;
 
 				if (bSendServerCommand)
 				{
diff --git a/Engine/Plugins/Online/OnlineSubsystemSteam/Source/Private/OnlineAuthHandlerSteam.cpp b/Engine/Plugins/Online/OnlineSubsystemSteam/Source/Private/OnlineAuthHandlerSteam.cpp
--- a/Engine/Plugins/Online/OnlineSubsystemSteam/Source/Private/OnlineAuthHandlerSteam.cpp
+++ b/Engine/Plugins/Online/OnlineSubsystemSteam/Source/Private/OnlineAuthHandlerSteam.cpp
@@ -124,7 +124,7 @@ FSteamAuthHandlerComponent::~FSteamAuthHandlerComponent()
 		return;
 	}
 
-	if (Handler->Mode == Handler::Mode::Client)
+	if (Handler->Mode == UE::Handler::Mode::Client)
 	{
 		AuthInterface->RevokeTicket(TicketHandle);
 	}
@@ -168,7 +168,7 @@ void FSteamAuthHandlerComponent::NotifyHandshakeBegin()
 		return;
 	}
 
-	if (Handler->Mode == Handler::Mode::Client)
+	if (Handler->Mode == UE::Handler::Mode::Client)
 	{
 		SendAuthKey(true);
 	}
@@ -280,7 +280,7 @@ void FSteamAuthHandlerComponent::RequestResend()
 	ResendWriter.WriteBit(1);
 
 	// Steam Auth is so simplistic that we really only have two messages we need to handle.
-	ResendingPacket.Type = (Handler->Mode == Handler::Mode::Server) ? 
+	ResendingPacket.Type = (Handler->Mode == UE::Handler::Mode::Server) ? 
 		ESteamAuthMsgType::ResendKey : ESteamAuthMsgType::ResendResult;
 
 	ResendWriter << ResendingPacket;
@@ -380,7 +380,7 @@ void FSteamAuthHandlerComponent::Incoming(FBitReader& Packet)
 			SendAuthKey(false);
 		}
 	}
-	else if (Handler && Handler->Mode == Handler::Mode::Server && Header.Type == ESteamAuthMsgType::ResendResult)
+	else if (Handler && Handler->Mode == UE::Handler::Mode::Server && Header.Type == ESteamAuthMsgType::ResendResult)
 	{
 		if (State == ESteamAuthHandlerState::Initialized)
 		{
diff --git a/Engine/Plugins/Runtime/PacketHandlers/AESGCMHandlerComponent/Source/Private/AESGCMHandlerComponent.cpp b/Engine/Plugins/Runtime/PacketHandlers/AESGCMHandlerComponent/Source/Private/AESGCMHandlerComponent.cpp
--- a/Engine/Plugins/Runtime/PacketHandlers/AESGCMHandlerComponent/Source/Private/AESGCMHandlerComponent.cpp
+++ b/Engine/Plugins/Runtime/PacketHandlers/AESGCMHandlerComponent/Source/Private/AESGCMHandlerComponent.cpp
@@ -55,7 +55,7 @@ bool FAESGCMHandlerComponent::IsEncryptionEnabled() const
 void FAESGCMHandlerComponent::Initialize()
 {
 	SetActive(true);
-	SetState(Handler::Component::State::Initialized);
+	SetState(UE::Handler::Component::State::Initialized);
 	Initialized();
 }
 
diff --git a/Engine/Plugins/Runtime/PacketHandlers/AESHandlerComponent/Source/Private/AESHandlerComponent.cpp b/Engine/Plugins/Runtime/PacketHandlers/AESHandlerComponent/Source/Private/AESHandlerComponent.cpp
--- a/Engine/Plugins/Runtime/PacketHandlers/AESHandlerComponent/Source/Private/AESHandlerComponent.cpp
+++ b/Engine/Plugins/Runtime/PacketHandlers/AESHandlerComponent/Source/Private/AESHandlerComponent.cpp
@@ -59,7 +59,7 @@ bool FAESHandlerComponent::IsEncryptionEnabled() const
 void FAESHandlerComponent::Initialize()
 {
 	SetActive(true);
-	SetState(Handler::Component::State::Initialized);
+	SetState(UE::Handler::Component::State::Initialized);
 	Initialized();
 }
 
diff --git a/Engine/Plugins/Runtime/PacketHandlers/DTLSHandlerComponent/Source/Private/DTLSHandlerComponent.cpp b/Engine/Plugins/Runtime/PacketHandlers/DTLSHandlerComponent/Source/Private/DTLSHandlerComponent.cpp
--- a/Engine/Plugins/Runtime/PacketHandlers/DTLSHandlerComponent/Source/Private/DTLSHandlerComponent.cpp
+++ b/Engine/Plugins/Runtime/PacketHandlers/DTLSHandlerComponent/Source/Private/DTLSHandlerComponent.cpp
@@ -92,7 +92,7 @@ void FDTLSHandlerComponent::SetEncryptionData(const FEncryptionData& EncryptionD
 	{
 		CertId = EncryptionData.Identifier;
 
-		if (Handler->Mode == Handler::Mode::Client)
+		if (Handler->Mode == UE::Handler::Mode::Client)
 		{
 			RemoteFingerprint = MakeUnique<FDTLSFingerprint>();
 
@@ -110,7 +110,7 @@ void FDTLSHandlerComponent::SetEncryptionData(const FEncryptionData& EncryptionD
 
 void FDTLSHandlerComponent::EnableEncryption()
 {
-	EDTLSContextType ContextType = (Handler->Mode == Handler::Mode::Server) ? EDTLSContextType::Server : EDTLSContextType::Client;
+	EDTLSContextType ContextType = (Handler->Mode == UE::Handler::Mode::Server) ? EDTLSContextType::Server : EDTLSContextType::Client;
 
 	UE_LOG(LogDTLSHandler, Log, TEXT("EnableEncryption: %s"), LexToString(ContextType));
 
@@ -157,7 +157,7 @@ bool FDTLSHandlerComponent::IsEncryptionEnabled() const
 void FDTLSHandlerComponent::Initialize()
 {
 	SetActive(true);
-	SetState(Handler::Component::State::Initialized);
+	SetState(UE::Handler::Component::State::Initialized);
 	Initialized();
 }
 
@@ -205,7 +205,7 @@ void FDTLSHandlerComponent::Incoming(FBitReader& Packet)
 				if (InternalState == EDTLSHandlerState::Unencrypted)
 				{
 					// server can move to handshaking now
-					check(Handler->Mode == Handler::Mode::Server);
+					check(Handler->Mode == UE::Handler::Mode::Server);
 					InternalState = EDTLSHandlerState::Handshaking;
 				}
 
diff --git a/Engine/Source/Runtime/Engine/Private/NetConnection.cpp b/Engine/Source/Runtime/Engine/Private/NetConnection.cpp
--- a/Engine/Source/Runtime/Engine/Private/NetConnection.cpp
+++ b/Engine/Source/Runtime/Engine/Private/NetConnection.cpp
@@ -635,7 +635,7 @@ void UNetConnection::InitHandler()
 
 		if (Handler.IsValid())
 		{
-			Handler::Mode Mode = Driver->ServerConnection != nullptr ? Handler::Mode::Client : Handler::Mode::Server;
+			UE::Handler::Mode Mode = Driver->ServerConnection != nullptr ? UE::Handler::Mode::Client : UE::Handler::Mode::Server;
 
 			FPacketHandlerNotifyAddHandler NotifyAddHandler;
 			
diff --git a/Engine/Source/Runtime/Engine/Private/NetDriver.cpp b/Engine/Source/Runtime/Engine/Private/NetDriver.cpp
--- a/Engine/Source/Runtime/Engine/Private/NetDriver.cpp
+++ b/Engine/Source/Runtime/Engine/Private/NetDriver.cpp
@@ -1513,7 +1513,7 @@ void UNetDriver::InitConnectionlessHandler()
 		if (ConnectionlessHandler.IsValid())
 		{
 			ConnectionlessHandler->NotifyAnalyticsProvider(AnalyticsProvider, AnalyticsAggregator);
-			ConnectionlessHandler->Initialize(Handler::Mode::Server, MAX_PACKET_SIZE, true, nullptr, nullptr, NetDriverDefinition);
+			ConnectionlessHandler->Initialize(UE::Handler::Mode::Server, MAX_PACKET_SIZE, true, nullptr, nullptr, NetDriverDefinition);
 
 			// Add handling for the stateless connect handshake, for connectionless packets, as the outermost layer
 			TSharedPtr<HandlerComponent> NewComponent =
diff --git a/Engine/Source/Runtime/Engine/Private/PacketHandlers/StatelessConnectHandlerComponent.cpp b/Engine/Source/Runtime/Engine/Private/PacketHandlers/StatelessConnectHandlerComponent.cpp
--- a/Engine/Source/Runtime/Engine/Private/PacketHandlers/StatelessConnectHandlerComponent.cpp
+++ b/Engine/Source/Runtime/Engine/Private/PacketHandlers/StatelessConnectHandlerComponent.cpp
@@ -466,7 +466,7 @@ void StatelessConnectHandlerComponent::SendInitialPacket(EHandshakeVersion Hands
 {
 	using namespace UE::Net;
 
-	if (Handler->Mode == Handler::Mode::Client)
+	if (Handler->Mode == UE::Handler::Mode::Client)
 	{
 		UNetConnection* ServerConn = (Driver != nullptr ? ToRawPtr(Driver->ServerConnection) : nullptr);
 
@@ -831,7 +831,7 @@ void StatelessConnectHandlerComponent::SetDriver(UNetDriver* InDriver)
 {
 	Driver = InDriver;
 
-	if (Handler->Mode == Handler::Mode::Server)
+	if (Handler->Mode == UE::Handler::Mode::Server)
 	{
 		StatelessConnectHandlerComponent* StatelessComponent = Driver->StatelessConnectComponent.Pin().Get();
 
@@ -885,7 +885,7 @@ void StatelessConnectHandlerComponent::Initialize()
 	using namespace UE::Net;
 
 	// On the server, initializes immediately. Clientside doesn't initialize until handshake completes.
-	if (Handler->Mode == Handler::Mode::Server)
+	if (Handler->Mode == UE::Handler::Mode::Server)
 	{
 		Initialized();
 
@@ -980,12 +980,12 @@ void StatelessConnectHandlerComponent::Incoming(FBitReader& Packet)
 		if (bHandshakePacket)
 		{
 			const bool bIsChallengePacket = HandshakeData.HandshakePacketType == EHandshakePacketType::Challenge && HandshakeData.Timestamp > 0.0;
-			const bool bIsInitialChallengePacket = bIsChallengePacket && State != Handler::Component::State::Initialized;
+			const bool bIsInitialChallengePacket = bIsChallengePacket && State != UE::Handler::Component::State::Initialized;
 			const bool bIsUpgradePacket = HandshakeData.HandshakePacketType == EHandshakePacketType::VersionUpgrade;
 
-			if (Handler->Mode == Handler::Mode::Client && bHasValidClientID && (bHasValidSessionID || bIsInitialChallengePacket || bIsUpgradePacket))
+			if (Handler->Mode == UE::Handler::Mode::Client && bHasValidClientID && (bHasValidSessionID || bIsInitialChallengePacket || bIsUpgradePacket))
 			{
-				if (State == Handler::Component::State::UnInitialized || State == Handler::Component::State::InitializedOnLocal)
+				if (State == UE::Handler::Component::State::UnInitialized || State == UE::Handler::Component::State::InitializedOnLocal)
 				{
 					if (HandshakeData.bRestartHandshake)
 					{
@@ -1007,7 +1007,7 @@ void StatelessConnectHandlerComponent::Incoming(FBitReader& Packet)
 						SendChallengeResponse(HandshakeData.RemoteCurVersion, HandshakeData.SecretId, HandshakeData.Timestamp, HandshakeData.Cookie);
 
 						// Utilize this state as an intermediary, indicating that the challenge response has been sent
-						SetState(Handler::Component::State::InitializedOnLocal);
+						SetState(UE::Handler::Component::State::InitializedOnLocal);
 					}
 					// Receiving challenge ack, verify the timestamp is < 0.0f
 					else if (HandshakeData.HandshakePacketType == EHandshakePacketType::Ack && HandshakeData.Timestamp < 0.0)
@@ -1032,7 +1032,7 @@ void StatelessConnectHandlerComponent::Incoming(FBitReader& Packet)
 						}
 
 						// Now finish initializing the handler - flushing the queued packet buffer in the process.
-						SetState(Handler::Component::State::Initialized);
+						SetState(UE::Handler::Component::State::Initialized);
 						Initialized();
 
 						bRestartedHandshake = false;
@@ -1134,7 +1134,7 @@ void StatelessConnectHandlerComponent::Incoming(FBitReader& Packet)
 
 							bRestartedHandshake = true;
 
-							SetState(Handler::Component::State::UnInitialized);
+							SetState(UE::Handler::Component::State::UnInitialized);
 							SendInitialPacket(LastRemoteHandshakeVersion);
 						}
 						else if (WithinHandshakeLogLimit())
@@ -1167,7 +1167,7 @@ void StatelessConnectHandlerComponent::Incoming(FBitReader& Packet)
 					// Ignore, could be a dupe/out-of-order challenge packet
 				}
 			}
-			else if (Handler->Mode == Handler::Mode::Server && bHasValidSessionID && bHasValidClientID)
+			else if (Handler->Mode == UE::Handler::Mode::Server && bHasValidSessionID && bHasValidClientID)
 			{
 				if (LastChallengeSuccessAddress.IsValid())
 				{
@@ -1224,7 +1224,7 @@ void StatelessConnectHandlerComponent::Incoming(FBitReader& Packet)
 	{
 		// Servers should wipe LastChallengeSuccessAddress shortly after the first non-handshake packet is received by the client,
 		// in order to disable challenge ack resending
-		if (LastInitTimestamp != 0.0 && LastChallengeSuccessAddress.IsValid() && Handler->Mode == Handler::Mode::Server)
+		if (LastInitTimestamp != 0.0 && LastChallengeSuccessAddress.IsValid() && Handler->Mode == UE::Handler::Mode::Server)
 		{
 			// Restart handshakes require extra time before disabling challenge ack resends, as NetConnection packets will already be in flight
 			const double RestartHandshakeAckResendWindow = 10.0;
@@ -1321,7 +1321,7 @@ void StatelessConnectHandlerComponent::IncomingConnectionless(FIncomingPacketRef
 			const bool bInitialConnect = HandshakeData.HandshakePacketType == EHandshakePacketType::InitialPacket &&
 												HandshakeData.Timestamp == 0.0;
 
-			if (Handler->Mode == Handler::Mode::Server && bValidVersion && (bHasValidSessionID || bInitialConnect))
+			if (Handler->Mode == UE::Handler::Mode::Server && bValidVersion && (bHasValidSessionID || bInitialConnect))
 			{
 				if (bInitialConnect)
 				{
@@ -1381,7 +1381,7 @@ void StatelessConnectHandlerComponent::IncomingConnectionless(FIncomingPacketRef
 					}
 				}
 			}
-			else if (Handler->Mode == Handler::Mode::Server && bValidVersion && !bHasValidSessionID)
+			else if (Handler->Mode == UE::Handler::Mode::Server && bValidVersion && !bHasValidSessionID)
 			{
 #if !UE_BUILD_SHIPPING
 				UE_LOG(LogHandshake, Log, TEXT("IncomingConnectionless: Rejecting packet with invalid session id: %u vs %u."),
@@ -1390,7 +1390,7 @@ void StatelessConnectHandlerComponent::IncomingConnectionless(FIncomingPacketRef
 
 				Packet.SetError();
 			}
-			else if (Handler->Mode == Handler::Mode::Server && !bValidVersion && bInitialConnect &&
+			else if (Handler->Mode == UE::Handler::Mode::Server && !bValidVersion && bInitialConnect &&
 						HandshakeData.RemoteCurVersion >= EHandshakeVersion::NetCLUpgradeMessage)
 			{
 				// Limit of 512 upgrade message packets, over 5 minutes
@@ -1565,11 +1565,11 @@ bool StatelessConnectHandlerComponent::ParseHandshakePacket(FBitReader& Packet,
 	}
 	else if (bRestartHandshakePacket)
 	{
-		bValidPacket = !Packet.IsError() && OutResult.bRestartHandshake && Handler->Mode == Handler::Mode::Client;
+		bValidPacket = !Packet.IsError() && OutResult.bRestartHandshake && Handler->Mode == UE::Handler::Mode::Client;
 	}
 	else if (bVersionUpgradePacket)
 	{
-		bValidPacket = !Packet.IsError() && !OutResult.bRestartHandshake && Handler->Mode == Handler::Mode::Client;
+		bValidPacket = !Packet.IsError() && !OutResult.bRestartHandshake && Handler->Mode == UE::Handler::Mode::Client;
 	}
 
 	if (bValidPacket)
@@ -1609,7 +1609,7 @@ bool StatelessConnectHandlerComponent::ParseHandshakePacketOriginal(FBitReader&
 		}
 		else if (OutResult.Timestamp > 0.0)
 		{
-			if (Handler->Mode == Handler::Mode::Client)
+			if (Handler->Mode == UE::Handler::Mode::Client)
 			{
 				OutResult.HandshakePacketType = EHandshakePacketType::Challenge;
 			}
@@ -1633,7 +1633,7 @@ bool StatelessConnectHandlerComponent::ParseHandshakePacketOriginal(FBitReader&
 	{
 		OutResult.HandshakePacketType = EHandshakePacketType::RestartHandshake;
 		OutResult.bRestartHandshake = !!Packet.ReadBit();
-		bValidPacket = !Packet.IsError() && OutResult.bRestartHandshake && Handler->Mode == Handler::Mode::Client;
+		bValidPacket = !Packet.IsError() && OutResult.bRestartHandshake && Handler->Mode == UE::Handler::Mode::Client;
 	}
 
 	return bValidPacket;
@@ -1797,9 +1797,9 @@ void StatelessConnectHandlerComponent::Tick(float DeltaTime)
 {
 	using namespace UE::Net;
 
-	if (Handler->Mode == Handler::Mode::Client)
+	if (Handler->Mode == UE::Handler::Mode::Client)
 	{
-		if (State != Handler::Component::State::Initialized && LastClientSendTimestamp != 0.0)
+		if (State != UE::Handler::Component::State::Initialized && LastClientSendTimestamp != 0.0)
 		{
 			double LastSendTimeDiff = FPlatformTime::Seconds() - LastClientSendTimestamp;
 
@@ -1809,10 +1809,10 @@ void StatelessConnectHandlerComponent::Tick(float DeltaTime)
 
 				if (bRestartChallenge)
 				{
-					SetState(Handler::Component::State::UnInitialized);
+					SetState(UE::Handler::Component::State::UnInitialized);
 				}
 
-				if (State == Handler::Component::State::UnInitialized)
+				if (State == UE::Handler::Component::State::UnInitialized)
 				{
 					UE_LOG(LogHandshake, Verbose, TEXT("Initial handshake packet timeout - resending."));
 
@@ -1834,7 +1834,7 @@ void StatelessConnectHandlerComponent::Tick(float DeltaTime)
 
 					SendInitialPacket(ResendVersion);
 				}
-				else if (State == Handler::Component::State::InitializedOnLocal && LastTimestamp != 0.0)
+				else if (State == UE::Handler::Component::State::InitializedOnLocal && LastTimestamp != 0.0)
 				{
 					UE_LOG(LogHandshake, Verbose, TEXT("Challenge response packet timeout - resending."));
 
diff --git a/Engine/Source/Runtime/PacketHandlers/PacketHandler/Private/PacketHandler.cpp b/Engine/Source/Runtime/PacketHandlers/PacketHandler/Private/PacketHandler.cpp
--- a/Engine/Source/Runtime/PacketHandlers/PacketHandler/Private/PacketHandler.cpp
+++ b/Engine/Source/Runtime/PacketHandlers/PacketHandler/Private/PacketHandler.cpp
@@ -90,7 +90,7 @@ BufferedPacket::~BufferedPacket()
  */
 
 PacketHandler::PacketHandler(FDDoSDetection* InDDoS/*=nullptr*/)
-	: Mode(Handler::Mode::Client)
+	: Mode(UE::Handler::Mode::Client)
 	, bConnectionlessHandler(false)
 	, DDoS(InDDoS)
 	, LowLevelSendDel()
@@ -99,7 +99,7 @@ PacketHandler::PacketHandler(FDDoSDetection* InDDoS/*=nullptr*/)
 	, IncomingPacket()
 	, HandlerComponents()
 	, MaxPacketBits(0)
-	, State(Handler::State::Uninitialized)
+	, State(UE::Handler::State::Uninitialized)
 	, BufferedPackets()
 	, QueuedPackets()
 	, QueuedRawPackets()
@@ -153,7 +153,7 @@ FPacketHandlerAddComponentDelegate& PacketHandler::GetAddComponentDelegate()
 	return AddComponentDelegate;
 }
 
-void PacketHandler::Initialize(Handler::Mode InMode, uint32 InMaxPacketBits, bool bConnectionlessOnly/*=false*/,
+void PacketHandler::Initialize(UE::Handler::Mode InMode, uint32 InMaxPacketBits, bool bConnectionlessOnly/*=false*/,
 								TSharedPtr<IAnalyticsProvider> InProvider/*=nullptr*/, FDDoSDetection* InDDoS/*=nullptr*/, FName InDriverProfile/*=NAME_None*/)
 {
 	Mode = InMode;
@@ -245,7 +245,7 @@ void PacketHandler::NotifyAnalyticsProvider(TSharedPtr<IAnalyticsProvider> InPro
 	Provider = InProvider;
 	Aggregator = InAggregator;
 
-	if (State != Handler::State::Uninitialized)
+	if (State != UE::Handler::State::Uninitialized)
 	{
 		for (const TSharedPtr<HandlerComponent>& CurComponent : HandlerComponents)
 		{
@@ -259,11 +259,11 @@ void PacketHandler::NotifyAnalyticsProvider(TSharedPtr<IAnalyticsProvider> InPro
 
 void PacketHandler::InitializeComponents()
 {
-	if (State == Handler::State::Uninitialized)
+	if (State == UE::Handler::State::Uninitialized)
 	{
 		if (HandlerComponents.Num() > 0)
 		{
-			SetState(Handler::State::InitializingComponents);
+			SetState(UE::Handler::State::InitializingComponents);
 		}
 		else
 		{
@@ -308,7 +308,7 @@ void PacketHandler::BeginHandshaking(FPacketHandlerHandshakeComplete InHandshake
 void PacketHandler::AddHandler(TSharedPtr<HandlerComponent>& NewHandler, bool bDeferInitialize/*=false*/)
 {
 	// This is never valid. Can end up silently changing maximum allow packet size, which could cause failure to send packets.
-	if (State != Handler::State::Uninitialized)
+	if (State != UE::Handler::State::Uninitialized)
 	{
 		LowLevelFatalError(TEXT("Handler added during runtime."));
 		return;
@@ -590,7 +590,7 @@ EIncomingResult PacketHandler::Incoming_Internal(FReceivedPacketView& PacketView
 
 		FPacketAudit::CheckStage(TEXT("PostPacketHandler"), ProcessedPacketReader);
 
-		if (State == Handler::State::Uninitialized)
+		if (State == UE::Handler::State::Uninitialized)
 		{
 			UpdateInitialState();
 		}
@@ -737,13 +737,13 @@ const ProcessedPacket PacketHandler::Outgoing_Internal(uint8* Packet, int32 Coun
 	{
 		OutgoingPacket.Reset();
 
-		if (State == Handler::State::Uninitialized)
+		if (State == UE::Handler::State::Uninitialized)
 		{
 			UpdateInitialState();
 		}
 
 
-		if (State == Handler::State::Initialized)
+		if (State == UE::Handler::State::Initialized)
 		{
 			OutgoingPacket.SerializeBits(Packet, CountBits);
 
@@ -813,7 +813,7 @@ PRAGMA_ENABLE_DEPRECATION_WARNINGS
 			}
 		}
 		// Buffer any packets being sent from game code until processors are initialized
-		else if (State == Handler::State::InitializingComponents && CountBits > 0)
+		else if (State == UE::Handler::State::InitializingComponents && CountBits > 0)
 		{
 			if (bConnectionless)
 			{
@@ -926,7 +926,7 @@ void PacketHandler::SendHandlerPacket(HandlerComponent* InComponent, FBitWriter&
 	// @todo #JohnB: There is duplication between this function and others, it would be nice to reduce this.
 
 	// Prevent any cases where a send happens before the handler is ready.
-	check(State != Handler::State::Uninitialized);
+	check(State != UE::Handler::State::Uninitialized);
 
 	if (LowLevelSendDel.IsBound())
 	{
@@ -997,7 +997,7 @@ PRAGMA_ENABLE_DEPRECATION_WARNINGS
 	}
 }
 
-void PacketHandler::SetState(Handler::State InState)
+void PacketHandler::SetState(UE::Handler::State InState)
 {
 	if (InState == State)
 	{
@@ -1011,7 +1011,7 @@ void PacketHandler::SetState(Handler::State InState)
 
 void PacketHandler::UpdateInitialState()
 {
-	if (State == Handler::State::Uninitialized)
+	if (State == UE::Handler::State::Uninitialized)
 	{
 		if (HandlerComponents.Num() > 0)
 		{
@@ -1063,7 +1063,7 @@ void PacketHandler::HandlerInitialized()
 
 	BufferedConnectionlessPackets.Empty();
 
-	SetState(Handler::State::Initialized);
+	SetState(UE::Handler::State::Initialized);
 
 	if (bBeganHandshaking)
 	{
@@ -1074,7 +1074,7 @@ void PacketHandler::HandlerInitialized()
 void PacketHandler::HandlerComponentInitialized(HandlerComponent* InComponent)
 {
 	// Check if all handlers are initialized
-	if (State != Handler::State::Initialized)
+	if (State != UE::Handler::State::Initialized)
 	{
 		bool bAllInitialized = true;
 		bool bEncounteredComponent = false;
@@ -1216,7 +1216,7 @@ int32 PacketHandler::GetTotalReservedPacketBits()
 
 HandlerComponent::HandlerComponent()
 	: Handler(nullptr)
-	, State(Handler::Component::State::UnInitialized)
+	, State(UE::Handler::Component::State::UnInitialized)
 	, MaxOutgoingBits(0)
 	, bRequiresHandshake(false)
 	, bRequiresReliability(false)
@@ -1227,7 +1227,7 @@ HandlerComponent::HandlerComponent()
 
 HandlerComponent::HandlerComponent(FName InName)
 	: Handler(nullptr)
-	, State(Handler::Component::State::UnInitialized)
+	, State(UE::Handler::Component::State::UnInitialized)
 	, MaxOutgoingBits(0)
 	, bRequiresHandshake(false)
 	, bRequiresReliability(false)
@@ -1247,7 +1247,7 @@ void HandlerComponent::SetActive(bool Active)
 	bActive = Active;
 }
 
-void HandlerComponent::SetState(Handler::Component::State InState)
+void HandlerComponent::SetState(UE::Handler::Component::State InState)
 {
 	State = InState;
 }
diff --git a/Engine/Source/Runtime/PacketHandlers/PacketHandler/Public/PacketHandler.h b/Engine/Source/Runtime/PacketHandlers/PacketHandler/Public/PacketHandler.h
--- a/Engine/Source/Runtime/PacketHandlers/PacketHandler/Public/PacketHandler.h
+++ b/Engine/Source/Runtime/PacketHandlers/PacketHandler/Public/PacketHandler.h
@@ -77,39 +77,42 @@ DECLARE_DELEGATE(FPacketHandlerHandshakeComplete);
  * Enums related to the PacketHandler
  */
 
-namespace Handler
+namespace UE
 {
-	/**
-	 * State of PacketHandler
-	 */
-	enum class State : uint8
-	{
-		Uninitialized,			// PacketHandler is uninitialized
-		InitializingComponents,	// PacketHandler is initializing HandlerComponents
-		Initialized				// PacketHandler and all HandlerComponents (if any) are initialized
-	};
-
-	/**
-	 * Mode of Packet Handler
-	 */
-	enum class Mode : uint8
-	{
-		Client,					// Clientside PacketHandler
-		Server					// Serverside PacketHandler
-	};
-
-	namespace Component
+	namespace Handler
 	{
 		/**
-		 * HandlerComponent State
-		 */
+		* State of PacketHandler
+		*/
 		enum class State : uint8
 		{
-			UnInitialized,		// HandlerComponent not yet initialized
-			InitializedOnLocal, // Initialized on local instance
-			InitializeOnRemote, // Initialized on remote instance, not on local instance
-			Initialized         // Initialized on both local and remote instances
+			Uninitialized,			// PacketHandler is uninitialized
+			InitializingComponents,	// PacketHandler is initializing HandlerComponents
+			Initialized				// PacketHandler and all HandlerComponents (if any) are initialized
+		};
+	
+		/**
+		* Mode of Packet Handler
+		*/
+		enum class Mode : uint8
+		{
+			Client,					// Clientside PacketHandler
+			Server					// Serverside PacketHandler
 		};
+	
+		namespace Component
+		{
+			/**
+			* HandlerComponent State
+			*/
+			enum class State : uint8
+			{
+				UnInitialized,		// HandlerComponent not yet initialized
+				InitializedOnLocal, // Initialized on local instance
+				InitializeOnRemote, // Initialized on remote instance, not on local instance
+				Initialized         // Initialized on both local and remote instances
+			};
+		}
 	}
 }
 
@@ -268,7 +271,7 @@ public:
 	 * @param InDDoS				Reference to the owning net drivers DDoS detection handler
 	 * @param InDriverProfile		The PacketHandler configuration profile to use
 	 */
-	void Initialize(Handler::Mode Mode, uint32 InMaxPacketBits, bool bConnectionlessOnly=false,
+	void Initialize(UE::Handler::Mode Mode, uint32 InMaxPacketBits, bool bConnectionlessOnly=false,
 					TSharedPtr<class IAnalyticsProvider> InProvider=nullptr, FDDoSDetection* InDDoS=nullptr, FName InDriverProfile=NAME_None);
 
 	/**
@@ -596,7 +599,7 @@ public:
 	 */
 	FORCEINLINE bool IsFullyInitialized()
 	{
-		return State == Handler::State::Initialized;
+		return State == UE::Handler::State::Initialized;
 	}
 
 	/** Returns a pointer to the DDoS detection handler */
@@ -615,7 +618,7 @@ private:
 	 *
 	 * @param InState	The new state for the handler
 	 */
-	void SetState(Handler::State InState);
+	void SetState(UE::Handler::State InState);
 
 	/**
 	 * Called when net send/receive functions are triggered, when the handler is still uninitialized - to set a valid initial state
@@ -644,7 +647,7 @@ private:
 
 public:
 	/** Mode of the handler, Client or Server */
-	Handler::Mode Mode;
+	UE::Handler::Mode Mode;
 
 	static FPacketHandlerAddComponentByNameDelegate& GetAddComponentByNameDelegate();
 	static FPacketHandlerAddComponentDelegate& GetAddComponentDelegate();
@@ -681,7 +684,7 @@ private:
 	uint32 MaxPacketBits;
 
 	/** State of the handler */
-	Handler::State State;
+	UE::Handler::State State;
 
 	/** Packets that are buffered while HandlerComponents are being initialized */
 	TArray<BufferedPacket*> BufferedPackets;
@@ -895,7 +898,7 @@ protected:
 	 *
 	 * @param State		The new state for the handler
 	 */
-	void SetState(Handler::Component::State State);
+	void SetState(UE::Handler::Component::State State);
 
 	/**
 	 * Should be called when the handler is fully initialized on both remote and local
@@ -909,7 +912,7 @@ public:
 
 protected:
 	/** The state of this handler */
-	Handler::Component::State State;
+	UE::Handler::Component::State State;
 
 	/** Maximum number of Outgoing packet bits supported (automatically calculated to factor in other HandlerComponent reserved bits) */
 	uint32 MaxOutgoingBits;
diff --git a/Engine/Source/Runtime/PacketHandlers/ReliabilityHandlerComponent/Private/ReliabilityHandlerComponent.cpp b/Engine/Source/Runtime/PacketHandlers/ReliabilityHandlerComponent/Private/ReliabilityHandlerComponent.cpp
--- a/Engine/Source/Runtime/PacketHandlers/ReliabilityHandlerComponent/Private/ReliabilityHandlerComponent.cpp
+++ b/Engine/Source/Runtime/PacketHandlers/ReliabilityHandlerComponent/Private/ReliabilityHandlerComponent.cpp
@@ -40,7 +40,7 @@ void ReliabilityHandlerComponent::Initialize()
 {
 	SetActive(true);
 	Initialized();
-	State = Handler::Component::State::Initialized;
+	State = UE::Handler::Component::State::Initialized;
 }
 
 bool ReliabilityHandlerComponent::IsValid() const
@@ -52,7 +52,7 @@ void ReliabilityHandlerComponent::Outgoing(FBitWriter& Packet, FOutPacketTraits&
 {
 	switch (State)
 	{
-		case Handler::Component::State::Initialized:
+		case UE::Handler::Component::State::Initialized:
 		{
 			check(IsActive() && IsValid());
 
@@ -83,7 +83,7 @@ void ReliabilityHandlerComponent::Incoming(FBitReader& Packet)
 {
 	switch (State)
 	{
-		case Handler::Component::State::Initialized:
+		case UE::Handler::Component::State::Initialized:
 		{
 			if (IsActive() && IsValid())
 			{
