Cameras: Refactored FMinimalViewInfo::CalculateProjectionMatrixGivenViewport to take a ViewRect instead, as not all usages of this projection matrix have a FViewport (ie: Movie Render Queue)

diff --git a/Engine/Source/Runtime/Engine/Classes/Camera/CameraTypes.h b/Engine/Source/Runtime/Engine/Classes/Camera/CameraTypes.h
--- a/Engine/Source/Runtime/Engine/Classes/Camera/CameraTypes.h
+++ b/Engine/Source/Runtime/Engine/Classes/Camera/CameraTypes.h
@@ -149,6 +149,8 @@ struct FMinimalViewInfo
 
 	/** Calculates the projection matrix (and potentially a constrained view rectangle) given a FMinimalViewInfo and partially configured projection data (must have the view rect already set) */
 	ENGINE_API static void CalculateProjectionMatrixGivenView(const FMinimalViewInfo& ViewInfo, TEnumAsByte<enum EAspectRatioAxisConstraint> AspectRatioAxisConstraint, class FViewport* Viewport, struct FSceneViewProjectionData& InOutProjectionData);
+	/** Calculates the projection matrix (and potentially a constrained view rectangle) given a FMinimalViewInfo and partially configured projection data (must have the view rect already set). ConstrainedViewRectangle is only used if the ViewInfo.bConstrainAspectRatio is set. */
+	ENGINE_API static void CalculateProjectionMatrixGivenViewRectangle(const FMinimalViewInfo& ViewInfo, TEnumAsByte<enum EAspectRatioAxisConstraint> AspectRatioAxisConstraint, const FIntRect& ConstrainedViewRectangle, FSceneViewProjectionData& InOutProjectionData);
 
 	/** The near plane distance of the perspective view (in world units). Returns the value of PerspectiveNearClipPlane if positive, and GNearClippingPlane otherwise */
 	ENGINE_API FORCEINLINE float GetFinalPerspectiveNearClipPlane() const
diff --git a/Engine/Source/Runtime/Engine/Private/Camera/CameraStackTypes.cpp b/Engine/Source/Runtime/Engine/Private/Camera/CameraStackTypes.cpp
--- a/Engine/Source/Runtime/Engine/Private/Camera/CameraStackTypes.cpp
+++ b/Engine/Source/Runtime/Engine/Private/Camera/CameraStackTypes.cpp
@@ -138,14 +138,14 @@ FMatrix FMinimalViewInfo::CalculateProjectionMatrix() const
 	return ProjectionMatrix;
 }
 
-void FMinimalViewInfo::CalculateProjectionMatrixGivenView(const FMinimalViewInfo& ViewInfo, TEnumAsByte<enum EAspectRatioAxisConstraint> AspectRatioAxisConstraint, FViewport* Viewport, FSceneViewProjectionData& InOutProjectionData)
+void FMinimalViewInfo::CalculateProjectionMatrixGivenViewRectangle(const FMinimalViewInfo& ViewInfo, TEnumAsByte<enum EAspectRatioAxisConstraint> AspectRatioAxisConstraint, const FIntRect& ConstrainedViewRectangle, FSceneViewProjectionData& InOutProjectionData)
 {
 	// Create the projection matrix (and possibly constrain the view rectangle)
 	if (ViewInfo.bConstrainAspectRatio)
 	{
 		// Enforce a particular aspect ratio for the render of the scene. 
 		// Results in black bars at top/bottom etc.
-		InOutProjectionData.SetConstrainedViewRectangle(Viewport->CalculateViewExtents(ViewInfo.AspectRatio, InOutProjectionData.GetViewRect()));
+		InOutProjectionData.SetConstrainedViewRectangle(ConstrainedViewRectangle);
 
 		InOutProjectionData.ProjectionMatrix = ViewInfo.CalculateProjectionMatrix();
 	}
@@ -236,3 +236,9 @@ void FMinimalViewInfo::CalculateProjectionMatrixGivenView(const FMinimalViewInfo
 	}
 }
 
+void FMinimalViewInfo::CalculateProjectionMatrixGivenView(const FMinimalViewInfo& ViewInfo, TEnumAsByte<enum EAspectRatioAxisConstraint> AspectRatioAxisConstraint, FViewport* Viewport, FSceneViewProjectionData& InOutProjectionData)
+{
+	FIntRect ViewExtents = Viewport->CalculateViewExtents(ViewInfo.AspectRatio, InOutProjectionData.GetViewRect());
+	CalculateProjectionMatrixGivenViewRectangle(ViewInfo, AspectRatioAxisConstraint, ViewExtents, InOutProjectionData);
+}
+
