[shaders] remove a sanity check that was incorrect; it's possible that if a shader fails to compile the shader map for a material will be set back to null after the compilation fail is processed (so not having a gamethreadshadermap at the point of the FinishCacheShaders call doesn't necessarily mean we didn't have one during BeginCacheShaders). in this case the logic following the existing sanity check will still function correctly; it will just attempt to compile the shader again and report the compile failure for any other instances/materials which share the shadermap.

diff --git a/Engine/Source/Runtime/Engine/Private/Materials/MaterialShared.cpp b/Engine/Source/Runtime/Engine/Private/Materials/MaterialShared.cpp
--- a/Engine/Source/Runtime/Engine/Private/Materials/MaterialShared.cpp
+++ b/Engine/Source/Runtime/Engine/Private/Materials/MaterialShared.cpp
@@ -2665,7 +2665,6 @@ bool FMaterial::CacheShaders(const FMaterialShaderMapId& ShaderMapId, EShaderPla
 	if (!GameThreadShaderMap)
 	{
 		TRefCountPtr<FMaterialShaderMap> ShaderMap;
-		check(!bFoundCompiling); // sanity check
 		if (CacheShadersPending.IsValid()) // we started a load above, check the result
 		{
 			ShaderMap = CacheShadersPending->Get();
