Improving const-correctness of StaticMeshVertexBuffer

diff --git a/Engine/Source/Runtime/Engine/Private/Rendering/StaticMeshVertexBuffer.cpp b/Engine/Source/Runtime/Engine/Private/Rendering/StaticMeshVertexBuffer.cpp
--- a/Engine/Source/Runtime/Engine/Private/Rendering/StaticMeshVertexBuffer.cpp
+++ b/Engine/Source/Runtime/Engine/Private/Rendering/StaticMeshVertexBuffer.cpp
@@ -453,7 +453,7 @@ void FStaticMeshVertexBuffer::AllocateData(bool bNeedsCPUAccess /*= true*/)
 	}
 }
 
-int FStaticMeshVertexBuffer::GetTangentSize()
+int FStaticMeshVertexBuffer::GetTangentSize() const
 {
 	if (GetUseHighPrecisionTangentBasis())
 	{
@@ -469,7 +469,7 @@ int FStaticMeshVertexBuffer::GetTangentSize()
 	}
 }
 
-int FStaticMeshVertexBuffer::GetTexCoordSize()
+int FStaticMeshVertexBuffer::GetTexCoordSize() const
 {
 	if (GetUseFullPrecisionUVs())
 	{
diff --git a/Engine/Source/Runtime/Engine/Public/Rendering/StaticMeshVertexBuffer.h b/Engine/Source/Runtime/Engine/Public/Rendering/StaticMeshVertexBuffer.h
--- a/Engine/Source/Runtime/Engine/Public/Rendering/StaticMeshVertexBuffer.h
+++ b/Engine/Source/Runtime/Engine/Public/Rendering/StaticMeshVertexBuffer.h
@@ -458,9 +458,9 @@ public:
 	FORCEINLINE_DEBUGGABLE void* GetTexCoordData() { return TexcoordDataPtr; }
 	FORCEINLINE_DEBUGGABLE const void* GetTexCoordData() const { return TexcoordDataPtr; }
 
-	ENGINE_API int GetTangentSize();
+	ENGINE_API int GetTangentSize() const;
 
-	ENGINE_API int GetTexCoordSize();
+	ENGINE_API int GetTexCoordSize() const;
 
 	FORCEINLINE_DEBUGGABLE bool GetAllowCPUAccess() const
 	{
@@ -501,10 +501,10 @@ private:
 	uint8* TexcoordDataPtr;
 
 	/** The cached Tangent stride. */
-	uint32 TangentsStride;
+	mutable uint32 TangentsStride; // Mutable to allow updating through const getter
 
 	/** The cached Texcoord stride. */
-	uint32 TexcoordStride;
+	mutable uint32 TexcoordStride; // Mutable to allow updating through const getter
 
 	/** The number of texcoords/vertex in the buffer. */
 	uint32 NumTexCoords;
