Fix for shared linear texture encoding and DDC2/IBuild interface.

diff --git a/Engine/Source/Developer/TextureBuild/Private/TextureBuildFunction.cpp b/Engine/Source/Developer/TextureBuild/Private/TextureBuildFunction.cpp
--- a/Engine/Source/Developer/TextureBuild/Private/TextureBuildFunction.cpp
+++ b/Engine/Source/Developer/TextureBuild/Private/TextureBuildFunction.cpp
@@ -359,7 +359,7 @@ void FTextureBuildFunction::Build(UE::DerivedData::FBuildContext& Context) const
 		return;
 	}
 
-	UE_LOG(LogTextureBuildFunction, Display, TEXT("Compressing %d source mip(s) (%dx%d) to %s..."), SourceMips.Num(), SourceMips[0].SizeX, SourceMips[0].SizeY, *BuildSettings.TextureFormatName.ToString());
+	UE_LOG(LogTextureBuildFunction, Display, TEXT("Compressing %s -> %d source mip(s) (%dx%d) to %s..."), *Context.GetName(), SourceMips.Num(), SourceMips[0].SizeX, SourceMips[0].SizeY, *BuildSettings.TextureFormatName.ToString());
 
 	ITextureCompressorModule& TextureCompressorModule = FModuleManager::GetModuleChecked<ITextureCompressorModule>(TEXTURE_COMPRESSOR_MODULENAME);
 	
@@ -444,7 +444,7 @@ void FTextureBuildFunction::Build(UE::DerivedData::FBuildContext& Context) const
 	}
 }
 
-void GenericTextureTilingBuildFunction(UE::DerivedData::FBuildContext& Context, const ITextureTiler* Tiler)
+void GenericTextureTilingBuildFunction(UE::DerivedData::FBuildContext& Context, const ITextureTiler* Tiler, const UE::DerivedData::FUtf8SharedString& BuildFunctionName)
 {
 	// The texture description is either passed as a constant or as an output from the other build ("build input").
 	FEncodedTextureDescription TextureDescription;
@@ -485,7 +485,7 @@ void GenericTextureTilingBuildFunction(UE::DerivedData::FBuildContext& Context,
 
 	UE::TextureBuildUtilities::FTextureBuildMetadata BuildMetadata(FCbObject(Context.FindInput(ANSITEXTVIEW("TextureBuildMetadata"))));
 
-	UE_LOG(LogTextureBuildFunction, Display, TEXT("Tiling %d source mip(s) with a tail of %d..."), TextureDescription.NumMips, TextureExtendedData.NumMipsInTail);
+	UE_LOG(LogTextureBuildFunction, Display, TEXT("Tiling %s with %S -> %d source mip(s) with a tail of %d..."), *Context.GetName(), *BuildFunctionName, TextureDescription.NumMips, TextureExtendedData.NumMipsInTail);
 
 	//
 	// Careful - the linear build might have a different streaming mip count than we output due to mip tail
@@ -523,6 +523,7 @@ void GenericTextureTilingBuildFunction(UE::DerivedData::FBuildContext& Context,
 			StreamingMipName << "Mip" << MipIndex;
 
 			FSharedBuffer SourceData = Context.FindInput(StreamingMipName);
+			check(SourceData.GetSize() == TextureDescription.GetMipSizeInBytes(MipIndex));
 			SourceMipView = SourceData.GetView();
 			InputTextureMipBuffers.Add(SourceData);
 		}
diff --git a/Engine/Source/Developer/TextureBuild/Public/TextureBuildFunction.h b/Engine/Source/Developer/TextureBuild/Public/TextureBuildFunction.h
--- a/Engine/Source/Developer/TextureBuild/Public/TextureBuildFunction.h
+++ b/Engine/Source/Developer/TextureBuild/Public/TextureBuildFunction.h
@@ -28,7 +28,7 @@ protected:
 *	This function does the meat of breaking out the inputs from the build context and handing them
 *	to the tiler, then packing them back up for the build process.
 */
-TEXTUREBUILD_API void GenericTextureTilingBuildFunction(UE::DerivedData::FBuildContext& Context, const ITextureTiler* Tiler);
+TEXTUREBUILD_API void GenericTextureTilingBuildFunction(UE::DerivedData::FBuildContext& Context, const ITextureTiler* Tiler, const UE::DerivedData::FUtf8SharedString& BuildFunctionName);
 
 
 /**
@@ -43,7 +43,7 @@ class FGenericTextureTilingBuildFunction : public UE::DerivedData::IBuildFunctio
 	virtual void Build(UE::DerivedData::FBuildContext& Context) const
 	{
 		ITextureTilerObject Tiler;
-		GenericTextureTilingBuildFunction(Context, &Tiler);
+		GenericTextureTilingBuildFunction(Context, &Tiler, GetName());
 	}
 	virtual void Configure(UE::DerivedData::FBuildConfigContext& Context) const
 	{
diff --git a/Engine/Source/Developer/TextureCompressor/Public/TextureCompressorModule.h b/Engine/Source/Developer/TextureCompressor/Public/TextureCompressorModule.h
--- a/Engine/Source/Developer/TextureCompressor/Public/TextureCompressorModule.h
+++ b/Engine/Source/Developer/TextureCompressor/Public/TextureCompressorModule.h
@@ -236,6 +236,10 @@ struct FTextureBuildSettings
 	// this so we can segregate the derived data keys.
 	bool bAffectedBySharedLinearEncoding = false;
 
+	// If we have a child format, this is the base format (i.e. will have the platform prefix removed). Otherwise equal
+	// to TextureFormatName.
+	FName BaseTextureFormatName;
+	
 	static constexpr uint32 MaxTextureResolutionDefault = TNumericLimits<uint32>::Max();
 
 	/** Default settings. */
diff --git a/Engine/Source/Runtime/Engine/Private/TextureDerivedData.cpp b/Engine/Source/Runtime/Engine/Private/TextureDerivedData.cpp
--- a/Engine/Source/Runtime/Engine/Private/TextureDerivedData.cpp
+++ b/Engine/Source/Runtime/Engine/Private/TextureDerivedData.cpp
@@ -679,6 +679,16 @@ static void FinalizeBuildSettingsForLayer(
 			if (TextureFormat)
 			{
 				bSupportsEncodeSpeed = TextureFormat->SupportsEncodeSpeed(OutSettings.TextureFormatName);
+				const FChildTextureFormat* ChildTextureFormat = TextureFormat->GetChildFormat();
+
+				if (ChildTextureFormat)
+				{
+					OutSettings.BaseTextureFormatName = ChildTextureFormat->GetBaseFormatName(OutSettings.TextureFormatName);
+				}
+				else
+				{
+					OutSettings.BaseTextureFormatName = OutSettings.TextureFormatName;
+				}
 
 				if (OutBuildResultMetadata)
 				{
@@ -704,16 +714,9 @@ static void FinalizeBuildSettingsForLayer(
 
 						// Shared linear encoding can only work if the base texture format does not expect to
 						// do the tiling itself (SupportsTiling == false).
-						const FChildTextureFormat* ChildTextureFormat = TextureFormat->GetChildFormat();
 						if (ChildTextureFormat && ChildTextureFormat->GetBaseFormatObject(OutSettings.TextureFormatName)->SupportsTiling() == false)
 						{
 							OutSettings.Tiler = ChildTextureFormat->GetTiler();
-							if (OutSettings.Tiler && IsUsingNewDerivedData())
-							{
-								// New derived data wants to treat everything as the base format and then have a separate tiling function
-								// afterwards.
-								OutSettings.TextureFormatName = ChildTextureFormat->GetBaseFormatName(OutSettings.TextureFormatName);
-							}
 						}
 					} // end if enabled
 				} // end if ddc2
diff --git a/Engine/Source/Runtime/Engine/Private/TextureDerivedDataTask.cpp b/Engine/Source/Runtime/Engine/Private/TextureDerivedDataTask.cpp
--- a/Engine/Source/Runtime/Engine/Private/TextureDerivedDataTask.cpp
+++ b/Engine/Source/Runtime/Engine/Private/TextureDerivedDataTask.cpp
@@ -1670,6 +1670,8 @@ bool DDC1_BuildTiledClassicTexture(
 	int64& BytesCached
 )
 {
+	TRACE_CPUPROFILER_EVENT_SCOPE(FTextureCacheDerivedDataWorker::BuildTiledClassicTexture);
+
 	// we know we are a child format if we have a tiler.
 	const ITextureTiler* Tiler = BuildSettingsPerLayerFetchOrBuild[0].Tiler;
 	const FChildTextureFormat* ChildFormat = GetTextureFormatManager()->FindTextureFormat(BuildSettingsPerLayerFetchOrBuild[0].TextureFormatName)->GetChildFormat();
@@ -1736,6 +1738,8 @@ bool DDC1_BuildTiledClassicTexture(
 		return false;
 	}
 
+	TRACE_CPUPROFILER_EVENT_SCOPE(FTextureCacheDerivedDataWorker::TileTexture);
+
 	check(LinearDerivedData.GetNumMipsInTail() == 0);
 
 	UE::TextureBuildUtilities::FTextureBuildMetadata BuildMetadata;
@@ -2915,15 +2919,37 @@ FTextureAsyncCacheDerivedDataTask* CreateTextureBuildTask(
 	ETextureCacheFlags Flags)
 {
 	using namespace UE::DerivedData;
+
+	// If we are tiling, we need to alter the build settings to act as though it's 
+	// for the linear base format for the build function - the tiling itself will be
+	// a separate build function that will consume the output of that.
+	// We have to do this here because if we do it where build settings are created, then
+	// the DDC key that is externally visible won't know anything about the tiling and
+	// the de-dupe code in BeginCacheForCookedPlatformData will delete the tiling build.
+	TOptional<FTextureBuildSettings> BaseSettingsFetch;
+	TOptional<FTextureBuildSettings> BaseSettingsFetchOrBuild;
 	FUtf8SharedString TilingFunctionName;
+	const FTextureBuildSettings* UseSettingsFetch = SettingsFetch;
+	const FTextureBuildSettings* UseSettingsFetchOrBuild = &SettingsFetchOrBuild;
 	if (SettingsFetchOrBuild.Tiler)
 	{
 		TilingFunctionName = SettingsFetchOrBuild.Tiler->GetBuildFunctionName();
-	}
 
-	if (FUtf8SharedString FunctionName = FindTextureBuildFunction(SettingsFetchOrBuild.TextureFormatName); !FunctionName.IsEmpty())
+		if (SettingsFetch)
+		{
+			BaseSettingsFetch = *SettingsFetch;
+			BaseSettingsFetch->TextureFormatName = BaseSettingsFetch->BaseTextureFormatName;
+			BaseSettingsFetch->Tiler = nullptr;
+			UseSettingsFetch = BaseSettingsFetch.GetPtrOrNull();
+		}
+		BaseSettingsFetchOrBuild = SettingsFetchOrBuild;
+		BaseSettingsFetchOrBuild->TextureFormatName = BaseSettingsFetchOrBuild->BaseTextureFormatName;
+		UseSettingsFetchOrBuild = BaseSettingsFetchOrBuild.GetPtrOrNull();
+	}
+	
+	if (FUtf8SharedString FunctionName = FindTextureBuildFunction(UseSettingsFetchOrBuild->TextureFormatName); !FunctionName.IsEmpty())
 	{
-		return new FTextureBuildTask(Texture, DerivedData, FunctionName, TilingFunctionName, SettingsFetch, SettingsFetchOrBuild, FetchMetadata, FetchOrBuildMetadata, Priority, Flags);
+		return new FTextureBuildTask(Texture, DerivedData, FunctionName, TilingFunctionName, UseSettingsFetch, *UseSettingsFetchOrBuild, FetchMetadata, FetchOrBuildMetadata, Priority, Flags);
 
 	}
 	return nullptr;
@@ -2941,7 +2967,7 @@ FTexturePlatformData::FStructuredDerivedDataKey CreateTextureDerivedDataKey(
 	{
 		TilingFunctionName = Settings.Tiler->GetBuildFunctionName();
 	}
-	if (FUtf8SharedString FunctionName = FindTextureBuildFunction(Settings.TextureFormatName); !FunctionName.IsEmpty())
+	if (FUtf8SharedString FunctionName = FindTextureBuildFunction(Settings.BaseTextureFormatName); !FunctionName.IsEmpty())
 	{
 		IBuild& Build = GetBuild();
 
