ExternalMorphSets: reuse old weights when calling RefreshExternalMorphTargetWeights

diff --git a/Engine/Source/Runtime/Engine/Private/Components/SkinnedMeshComponent.cpp b/Engine/Source/Runtime/Engine/Private/Components/SkinnedMeshComponent.cpp
--- a/Engine/Source/Runtime/Engine/Private/Components/SkinnedMeshComponent.cpp
+++ b/Engine/Source/Runtime/Engine/Private/Components/SkinnedMeshComponent.cpp
@@ -772,8 +772,10 @@ void USkinnedMeshComponent::RefreshExternalMorphTargetWeights()
 	}
 	check(ExternalMorphSets.Num() == NumLODs);
 
-	ExternalMorphWeightData.Reset();
-	ExternalMorphWeightData.AddDefaulted(NumLODs);
+	if (ExternalMorphWeightData.Num() != NumLODs)
+	{
+		ExternalMorphWeightData.SetNum(NumLODs);
+	}
 
 	// Make sure that for every LOD's list of morph target sets, we have the correct number of morph target weights.
 	// This resets all weights to 0 as well.
@@ -796,8 +798,7 @@ void USkinnedMeshComponent::RefreshExternalMorphTargetWeights()
 			check(MorphSetWeights != nullptr);
 
 			MorphSetWeights->Name = Item.Value->Name;
-			MorphSetWeights->Weights.Reset();
-			MorphSetWeights->Weights.AddZeroed(MorphBuffer.GetNumMorphs());
+			MorphSetWeights->Weights.SetNumZeroed(MorphBuffer.GetNumMorphs());
 		}
 
 		ExternalWeights.UpdateNumActiveMorphTargets();
