Partial fix bones transforms when a skeletal mesh body is welded

diff --git a/Engine/Source/Runtime/Engine/Private/PhysicsEngine/PhysAnim.cpp b/Engine/Source/Runtime/Engine/Private/PhysicsEngine/PhysAnim.cpp
--- a/Engine/Source/Runtime/Engine/Private/PhysicsEngine/PhysAnim.cpp
+++ b/Engine/Source/Runtime/Engine/Private/PhysicsEngine/PhysAnim.cpp
@@ -226,6 +226,16 @@ void USkeletalMeshComponent::PerformBlendPhysicsBones(
 #endif
 				FBodyInstance* PhysicsAssetBodyInstance = Bodies[BodyIndex];
 
+				// If this body is welded to something, it will not have an updated transform so we will use the
+				// bone transform. This means that if the mesh continues to play an animation, the pose will not 
+				// match the pose when the weld happened. Ideally we would restore the relative transform at the
+				// time of the weld, but we do not explicitly store that data (though it could perhaps be 
+				// recovered from the collision geometry hierarchy, it's easy to just not animate.)
+				if (PhysicsAssetBodyInstance->WeldParent != nullptr)
+				{
+					BodyIndex = INDEX_NONE;
+				}
+
 				bDriveMeshWhenKinematic =
 					bUpdateMeshWhenKinematic &&
 					PhysicsAssetBodyInstance->IsValidBodyInstance();
diff --git a/Engine/Source/Runtime/Engine/Private/PhysicsEngine/PhysicsHandleComponent.cpp b/Engine/Source/Runtime/Engine/Private/PhysicsEngine/PhysicsHandleComponent.cpp
--- a/Engine/Source/Runtime/Engine/Private/PhysicsEngine/PhysicsHandleComponent.cpp
+++ b/Engine/Source/Runtime/Engine/Private/PhysicsEngine/PhysicsHandleComponent.cpp
@@ -101,6 +101,12 @@ void UPhysicsHandleComponent::GrabComponentImp(UPrimitiveComponent* InComponent,
 		return;
 	}
 
+	// If the grabbed body is welded to another, grab that instead
+	if (BodyInstance->WeldParent != nullptr)
+	{
+		BodyInstance = BodyInstance->WeldParent;
+	}
+
 	// simulatable bodies should have handles.
 	FPhysicsActorHandle& InHandle = BodyInstance->GetPhysicsActorHandle();
 	if (!InHandle)
