Preventing a crash in UStaticMesh caused by bad Indices data.

diff --git a/Engine/Source/Runtime/Engine/Private/StaticMesh.cpp b/Engine/Source/Runtime/Engine/Private/StaticMesh.cpp
--- a/Engine/Source/Runtime/Engine/Private/StaticMesh.cpp
+++ b/Engine/Source/Runtime/Engine/Private/StaticMesh.cpp
@@ -7158,6 +7158,11 @@ bool UStaticMesh::GetPhysicsTriMeshDataCheckComplex(struct FTriMeshCollisionData
 
 			for (uint32 TriIdx = Section.FirstIndex; TriIdx < OnePastLastIndex; TriIdx += 3)
 			{
+				if ((TriIdx + 2) >= static_cast<uint32>(Indices.Num()))
+				{
+					UE_LOG(LogStaticMesh, Error, TEXT("UStaticMesh::GetPhysicsTriMeshData: Triangle data from '%s' is unexpectedly missing."), *GetFullName());
+					return false;
+				}
 				FTriIndices TriIndex;
 				TriIndex.v0 = GetCollisionVertIndexForMeshVertIndex(Indices[TriIdx +0], MeshToCollisionVertMap, CollisionData->Vertices, CollisionData->UVs, LOD.VertexBuffers.PositionVertexBuffer, LOD.VertexBuffers.StaticMeshVertexBuffer);
 				TriIndex.v1 = GetCollisionVertIndexForMeshVertIndex(Indices[TriIdx +1], MeshToCollisionVertMap, CollisionData->Vertices, CollisionData->UVs, LOD.VertexBuffers.PositionVertexBuffer, LOD.VertexBuffers.StaticMeshVertexBuffer);
