Added world partition runtime grid origin in grid settings to help properly align game content on runtime grid cells (for example landscape proxy actors)
Also added a "Preview Grid Level" option to preview other levels of the grids

diff --git a/Engine/Source/Editor/WorldPartitionEditor/Private/WorldPartition/Customizations/WorldPartitionRuntimeSpatialHashDetailsCustomization.cpp b/Engine/Source/Editor/WorldPartitionEditor/Private/WorldPartition/Customizations/WorldPartitionRuntimeSpatialHashDetailsCustomization.cpp
--- a/Engine/Source/Editor/WorldPartitionEditor/Private/WorldPartition/Customizations/WorldPartitionRuntimeSpatialHashDetailsCustomization.cpp
+++ b/Engine/Source/Editor/WorldPartitionEditor/Private/WorldPartition/Customizations/WorldPartitionRuntimeSpatialHashDetailsCustomization.cpp
@@ -7,6 +7,7 @@
 #include "DetailWidgetRow.h"
 #include "Widgets/Text/STextBlock.h"
 #include "Widgets/Input/SCheckBox.h"
+#include "Widgets/Input/SNumericEntryBox.h"
 #include "ScopedTransaction.h"
 
 #define LOCTEXT_NAMESPACE "FWorldPartitionRuntimeSpatialHashDetails"
@@ -44,6 +45,29 @@ void FWorldPartitionRuntimeSpatialHashDetails::CustomizeDetails(IDetailLayoutBui
 				WorldPartitionRuntimeSpatialHash->SetPreviewGrids(InState == ECheckBoxState::Checked);
 			})
 		];
+
+	RuntimeSettingsCategory.AddCustomRow(LOCTEXT("PreviewGridLevel", "Preview Grid Level"), false)
+		.NameContent()
+		[
+			SNew(STextBlock)
+			.IsEnabled_Lambda([this]() { return WorldPartitionRuntimeSpatialHash->GetPreviewGrids(); })
+			.Text(LOCTEXT("WorldPartitionRuntimeSpatialHashPreviewGridLevel", "Preview Grid Level"))
+			.ToolTipText(LOCTEXT("WorldPartitionRuntimeSpatialHashPreviewGridLevel_ToolTip", "The grid level to preview"))
+			.Font(IDetailLayoutBuilder::GetDetailFont())
+		]
+		.ValueContent()
+		[
+			SNew(SNumericEntryBox<int32>)
+			.IsEnabled_Lambda([this]() { return WorldPartitionRuntimeSpatialHash->GetPreviewGrids(); })
+			.AllowSpin(false)
+			.MinValue(0)
+			.Value_Lambda([this]() { return WorldPartitionRuntimeSpatialHash->GetPreviewGridLevel(); })
+			.OnValueCommitted_Lambda([this](int32 NewValue, ETextCommit::Type CommitType)
+			{
+				FScopedTransaction Transaction(LOCTEXT("SetRuntimeSpatialHashPreviewGridLevel", "Change the Runtime Spatial Hash Preview Grid Level"));
+				WorldPartitionRuntimeSpatialHash->SetPreviewGridLevel(NewValue);
+			})
+		];
 }
 
 #undef LOCTEXT_NAMESPACE
diff --git a/Engine/Source/Runtime/Engine/Private/WorldPartition/RuntimeSpatialHash/RuntimeSpatialHashGridHelper.cpp b/Engine/Source/Runtime/Engine/Private/WorldPartition/RuntimeSpatialHash/RuntimeSpatialHashGridHelper.cpp
--- a/Engine/Source/Runtime/Engine/Private/WorldPartition/RuntimeSpatialHash/RuntimeSpatialHashGridHelper.cpp
+++ b/Engine/Source/Runtime/Engine/Private/WorldPartition/RuntimeSpatialHash/RuntimeSpatialHashGridHelper.cpp
@@ -137,10 +137,9 @@ int32 FSquare2DGridHelper::ForEachIntersectingCells(const FSphericalSector& InSh
 }
 
 #if WITH_EDITOR
-FSquare2DGridHelper GetGridHelper(const FBox& WorldBounds, int64 GridCellSize)
+FSquare2DGridHelper GetGridHelper(const FBox& WorldBounds, const FVector& GridOrigin, int64 GridCellSize)
 {
 	// Default grid to a minimum of 1 level and 1 cell, for always loaded actors
-	FVector GridOrigin = FVector::ZeroVector;
 	return FSquare2DGridHelper(WorldBounds, GridOrigin, GridCellSize);
 }
 
@@ -151,7 +150,7 @@ FSquare2DGridHelper GetPartitionedActors(const FBox& WorldBounds, const FSpatial
 	//
 	// Create the hierarchical grids for the game
 	//	
-	FSquare2DGridHelper PartitionedActors = GetGridHelper(WorldBounds, Grid.CellSize);
+	FSquare2DGridHelper PartitionedActors = GetGridHelper(WorldBounds, FVector(Grid.Origin, 0), Grid.CellSize);
 	if (ensure(PartitionedActors.Levels.Num()) && WorldBounds.IsValid)
 	{
 		int32 IntersectingCellCount = 0;
diff --git a/Engine/Source/Runtime/Engine/Private/WorldPartition/WorldPartitionRuntimeSpatialHash.cpp b/Engine/Source/Runtime/Engine/Private/WorldPartition/WorldPartitionRuntimeSpatialHash.cpp
--- a/Engine/Source/Runtime/Engine/Private/WorldPartition/WorldPartitionRuntimeSpatialHash.cpp
+++ b/Engine/Source/Runtime/Engine/Private/WorldPartition/WorldPartitionRuntimeSpatialHash.cpp
@@ -958,6 +958,7 @@ bool FSpatialHashRuntimeGrid::operator == (const FSpatialHashRuntimeGrid& Other)
 		   CellSize == Other.CellSize &&
 		   LoadingRange == Other.LoadingRange &&
 		   bBlockOnSlowStreaming == Other.bBlockOnSlowStreaming &&
+		   Origin == Other.Origin &&
 		   Priority == Other.Priority &&
 		   DebugColor == Other.DebugColor &&
 		   bClientOnlyVisible == Other.bClientOnlyVisible &&
@@ -985,6 +986,7 @@ UWorldPartitionRuntimeSpatialHash::UWorldPartitionRuntimeSpatialHash(const FObje
 	: Super(ObjectInitializer)
 #if WITH_EDITORONLY_DATA
 	, bPreviewGrids(false)
+	, PreviewGridLevel(0)
 #endif
 	, bIsNameToGridMappingDirty(true)
 {}
@@ -1017,7 +1019,7 @@ FString UWorldPartitionRuntimeSpatialHash::GetCellCoordString(const FGridCellCoo
 #if WITH_EDITOR
 void UWorldPartitionRuntimeSpatialHash::DrawPreview() const
 {
-	GridPreviewer.Draw(GetWorld(), Grids, bPreviewGrids);
+	GridPreviewer.Draw(GetWorld(), Grids, bPreviewGrids, PreviewGridLevel);
 }
 
 URuntimeHashExternalStreamingObjectBase* UWorldPartitionRuntimeSpatialHash::StoreToExternalStreamingObject(UObject* StreamingObjectOuter, FName StreamingObjectName)
@@ -1182,6 +1184,7 @@ void UWorldPartitionRuntimeSpatialHash::DumpStateLog(FHierarchicalLogArchive& Ar
 		Ar.Printf(TEXT("     Loading Range: %3.2f"), StreamingGrid.LoadingRange);
 		Ar.Printf(TEXT("Block Slow Loading: %s"), StreamingGrid.bBlockOnSlowStreaming ? TEXT("Yes") : TEXT("No"));
 		Ar.Printf(TEXT(" ClientOnlyVisible: %s"), StreamingGrid.bClientOnlyVisible ? TEXT("Yes") : TEXT("No"));
+		Ar.Printf(TEXT("            Origin: %s"), *StreamingGrid.Origin.ToString());
 		Ar.Printf(TEXT(""));
 		if (const UHLODLayer* HLODLayer = StreamingGrid.HLODLayer)
 		{
@@ -1296,6 +1299,17 @@ void UWorldPartitionRuntimeSpatialHash::SetPreviewGrids(bool bInPreviewGrids)
 	bPreviewGrids = bInPreviewGrids;
 }
 
+int32 UWorldPartitionRuntimeSpatialHash::GetPreviewGridLevel() const
+{
+	return PreviewGridLevel;
+}
+
+void UWorldPartitionRuntimeSpatialHash::SetPreviewGridLevel(int32 InPreviewGridLevel)
+{
+	Modify(false);
+	PreviewGridLevel = FMath::Max(0, InPreviewGridLevel);
+}
+
 bool UWorldPartitionRuntimeSpatialHash::CreateStreamingGrid(const FSpatialHashRuntimeGrid& RuntimeGrid, const FSquare2DGridHelper& PartionedActors, UWorldPartitionStreamingPolicy* StreamingPolicy, TArray<FString>* OutPackagesToGenerate)
 {
 	TRACE_CPUPROFILER_EVENT_SCOPE(CreateStreamingGrid);
@@ -1318,6 +1332,7 @@ bool UWorldPartitionRuntimeSpatialHash::CreateStreamingGrid(const FSpatialHashRu
 	CurrentStreamingGrid.WorldBounds = PartionedActors.WorldBounds;
 	CurrentStreamingGrid.LoadingRange = RuntimeGrid.LoadingRange;
 	CurrentStreamingGrid.bBlockOnSlowStreaming = RuntimeGrid.bBlockOnSlowStreaming;
+	CurrentStreamingGrid.Origin = FVector(RuntimeGrid.Origin, 0);
 	CurrentStreamingGrid.DebugColor = RuntimeGrid.DebugColor;
 	CurrentStreamingGrid.bClientOnlyVisible = RuntimeGrid.bClientOnlyVisible;
 	CurrentStreamingGrid.HLODLayer = RuntimeGrid.HLODLayer;
diff --git a/Engine/Source/Runtime/Engine/Private/WorldPartition/WorldPartitionRuntimeSpatialHashGridPreviewer.cpp b/Engine/Source/Runtime/Engine/Private/WorldPartition/WorldPartitionRuntimeSpatialHashGridPreviewer.cpp
--- a/Engine/Source/Runtime/Engine/Private/WorldPartition/WorldPartitionRuntimeSpatialHashGridPreviewer.cpp
+++ b/Engine/Source/Runtime/Engine/Private/WorldPartition/WorldPartitionRuntimeSpatialHashGridPreviewer.cpp
@@ -28,7 +28,7 @@ FWorldPartitionRuntimeSpatialHashGridPreviewer::FWorldPartitionRuntimeSpatialHas
 }
 
 #if WITH_EDITOR
-void FWorldPartitionRuntimeSpatialHashGridPreviewer::Draw(UWorld* World, const TArray<FSpatialHashRuntimeGrid>& Grids, bool bEnabled)
+void FWorldPartitionRuntimeSpatialHashGridPreviewer::Draw(UWorld* World, const TArray<FSpatialHashRuntimeGrid>& Grids, bool bEnabled, int32 PreviewGridLevel)
 {
 	if (bEnabled && Material)
 	{
@@ -63,10 +63,12 @@ void FWorldPartitionRuntimeSpatialHashGridPreviewer::Draw(UWorld* World, const T
 				{
 					const FSpatialHashRuntimeGrid& Grid = Grids[i];
 
-					if (CachedParameters.CellSize != Grid.CellSize)
+					const int32 PreviewCellSize = (1 << FMath::Max(0, PreviewGridLevel)) * Grid.CellSize;
+
+					if (CachedParameters.CellSize != PreviewCellSize)
 					{
-						MID->SetScalarParameterValue(*FString::Printf(TEXT("Grid%d_CellSize"), i), (float)Grid.CellSize);
-						CachedParameters.CellSize = Grid.CellSize;
+						MID->SetScalarParameterValue(*FString::Printf(TEXT("Grid%d_CellSize"), i), (float)PreviewCellSize);
+						CachedParameters.CellSize = PreviewCellSize;
 					}
 
 					if (CachedParameters.LoadingRange != Grid.LoadingRange)
@@ -81,7 +83,7 @@ void FWorldPartitionRuntimeSpatialHashGridPreviewer::Draw(UWorld* World, const T
 						CachedParameters.GridColor = Grid.DebugColor;
 					}
 
-					FVector GridOffset = GRuntimeSpatialHashUseAlignedGridLevels ? FVector(0.5 * Grid.CellSize) : FVector::ZeroVector;
+					FVector GridOffset = FVector(Grid.Origin, 0) + (GRuntimeSpatialHashUseAlignedGridLevels ? FVector(0.5 * PreviewCellSize) : FVector::ZeroVector);
 					if (CachedParameters.GridOffset != GridOffset)
 					{
 						MID->SetVectorParameterValue(*FString::Printf(TEXT("Grid%d_Offset"), i), GridOffset);
diff --git a/Engine/Source/Runtime/Engine/Public/WorldPartition/RuntimeSpatialHash/RuntimeSpatialHashGridHelper.h b/Engine/Source/Runtime/Engine/Public/WorldPartition/RuntimeSpatialHash/RuntimeSpatialHashGridHelper.h
--- a/Engine/Source/Runtime/Engine/Public/WorldPartition/RuntimeSpatialHash/RuntimeSpatialHashGridHelper.h
+++ b/Engine/Source/Runtime/Engine/Public/WorldPartition/RuntimeSpatialHash/RuntimeSpatialHashGridHelper.h
@@ -473,6 +473,6 @@ public:
 };
 
 #if WITH_EDITOR
-FSquare2DGridHelper GetGridHelper(const FBox& WorldBounds, int64 GridCellSize);
+FSquare2DGridHelper GetGridHelper(const FBox& WorldBounds, const FVector& GridOrigin, int64 GridCellSize);
 FSquare2DGridHelper GetPartitionedActors(const FBox& WorldBounds, const FSpatialHashRuntimeGrid& Grid, const TArray<const IStreamingGenerationContext::FActorSetInstance*>& ActorSetInstances);
 #endif // #if WITH_EDITOR
\ No newline at end of file
diff --git a/Engine/Source/Runtime/Engine/Public/WorldPartition/WorldPartitionRuntimeSpatialHash.h b/Engine/Source/Runtime/Engine/Public/WorldPartition/WorldPartitionRuntimeSpatialHash.h
--- a/Engine/Source/Runtime/Engine/Public/WorldPartition/WorldPartitionRuntimeSpatialHash.h
+++ b/Engine/Source/Runtime/Engine/Public/WorldPartition/WorldPartitionRuntimeSpatialHash.h
@@ -177,6 +177,7 @@ struct FSpatialHashRuntimeGrid
 		: CellSize(12800)
 		, LoadingRange(25600)
 		, bBlockOnSlowStreaming(false)
+		, Origin(FVector2D::ZeroVector)
 		, Priority(0)
 		, DebugColor(FLinearColor::MakeRandomColor())
 		, bClientOnlyVisible(false)
@@ -195,9 +196,12 @@ struct FSpatialHashRuntimeGrid
 	float LoadingRange;
 
 	/** Should streaming block in situations where cells aren't getting loaded fast enough. */
-	UPROPERTY(EditAnywhere, Category = Settings)
+	UPROPERTY(EditAnywhere, Category=Settings)
 	bool bBlockOnSlowStreaming;
 
+	UPROPERTY(EditAnywhere, Category=Settings)
+	FVector2D Origin;
+
 	UPROPERTY(EditAnywhere, Category=Settings)
 	int32 Priority;
 
@@ -267,6 +271,9 @@ public:
 
 	bool GetPreviewGrids() const;
 	void SetPreviewGrids(bool bInPreviewGrids);
+
+	int32 GetPreviewGridLevel() const;
+	void SetPreviewGridLevel(int32 InPreviewGridLevel);
 #endif
 
 	// streaming interface
@@ -307,6 +314,9 @@ private:
 	UPROPERTY(Transient)
 	bool bPreviewGrids;
 
+	UPROPERTY(Transient)
+	int32 PreviewGridLevel;
+
 	UPROPERTY(Transient)
 	mutable FWorldPartitionRuntimeSpatialHashGridPreviewer GridPreviewer;
 #endif
diff --git a/Engine/Source/Runtime/Engine/Public/WorldPartition/WorldPartitionRuntimeSpatialHashGridPreviewer.h b/Engine/Source/Runtime/Engine/Public/WorldPartition/WorldPartitionRuntimeSpatialHashGridPreviewer.h
--- a/Engine/Source/Runtime/Engine/Public/WorldPartition/WorldPartitionRuntimeSpatialHashGridPreviewer.h
+++ b/Engine/Source/Runtime/Engine/Public/WorldPartition/WorldPartitionRuntimeSpatialHashGridPreviewer.h
@@ -17,7 +17,7 @@ struct ENGINE_API FWorldPartitionRuntimeSpatialHashGridPreviewer
 	FWorldPartitionRuntimeSpatialHashGridPreviewer();
 
 #if WITH_EDITOR
-	void Draw(UWorld* World, const TArray<FSpatialHashRuntimeGrid>& Grids, bool bEnabled);
+	void Draw(UWorld* World, const TArray<FSpatialHashRuntimeGrid>& Grids, bool bEnabled, int32 PreviewGridLevel);
 #endif
 
 #if WITH_EDITORONLY_DATA
