Prevent ExternalActors to be considerer as Asset Managers for their UWorld
Fix all maps referencing using content bundles actors being cooked.

#test fornite cook, qagame cook

diff --git a/Engine/Source/Editor/UnrealEd/Private/Commandlets/AssetRegistryGenerator.cpp b/Engine/Source/Editor/UnrealEd/Private/Commandlets/AssetRegistryGenerator.cpp
--- a/Engine/Source/Editor/UnrealEd/Private/Commandlets/AssetRegistryGenerator.cpp
+++ b/Engine/Source/Editor/UnrealEd/Private/Commandlets/AssetRegistryGenerator.cpp
@@ -1423,12 +1423,13 @@ void FAssetRegistryGenerator::ComputePackageDifferences(const FComputeDifference
 			for (const FAssetIdentifier& Referencer : Referencers)
 			{
 				FName ReferencerPackageName = Referencer.PackageName;
-				// TODO: Replace this workaround for ExternalActors with a modification to the ExternalActor Packages' 
+				// EXTERNALACTOR_TODO: Replace this workaround for ExternalActors with a modification to the ExternalActor Packages' 
 				// dependencies. External actors have an import dependency (hard, build, game) on their Map package because
 				// the map package is their outer. But unless they have some other use of the Map package, we do not want to
 				// mark them as modified even if their map package is modified. Doing so would mark all actors in the map as
 				// modified anytime one of them changed.
 				// Workaround: Detect external actors by naming convention and suppress their reference to the map package.
+				// See also UAssetManager::ShouldSetManager
 				TStringBuilder<256> ReferencerPackageNameStr(InPlace, ReferencerPackageName);
 				int32 ExternalActorsFolderIndex = ReferencerPackageNameStr.ToView().Find(ExternalActorsFolderName);
 				if (ExternalActorsFolderIndex != INDEX_NONE)
diff --git a/Engine/Source/Runtime/Engine/Private/AssetManager.cpp b/Engine/Source/Runtime/Engine/Private/AssetManager.cpp
--- a/Engine/Source/Runtime/Engine/Private/AssetManager.cpp
+++ b/Engine/Source/Runtime/Engine/Private/AssetManager.cpp
@@ -25,9 +25,11 @@
 #include "Misc/DelayedAutoRegister.h"
 #include "Misc/ScopedSlowTask.h"
 #include "Misc/PackageName.h"
+#include "Misc/PathViews.h"
 #include "MoviePlayerProxy.h"
 #include "Modules/ModuleManager.h"
 #include "Stats/StatsMisc.h"
+#include "String/Find.h"
 #include "UObject/ConstructorHelpers.h"
 #include "UObject/LinkerLoad.h"
 #include "UObject/ObjectSaveContext.h"
@@ -3660,6 +3662,30 @@ EAssetSetManagerResult::Type UAssetManager::ShouldSetManager(const FAssetIdentif
 		return EAssetSetManagerResult::DoNotSet;
 	}
 
+	// EXTERNALACTOR_TODO: Replace this workaround for ExternalActors with a modification to the ExternalActor Packages' 
+	// dependencies. External actors have an import dependency (hard, build, game) on their Map package because
+	// the map package is their outer. At cook time they are saved in umaps (WorldPartition cells, generated packages). Prevent
+	// their dependency on the WorldPartition generator package to set them as a manager of the generator package.
+	// Workaround: Detect external actors by naming convention and suppress their reference to the map package.
+	// Long-Term Fix: Make the external actors dependency on their map a non-game one so they are not considered while evaluating the manager asset.
+	// See also FAssetRegistryGenerator::ComputePackageDifferences
+	TStringBuilder<256> SourcePackageString;
+	Source.PackageName.ToString(SourcePackageString);
+	int32 ExternalActorIdx = UE::String::FindFirst(SourcePackageString, ULevel::GetExternalActorsFolderName(), ESearchCase::IgnoreCase);
+	if (ExternalActorIdx != INDEX_NONE)
+	{
+		FStringView TargetMountPoint = FPathViews::GetMountPointNameFromPath(TargetPackageString);
+		FStringView TargetRelativePath = FStringView(TargetPackageString).RightChop(TargetMountPoint.Len() + 1);
+
+		FStringView ChoppedExternalActor = FStringView(SourcePackageString).RightChop(ExternalActorIdx + FStringView(ULevel::GetExternalActorsFolderName()).Len());
+		bool bIsTargetWorld = UE::String::FindFirst(ChoppedExternalActor, TargetRelativePath, ESearchCase::IgnoreCase) != INDEX_NONE;
+		if (bIsTargetWorld)
+		{
+			// If the Target is the UWorld and its source an External actor, then do not set. ExternalActors do not influence their level's chunk.
+			return EAssetSetManagerResult::DoNotSet;
+		}
+	}
+
 	if (Flags & EAssetSetManagerFlags::TargetHasExistingManager)
 	{
 		// If target has a higher priority manager, never recurse and only set manager if direct
