large number of changes to the cluster union to make sure it better works in a networked environment and introduces the removal functionality

- force the physics thread to be authoritative on the cluster union makeup - no GT state on what components/particles are in the cluster union is changed until the PT -> GT sync.
- add ability to remove particles from cluster unions
- changing X/R/V/W will now actually push data from the GT to the PT for cluster unions on the client
- change the usage of BoneNames to use BoneIds to avoid any costs related to using an FName (especially for the geometry collection).
- removed the is anchored and child to parent particle properties and we now directly enqueue a command on the PT to change those.
- added various overrides of functions of the primitive component to make them return a value appropriate for the cluster union component given the absense of an actual body instance
- switch from using a implicit object cluster union to a regular union for cluster unions
- preemptively remove particles from clusters when the PT wants to break them
- properly set the cluster union's shapes' collision filters

diff --git a/Engine/Source/Runtime/Engine/Classes/PhysicsEngine/ClusterUnionComponent.h b/Engine/Source/Runtime/Engine/Classes/PhysicsEngine/ClusterUnionComponent.h
--- a/Engine/Source/Runtime/Engine/Classes/PhysicsEngine/ClusterUnionComponent.h
+++ b/Engine/Source/Runtime/Engine/Classes/PhysicsEngine/ClusterUnionComponent.h
@@ -31,7 +31,7 @@ struct FClusteredComponentData
 {
 	GENERATED_BODY()
 
-	TArray<Chaos::FPhysicsObjectHandle> PhysicsObjects;
+	TSet<Chaos::FPhysicsObjectHandle> PhysicsObjects;
 
 	// Using a TWeakObjectPtr here because the UClusterUnionReplicatedProxyComponent will have a pointer back
 	// and we don't want to get into a situation where a circular reference occurs.
@@ -65,6 +65,9 @@ struct FClusterUnionReplicatedData
 	UPROPERTY()
 	FVector_NetQuantize100 AngVel;
 
+	UPROPERTY()
+	uint8 ObjectState = 0;
+
 	UPROPERTY()
 	bool bIsAnchored = false;
 };
@@ -75,7 +78,7 @@ struct FClusterUnionPendingAddData
 	GENERATED_BODY()
 
 	UPROPERTY()
-	TArray<FName> BoneNames;
+	TArray<int32> BoneIds;
 };
 
 /**
@@ -85,6 +88,21 @@ struct FClusterUnionPendingAddData
  * This component can be used as part of AClusterUnionActor or on its own as its list of
  * clustered components/actors can be specified dynamically at runtime and/or statically
  * on asset creation.
+ * 
+ * The cluster union component needs to not only maintain a game thread representation of what's happening on the
+ * physics thread but it also needs to make sure this data gets replicated to every client. A general model of how
+ * the data flow happens is as follows:
+ * 
+ *  [Server GT Command] -> [Server PT Command] -> [Server Modifies PT Data] -> [Server Sync PT Data back to GT Data].
+ * 
+ * This enables GT control over what happens to the cluster union BUT ALSO maintains a physics-first approach
+ * to the cluster union where a physics event can possibly cause the cluster union to break.
+ * 
+ * The GT data is replicated from the server to the clients either via the FClusterUnionReplicatedData on the cluster union component
+ * or per-child component data is replicated via the UClusterUnionReplicatedProxyComponent. Generally, the same flow is
+ * replicated on the client. The only exception is for replicating the X/R/V/W properties on the cluster union particle which does
+ * a GT -> PT data sync. There's no particula reason this happens...it just mirrors the single particle physics proxy here.
+ *
  */
 UCLASS()
 class ENGINE_API UClusterUnionComponent : public UPrimitiveComponent
@@ -94,17 +112,18 @@ public:
 	UClusterUnionComponent(const FObjectInitializer& ObjectInitializer);
 
 	UFUNCTION(BlueprintCallable, Category="Cluster Union")
-	void AddComponentToCluster(UPrimitiveComponent* InComponent, const TArray<FName>& BoneNames);
+	void AddComponentToCluster(UPrimitiveComponent* InComponent, const TArray<int32>& BoneIds);
+
+	UFUNCTION(BlueprintCallable, Category = "Cluster Union")
+	void RemoveComponentFromCluster(UPrimitiveComponent* InComponent);
 
 	UFUNCTION(BlueprintCallable, Category = "Cluster Union")
 	void SetIsAnchored(bool bIsAnchored);
 
 	// The SyncVelocitiesFromPhysics will set replicated state using data from the physics thread. 
 	void SyncVelocitiesFromPhysics(const FVector& LinearVelocity, const FVector& AngularVelocity);
-	// The SyncIsAnchoredFromPhysics will set replicated state using data from the physics thread.
-	void SyncIsAnchoredFromPhysics(bool bIsAnchored);
-	// SyncChildToParentFromProxy will take the ChildToParent transform data from the proxy and make sure that gets replicated to the correct replicated proxy components.
-	void SyncChildToParentFromProxy();
+	// SyncClusterUnionFromProxy will examine the make up of the cluster union (particles, child to parent, etc.) and do whatever is needed on the GT in terms of bookkeeping.
+	void SyncClusterUnionFromProxy();
 
 	UFUNCTION()
 	bool IsComponentAdded(UPrimitiveComponent* Component) { return ComponentToPhysicsObjects.Contains(Component); }
@@ -114,7 +133,7 @@ protected:
 
 	// This should only be called on the client when replication happens.
 	UFUNCTION()
-	void ForceSetChildToParent(UPrimitiveComponent* InComponent, const TArray<FName>& BoneNames, const TArray<FTransform>& ChildToParent);
+	void ForceSetChildToParent(UPrimitiveComponent* InComponent, const TArray<int32>& BoneIds, const TArray<FTransform>& ChildToParent);
 
 private:
 	// These are the statically clustered components. These should
@@ -164,6 +183,11 @@ private:
 	UFUNCTION()
 	void HandleComponentPhysicsStateChange(UPrimitiveComponent* ChangedComponent, EComponentPhysicsStateChange StateChange);
 
+	// These functions only get called when the physics thread syncs to the game thread thereby enforcing a physics thread authoritative view of
+	// what particles are currently contained within the cluster union.
+	void HandleAddOrModifiedClusteredComponent(UPrimitiveComponent* ChangedComponent, const TMap<int32, FTransform>& PerBoneChildToParent);
+	void HandleRemovedClusteredComponent(UPrimitiveComponent* ChangedComponent, bool bDestroyReplicatedProxy);
+
 	// Whether or not this code is running on the server.
 	UFUNCTION()
 	bool IsAuthority() const;
@@ -177,9 +201,19 @@ public:
 	virtual void GetLifetimeReplicatedProps(TArray<FLifetimeProperty>& OutLifetimeProps) const override;
 	//~ End UActorComponent Interface
 
+	//~ Begin UPrimitiveComponent Interface
+public:
+	virtual FBodyInstance* GetBodyInstance(FName BoneName, bool bGetWelded, int32 Index) const override { return nullptr; }
+	virtual void SetSimulatePhysics(bool bSimulate) override;
+	virtual bool CanEditSimulatePhysics() override { return true; }
+	virtual bool LineTraceComponent(FHitResult& OutHit, const FVector Start, const FVector End, const FCollisionQueryParams& Params) override;
+	virtual bool SweepComponent(FHitResult& OutHit, const FVector Start, const FVector End, const FQuat& ShapeWorldRotation, const FCollisionShape& CollisionShape, bool bTraceComplex) override;
+	virtual bool OverlapComponentWithResult(const FVector& Pos, const FQuat& Rot, const FCollisionShape& CollisionShape, TArray<FOverlapResult>& OutOverlap) const override;
+	//~ End UPrimitiveComponent Interface
+
 	//~ Begin USceneComponent Interface
 public:
-	virtual void OnUpdateTransform(EUpdateTransformFlags UpdateTransformFlags, ETeleportType Teleport = ETeleportType::None) override;
+	virtual void OnUpdateTransform(EUpdateTransformFlags UpdateTransformFlags, ETeleportType Teleport) override;
 	//~ End USceneComponent Interface
 
 	//~ Begin IPhysicsComponent Interface.
@@ -187,6 +221,7 @@ public:
 	virtual Chaos::FPhysicsObject* GetPhysicsObjectById(int32 Id) const override;
 	virtual Chaos::FPhysicsObject* GetPhysicsObjectByName(const FName& Name) const override;
 	virtual TArray<Chaos::FPhysicsObject*> GetAllPhysicsObjects() const override;
+	virtual int32 GetIdFromGTParticle(Chaos::FGeometryParticle* Particle) const override;
 	//~ End IPhysicsComponent Interface.
 
 	//~ Begin UObject Interface.
diff --git a/Engine/Source/Runtime/Engine/Classes/PhysicsEngine/ClusterUnionReplicatedProxyComponent.h b/Engine/Source/Runtime/Engine/Classes/PhysicsEngine/ClusterUnionReplicatedProxyComponent.h
--- a/Engine/Source/Runtime/Engine/Classes/PhysicsEngine/ClusterUnionReplicatedProxyComponent.h
+++ b/Engine/Source/Runtime/Engine/Classes/PhysicsEngine/ClusterUnionReplicatedProxyComponent.h
@@ -28,6 +28,9 @@ class ENGINE_API UClusterUnionReplicatedProxyComponent : public UActorComponent
 public:
 	UClusterUnionReplicatedProxyComponent(const FObjectInitializer& ObjectInitializer);
 
+	UFUNCTION()
+	UClusterUnionComponent* GetParentClusterUnionComponent() const { return ParentClusterUnion; }
+
 	UFUNCTION()
 	void SetParentClusterUnion(UClusterUnionComponent* InComponent);
 
@@ -35,10 +38,10 @@ public:
 	void SetChildClusteredComponent(UPrimitiveComponent* InComponent);
 
 	UFUNCTION()
-	void SetParticleBoneNames(const TArray<FName>& InNames);
+	void SetParticleBoneIds(const TArray<int32>& InIds);
 
 	UFUNCTION()
-	void SetParticleChildToParent(const FName& BoneName, const FTransform& ChildToParent);
+	void SetParticleChildToParent(int32 BoneId, const FTransform& ChildToParent);
 
 protected:
 
@@ -49,7 +52,7 @@ protected:
 	void OnRep_ChildClusteredComponent();
 
 	UFUNCTION()
-	void OnRep_ParticleBoneNames();
+	void OnRep_ParticleBoneIds();
 
 	UFUNCTION()
 	void OnRep_ParticleChildToParents();
@@ -67,11 +70,11 @@ private:
 	UPROPERTY()
 	bool bNetUpdateChildClusteredComponent;
 
-	UPROPERTY(ReplicatedUsing=OnRep_ParticleBoneNames)
-	TArray<FName> ParticleBoneNames;
+	UPROPERTY(ReplicatedUsing=OnRep_ParticleBoneIds)
+	TArray<int32> ParticleBoneIds;
 
 	UPROPERTY()
-	bool bNetUpdateParticleBoneNames;
+	bool bNetUpdateParticleBoneIds;
 
 	UPROPERTY(ReplicatedUsing=OnRep_ParticleChildToParents)
 	TArray<FTransform> ParticleChildToParents;
@@ -83,6 +86,7 @@ private:
 
 	//~ Begin UActorComponent Interface
 public:
+	virtual void EndPlay(const EEndPlayReason::Type EndPlayReason) override;
 	virtual void GetLifetimeReplicatedProps(TArray<FLifetimeProperty>& OutLifetimeProps) const override;
 	//~ End UActorComponent Interface
 
diff --git a/Engine/Source/Runtime/Engine/Private/PhysicsEngine/ClusterUnionActor.cpp b/Engine/Source/Runtime/Engine/Private/PhysicsEngine/ClusterUnionActor.cpp
--- a/Engine/Source/Runtime/Engine/Private/PhysicsEngine/ClusterUnionActor.cpp
+++ b/Engine/Source/Runtime/Engine/Private/PhysicsEngine/ClusterUnionActor.cpp
@@ -13,4 +13,5 @@ AClusterUnionActor::AClusterUnionActor(const FObjectInitializer& ObjectInitializ
 
 	SetRemoteRoleForBackwardsCompat(ROLE_SimulatedProxy);
 	bReplicates = true;
+	SetReplicatingMovement(true);
 }
\ No newline at end of file
diff --git a/Engine/Source/Runtime/Engine/Private/PhysicsEngine/ClusterUnionComponent.cpp b/Engine/Source/Runtime/Engine/Private/PhysicsEngine/ClusterUnionComponent.cpp
--- a/Engine/Source/Runtime/Engine/Private/PhysicsEngine/ClusterUnionComponent.cpp
+++ b/Engine/Source/Runtime/Engine/Private/PhysicsEngine/ClusterUnionComponent.cpp
@@ -14,6 +14,27 @@
 
 DEFINE_LOG_CATEGORY(LogClusterUnion);
 
+namespace
+{
+	// This can probably be generalized into the physics engine somewhere.
+	UPrimitiveComponent* GetComponentFromUserData(void* UserData)
+	{
+		if (!UserData)
+		{
+			return nullptr;
+		}
+		else if (FBodyInstance* BodyInstance = FChaosUserData::Get<FBodyInstance>(UserData))
+		{
+			return BodyInstance->OwnerComponent.Get();
+		}
+		else if (UPrimitiveComponent* Component = FChaosUserData::Get<UPrimitiveComponent>(UserData))
+		{
+			return Component;
+		}
+		return nullptr;
+	}
+}
+
 UClusterUnionComponent::UClusterUnionComponent(const FObjectInitializer& ObjectInitializer)
 	: UPrimitiveComponent(ObjectInitializer)
 {
@@ -32,7 +53,7 @@ FPhysScene_Chaos* UClusterUnionComponent::GetChaosScene() const
 	return GWorld->GetPhysicsScene();
 }
 
-void UClusterUnionComponent::AddComponentToCluster(UPrimitiveComponent* InComponent, const TArray<FName>& BoneNames)
+void UClusterUnionComponent::AddComponentToCluster(UPrimitiveComponent* InComponent, const TArray<int32>& BoneIds)
 {
 	if (!InComponent || !PhysicsProxy)
 	{
@@ -45,7 +66,7 @@ void UClusterUnionComponent::AddComponentToCluster(UPrimitiveComponent* InCompon
 		{
 			// Early out - defer adding the component to the cluster until the component has a valid physics state.
 			FClusterUnionPendingAddData Data;
-			Data.BoneNames = BoneNames;
+			Data.BoneIds = BoneIds;
 			PendingComponentsToAdd.Add(InComponent, Data);
 			InComponent->OnComponentPhysicsStateChanged.AddDynamic(this, &UClusterUnionComponent::HandleComponentPhysicsStateChange);
 		}
@@ -56,22 +77,22 @@ void UClusterUnionComponent::AddComponentToCluster(UPrimitiveComponent* InCompon
 
 	TArray<Chaos::FPhysicsObjectHandle> Objects;
 
-	if (BoneNames.IsEmpty())
+	if (BoneIds.IsEmpty())
 	{
 		Objects = InComponent->GetAllPhysicsObjects();
 	}
 	else
 	{
-		Objects.Reserve(BoneNames.Num());
-		for (const FName& Name : BoneNames)
+		Objects.Reserve(BoneIds.Num());
+		for (int32 Id : BoneIds)
 		{
-			Objects.Add(InComponent->GetPhysicsObjectByName(Name));
+			Objects.Add(InComponent->GetPhysicsObjectById(Id));
 		}
 	}
 
 	FLockedReadPhysicsObjectExternalInterface Interface = FPhysicsObjectExternalInterface::LockRead(Objects);
 
-	if (BoneNames.IsEmpty())
+	if (BoneIds.IsEmpty())
 	{
 		Objects = Objects.FilterByPredicate(
 			[&Interface](Chaos::FPhysicsObjectHandle Handle)
@@ -87,71 +108,31 @@ void UClusterUnionComponent::AddComponentToCluster(UPrimitiveComponent* InCompon
 		return;
 	}
 
-	// Force the component and its parent actor to stop replicating movement.
-	// Setting the component to not replicate should be sufficient since a simulating
-	// component shouldn't be doing much more than replicating its position anyway.
-	if (AActor* Owner = InComponent->GetOwner())
-	{
-		if (FClusteredActorData* Data = ActorToComponents.Find(Owner))
-		{
-			Data->Components.Add(InComponent);
-		}
-		else
-		{
-			FClusteredActorData NewData;
-			NewData.Components.Add(InComponent);
-			NewData.bWasReplicatingMovement = Owner->IsReplicatingMovement();
-			ActorToComponents.Add(Owner, NewData);
+	PhysicsProxy->AddPhysicsObjects_External(Objects);
+}
 
-			if (IsAuthority())
-			{
-				Owner->SetReplicatingMovement(false);
-			}
-		}
+void UClusterUnionComponent::RemoveComponentFromCluster(UPrimitiveComponent* InComponent)
+{
+	if (!InComponent || !PhysicsProxy || !ensure(InComponent->HasValidPhysicsState()))
+	{
+		return;
 	}
 
-	FClusteredComponentData Data;
-	Data.PhysicsObjects = Objects;
-	Data.bWasReplicating = InComponent->GetIsReplicated();
-
-	if (IsAuthority())
+	const int32 NumRemoved = PendingComponentsToAdd.Remove(InComponent);
+	if (NumRemoved > 0)
 	{
-		InComponent->SetIsReplicated(false);
-		if (AActor* Owner = InComponent->GetOwner())
-		{
-			// Create a replicated proxy component and add it to the actor being added to the cluster.
-			// This component will take care of replicating this addition into the cluster.
-			TObjectPtr<UClusterUnionReplicatedProxyComponent> ReplicatedProxy = NewObject<UClusterUnionReplicatedProxyComponent>(Owner);
-			if (ensure(ReplicatedProxy))
-			{
-				ReplicatedProxy->RegisterComponent();
-				ReplicatedProxy->SetParentClusterUnion(this);
-				ReplicatedProxy->SetChildClusteredComponent(InComponent);
-
-				TArray<FName> PhysicsObjectBoneNames;
-				PhysicsObjectBoneNames.Reserve(Objects.Num());
-
-				for (Chaos::FPhysicsObjectHandle PhysicsObject : Objects)
-				{
-					PhysicsObjectBoneNames.Add(Chaos::FPhysicsObjectInterface::GetName(PhysicsObject));
-				}
-
-				ReplicatedProxy->SetParticleBoneNames(PhysicsObjectBoneNames);
-				ReplicatedProxy->SetIsReplicated(true);
-			}
-
-			Data.ReplicatedProxyComponent = ReplicatedProxy;
-			Owner->FlushNetDormancy();
-		}
+		// We haven't actually added yet so we can early out.
+		return;
 	}
 
-	ComponentToPhysicsObjects.Add(InComponent, Data);
+	TSet<Chaos::FPhysicsObjectHandle> PhysicsObjectsToRemove;
 
-	for (Chaos::FPhysicsObjectHandle Handle : Objects)
+	if (FClusteredComponentData* ComponentData = ComponentToPhysicsObjects.Find(InComponent))
 	{
-		PhysicsObjectToComponent.Add(Handle, InComponent);
+		PhysicsObjectsToRemove = ComponentData->PhysicsObjects;
 	}
-	PhysicsProxy->AddPhysicsObjects_External(Objects);
+
+	PhysicsProxy->RemovePhysicsObjects_External(PhysicsObjectsToRemove);
 }
 
 void UClusterUnionComponent::SetIsAnchored(bool bIsAnchored)
@@ -202,7 +183,7 @@ void UClusterUnionComponent::OnCreatePhysicsState()
 	Parameters.ConnectionMethod = Chaos::FClusterCreationParameters::EConnectionMethod::DelaunayTriangulation;
 
 	FChaosUserData::Set<UPrimitiveComponent>(&PhysicsUserData, this);
-	PhysicsProxy = new Chaos::FClusterUnionPhysicsProxy{this, Parameters, static_cast<void*>(&PhysicsUserData), IsAuthority() ? Chaos::EThreadContext::Internal : Chaos::EThreadContext::External };
+	PhysicsProxy = new Chaos::FClusterUnionPhysicsProxy{this, Parameters, static_cast<void*>(&PhysicsUserData), GetOwner()->GetUniqueID(), GetUniqueID()};
 	PhysicsProxy->Initialize_External();
 	if (FPhysScene_Chaos* Scene = GetChaosScene())
 	{
@@ -233,6 +214,22 @@ void UClusterUnionComponent::OnDestroyPhysicsState()
 		return;
 	}
 
+	// We need to make sure we *immediately* disconnect on the GT side since there's no guarantee the normal flow
+	// will happen once we've destroyed things.
+	TSet<TObjectPtr<UPrimitiveComponent>> RemainingComponents;
+	for (const TPair<TObjectKey<UPrimitiveComponent>, FClusteredComponentData>& Kvp : ComponentToPhysicsObjects)
+	{
+		RemainingComponents.Add(Kvp.Key.ResolveObjectPtr());
+	}
+
+	for (TObjectPtr<UPrimitiveComponent> Component : RemainingComponents)
+	{
+		if (Component)
+		{
+			HandleRemovedClusteredComponent(Component, false);
+		}
+	}
+
 	if (FPhysScene_Chaos* Scene = GetChaosScene())
 	{
 		Scene->RemoveObject(PhysicsProxy);
@@ -283,6 +280,11 @@ TArray<Chaos::FPhysicsObject*> UClusterUnionComponent::GetAllPhysicsObjects() co
 	return { GetPhysicsObjectById(0) };
 }
 
+int32 UClusterUnionComponent::GetIdFromGTParticle(Chaos::FGeometryParticle* Particle) const
+{
+	return 0;
+}
+
 void UClusterUnionComponent::HandleComponentPhysicsStateChange(UPrimitiveComponent* ChangedComponent, EComponentPhysicsStateChange StateChange)
 {
 	// TODO: Maybe we should handle the destroyed state change too?
@@ -295,7 +297,7 @@ void UClusterUnionComponent::HandleComponentPhysicsStateChange(UPrimitiveCompone
 
 	if (FClusterUnionPendingAddData* PendingData = PendingComponentsToAdd.Find(ChangedComponent))
 	{
-		AddComponentToCluster(ChangedComponent, PendingData->BoneNames);
+		AddComponentToCluster(ChangedComponent, PendingData->BoneIds);
 	}
 }
 
@@ -311,44 +313,198 @@ void UClusterUnionComponent::SyncVelocitiesFromPhysics(const FVector& LinearVelo
 	MARK_PROPERTY_DIRTY_FROM_NAME(UClusterUnionComponent, ReplicatedRigidState, this);
 }
 
-void UClusterUnionComponent::SyncIsAnchoredFromPhysics(bool bIsAnchored)
+void UClusterUnionComponent::SyncClusterUnionFromProxy()
 {
-	if (!IsAuthority())
+	// NOTE THAT WE ARE ON THE GAME THREAD HERE.
+	if (!PhysicsProxy)
 	{
 		return;
 	}
 
-	ReplicatedRigidState.bIsAnchored = bIsAnchored;
+	ReplicatedRigidState.bIsAnchored = PhysicsProxy->IsAnchored_External();
+	ReplicatedRigidState.ObjectState = static_cast<uint8>(PhysicsProxy->GetObjectState_External());
 	MARK_PROPERTY_DIRTY_FROM_NAME(UClusterUnionComponent, ReplicatedRigidState, this);
+	
+	const Chaos::FClusterUnionSyncedData& FullData = PhysicsProxy->GetSyncedData_External();
+
+	// Note that at the UClusterUnionComponent level we really only want to be dealing with components.
+	// Hence why we need to modify each of the particles that we synced from the game thread into a
+	// component + bone id combination for identification. 
+	TMap<TObjectKey<UPrimitiveComponent>, TMap<int32, FTransform>> MappedData;
+	for (const Chaos::FClusterUnionChildData& ChildData : FullData.ChildParticles)
+	{
+		if (UPrimitiveComponent* Component = GetComponentFromUserData(ChildData.Particle->UserData()))
+		{
+			const int32 BoneId = Component->GetIdFromGTParticle(ChildData.Particle);
+			MappedData.FindOrAdd(Component).Add(BoneId, ChildData.ChildToParent);
+		}
+	}
+
+	// We need to handle any additions, deletions, and modifications to any child in the cluster union here.
+	// If a component lives in MappedData but not in ComponentToPhysicsObjects, new component!
+	// If a component lives in both, then it's a modified component.
+	for (const TPair<TObjectKey<UPrimitiveComponent>, TMap<int32, FTransform>>& Kvp : MappedData)
+	{
+		HandleAddOrModifiedClusteredComponent(Kvp.Key.ResolveObjectPtr(), Kvp.Value);
+	}
+
+	// If a component lives in ComponentToPhysicsObjects but not in MappedData, deleted component!
+	TArray<TObjectPtr<UPrimitiveComponent>> ComponentsToRemove;
+	for (const TPair<TObjectKey<UPrimitiveComponent>, FClusteredComponentData>& Kvp : ComponentToPhysicsObjects)
+	{
+		if (!MappedData.Contains(Kvp.Key))
+		{
+			ComponentsToRemove.Add(Kvp.Key.ResolveObjectPtr());
+		}
+	}
+
+	for (TObjectPtr<UPrimitiveComponent> Component : ComponentsToRemove)
+	{
+		HandleRemovedClusteredComponent(Component, true);
+	}
 }
 
-void UClusterUnionComponent::SyncChildToParentFromProxy()
+void UClusterUnionComponent::HandleAddOrModifiedClusteredComponent(UPrimitiveComponent* ChangedComponent, const TMap<int32, FTransform>& PerBoneChildToParent)
 {
-	if (!PhysicsProxy)
+	if (!ChangedComponent)
 	{
 		return;
 	}
 
-	const TArray<Chaos::FPhysicsObjectHandle>& ChildPhysicsObjects = PhysicsProxy->GetChildPhysicsObjects_External();
-	const Chaos::FClusterUnionSyncedData& SyncedData = PhysicsProxy->GetSyncedData_External();
-	if (ChildPhysicsObjects.Num() != SyncedData.ChildToParent.Num())
+	const bool bIsNew = !ComponentToPhysicsObjects.Contains(ChangedComponent);
+	FClusteredComponentData& ComponentData = ComponentToPhysicsObjects.FindOrAdd(ChangedComponent);
+
+	// If this is a *new* component that we're keeping track of then there's additional book-keeping
+	// we need to do to make sure we don't forget what exactly we're tracking. Additionally, we need to
+	// modify the component and its parent actor to ensure their replication stops.
+	if (bIsNew)
+	{
+		// Force the component and its parent actor to stop replicating movement.
+		// Setting the component to not replicate should be sufficient since a simulating
+		// component shouldn't be doing much more than replicating its position anyway.
+		if (AActor* Owner = ChangedComponent->GetOwner())
+		{
+			if (FClusteredActorData* Data = ActorToComponents.Find(Owner))
+			{
+				Data->Components.Add(ChangedComponent);
+			}
+			else
+			{
+				FClusteredActorData NewData;
+				NewData.Components.Add(ChangedComponent);
+				NewData.bWasReplicatingMovement = Owner->IsReplicatingMovement();
+				ActorToComponents.Add(Owner, NewData);
+
+				if (IsAuthority())
+				{
+					Owner->SetReplicatingMovement(false);
+				}
+			}
+		}
+
+		ComponentData.bWasReplicating = ChangedComponent->GetIsReplicated();
+		if (IsAuthority())
+		{
+			ChangedComponent->SetIsReplicated(false);
+			if (AActor* Owner = ChangedComponent->GetOwner())
+			{
+				// Create a replicated proxy component and add it to the actor being added to the cluster.
+				// This component will take care of replicating this addition into the cluster.
+				TObjectPtr<UClusterUnionReplicatedProxyComponent> ReplicatedProxy = NewObject<UClusterUnionReplicatedProxyComponent>(Owner);
+				if (ensure(ReplicatedProxy))
+				{
+					ReplicatedProxy->RegisterComponent();
+					ReplicatedProxy->SetParentClusterUnion(this);
+					ReplicatedProxy->SetChildClusteredComponent(ChangedComponent);
+					ReplicatedProxy->SetIsReplicated(true);
+				}
+
+				ComponentData.ReplicatedProxyComponent = ReplicatedProxy;
+			}
+		}
+	}
+
+	if (IsAuthority() && ComponentData.ReplicatedProxyComponent.IsValid())
+	{
+		// We really only need to do modifications on the server since that's where we're changing the replicated proxy to broadcast this data change.
+		TSet<int32> BoneIds;
+		PerBoneChildToParent.GetKeys(BoneIds);
+
+		TObjectPtr<UClusterUnionReplicatedProxyComponent> ReplicatedProxy = ComponentData.ReplicatedProxyComponent.Get();
+		ReplicatedProxy->SetParticleBoneIds(BoneIds.Array());
+		for (const TPair<int32, FTransform>& Kvp : PerBoneChildToParent)
+		{
+			ReplicatedProxy->SetParticleChildToParent(Kvp.Key, Kvp.Value);
+		}
+
+		if (AActor* Owner = ChangedComponent->GetOwner())
+		{
+			Owner->FlushNetDormancy();
+		}
+	}
+
+	// One more loop to ensure that our sets of physics objects are valid and up to date.
+	// This needs to happen on both the client and the server.
+	for (const TPair<int32, FTransform>& Kvp : PerBoneChildToParent)
+	{
+		Chaos::FPhysicsObjectHandle PhysicsObject = ChangedComponent->GetPhysicsObjectById(Kvp.Key);
+		ComponentData.PhysicsObjects.Add(PhysicsObject);
+		PhysicsObjectToComponent.Add(PhysicsObject, ChangedComponent);
+	}
+}
+
+void UClusterUnionComponent::HandleRemovedClusteredComponent(UPrimitiveComponent* ChangedComponent, bool bDestroyReplicatedProxy)
+{
+	if (!ChangedComponent)
+	{
+		return;
+	}
+
+	// At this point the component's particles are no longer a part of the cluster union. So we just need
+	// to get our book-keeping and game thread state to match that.
+	AActor* Owner = ChangedComponent->GetOwner();
+	if (!ensure(Owner))
 	{
 		return;
 	}
 
-	for (int32 Index = 0; Index < ChildPhysicsObjects.Num(); ++Index)
+	if (FClusteredComponentData* ComponentData = ComponentToPhysicsObjects.Find(ChangedComponent))
+	{
+		for (Chaos::FPhysicsObjectHandle Handle : ComponentData->PhysicsObjects)
+		{
+			PhysicsObjectToComponent.Remove(Handle);
+		}
+
+		if (IsAuthority())
+		{
+			ChangedComponent->SetIsReplicated(ComponentData->bWasReplicating);
+
+			if (bDestroyReplicatedProxy && ensure(ComponentData->ReplicatedProxyComponent.IsValid()))
+			{
+				UClusterUnionReplicatedProxyComponent* ProxyComponent = ComponentData->ReplicatedProxyComponent.Get();
+				ProxyComponent->DestroyComponent();
+			}
+		}
+
+		ComponentToPhysicsObjects.Remove(ChangedComponent);
+	}
+
+
+	if (FClusteredActorData* ActorData = ActorToComponents.Find(Owner))
 	{
-		Chaos::FPhysicsObjectHandle PhysicsObject = ChildPhysicsObjects[Index];
-		const FTransform& ChildToParent = SyncedData.ChildToParent[Index];
+		ActorData->Components.Remove(ChangedComponent);
 
-		if (UPrimitiveComponent** Component = PhysicsObjectToComponent.Find(PhysicsObject); Component && *Component)
+		if (ActorData->Components.IsEmpty())
 		{
-			if (FClusteredComponentData* Data = ComponentToPhysicsObjects.Find(*Component); Data && Data->ReplicatedProxyComponent.IsValid())
+			if (IsAuthority())
 			{
-				Data->ReplicatedProxyComponent.Get()->SetParticleChildToParent(Chaos::FPhysicsObjectInterface::GetName(PhysicsObject), ChildToParent);
+				Owner->SetReplicatingMovement(ActorData->bWasReplicatingMovement);
 			}
+			ActorToComponents.Remove(Owner);
 		}
 	}
+
+	Owner->FlushNetDormancy();
 }
 
 void UClusterUnionComponent::OnRep_RigidState()
@@ -361,6 +517,7 @@ void UClusterUnionComponent::OnRep_RigidState()
 	PhysicsProxy->SetLinearVelocity_External(ReplicatedRigidState.LinVel);
 	PhysicsProxy->SetAngularVelocity_External(ReplicatedRigidState.AngVel);
 	PhysicsProxy->SetIsAnchored_External(ReplicatedRigidState.bIsAnchored);
+	PhysicsProxy->SetObjectState_External(static_cast<Chaos::EObjectStateType>(ReplicatedRigidState.ObjectState));
 }
 
 void UClusterUnionComponent::GetLifetimeReplicatedProps(TArray<FLifetimeProperty>& OutLifetimeProps) const
@@ -373,20 +530,86 @@ void UClusterUnionComponent::GetLifetimeReplicatedProps(TArray<FLifetimeProperty
 	DOREPLIFETIME_WITH_PARAMS_FAST(UClusterUnionComponent, ReplicatedRigidState, Params);
 }
 
-void UClusterUnionComponent::ForceSetChildToParent(UPrimitiveComponent* InComponent, const TArray<FName>& BoneNames, const TArray<FTransform>& ChildToParent)
+void UClusterUnionComponent::ForceSetChildToParent(UPrimitiveComponent* InComponent, const TArray<int32>& BoneIds, const TArray<FTransform>& ChildToParent)
 {
-	if (IsAuthority() || !PhysicsProxy || !ensure(InComponent) || !ensure(BoneNames.Num() == ChildToParent.Num()))
+	if (IsAuthority() || !PhysicsProxy || !ensure(InComponent) || !ensure(BoneIds.Num() == ChildToParent.Num()))
 	{
 		return;
 	}
 
-	for (int32 Index = 0; Index < BoneNames.Num(); ++Index)
+	for (int32 Index = 0; Index < BoneIds.Num(); ++Index)
 	{
-		Chaos::FPhysicsObjectHandle Handle = InComponent->GetPhysicsObjectByName(BoneNames[Index]);
+		Chaos::FPhysicsObjectHandle Handle = InComponent->GetPhysicsObjectById(BoneIds[Index]);
 		PhysicsProxy->SetChildToParent_External(Handle, ChildToParent[Index]);
 	}
 }
 
+void UClusterUnionComponent::SetSimulatePhysics(bool bSimulate)
+{
+	if (!PhysicsProxy)
+	{
+		return;
+	}
+
+	PhysicsProxy->SetObjectState_External(bSimulate ? Chaos::EObjectStateType::Dynamic : Chaos::EObjectStateType::Kinematic);
+}
+
+bool UClusterUnionComponent::LineTraceComponent(FHitResult& OutHit, const FVector Start, const FVector End, const FCollisionQueryParams& Params)
+{
+	bool bHasHit = false;
+	OutHit.Distance = TNumericLimits<float>::Max();
+	for (const TPair<TObjectKey<UPrimitiveComponent>, FClusteredComponentData>& Kvp : ComponentToPhysicsObjects)
+	{
+		if (UPrimitiveComponent* Component = Kvp.Key.ResolveObjectPtr())
+		{
+			FHitResult ComponentHit;
+			if (Component->LineTraceComponent(ComponentHit, Start, End, Params) && ComponentHit.Distance < OutHit.Distance)
+			{
+				bHasHit = true;
+				OutHit = ComponentHit;
+			}
+		}
+	}
+	return bHasHit;
+}
+
+bool UClusterUnionComponent::SweepComponent(FHitResult& OutHit, const FVector Start, const FVector End, const FQuat& ShapeWorldRotation, const FCollisionShape& CollisionShape, bool bTraceComplex)
+{
+	bool bHasHit = false;
+	OutHit.Distance = TNumericLimits<float>::Max();
+	for (const TPair<TObjectKey<UPrimitiveComponent>, FClusteredComponentData>& Kvp : ComponentToPhysicsObjects)
+	{
+		if (UPrimitiveComponent* Component = Kvp.Key.ResolveObjectPtr())
+		{
+			FHitResult ComponentHit;
+			if (Component->SweepComponent(ComponentHit, Start, End, ShapeWorldRotation, CollisionShape, bTraceComplex) && ComponentHit.Distance < OutHit.Distance)
+			{
+				bHasHit = true;
+				OutHit = ComponentHit;
+			}
+		}
+	}
+	return bHasHit;
+}
+
+bool UClusterUnionComponent::OverlapComponentWithResult(const FVector& Pos, const FQuat& Rot, const FCollisionShape& CollisionShape, TArray<FOverlapResult>& OutOverlap) const
+{
+	bool bHasOverlap = false;
+	for (const TPair<TObjectKey<UPrimitiveComponent>, FClusteredComponentData>& Kvp : ComponentToPhysicsObjects)
+	{
+		if (UPrimitiveComponent* Component = Kvp.Key.ResolveObjectPtr())
+		{
+			TArray<FOverlapResult> SubOverlaps;
+			if (Component->OverlapComponentWithResult(Pos, Rot, CollisionShape, SubOverlaps))
+			{
+				bHasOverlap = true;
+				OutOverlap.Append(SubOverlaps);
+			}
+		}
+	}
+	return bHasOverlap;
+}
+
 void UClusterUnionComponent::AddReferencedObjects(UObject* InThis, FReferenceCollector& Collector)
 {
 	UPrimitiveComponent::AddReferencedObjects(InThis, Collector);
diff --git a/Engine/Source/Runtime/Engine/Private/PhysicsEngine/ClusterUnionReplicatedProxyComponent.cpp b/Engine/Source/Runtime/Engine/Private/PhysicsEngine/ClusterUnionReplicatedProxyComponent.cpp
--- a/Engine/Source/Runtime/Engine/Private/PhysicsEngine/ClusterUnionReplicatedProxyComponent.cpp
+++ b/Engine/Source/Runtime/Engine/Private/PhysicsEngine/ClusterUnionReplicatedProxyComponent.cpp
@@ -19,7 +19,7 @@ UClusterUnionReplicatedProxyComponent::UClusterUnionReplicatedProxyComponent(con
 
 	bNetUpdateParentClusterUnion = false;
 	bNetUpdateChildClusteredComponent = false;
-	bNetUpdateParticleBoneNames = false;
+	bNetUpdateParticleBoneIds = false;
 	bNetUpdateParticleChildToParents = false;
 }
 
@@ -32,10 +32,20 @@ void UClusterUnionReplicatedProxyComponent::GetLifetimeReplicatedProps(TArray<FL
 
 	DOREPLIFETIME_WITH_PARAMS_FAST(UClusterUnionReplicatedProxyComponent, ParentClusterUnion, Params);
 	DOREPLIFETIME_WITH_PARAMS_FAST(UClusterUnionReplicatedProxyComponent, ChildClusteredComponent, Params);
-	DOREPLIFETIME_WITH_PARAMS_FAST(UClusterUnionReplicatedProxyComponent, ParticleBoneNames, Params);
+	DOREPLIFETIME_WITH_PARAMS_FAST(UClusterUnionReplicatedProxyComponent, ParticleBoneIds, Params);
 	DOREPLIFETIME_WITH_PARAMS_FAST(UClusterUnionReplicatedProxyComponent, ParticleChildToParents, Params);
 }
 
+void UClusterUnionReplicatedProxyComponent::EndPlay(const EEndPlayReason::Type EndPlayReason)
+{
+	UActorComponent::EndPlay(EndPlayReason);
+
+	if (!GetOwner()->HasAuthority() && ParentClusterUnion && ChildClusteredComponent)
+	{
+		ParentClusterUnion->RemoveComponentFromCluster(ChildClusteredComponent);
+	}
+}
+
 void UClusterUnionReplicatedProxyComponent::SetParentClusterUnion(UClusterUnionComponent* InComponent)
 {
 	ParentClusterUnion = InComponent;
@@ -48,24 +58,24 @@ void UClusterUnionReplicatedProxyComponent::SetChildClusteredComponent(UPrimitiv
 	MARK_PROPERTY_DIRTY_FROM_NAME(UClusterUnionReplicatedProxyComponent, ChildClusteredComponent, this);
 }
 
-void UClusterUnionReplicatedProxyComponent::SetParticleBoneNames(const TArray<FName>& InNames)
+void UClusterUnionReplicatedProxyComponent::SetParticleBoneIds(const TArray<int32>& InIds)
 {
-	ParticleBoneNames = InNames;
+	ParticleBoneIds = InIds;
 
 	ParticleChildToParents.Empty();
-	ParticleChildToParents.Reserve(InNames.Num());
-	for (int32 Index = 0; Index < InNames.Num(); ++Index)
+	ParticleChildToParents.Reserve(InIds.Num());
+	for (int32 Index = 0; Index < InIds.Num(); ++Index)
 	{
 		ParticleChildToParents.Add(FTransform::Identity);
 	}
 
-	MARK_PROPERTY_DIRTY_FROM_NAME(UClusterUnionReplicatedProxyComponent, ParticleBoneNames, this);
+	MARK_PROPERTY_DIRTY_FROM_NAME(UClusterUnionReplicatedProxyComponent, ParticleBoneIds, this);
 }
 
-void UClusterUnionReplicatedProxyComponent::SetParticleChildToParent(const FName& BoneName, const FTransform& ChildToParent)
+void UClusterUnionReplicatedProxyComponent::SetParticleChildToParent(int32 BoneId, const FTransform& ChildToParent)
 {
 	int32 Index = INDEX_NONE;
-	if (ParticleBoneNames.Find(BoneName, Index))
+	if (ParticleBoneIds.Find(BoneId, Index))
 	{
 		ParticleChildToParents[Index] = ChildToParent;
 		MARK_PROPERTY_DIRTY_FROM_NAME(UClusterUnionReplicatedProxyComponent, ParticleChildToParents, this);
@@ -82,9 +92,9 @@ void UClusterUnionReplicatedProxyComponent::OnRep_ChildClusteredComponent()
 	bNetUpdateChildClusteredComponent = true;
 }
 
-void UClusterUnionReplicatedProxyComponent::OnRep_ParticleBoneNames()
+void UClusterUnionReplicatedProxyComponent::OnRep_ParticleBoneIds()
 {
-	bNetUpdateParticleBoneNames = true;
+	bNetUpdateParticleBoneIds = true;
 }
 
 void UClusterUnionReplicatedProxyComponent::OnRep_ParticleChildToParents()
@@ -97,28 +107,28 @@ void UClusterUnionReplicatedProxyComponent::PostRepNotifies()
 	UActorComponent::PostRepNotifies();
 
 	// These three properties should only get set once when the component is created.
-	const bool bIsInitialReplication = bNetUpdateParentClusterUnion || bNetUpdateChildClusteredComponent || bNetUpdateParticleBoneNames;
-	const bool bIsValid = ParentClusterUnion && ChildClusteredComponent && !ParticleBoneNames.IsEmpty();
+	const bool bIsInitialReplication = bNetUpdateParentClusterUnion || bNetUpdateChildClusteredComponent || bNetUpdateParticleBoneIds;
+	const bool bIsValid = ParentClusterUnion && ChildClusteredComponent && !ParticleBoneIds.IsEmpty();
 	if (bIsInitialReplication)
 	{
 		if (bIsValid)
 		{
-			ParentClusterUnion->AddComponentToCluster(ChildClusteredComponent, ParticleBoneNames);
+			ParentClusterUnion->AddComponentToCluster(ChildClusteredComponent, ParticleBoneIds);
 		}
 
 		bNetUpdateParentClusterUnion = false;
 		bNetUpdateChildClusteredComponent = false;
-		bNetUpdateParticleBoneNames = false;
+		bNetUpdateParticleBoneIds = false;
 	}
 
-	if (bIsValid && bNetUpdateParticleChildToParents && ParticleBoneNames.Num() == ParticleChildToParents.Num())
+	if (bIsValid && bNetUpdateParticleChildToParents && ParticleBoneIds.Num() == ParticleChildToParents.Num())
 	{
 		// This particular bit can't happen utnil *after* we add the component to the cluster union. There's an additional deferral
 		// in AddComponentToCluster that we have to wait for.
 		DeferUntilChildClusteredComponentInParentUnion(
 			[this]()
 			{
-				ParentClusterUnion->ForceSetChildToParent(ChildClusteredComponent, ParticleBoneNames, ParticleChildToParents);
+				ParentClusterUnion->ForceSetChildToParent(ChildClusteredComponent, ParticleBoneIds, ParticleChildToParents);
 			}
 		);
 		bNetUpdateParticleChildToParents = false;
diff --git a/Engine/Source/Runtime/Engine/Private/PhysicsEngine/Experimental/PhysScene_Chaos.cpp b/Engine/Source/Runtime/Engine/Private/PhysicsEngine/Experimental/PhysScene_Chaos.cpp
--- a/Engine/Source/Runtime/Engine/Private/PhysicsEngine/Experimental/PhysScene_Chaos.cpp
+++ b/Engine/Source/Runtime/Engine/Private/PhysicsEngine/Experimental/PhysScene_Chaos.cpp
@@ -717,8 +717,11 @@ void FPhysScene_Chaos::RemoveObject(Chaos::FClusterUnionPhysicsProxy* InObject)
 {
 	Chaos::FPhysicsSolver* Solver = InObject->GetSolver<Chaos::FPhysicsSolver>();
 
+	RemoveActorFromAccelerationStructureImp(InObject->GetParticle_External());
+
 	if (Solver)
 	{
+		Solver->UpdateParticleInAccelerationStructure_External(InObject->GetParticle_External(), true);
 		Solver->UnregisterObject(InObject);
 	}
 
@@ -1892,12 +1895,10 @@ void FPhysScene_Chaos::OnSyncBodies(Chaos::FPhysicsSolverBase* Solver)
 
 					// Any changes in the linear/angular velocity of the particle is meaningless if the X/R doesn't change too.
 					ParentComponent->SyncVelocitiesFromPhysics(DirtyParticle->V(), DirtyParticle->W());
-
-					Interface->AddToSpatialAcceleration({ &Handle, 1 }, Outer->GetSpacialAcceleration());
 				}
 
-				ParentComponent->SyncIsAnchoredFromPhysics(Proxy->IsAnchored_External());
-				ParentComponent->SyncChildToParentFromProxy();
+				Interface->AddToSpatialAcceleration({ &Handle, 1 }, Outer->GetSpacialAcceleration());
+				ParentComponent->SyncClusterUnionFromProxy();
 				DirtyParticle->ClearEvents();
 			}
 		}
diff --git a/Engine/Source/Runtime/Experimental/Chaos/Private/Chaos/ClusterUnionManager.cpp b/Engine/Source/Runtime/Experimental/Chaos/Private/Chaos/ClusterUnionManager.cpp
--- a/Engine/Source/Runtime/Experimental/Chaos/Private/Chaos/ClusterUnionManager.cpp
+++ b/Engine/Source/Runtime/Experimental/Chaos/Private/Chaos/ClusterUnionManager.cpp
@@ -17,27 +17,42 @@ namespace Chaos
 	}
 
 	DECLARE_CYCLE_STAT(TEXT("FClusterUnionManager::CreateNewClusterUnion"), STAT_CreateNewClusterUnion, STATGROUP_Chaos);
-	FClusterUnionIndex FClusterUnionManager::CreateNewClusterUnion(const FClusterCreationParameters& Parameters, FClusterUnionExplicitIndex ExplicitIndex, const FUniqueIdx* UniqueIndex)
+	FClusterUnionIndex FClusterUnionManager::CreateNewClusterUnion(const FClusterCreationParameters& Parameters, const FClusterUnionCreationParameters& ClusterUnionParameters)
 	{
 		SCOPE_CYCLE_COUNTER(STAT_CreateNewClusterUnion);
 		FClusterUnionIndex NewIndex = ClaimNextUnionIndex();
 		check(NewIndex > 0);
 
 		FClusterUnion NewUnion;
-		NewUnion.ExplicitIndex = ExplicitIndex;
+		NewUnion.ExplicitIndex = ClusterUnionParameters.ExplicitIndex;
 		NewUnion.SharedGeometry = ForceRecreateClusterUnionSharedGeometry(NewUnion);
-		NewUnion.InternalCluster = MClustering.CreateClusterParticle(-NewIndex, {}, Parameters, NewUnion.SharedGeometry, nullptr, UniqueIndex);
+		NewUnion.InternalCluster = MClustering.CreateClusterParticle(-NewIndex, {}, Parameters, NewUnion.SharedGeometry, nullptr, ClusterUnionParameters.UniqueIndex);
 		NewUnion.Parameters = Parameters;
+		NewUnion.ClusterUnionParameters = ClusterUnionParameters;
+
+		// Some parameters aren't relevant after creation.
+		NewUnion.ClusterUnionParameters.UniqueIndex = nullptr;
+
+
 		if (ensure(NewUnion.InternalCluster != nullptr))
 		{
 			NewUnion.InternalCluster->SetInternalCluster(true);
-
+			if (AccelerationStructureSplitStaticAndDynamic == 1)
+			{
+				NewUnion.InternalCluster->SetSpatialIdx(FSpatialAccelerationIdx{ 0, 1 });
+			}
+			else
+			{
+				NewUnion.InternalCluster->SetSpatialIdx(FSpatialAccelerationIdx{ 0, 0 });
+			}
 			// No bounds for now since we don't have particles. When/if we do get particles later, updating the
 			// geometry should switch this flag back on.
 			NewUnion.InternalCluster->SetHasBounds(false);
 		}
+		MEvolution.DisableParticle(NewUnion.InternalCluster);
 
 		ClusterUnions.Add(NewIndex, NewUnion);
+		ParticleToClusterUnionIndex.Add(NewUnion.InternalCluster, NewIndex);
 		return NewIndex;
 	}
 
@@ -48,8 +63,25 @@ namespace Chaos
 
 		if (FClusterUnion* ClusterUnion = FindClusterUnion(Index))
 		{
+			// Need to actually remove the particles and set them back into a simulatable state.
+			// We need a clean removal here just in case the cluster union is actually being destroyed 
+			// on the game thread prior to its children (which would live on another actor).
+			// 
+			// Note that we need to make a copy of the array here since the children list will be modified by the HandleRemoveOperation.
+			// However, the function does not expect that the input array will change.
+			TArray<FPBDRigidParticleHandle*> ChildrenCopy = ClusterUnion->ChildParticles;
+			HandleRemoveOperation(Index, ChildrenCopy, false);
 			ClusterUnion->ChildParticles.Empty();
 			MClustering.DestroyClusterParticle(ClusterUnion->InternalCluster);
+
+			if (ClusterUnion->ExplicitIndex != INDEX_NONE)
+			{
+				ExplicitIndexMap.Remove(ClusterUnion->ExplicitIndex);
+				PendingExplicitIndexOperations.Remove(ClusterUnion->ExplicitIndex);
+			}
+			ReusableIndices.Add(Index);
+			PendingClusterIndexOperations.Remove(Index);
+			ClusterUnions.Remove(Index);
 		}
 	}
 
@@ -65,9 +97,7 @@ namespace Chaos
 
 		const FRigidTransform3 ClusterWorldTM(Union.InternalCluster->X(), Union.InternalCluster->R());
 		TArray<TUniquePtr<FImplicitObject>> Objects;
-		TArray<FPBDRigidParticleHandle*> ChildParticleHandles;
 		Objects.Reserve(Union.ChildParticles.Num());
-		ChildParticleHandles.Reserve(Union.ChildParticles.Num());
 
 		for (FPBDRigidParticleHandle* Child : Union.ChildParticles)
 		{
@@ -76,11 +106,10 @@ namespace Chaos
 			if (Child->Geometry())
 			{
 				Objects.Add(TUniquePtr<FImplicitObject>(new TImplicitObjectTransformed<FReal, 3>(Child->Geometry(), Frame)));
-				ChildParticleHandles.Add(Child);
 			}
 		}
 
-		return MakeShared<FImplicitObjectUnionClustered>(MoveTemp(Objects), ChildParticleHandles);
+		return MakeShared<FImplicitObjectUnion>(MoveTemp(Objects));
 	}
 
 	DECLARE_CYCLE_STAT(TEXT("FClusterUnionManager::ClaimNextUnionIndex"), STAT_ClaimNextUnionIndex, STATGROUP_Chaos);
@@ -142,6 +171,9 @@ namespace Chaos
 				case EClusterUnionOperation::AddReleased:
 					HandleAddOperation(OpMap.Key, Op.Value, Op.Key == EClusterUnionOperation::AddReleased);
 					break;
+				case EClusterUnionOperation::Remove:
+					HandleRemoveOperation(OpMap.Key, Op.Value, true);
+					break;
 				}
 			}
 		}
@@ -174,6 +206,10 @@ namespace Chaos
 			return;
 		}
 
+		// If we're adding particles to a cluster we need to first make sure they're not part of any other cluster.
+		// Book-keeping might get a bit odd if we try to add a particle to a new clutser and then only later remove the particle from its old cluster.
+		HandleRemoveOperationWithClusterLookup(Particles, true);
+
 		TGuardValue_Bitfield_Cleanup<TFunction<void()>> Cleanup(
 			[this, OldGenerateClusterBreaking=MClustering.GetDoGenerateBreakingData()]() {
 				MClustering.SetGenerateClusterBreaking(OldGenerateClusterBreaking);
@@ -240,6 +276,10 @@ namespace Chaos
 		}
 
 		Cluster->ChildParticles.Append(FinalParticlesToAdd);
+		for (FPBDRigidParticleHandle* Particle : FinalParticlesToAdd)
+		{
+			ParticleToClusterUnionIndex.Add(Particle, ClusterIndex);
+		}
 
 		MClustering.AddParticlesToCluster(Cluster->InternalCluster, FinalParticlesToAdd, ChildToParentMap);
 
@@ -257,30 +297,12 @@ namespace Chaos
 			}
 		}
 
-		// Update cluster properties.
-		FMatrix33 ClusterInertia(0);
-
-		// TODO: These functions are generally just re-building the cluster from scratch. Need to figure out a way
-		// to get these functions to update the already existing cluster instead.
-		TSet<FPBDRigidParticleHandle*> FullChildrenSet(FinalParticlesToAdd);
-
-		const FRigidTransform3 ForceMassOrientation{ Cluster->InternalCluster->X(), Cluster->InternalCluster->R() };
-		UpdateClusterMassProperties(Cluster->InternalCluster, FullChildrenSet, ClusterInertia, bIsNewCluster ? nullptr : &ForceMassOrientation);
-
 		if (bIsNewCluster && bIsAnchored)
 		{
 			// The anchored flag is taken care of in UpdateKinematicProperties so it must be set before that is called.
 			Cluster->InternalCluster->SetIsAnchored(true);
 		}
-		UpdateKinematicProperties(Cluster->InternalCluster, MClustering.GetChildrenMap(), MEvolution);
-
-		// The recreation of the geometry must happen after the call to UpdateClusterMassProperties.
-		// Creating the geometry requires knowing the relative frame between the parent cluster and the child clusters. The
-		// parent transform is not set properly for a new empty cluster until UpdateClusterMassProperties is called for the first time.
-		Cluster->SharedGeometry = ForceRecreateClusterUnionSharedGeometry(*Cluster);
-
-		UpdateGeometry(Cluster->InternalCluster, FullChildrenSet, MClustering.GetChildrenMap(), Cluster->SharedGeometry, Cluster->Parameters);
-		MClustering.GenerateConnectionGraph(Cluster->InternalCluster, Cluster->Parameters);
+		UpdateAllClusterUnionProperties(*Cluster, bIsNewCluster);
 
 		if (OldProxy)
 		{
@@ -296,6 +318,157 @@ namespace Chaos
 
 			MEvolution.SetPhysicsMaterial(Cluster->InternalCluster, MEvolution.GetPhysicsMaterial(FinalParticlesToAdd[0]));
 		}
+
+		if (Cluster->InternalCluster->Disabled())
+		{
+			MEvolution.EnableParticle(Cluster->InternalCluster);
+		}
+
+		MEvolution.DirtyParticle(*Cluster->InternalCluster);
+		MEvolution.GetParticles().MarkTransientDirtyParticle(Cluster->InternalCluster);
+	}
+
+	DECLARE_CYCLE_STAT(TEXT("FClusterUnionManager::HandleRemoveOperation"), STAT_HandleRemoveOperation, STATGROUP_Chaos);
+	void FClusterUnionManager::HandleRemoveOperation(FClusterUnionIndex ClusterIndex, const TArray<FPBDRigidParticleHandle*>& Particles, bool bUpdateClusterProperties)
+	{
+		SCOPE_CYCLE_COUNTER(STAT_HandleRemoveOperation);
+		FClusterUnion* Cluster = ClusterUnions.Find(ClusterIndex);
+		if (!Cluster || Particles.IsEmpty())
+		{
+			return;
+		}
+
+		TGuardValue_Bitfield_Cleanup<TFunction<void()>> Cleanup(
+			[this, OldGenerateClusterBreaking = MClustering.GetDoGenerateBreakingData()]() {
+				MClustering.SetGenerateClusterBreaking(OldGenerateClusterBreaking);
+			}
+		);
+		MClustering.SetGenerateClusterBreaking(false);
+
+		IPhysicsProxyBase* OldProxy = Cluster->InternalCluster->PhysicsProxy();
+		TArray<int32> ParticleIndicesToRemove;
+		ParticleIndicesToRemove.Reserve(Particles.Num());
+
+		for (FPBDRigidParticleHandle* Handle : Particles)
+		{
+			const int32 ParticleIndex = Cluster->ChildParticles.Find(Handle);
+			if (ParticleIndex != INDEX_NONE)
+			{
+				ParticleIndicesToRemove.Add(ParticleIndex);
+				// Remove the parent proxy only if it's a cluster union proxy.
+				if (IPhysicsProxyBase* Proxy = Handle->PhysicsProxy(); Proxy && Proxy->GetParentProxy() && Proxy->GetParentProxy()->GetType() == EPhysicsProxyType::ClusterUnionProxy)
+				{
+					Proxy->SetParentProxy(nullptr);
+				}
+			}
+		}
+
+		ParticleIndicesToRemove.Sort();
+		for (int32 Index = ParticleIndicesToRemove.Num() - 1; Index >= 0; --Index)
+		{
+			const int32 ParticleIndex = ParticleIndicesToRemove[Index];
+			ParticleToClusterUnionIndex.Remove(Cluster->ChildParticles[ParticleIndex]);
+			Cluster->ChildParticles.RemoveAt(ParticleIndex);
+		}
+
+		MClustering.RemoveParticlesFromCluster(Cluster->InternalCluster, Particles);
+
+		if (bUpdateClusterProperties)
+		{
+			UpdateAllClusterUnionProperties(*Cluster, false);
+		}
+
+		// Removing a particle should have no bearing on the proxy of the cluster.
+		// This gets changed because we go through an internal initialization route when we update the cluster union particle's properties.
+		Cluster->InternalCluster->SetPhysicsProxy(OldProxy);
+
+		if (!Cluster->ChildParticles.IsEmpty())
+		{
+			MEvolution.DirtyParticle(*Cluster->InternalCluster);
+		}
+		else
+		{
+			// Note that if we have 0 child particles, our implicit object union will have an invalid bounding box.
+			// We must eject from the acceleration structure otherwise we risk cashes.
+			MEvolution.DisableParticle(Cluster->InternalCluster);
+		}
+		MEvolution.GetParticles().MarkTransientDirtyParticle(Cluster->InternalCluster);
+	}
+
+	DECLARE_CYCLE_STAT(TEXT("FClusterUnionManager::HandleRemoveOperationWithClusterLookup"), STAT_HandleRemoveOperationWithClusterLookup, STATGROUP_Chaos);
+	void FClusterUnionManager::HandleRemoveOperationWithClusterLookup(const TArray<FPBDRigidParticleHandle*>& InParticles, bool bUpdateClusterProperties)
+	{
+		SCOPE_CYCLE_COUNTER(STAT_HandleRemoveOperationWithClusterLookup);
+		TMap<FClusterUnionIndex, TSet<FPBDRigidParticleHandle*>> ParticlesPerCluster;
+		for (FPBDRigidParticleHandle* Particle : InParticles)
+		{
+			if (const int32 Index = FindClusterUnionIndexFromParticle(Particle); Index != INDEX_NONE)
+			{
+				ParticlesPerCluster.FindOrAdd(Index).Add(Particle);
+			}
+		}
+
+		for (const TPair<FClusterUnionIndex, TSet<FPBDRigidParticleHandle*>>& Kvp : ParticlesPerCluster)
+		{
+			HandleRemoveOperation(Kvp.Key, Kvp.Value.Array(), bUpdateClusterProperties);
+		}
+	}
+
+	DECLARE_CYCLE_STAT(TEXT("FClusterUnionManager::UpdateClusterUnionProperties"), STAT_UpdateClusterUnionProperties, STATGROUP_Chaos);
+	void FClusterUnionManager::UpdateAllClusterUnionProperties(FClusterUnion& ClusterUnion, bool bRecomputeMassOrientation)
+	{
+		SCOPE_CYCLE_COUNTER(STAT_UpdateClusterUnionProperties);
+		// Update cluster properties.
+		FMatrix33 ClusterInertia(0);
+
+		// TODO: These functions are generally just re-building the cluster from scratch. Need to figure out a way
+		// to get these functions to update the already existing cluster instead.
+		TSet<FPBDRigidParticleHandle*> FullChildrenSet(ClusterUnion.ChildParticles);
+
+		const FRigidTransform3 ForceMassOrientation{ ClusterUnion.InternalCluster->X(), ClusterUnion.InternalCluster->R() };
+		UpdateClusterMassProperties(ClusterUnion.InternalCluster, FullChildrenSet, ClusterInertia, bRecomputeMassOrientation ? nullptr : &ForceMassOrientation);
+		UpdateKinematicProperties(ClusterUnion.InternalCluster, MClustering.GetChildrenMap(), MEvolution);
+
+		// The recreation of the geometry must happen after the call to UpdateClusterMassProperties.
+		// Creating the geometry requires knowing the relative frame between the parent cluster and the child clusters. The
+		// parent transform is not set properly for a new empty cluster until UpdateClusterMassProperties is called for the first time.
+		ClusterUnion.SharedGeometry = ForceRecreateClusterUnionSharedGeometry(ClusterUnion);
+		UpdateGeometry(ClusterUnion.InternalCluster, FullChildrenSet, MClustering.GetChildrenMap(), ClusterUnion.SharedGeometry, ClusterUnion.Parameters);
+
+		// TODO: Need to figure out how to do the mapping back to the child shape if we ever do shape simplification...
+		if (!ClusterUnion.ChildParticles.IsEmpty() && ClusterUnion.ChildParticles.Num() == ClusterUnion.InternalCluster->ShapesArray().Num())
+		{
+			for (int32 ChildIndex = 0; ChildIndex < ClusterUnion.ChildParticles.Num(); ++ChildIndex)
+			{
+				// TODO: Is there a better way to do this merge?
+				const TUniquePtr<Chaos::FPerShapeData>& TemplateShape = ClusterUnion.ChildParticles[ChildIndex]->ShapesArray()[0];
+				const TUniquePtr<Chaos::FPerShapeData>& ShapeData = ClusterUnion.InternalCluster->ShapesArray()[ChildIndex];
+				if (ShapeData && TemplateShape)
+				{
+					{
+						FCollisionData Data = TemplateShape->GetCollisionData();
+						Data.UserData = nullptr;
+						ShapeData->SetCollisionData(Data);
+					}
+
+					{
+						FCollisionFilterData Data = TemplateShape->GetQueryData();
+						Data.Word0 = ClusterUnion.ClusterUnionParameters.ActorId;
+						ShapeData->SetQueryData(Data);
+					}
+
+					{
+						FCollisionFilterData Data = TemplateShape->GetSimData();
+						Data.Word0 = 0;
+						Data.Word2 = ClusterUnion.ClusterUnionParameters.ComponentId;
+						ShapeData->SetSimData(Data);
+					}
+				}
+			}
+		}
+
+		MClustering.ClearConnectionGraph(ClusterUnion.InternalCluster);
+		MClustering.GenerateConnectionGraph(ClusterUnion.InternalCluster, ClusterUnion.Parameters);
 	}
 
 	DECLARE_CYCLE_STAT(TEXT("FClusterUnionManager::GetOrCreateClusterUnionIndexFromExplicitIndex"), STAT_GetOrCreateClusterUnionIndexFromExplicitIndex, STATGROUP_Chaos);
@@ -308,7 +481,7 @@ namespace Chaos
 			return *OutIndex;
 		}
 
-		FClusterUnionIndex NewIndex = CreateNewClusterUnion(DefaultClusterCreationParameters(), InIndex);
+		FClusterUnionIndex NewIndex = CreateNewClusterUnion(DefaultClusterCreationParameters());
 		ExplicitIndexMap.Add(InIndex, NewIndex);
 		return NewIndex;
 	}
@@ -319,4 +492,20 @@ namespace Chaos
 		 Parameters.ConnectionMethod = MClustering.GetClusterUnionConnectionType();
 		 return Parameters;
 	}
+
+	DECLARE_CYCLE_STAT(TEXT("FClusterUnionManager::FindClusterUnionIndexFromParticle"), STAT_FindClusterUnionIndexFromParticle, STATGROUP_Chaos);
+	FClusterUnionIndex FClusterUnionManager::FindClusterUnionIndexFromParticle(FPBDRigidParticleHandle* ChildParticle)
+	{
+		SCOPE_CYCLE_COUNTER(STAT_FindClusterUnionIndexFromParticle);
+		if (!ChildParticle)
+		{
+			return INDEX_NONE;
+		}
+
+		if (FClusterUnionIndex* Index = ParticleToClusterUnionIndex.Find(ChildParticle))
+		{
+			return *Index;
+		}
+		return INDEX_NONE;
+	}
 }
\ No newline at end of file
diff --git a/Engine/Source/Runtime/Experimental/Chaos/Private/Chaos/PBDRigidClustering.cpp b/Engine/Source/Runtime/Experimental/Chaos/Private/Chaos/PBDRigidClustering.cpp
--- a/Engine/Source/Runtime/Experimental/Chaos/Private/Chaos/PBDRigidClustering.cpp
+++ b/Engine/Source/Runtime/Experimental/Chaos/Private/Chaos/PBDRigidClustering.cpp
@@ -223,7 +223,19 @@ namespace Chaos
 
 		// Disable all the input children since they no longer need to be simulated.
 		TSet<FPBDRigidParticleHandle*> InChildrenSet(InChildren);
-		MEvolution.DisableParticles(reinterpret_cast<TSet<FGeometryParticleHandle*>&>(InChildrenSet));
+		for (FPBDRigidParticleHandle* Handle : InChildren)
+		{
+			MEvolution.DisableParticle(Handle);
+			MEvolution.GetParticles().MarkTransientDirtyParticle(Handle);
+
+			if (FPBDRigidClusteredParticleHandle* ClusteredChild = Handle->CastToClustered())
+			{
+				TopLevelClusterParents.Remove(ClusteredChild);
+				TopLevelClusterParentsStrained.Remove(ClusteredChild);
+
+				ClusteredChild->ClusterIds().Id = Cluster;
+			}
+		}
 
 		// Note that we want to compute the internal strain on the cluster the same if we build it up incrementally as well as if we
 		// build it all at the same time. The parent cluster's internal strain should be the average of all the child strains.
@@ -232,18 +244,51 @@ namespace Chaos
 		Cluster->SetInternalStrains(Cluster->GetInternalStrains() * OldNumChildren);
 		Cluster->ClusterIds().NumChildren = Children.Num();
 
+		UpdateClusterParticlePropertiesFromChildren(Cluster, InChildren, ChildToParentMap);
+	}
+
+	DECLARE_CYCLE_STAT(TEXT("TPBDRigidClustering<>::RemoveParticlesFromCluster"), STAT_RemoveParticlesFromCluster, STATGROUP_Chaos);
+	void
+	FRigidClustering::RemoveParticlesFromCluster(
+		FPBDRigidClusteredParticleHandle* Cluster,
+		const TArray<FPBDRigidParticleHandle*>& InChildren)
+	{
+		SCOPE_CYCLE_COUNTER(STAT_RemoveParticlesFromCluster);
+
+		FRigidHandleArray& Children = MChildren.FindOrAdd(Cluster);
+		for (FPBDRigidParticleHandle* Child : InChildren)
+		{
+			if (int32 Index = Children.Find(Child); Child && Index != INDEX_NONE)
+			{
+				RemoveChildFromParent(Child, Cluster);
+				Children.RemoveAtSwap(Index);
+				MEvolution.DirtyParticle(*Child);
+				MEvolution.GetParticles().MarkTransientDirtyParticle(Child);
+			}
+		}
+
+		Cluster->ClusterIds().NumChildren = Children.Num();
+		Cluster->SetInternalStrains(0.0);
+		Cluster->SetCollisionGroup(INT_MAX);
+		Cluster->ClearPhysicsProxies();
+
+		// We need to fully rebuild the cluster properties from the set of children.
+		UpdateClusterParticlePropertiesFromChildren(Cluster, Children, {});
+		MEvolution.DirtyParticle(*Cluster);
+	}
+
+	DECLARE_CYCLE_STAT(TEXT("TPBDRigidClustering<>::UpdateClusterParticlePropertiesFromChildren"), STAT_UpdateClusterParticlePropertiesFromChildren, STATGROUP_Chaos);
+	void
+	FRigidClustering::UpdateClusterParticlePropertiesFromChildren(
+		FPBDRigidClusteredParticleHandle* Cluster,
+		const FRigidHandleArray& Children,
+		const TMap<FPBDRigidParticleHandle*, FPBDRigidParticleHandle*>& ChildToParentMap)
+	{
 		// An initial pass through the children to transfer some of their cluster properties to their new parent.
-		for (FPBDRigidParticleHandle* Child : InChildrenSet)
+		for (FPBDRigidParticleHandle* Child : Children)
 		{
 			if (FPBDRigidClusteredParticleHandle* ClusteredChild = Child->CastToClustered())
 			{
-				TopLevelClusterParents.Remove(ClusteredChild);
-				TopLevelClusterParentsStrained.Remove(ClusteredChild);
-
-				// Cluster group id 0 means "don't union with other things"
-				// TODO: Use INDEX_NONE instead of 0?
-				ClusteredChild->SetClusterGroupIndex(0);
-				ClusteredChild->ClusterIds().Id = Cluster;
 				Cluster->SetInternalStrains(Cluster->GetInternalStrains() + ClusteredChild->GetInternalStrains());
 				Cluster->SetCollisionImpulses(FMath::Max(Cluster->CollisionImpulses(), ClusteredChild->CollisionImpulses()));
 
@@ -257,9 +302,9 @@ namespace Chaos
 			MEvolution.DoInternalParticleInitilization(ProxyParticle, Cluster);
 		}
 
-		if (!Children.IsEmpty())
+		if (Cluster->ClusterIds().NumChildren > 0)
 		{
-			Cluster->SetInternalStrains(Cluster->GetInternalStrains() / static_cast<FReal>(Children.Num()));
+			Cluster->SetInternalStrains(Cluster->GetInternalStrains() / static_cast<FReal>(Cluster->ClusterIds().NumChildren));
 		}
 	}
 
@@ -618,11 +663,16 @@ namespace Chaos
 	{
 		SCOPE_CYCLE_COUNTER(STAT_ReleaseClusterParticles_STRAIN);
 
-		if (!ensureMsgf((ClusteredParticle->Parent() == nullptr), TEXT("Removing a cluster that still has a parent")))
+		if (FPBDRigidClusteredParticleHandle* Parent = ClusteredParticle->Parent())
 		{
-			TSet<FPBDRigidParticleHandle*> EmptySet;
-			return EmptySet;
-		};
+			// Having a parent is only OK if the parent is a cluster union since ReleaseClusterParticlesImpl will
+			// cause it to be ejected from the cluster union.
+			if (!ensureMsgf((ClusterUnionManager.FindClusterUnionIndexFromParticle(Parent) != INDEX_NONE), TEXT("Removing a cluster that still has a non-cluster union parent")))
+			{
+				TSet<FPBDRigidParticleHandle*> EmptySet;
+				return EmptySet;
+			}
+		}
 
 		return ReleaseClusterParticlesImpl(ClusteredParticle, bForceRelease, true /*bCreateNewClusters*/);
 	}
@@ -674,6 +724,8 @@ namespace Chaos
 		TArray<FPBDRigidParticleHandle*>& Children = MChildren[ClusteredParticle];
 		const bool bParentCrumbled = CrumbledSinceLastUpdate.Contains(ClusteredParticle);
 
+		bool bFoundFirstRelease = false;
+
 		// only used for propagation
 		TMap<FPBDRigidParticleHandle*, FReal> AppliedStrains;
 
@@ -690,6 +742,12 @@ namespace Chaos
 			const FReal MaxAppliedStrain = FMath::Max(Child->CollisionImpulses(), Child->GetExternalStrain());
 			if ((MaxAppliedStrain >= Child->GetInternalStrains()) || bForceRelease)
 			{
+				if (!bFoundFirstRelease)
+				{
+					ClusterUnionManager.HandleRemoveOperationWithClusterLookup({ ClusteredParticle }, true);
+					bFoundFirstRelease = true;
+				}
+
 				//UE_LOG(LogTemp, Warning, TEXT("Releasing child %d from parent %p due to strain %.5f Exceeding internal strain %.5f (Source: %s)"), ChildIdx, ClusteredParticle, ChildStrain, Child->Strain(), bForceRelease ? TEXT("Forced by caller") : ExternalStrainMap ? TEXT("External") : TEXT("Collision"));
 				
 				// The piece that hits just breaks off - we may want more control 
@@ -1134,6 +1192,28 @@ namespace Chaos
 		}
 	}
 
+	DECLARE_CYCLE_STAT(TEXT("TPBDRigidClustering<>::ClearConnectionGraph"), STAT_ClearConnectionGraph, STATGROUP_Chaos);
+	void
+	FRigidClustering::ClearConnectionGraph(FPBDRigidClusteredParticleHandle* Parent)
+	{
+		SCOPE_CYCLE_COUNTER(STAT_ClearConnectionGraph);
+		if (!MChildren.Contains(Parent))
+		{
+			return;
+		}
+
+		const TArray<FPBDRigidParticleHandle*>& Children = MChildren[Parent];
+		for (FPBDRigidParticleHandle* Handle : Children)
+		{
+			if (!Handle)
+			{
+				continue;
+			}
+
+			RemoveNodeConnections(Handle);
+		}
+	}
+
 	FRealSingle MinImpulseForStrainEval = 980 * 2 * 1.f / 30.f; //ignore impulses caused by just keeping object on ground. This is a total hack, we should not use accumulated impulse directly. Instead we need to look at delta v along constraint normal
 	FAutoConsoleVariableRef CVarMinImpulseForStrainEval(TEXT("p.chaos.MinImpulseForStrainEval"), MinImpulseForStrainEval, TEXT("Minimum accumulated impulse before accumulating for strain eval "));
 
@@ -1376,11 +1456,14 @@ namespace Chaos
 
 	bool FRigidClustering::BreakCluster(FPBDRigidClusteredParticleHandle* ClusteredParticle)
 	{
-		if (!ClusteredParticle || ClusteredParticle->Disabled())
+		if (!ClusteredParticle)
 		{
 			return false;
 		}
 
+		// Pre-emptively remove the particle from cluster unions it might be in.
+		ClusterUnionManager.HandleRemoveOperationWithClusterLookup({ ClusteredParticle }, true);
+
 		// max strain will allow to unconditionally release the children when strain is evaluated
 		constexpr FReal MaxStrain = TNumericLimits<FReal>::Max();
 		if (TArray<FPBDRigidParticleHandle*>* ChildrenHandles = GetChildrenMap().Find(ClusteredParticle))
@@ -1409,11 +1492,40 @@ namespace Chaos
 		// max strain will allow to unconditionally release the children when strain is evaluated
 		constexpr FReal MaxStrain = TNumericLimits<FReal>::Max();
 
-		// we should probably have a way to retrieve all the active clusters per proxy instead of going through the map 
+		// we should probably have a way to retrieve all the active clusters per proxy instead of having to do this iteration
 		for (FPBDRigidClusteredParticleHandle* ClusteredHandle : GetTopLevelClusterParents())
 		{
-			if (ClusteredHandle && ClusteredHandle->PhysicsProxy() == Proxy)
+			const bool bIsInputProxy = ClusteredHandle->PhysicsProxy() == Proxy;
+
+			// This handles the case where we want to break a GC but it's still in a cluster union.
+			const bool bIsInPhysicsProxiesSet = ClusteredHandle->PhysicsProxies().Contains(Proxy);
+			if (ClusteredHandle && (bIsInputProxy || bIsInPhysicsProxiesSet))
 			{
+				// Now we need to go from the parent cluster union particle to the GC particle that corresponds to the proxy.
+				if (bIsInPhysicsProxiesSet)
+				{
+					ClusteredHandle = nullptr;
+					if (TArray<FPBDRigidParticleHandle*>* Children = MChildren.Find(ClusteredHandle))
+					{
+						FPBDRigidParticleHandle** Candidate = Children->FindByPredicate(
+							[Proxy](FPBDRigidParticleHandle* Particle)
+							{
+								return Particle->PhysicsProxy() == Proxy && Particle->CastToClustered();
+							}
+						);
+
+						if (Candidate)
+						{
+							ClusteredHandle = (*Candidate)->CastToClustered();
+						}
+					}
+
+					if (!ClusteredHandle)
+					{
+						continue;
+					}
+				}
+
 				if (const TArray<FPBDRigidParticleHandle*>* Children = MChildren.Find(ClusteredHandle))
 				{
 					for (FPBDRigidParticleHandle* ChildHandle : *Children)
@@ -1432,6 +1544,7 @@ namespace Chaos
 				bCrumbledAnyCluster = true;
 			}
 		}
+
 		return bCrumbledAnyCluster;
 	}
 
diff --git a/Engine/Source/Runtime/Experimental/Chaos/Private/PhysicsProxy/ClusterUnionPhysicsProxy.cpp b/Engine/Source/Runtime/Experimental/Chaos/Private/PhysicsProxy/ClusterUnionPhysicsProxy.cpp
--- a/Engine/Source/Runtime/Experimental/Chaos/Private/PhysicsProxy/ClusterUnionPhysicsProxy.cpp
+++ b/Engine/Source/Runtime/Experimental/Chaos/Private/PhysicsProxy/ClusterUnionPhysicsProxy.cpp
@@ -43,19 +43,38 @@ namespace Chaos
 			BufferData.R = Particle->R();
 			BufferData.V = Particle->V();
 			BufferData.W = Particle->W();
+			BufferData.ObjectState = Particle->ObjectState();
+
+			const FShapesArray& ShapeArray = Particle->ShapesArray();
+			BufferData.CollisionData.Empty(ShapeArray.Num());
+			BufferData.QueryData.Empty(ShapeArray.Num());
+			BufferData.SimData.Empty(ShapeArray.Num());
+
+			for (const TUniquePtr<Chaos::FPerShapeData>& ShapeData : ShapeArray)
+			{
+				BufferData.CollisionData.Add(ShapeData->GetCollisionData());
+				BufferData.QueryData.Add(ShapeData->GetQueryData());
+				BufferData.SimData.Add(ShapeData->GetSimData());
+			}
 
 			if constexpr (std::is_base_of_v<FClusterUnionPhysicsProxy::FInternalParticle, TParticle>)
 			{
 				BufferData.bIsAnchored = Particle->IsAnchored();
+				BufferData.SharedGeometry = ConstCastSharedPtr<FImplicitObject, const FImplicitObject, ESPMode::ThreadSafe>(Particle->SharedGeometry());
+			}
+			else if constexpr (std::is_base_of_v<FClusterUnionPhysicsProxy::FExternalParticle, TParticle>)
+			{
+				BufferData.SharedGeometry = ConstCastSharedPtr<FImplicitObject, const FImplicitObject, ESPMode::ThreadSafe>(Particle->SharedGeometryLowLevel());
 			}
 		}
 	}
 
-	FClusterUnionPhysicsProxy::FClusterUnionPhysicsProxy(UObject* InOwner, const FClusterCreationParameters& InParameters, void* InUserData, EThreadContext InAuthoritativeContext)
+	FClusterUnionPhysicsProxy::FClusterUnionPhysicsProxy(UObject* InOwner, const FClusterCreationParameters& InParameters, void* InUserData, uint32 InActorId, uint32 InComponentId)
 		: Base(InOwner)
 		, ClusterParameters(InParameters)
 		, UserData(InUserData)
-		, AuthoritativeContext(InAuthoritativeContext)
+		, ActorId(InActorId)
+		, ComponentId(InComponentId)
 	{
 	}
 
@@ -94,11 +113,18 @@ namespace Chaos
 
 		FUniqueIdx UniqueIndex = Particle_External->UniqueIdx();
 		FClusterUnionManager& ClusterUnionManager = Evolution->GetRigidClustering().GetClusterUnionManager();
-		ClusterUnionIndex = ClusterUnionManager.CreateNewClusterUnion(ClusterParameters, INDEX_NONE, &UniqueIndex);
+
+		FClusterUnionCreationParameters ClusterUnionParameters;
+		ClusterUnionParameters.UniqueIndex = &UniqueIndex;
+		ClusterUnionParameters.ActorId = ActorId;
+		ClusterUnionParameters.ComponentId = ComponentId;
+
+		ClusterUnionIndex = ClusterUnionManager.CreateNewClusterUnion(ClusterParameters, ClusterUnionParameters);
 		if (FClusterUnion* ClusterUnion = ClusterUnionManager.FindClusterUnion(ClusterUnionIndex); ensure(ClusterUnion != nullptr))
 		{
 			Particle_Internal = ClusterUnion->InternalCluster;
 			Particle_Internal->SetPhysicsProxy(this);
+			Particle_Internal->GTGeometryParticle() = Particle_External.Get();
 		}
 	}
 
@@ -118,22 +144,33 @@ namespace Chaos
 		{
 			return;
 		}
+		
+		Solver->EnqueueCommandImmediate(
+			[this, Objects=Objects]() mutable
+			{
+				FReadPhysicsObjectInterface_Internal Interface = FPhysicsObjectInternalInterface::GetRead();
+				if (FPBDRigidsEvolutionGBF* Evolution = GetEvolution(this))
+				{
+					Evolution->GetRigidClustering().GetClusterUnionManager().HandleAddOperation(ClusterUnionIndex, Interface.GetAllRigidParticles(Objects), false);
+				}
+			}
+		);
+	}
 
-		int32 OldNum = ChildPhysicsObjects.Num();
-		ChildPhysicsObjects.Append(Objects);
-		for (FPhysicsObjectHandle Handle : Objects)
+	void FClusterUnionPhysicsProxy::RemovePhysicsObjects_External(const TSet<FPhysicsObjectHandle>& Objects)
+	{
+		if (!Solver || Objects.IsEmpty())
 		{
-			ChildPhysicsObjectIndexMap.Add(Handle, OldNum++);
-			SyncedData_External.ChildToParent.Add(FTransform::Identity);
+			return;
 		}
-		
+
 		Solver->EnqueueCommandImmediate(
-			[this, Objects=Objects]() mutable
+			[this, Objects = Objects]() mutable
 			{
 				FReadPhysicsObjectInterface_Internal Interface = FPhysicsObjectInternalInterface::GetRead();
 				if (FPBDRigidsEvolutionGBF* Evolution = GetEvolution(this))
 				{
-					Evolution->GetRigidClustering().GetClusterUnionManager().AddPendingClusterIndexOperation(ClusterUnionIndex, EClusterUnionOperation::Add, Interface.GetAllRigidParticles(Objects));
+					Evolution->GetRigidClustering().GetClusterUnionManager().HandleRemoveOperation(ClusterUnionIndex, Interface.GetAllRigidParticles(Objects.Array()), true);
 				}
 			}
 		);
@@ -146,12 +183,56 @@ namespace Chaos
 			return;
 		}
 
-		SyncedData_External.bIsAnchored = bIsAnchored;
-		Particle_External->ForceDirty(EChaosPropertyFlags::ClusterIsAnchored);
-		Solver->AddDirtyProxy(this);
+		Solver->EnqueueCommandImmediate(
+			[this, bIsAnchored]() mutable
+			{
+				if (Particle_Internal->IsAnchored() != bIsAnchored)
+				{
+					if (FPBDRigidsEvolutionGBF* Evolution = GetEvolution(this))
+					{
+						Particle_Internal->SetIsAnchored(bIsAnchored);
+						if (!bIsAnchored && !Particle_Internal->IsDynamic())
+						{
+							FKinematicTarget NoKinematicTarget;
+							Evolution->SetParticleObjectState(Particle_Internal, EObjectStateType::Dynamic);
+							Evolution->SetParticleKinematicTarget(Particle_Internal, NoKinematicTarget);
+						}
+						Evolution->GetParticles().MarkTransientDirtyParticle(Particle_Internal);
+					}
+				}
+			}
+		);
+	}
 
-		FClusterUnionProxyTimestamp& SyncTS = GetSyncTimestampAs<FClusterUnionProxyTimestamp>();
-		SyncTS.OverWriteAnchored.Set(GetSolverSyncTimestamp_External(), bIsAnchored);
+	EObjectStateType FClusterUnionPhysicsProxy::GetObjectState_External() const
+	{
+		if (!ensure(Particle_External))
+		{
+			return EObjectStateType::Uninitialized;
+		}
+
+		return Particle_External->ObjectState();
+	}
+
+	void FClusterUnionPhysicsProxy::SetObjectState_External(EObjectStateType State)
+	{
+		if (!Solver || !ensure(Particle_External))
+		{
+			return;
+		}
+
+		Solver->EnqueueCommandImmediate(
+			[this, State]() mutable
+			{
+				if (Particle_Internal->ObjectState() != State)
+				{
+					if (FPBDRigidsEvolutionGBF* Evolution = GetEvolution(this))
+					{
+						Evolution->SetParticleObjectState(Particle_Internal, State);
+					}
+				}
+			}
+		);
 	}
 
 	void FClusterUnionPhysicsProxy::PushToPhysicsState(const FDirtyPropertiesManager& Manager, int32 DataIdx, const FDirtyProxy& Dirty)
@@ -180,42 +261,6 @@ namespace Chaos
 		{
 			Particle_Internal->SetVelocities(*NewVelocities);
 		}
-
-		if (const bool* bNewIsAnchored = ParticleData.FindClusterIsAnchored(Manager, DataIdx); bNewIsAnchored && *bNewIsAnchored != Particle_Internal->IsAnchored())
-		{
-			Particle_Internal->SetIsAnchored(*bNewIsAnchored);
-			if (!*bNewIsAnchored && !Particle_Internal->IsDynamic())
-			{
-				FKinematicTarget NoKinematicTarget;
-				Evolution.SetParticleObjectState(Particle_Internal, EObjectStateType::Dynamic);
-				Evolution.SetParticleKinematicTarget(Particle_Internal, NoKinematicTarget);
-			}
-		}
-
-		if (AuthoritativeContext == EThreadContext::External)
-		{
-			if (const TArray<FTransform>* NewChildToParent = ParticleData.FindClusterChildToParent(Manager, DataIdx))
-			{
-				if (FClusterUnion* ClusterUnion = ClusterUnionManager.FindClusterUnion(ClusterUnionIndex))
-				{
-					if (NewChildToParent->Num() == ClusterUnion->ChildParticles.Num())
-					{
-						for (int32 Index = 0; Index < NewChildToParent->Num(); ++Index)
-						{
-							if (FPBDRigidClusteredParticleHandle* ChildHandle = ClusterUnion->ChildParticles[Index]->CastToClustered())
-							{
-								if (ChildHandle->ChildToParent().Equals((*NewChildToParent)[Index]))
-								{
-									continue;
-								}
-								ChildHandle->SetChildToParent((*NewChildToParent)[Index]);
-								RigidsSolver.GetParticles().MarkTransientDirtyParticle(ChildHandle);
-							}
-						}
-					}
-				}
-			}
-		}
 	}
 
 	bool FClusterUnionPhysicsProxy::PullFromPhysicsState(const FDirtyClusterUnionData& PullData, int32 SolverSyncTimestamp, const FDirtyClusterUnionData* NextPullData, const FRealSingle* Alpha)
@@ -231,9 +276,39 @@ namespace Chaos
 			return false;
 		}
 
-		if (SolverSyncTimestamp >= ProxyTimestamp->OverWriteAnchored.Timestamp)
+		
+		SyncedData_External.bIsAnchored = PullData.bIsAnchored;
+		SyncedData_External.ChildParticles.Empty(PullData.ChildParticles.Num());
+		for (const FDirtyClusterUnionParticleData& InData : PullData.ChildParticles)
 		{
-			SyncedData_External.bIsAnchored = PullData.bIsAnchored;
+			FClusterUnionChildData ConvertedData;
+			ConvertedData.Particle = InData.Particle;
+			ConvertedData.ChildToParent = InData.ChildToParent;
+			SyncedData_External.ChildParticles.Add(ConvertedData);
+		}
+
+		Particle_External->SetGeometry(PullData.SharedGeometry);
+		Particle_External->SetObjectState(PullData.ObjectState, true, /*bInvalidate=*/false);
+
+		for (int32 ShapeIndex = 0; ShapeIndex < Particle_External->ShapesArray().Num(); ++ShapeIndex)
+		{
+			if (const TUniquePtr<Chaos::FPerShapeData>& ShapeData = Particle_External->ShapesArray()[ShapeIndex])
+			{
+				if (PullData.CollisionData.IsValidIndex(ShapeIndex))
+				{
+					ShapeData->SetCollisionData(PullData.CollisionData[ShapeIndex]);
+				}
+
+				if (PullData.QueryData.IsValidIndex(ShapeIndex))
+				{
+					ShapeData->SetQueryData(PullData.QueryData[ShapeIndex]);
+				}
+
+				if (PullData.SimData.IsValidIndex(ShapeIndex))
+				{
+					ShapeData->SetSimData(PullData.SimData[ShapeIndex]);
+				}
+			}
 		}
 
 		if (NextPullData)
@@ -295,14 +370,6 @@ namespace Chaos
 				Particle_External->SetW(PullData.W, false);
 			}
 		}
-
-		if (AuthoritativeContext == EThreadContext::Internal && SolverSyncTimestamp >= ProxyTimestamp->OverWriteChildToParent.Timestamp)
-		{
-			for (int32 Index = 0; Index < PullData.ChildToParent.Num() && Index < SyncedData_External.ChildToParent.Num(); ++Index)
-			{
-				SyncedData_External.ChildToParent[Index] = PullData.ChildToParent[Index];
-			}
-		}
 		
 		Particle_External->UpdateShapeBounds();
 		return true;
@@ -314,20 +381,32 @@ namespace Chaos
 	
 		FPBDRigidsEvolutionGBF& Evolution = *static_cast<FPBDRigidsSolver*>(Solver)->GetEvolution();
 		FClusterUnionManager& ClusterUnionManager = Evolution.GetRigidClustering().GetClusterUnionManager();
+
 		if (FClusterUnion* ClusterUnion = ClusterUnionManager.FindClusterUnion(ClusterUnionIndex))
 		{
-			BufferData.ChildToParent.Empty(ClusterUnion->ChildParticles.Num());
+			BufferData.ChildParticles.Empty(ClusterUnion->ChildParticles.Num());
+
 			for (FPBDRigidParticleHandle* Particle : ClusterUnion->ChildParticles)
 			{
+				if (!ensure(Particle->GTGeometryParticle()))
+				{
+					continue;
+				}
+
+				FDirtyClusterUnionParticleData Data;
+				Data.Particle = Particle->GTGeometryParticle()->CastToRigidParticle();
 				if (FPBDRigidClusteredParticleHandle* ClusteredParticle = Particle->CastToClustered())
 				{
-					BufferData.ChildToParent.Add(ClusteredParticle->ChildToParent());
+					Data.ChildToParent = ClusteredParticle->ChildToParent();
 				}
 				else
 				{
-					BufferData.ChildToParent.Add(FTransform::Identity);
+					Data.ChildToParent = FRigidTransform3::Identity;
 				}
+				BufferData.ChildParticles.Add(Data);
 			}
+
+			BufferData.SharedGeometry = ClusterUnion->SharedGeometry;
 		}
 	}
 
@@ -335,7 +414,15 @@ namespace Chaos
 	{
 		BufferPhysicsResultsImp(this, Particle_External.Get(), BufferData);
 		BufferData.bIsAnchored = SyncedData_External.bIsAnchored;
-		BufferData.ChildToParent = SyncedData_External.ChildToParent;
+
+		BufferData.ChildParticles.Empty(SyncedData_External.ChildParticles.Num());
+		for (const FClusterUnionChildData& Data : SyncedData_External.ChildParticles)
+		{
+			FDirtyClusterUnionParticleData ConvertedData;
+			ConvertedData.Particle = Data.Particle;
+			ConvertedData.ChildToParent = Data.ChildToParent;
+			BufferData.ChildParticles.Add(ConvertedData);
+		}
 	}
 
 	void FClusterUnionPhysicsProxy::SyncRemoteData(FDirtyPropertiesManager& Manager, int32 DataIdx, FDirtyChaosProperties& RemoteData) const
@@ -362,8 +449,6 @@ namespace Chaos
 		// SyncRemote will check the dirty flags and will skip the change in value if the dirty flag is not actually set.
 		RemoteData.SyncRemote<FParticlePositionRotation, EChaosProperty::ClusterXR>(Manager, DataIdx, Particle_External->XR());
 		RemoteData.SyncRemote<FParticleVelocities, EChaosProperty::ClusterVelocities>(Manager, DataIdx, Particle_External->Velocities());
-		RemoteData.SyncRemote<bool, EChaosProperty::ClusterIsAnchored>(Manager, DataIdx, SyncedData_External.bIsAnchored);
-		RemoteData.SyncRemote<TArray<FTransform>, EChaosProperty::ClusterChildToParent>(Manager, DataIdx, SyncedData_External.ChildToParent);
 	}
 
 	void FClusterUnionPhysicsProxy::ClearAccumulatedData()
@@ -378,7 +463,7 @@ namespace Chaos
 
 	void FClusterUnionPhysicsProxy::SetXR_External(const FVector& X, const FQuat& R)
 	{
-		if (!ensure(Particle_External))
+		if (!Solver || !ensure(Particle_External))
 		{
 			return;
 		}
@@ -390,11 +475,13 @@ namespace Chaos
 		FClusterUnionProxyTimestamp& SyncTS = GetSyncTimestampAs<FClusterUnionProxyTimestamp>();
 		SyncTS.OverWriteX.Set(GetSolverSyncTimestamp_External(), X);
 		SyncTS.OverWriteR.Set(GetSolverSyncTimestamp_External(), R);
+
+		Solver->AddDirtyProxy(this);
 	}
 
 	void FClusterUnionPhysicsProxy::SetLinearVelocity_External(const FVector& V)
 	{
-		if (!ensure(Particle_External))
+		if (!Solver || !ensure(Particle_External))
 		{
 			return;
 		}
@@ -404,6 +491,8 @@ namespace Chaos
 
 		FClusterUnionProxyTimestamp& SyncTS = GetSyncTimestampAs<FClusterUnionProxyTimestamp>();
 		SyncTS.OverWriteV.Set(GetSolverSyncTimestamp_External(), V);
+
+		Solver->AddDirtyProxy(this);
 	}
 
 	void FClusterUnionPhysicsProxy::SetAngularVelocity_External(const FVector& W)
@@ -418,6 +507,8 @@ namespace Chaos
 
 		FClusterUnionProxyTimestamp& SyncTS = GetSyncTimestampAs<FClusterUnionProxyTimestamp>();
 		SyncTS.OverWriteW.Set(GetSolverSyncTimestamp_External(), W);
+
+		Solver->AddDirtyProxy(this);
 	}
 
 	void FClusterUnionPhysicsProxy::SetChildToParent_External(FPhysicsObjectHandle Child, const FTransform& RelativeTransform)
@@ -427,14 +518,32 @@ namespace Chaos
 			return;
 		}
 
-		if (int32* Index = ChildPhysicsObjectIndexMap.Find(Child); Index && *Index >=0 && *Index < ChildPhysicsObjects.Num() && ensure(ChildPhysicsObjects[*Index] == Child))
-		{
-			SyncedData_External.ChildToParent[*Index] = RelativeTransform;
-			Particle_External->ForceDirty(EChaosPropertyFlags::ClusterChildToParent);
-			Solver->AddDirtyProxy(this);
+		Solver->EnqueueCommandImmediate(
+			[this, Child, RelativeTransform]() mutable
+			{
+				FPBDRigidsEvolutionGBF& Evolution = *GetEvolution(this);
+				FClusterUnionManager& ClusterUnionManager = Evolution.GetRigidClustering().GetClusterUnionManager();
 
-			FClusterUnionProxyTimestamp& SyncTS = GetSyncTimestampAs<FClusterUnionProxyTimestamp>();
-			SyncTS.OverWriteChildToParent.Set(GetSolverSyncTimestamp_External(), SyncedData_External.ChildToParent);
-		}
+				if (FClusterUnion* ClusterUnion = ClusterUnionManager.FindClusterUnion(ClusterUnionIndex))
+				{
+					FGeometryParticleHandle* Particle = Child->GetParticle<EThreadContext::Internal>();
+					if (!ensure(Particle))
+					{
+						return;
+					}
+
+					const int32 ParticleIndex = ClusterUnion->ChildParticles.Find(Particle->CastToRigidParticle());
+					if (ParticleIndex != INDEX_NONE)
+					{
+						if (FPBDRigidClusteredParticleHandle* ChildHandle = ClusterUnion->ChildParticles[ParticleIndex]->CastToClustered())
+						{
+							ChildHandle->SetChildToParent(RelativeTransform);
+							Evolution.GetParticles().MarkTransientDirtyParticle(ChildHandle);
+							Evolution.GetParticles().MarkTransientDirtyParticle(Particle);
+						}
+					}
+				}
+			}
+		);
 	}
 }
\ No newline at end of file
diff --git a/Engine/Source/Runtime/Experimental/Chaos/Private/PhysicsProxy/GeometryCollectionPhysicsProxy.cpp b/Engine/Source/Runtime/Experimental/Chaos/Private/PhysicsProxy/GeometryCollectionPhysicsProxy.cpp
--- a/Engine/Source/Runtime/Experimental/Chaos/Private/PhysicsProxy/GeometryCollectionPhysicsProxy.cpp
+++ b/Engine/Source/Runtime/Experimental/Chaos/Private/PhysicsProxy/GeometryCollectionPhysicsProxy.cpp
@@ -1957,6 +1957,8 @@ void FGeometryCollectionPhysicsProxy::BreakClusters_External(TArray<FGeometryCol
 			{
 				if (Chaos::FPBDRigidClusteredParticleHandle* ClusteredHandle = FindClusteredParticleHandleByItemIndex_Internal(ItemIndex))
 				{
+					// Pre-emptively remove the particle from cluster unions it might be in.
+					Clustering.GetClusterUnionManager().HandleRemoveOperationWithClusterLookup({ ClusteredHandle }, true);
 					Clustering.BreakCluster(ClusteredHandle);
 				}
 			}
diff --git a/Engine/Source/Runtime/Experimental/Chaos/Public/Chaos/ClusterUnionManager.h b/Engine/Source/Runtime/Experimental/Chaos/Public/Chaos/ClusterUnionManager.h
--- a/Engine/Source/Runtime/Experimental/Chaos/Public/Chaos/ClusterUnionManager.h
+++ b/Engine/Source/Runtime/Experimental/Chaos/Public/Chaos/ClusterUnionManager.h
@@ -20,7 +20,16 @@ namespace Chaos
 		Add,
 		// AddReleased is the original behavior where if the particle to be added is a cluster, we will release the cluster first
 		// and add its children instead.
-		AddReleased
+		AddReleased,
+		Remove
+	};
+
+	struct CHAOS_API FClusterUnionCreationParameters
+	{
+		FClusterUnionExplicitIndex ExplicitIndex = INDEX_NONE;
+		const FUniqueIdx* UniqueIndex = nullptr;
+		uint32 ActorId = 0;
+		uint32 ComponentId = 0;
 	};
 
 	struct CHAOS_API FClusterUnion
@@ -39,8 +48,13 @@ namespace Chaos
 
 		// Need to remember the parameters used to create the cluster so we can update it later.
 		FClusterCreationParameters Parameters;
+
+		// Other parameters that aren't related to clusters generally but are related to info we need about the cluster union.
+		FClusterUnionCreationParameters ClusterUnionParameters;
 	};
 
+	
+
 	/**
 	 * This class is used by Chaos to create internal clusters that will
 	 * cause one or more clusters to simulate together as a single rigid
@@ -52,15 +66,24 @@ namespace Chaos
 		FClusterUnionManager(FRigidClustering& InClustering, FPBDRigidsEvolutionGBF& InEvolution);
 
 		// Creates a new cluster union with an automatically assigned cluster union index.
-		FClusterUnionIndex CreateNewClusterUnion(const FClusterCreationParameters& Parameters, FClusterUnionExplicitIndex ExplicitIndex = INDEX_NONE, const FUniqueIdx* UniqueIndex = nullptr);
+		FClusterUnionIndex CreateNewClusterUnion(const FClusterCreationParameters& Parameters, const FClusterUnionCreationParameters& ClusterUnionParameters = FClusterUnionCreationParameters{});
 
 		// Destroy a given cluster union.
 		void DestroyClusterUnion(FClusterUnionIndex Index);
 
-		// Add a new operation to the queue.
+		// Add a new operation to the queue. Note that we currently only support the pending/flush operations for explicit operations. The behavior is legacy anyway so this should be fine.
 		void AddPendingExplicitIndexOperation(FClusterUnionExplicitIndex Index, EClusterUnionOperation Op, const TArray<FPBDRigidParticleHandle*>& Particles);
 		void AddPendingClusterIndexOperation(FClusterUnionIndex Index, EClusterUnionOperation Op, const TArray<FPBDRigidParticleHandle*>& Particles);
 
+		// Actually performs the change specified in the FClusterUnionOperationData structure.
+		void HandleAddOperation(FClusterUnionIndex ClusterIndex, const TArray<FPBDRigidParticleHandle*>& Particles, bool bReleaseClustersFirst);
+
+		// Removes the specified particles from the specified cluster.
+		void HandleRemoveOperation(FClusterUnionIndex ClusterIndex, const TArray<FPBDRigidParticleHandle*>& Particles, bool bUpdateClusterProperties);
+
+		// Helper function to remove particles given only the particle handle. This will consult the lookup table to find which cluster the particle is in before doing a normal remove operation.
+		void HandleRemoveOperationWithClusterLookup(const TArray<FPBDRigidParticleHandle*>& InParticles, bool bUpdateClusterProperties);
+
 		// Will be called at the beginning of every time step to ensure that all the expected cluster unions have been modified.
 		void FlushPendingOperations();
 
@@ -68,6 +91,8 @@ namespace Chaos
 		FClusterUnion* FindClusterUnionFromExplicitIndex(FClusterUnionExplicitIndex Index);
 		FClusterUnion* FindClusterUnion(FClusterUnionIndex Index);
 
+		FClusterUnionIndex FindClusterUnionIndexFromParticle(FPBDRigidParticleHandle* Particle);
+
 	private:
 		FRigidClustering& MClustering;
 		FPBDRigidsEvolutionGBF& MEvolution;
@@ -107,14 +132,17 @@ namespace Chaos
 		TMap<FClusterUnionExplicitIndex, FClusterUnionIndex> ExplicitIndexMap;
 		FClusterCreationParameters DefaultClusterCreationParameters() const;
 
+		// A particle lookup table. Lets us go from a given particle pointer to the cluster it's in.
+		TMap<FPBDRigidParticleHandle*, FClusterUnionIndex> ParticleToClusterUnionIndex;
+
 		// If no cluster index is set but an explicit index is set, map the explicit index to a regular index.
 		FClusterUnionIndex GetOrCreateClusterUnionIndexFromExplicitIndex(FClusterUnionExplicitIndex InIndex);
 
 		// Forcefully recreate the shared geometry on a cluster. Potentially expensive so ideally should be used rarely.
 		TSharedPtr<FImplicitObject, ESPMode::ThreadSafe> ForceRecreateClusterUnionSharedGeometry(const FClusterUnion& Union);
 
-		// Actually performs the change specified in the FClusterUnionOperationData structure.
-		void HandleAddOperation(FClusterUnionIndex ClusterIndex, const TArray<FPBDRigidParticleHandle*>& Particles, bool bReleaseClustersFirst);
+		// Update the cluster union's properties after its set of particle changes.
+		void UpdateAllClusterUnionProperties(FClusterUnion& ClusterUnion, bool bRecomputeMassOrientation);
 	};
 
 }
\ No newline at end of file
diff --git a/Engine/Source/Runtime/Experimental/Chaos/Public/Chaos/Framework/PhysicsProxyBase.h b/Engine/Source/Runtime/Experimental/Chaos/Public/Chaos/Framework/PhysicsProxyBase.h
--- a/Engine/Source/Runtime/Experimental/Chaos/Public/Chaos/Framework/PhysicsProxyBase.h
+++ b/Engine/Source/Runtime/Experimental/Chaos/Public/Chaos/Framework/PhysicsProxyBase.h
@@ -65,8 +65,6 @@ struct CHAOS_API FClusterUnionProxyTimestamp : public FProxyTimestampBase
 	TTimestampProperty<Chaos::FRotation3> OverWriteR;
 	TTimestampProperty<Chaos::FVec3> OverWriteV;
 	TTimestampProperty<Chaos::FVec3> OverWriteW;
-	TTimestampProperty<bool> OverWriteAnchored;
-	TTimestampProperty<TArray<FTransform>> OverWriteChildToParent;
 };
 
 class CHAOS_API IPhysicsProxyBase
diff --git a/Engine/Source/Runtime/Experimental/Chaos/Public/Chaos/PBDRigidClustering.h b/Engine/Source/Runtime/Experimental/Chaos/Public/Chaos/PBDRigidClustering.h
--- a/Engine/Source/Runtime/Experimental/Chaos/Public/Chaos/PBDRigidClustering.h
+++ b/Engine/Source/Runtime/Experimental/Chaos/Public/Chaos/PBDRigidClustering.h
@@ -106,6 +106,12 @@ public:
 		const TArray<FPBDRigidParticleHandle*>& InChildren,
 		const TMap<FPBDRigidParticleHandle*, FPBDRigidParticleHandle*>& ChildToParentMap);
 
+	/**
+	 * Manually remove a set of particles from the cluster.
+	 */
+	void RemoveParticlesFromCluster(
+		FPBDRigidClusteredParticleHandle* Cluster,
+		const TArray<FPBDRigidParticleHandle*>& InChildren);
 	/**
 	 *  UnionClusterGroups
 	 *    Clusters that share a group index should be unioned into a single cluster prior to simulation.
@@ -351,6 +357,8 @@ public:
 		Chaos::FPBDRigidClusteredParticleHandle* Parent,
 		const FClusterCreationParameters & Parameters = FClusterCreationParameters());
 
+	void ClearConnectionGraph(FPBDRigidClusteredParticleHandle* Parent);
+
 	const TSet<Chaos::FPBDRigidClusteredParticleHandle*>& GetTopLevelClusterParents() const { return TopLevelClusterParents; }
 	TSet<Chaos::FPBDRigidClusteredParticleHandle*>& GetTopLevelClusterParents() { return TopLevelClusterParents; }
 
@@ -410,6 +418,15 @@ public:
 
 	void UpdateTopLevelParticle(FPBDRigidClusteredParticleHandle* Particle);
 
+	/**
+	 * This function is a bit more versatile than the name suggests. This function can either be used to update the cluster properties
+	 * incrementally or entirely rebuild the properties all over again. This all depends on whether the input children is
+	 * either 1) the new children or 2) all the children as well as what those initial properties are set to.
+	 */
+	void UpdateClusterParticlePropertiesFromChildren(
+		FPBDRigidClusteredParticleHandle* Cluster,
+		const FRigidHandleArray& Children,
+		const TMap<FPBDRigidParticleHandle*, FPBDRigidParticleHandle*>& ChildToParentMap);
 private:
 
 	FRigidEvolution& MEvolution;
diff --git a/Engine/Source/Runtime/Experimental/Chaos/Public/Chaos/ParticleProperties.inl b/Engine/Source/Runtime/Experimental/Chaos/Public/Chaos/ParticleProperties.inl
--- a/Engine/Source/Runtime/Experimental/Chaos/Public/Chaos/ParticleProperties.inl
+++ b/Engine/Source/Runtime/Experimental/Chaos/Public/Chaos/ParticleProperties.inl
@@ -32,7 +32,5 @@ CHAOS_PROPERTY(GroundParticleProxy, FParticleProxyProperty, EPhysicsProxyType::C
 
 CHAOS_PROPERTY(ClusterXR, FParticlePositionRotation, EPhysicsProxyType::ClusterUnionProxy)
 CHAOS_PROPERTY(ClusterVelocities, FParticleVelocities, EPhysicsProxyType::ClusterUnionProxy)
-CHAOS_PROPERTY(ClusterIsAnchored, bool, EPhysicsProxyType::ClusterUnionProxy)
-CHAOS_PROPERTY(ClusterChildToParent, TArray<FTransform>, EPhysicsProxyType::ClusterUnionProxy)
 
 #undef CHAOS_PROPERTY
\ No newline at end of file
diff --git a/Engine/Source/Runtime/Experimental/Chaos/Public/Chaos/PullPhysicsDataImp.h b/Engine/Source/Runtime/Experimental/Chaos/Public/Chaos/PullPhysicsDataImp.h
--- a/Engine/Source/Runtime/Experimental/Chaos/Public/Chaos/PullPhysicsDataImp.h
+++ b/Engine/Source/Runtime/Experimental/Chaos/Public/Chaos/PullPhysicsDataImp.h
@@ -63,14 +63,25 @@ struct FDirtyGeometryCollectionData : public TBasePullData<FGeometryCollectionPh
 	FGeometryCollectionResults Results;
 };
 
+struct FDirtyClusterUnionParticleData
+{
+	FPBDRigidParticle* Particle = nullptr;
+	FRigidTransform3 ChildToParent;
+};
+
 struct FDirtyClusterUnionData : public TBasePullData<FClusterUnionPhysicsProxy, FClusterUnionProxyTimestamp>
 {
 	FVec3 X;
 	FQuat R;
 	FVec3 V;
 	FVec3 W;
+	EObjectStateType ObjectState;
 	bool bIsAnchored = false;
-	TArray<FTransform> ChildToParent;
+	TArray<FDirtyClusterUnionParticleData> ChildParticles;
+	TSharedPtr<FImplicitObject, ESPMode::ThreadSafe> SharedGeometry;
+	TArray<FCollisionData> CollisionData;
+	TArray<FCollisionFilterData> QueryData;
+	TArray<FCollisionFilterData> SimData;
 };
 
 struct FJointConstraintOutputData {
diff --git a/Engine/Source/Runtime/Experimental/Chaos/Public/PhysicsProxy/ClusterUnionPhysicsProxy.h b/Engine/Source/Runtime/Experimental/Chaos/Public/PhysicsProxy/ClusterUnionPhysicsProxy.h
--- a/Engine/Source/Runtime/Experimental/Chaos/Public/PhysicsProxy/ClusterUnionPhysicsProxy.h
+++ b/Engine/Source/Runtime/Experimental/Chaos/Public/PhysicsProxy/ClusterUnionPhysicsProxy.h
@@ -17,6 +17,12 @@ namespace Chaos
 {
 	struct FDirtyClusterUnionData;
 
+	struct FClusterUnionChildData
+	{
+		FPBDRigidParticle* Particle = nullptr;
+		FTransform ChildToParent;
+	};
+
 	/**
 	 * Extra data that needs to be synced between the PT and GT for cluster unions.
 	 */
@@ -25,12 +31,8 @@ namespace Chaos
 		// Whether the cluster is anchored or not.
 		bool bIsAnchored;
 
-		// For every child, we need to sync its ChildToParent transform. The order of the transforms is identical
-		// to the order of the physics objects stored in the proxy. The *direction* that we sync ChildToParent
-		// is determined by whether we're on the server or the client. On the server, the ChildToParent is physics
-		// thread authoritative so the sync goes from the PT to the GT. On the client, we want to use whatever
-		// ChildToParent was given by the server so it's GT authoritative so the sync goes the other way.
-		TArray<FTransform> ChildToParent;
+		// Data on every child particle in the cluster union.
+		TArray<FClusterUnionChildData> ChildParticles;
 	};
 
 
@@ -43,16 +45,22 @@ namespace Chaos
 		using FInternalParticle = TPBDRigidClusteredParticleHandle<Chaos::FReal, 3>;
 
 		FClusterUnionPhysicsProxy() = delete;
-		FClusterUnionPhysicsProxy(UObject* InOwner, const FClusterCreationParameters& InParameters, void* InUserData, EThreadContext InAuthoritativeContext);
+		FClusterUnionPhysicsProxy(UObject* InOwner, const FClusterCreationParameters& InParameters, void* InUserData, uint32 InActorId, uint32 InComponentId);
 
 		// Add physics objects to the cluster union. Should only be called from the game thread.
 		void AddPhysicsObjects_External(const TArray<FPhysicsObjectHandle>& Objects);
-		const TArray<FPhysicsObjectHandle>& GetChildPhysicsObjects_External() const { return ChildPhysicsObjects; }
 		const FClusterUnionSyncedData& GetSyncedData_External() const { return SyncedData_External; }
 
+		// Remove physics objects from the cluster union.
+		void RemovePhysicsObjects_External(const TSet<FPhysicsObjectHandle>& Objects);
+
 		// Set/Remove any anchors on the cluster particle.
 		void SetIsAnchored_External(bool bIsAnchored);
 
+		// Set Object State.
+		EObjectStateType GetObjectState_External() const;
+		void SetObjectState_External(EObjectStateType State);
+
 		// Cluster union proxy initialization happens in two first on the game thread (external) then on the
 		// physics thread (internal). Cluster unions are a primarily physics concept so the things exposed to
 		// an external context is just there to be able to safely query the state on the physics thread and to
@@ -67,7 +75,7 @@ namespace Chaos
 		virtual void* GetHandleUnsafe() const override { return Particle_Internal; }
 		FPhysicsObjectHandle GetPhysicsObjectHandle() const { return PhysicsObject.Get(); }
 
-		bool HasChildren_External() const { return !ChildrenParticles_External.IsEmpty(); }
+		bool HasChildren_External() const { return !SyncedData_External.ChildParticles.IsEmpty(); }
 		bool HasChildren_Internal() const;
 
 		bool IsAnchored_External() const { return SyncedData_External.bIsAnchored; }
@@ -94,29 +102,21 @@ namespace Chaos
 		const FProxyInterpolationData& GetInterpolationData() const { return InterpolationData; }
 
 		FClusterUnionIndex GetClusterUnionIndex() const { return ClusterUnionIndex; }
+
 	private:
 		bool bIsInitializedOnPhysicsThread = false;
 		FClusterCreationParameters ClusterParameters;
 		void* UserData = nullptr;
-		EThreadContext AuthoritativeContext;
+		uint32 ActorId;
+		uint32 ComponentId;
 
 		FPhysicsObjectUniquePtr PhysicsObject;
 		TUniquePtr<FExternalParticle> Particle_External;
 		FClusterUnionSyncedData SyncedData_External;
 
-		// This isn't captured in TPBDRigidParticle hence why we need to track it manually.
-		TSet<FExternalParticle*> ChildrenParticles_External;
-
 		FInternalParticle* Particle_Internal = nullptr;
 		FClusterUnionIndex ClusterUnionIndex = INDEX_NONE;
 
-		// Children. Assuming that the ChildPhysicsObjects array can get large so we also maintain
-		// the ChildPhysicsObjectIndexMap so that it's easy to find a given physics object handle's index.
-		// This is useful because other things such as the FClusterUnionSyncedData::ChildToParent array
-		// will utilize the same index.
-		TArray<FPhysicsObjectHandle> ChildPhysicsObjects;
-		TMap<FPhysicsObjectHandle, int32> ChildPhysicsObjectIndexMap;
-
 		FProxyInterpolationData InterpolationData;
 
 		//~ Begin TPhysicsProxy Interface
