Fixed null SRV assert on switching editor to a mobile preview

diff --git a/Engine/Source/Runtime/Engine/Private/DynamicMeshBuilder.cpp b/Engine/Source/Runtime/Engine/Private/DynamicMeshBuilder.cpp
--- a/Engine/Source/Runtime/Engine/Private/DynamicMeshBuilder.cpp
+++ b/Engine/Source/Runtime/Engine/Private/DynamicMeshBuilder.cpp
@@ -465,7 +465,7 @@ public:
 		RHIUnlockBuffer(TexCoordBuffer.VertexBufferRHI);
 		RHIUnlockBuffer(ColorBuffer.VertexBufferRHI);
 
-		if (RHISupportsManualVertexFetch(GetFeatureLevelShaderPlatform(FeatureLevel)))
+		if (RHISupportsManualVertexFetch(GMaxRHIShaderPlatform))
 		{
 			TangentBufferSRV = RHICreateShaderResourceView(TangentBuffer.VertexBufferRHI, 4, PF_R8G8B8A8_SNORM);
 			TexCoordBufferSRV = RHICreateShaderResourceView(TexCoordBuffer.VertexBufferRHI, TextureStride, TextureFormat);
