Fixed bug on level not using OFPA with Actor Folders enabled : internal acceleration table FolderLabelToActorFolders was not updated.

diff --git a/Engine/Source/Runtime/Engine/Private/Level.cpp b/Engine/Source/Runtime/Engine/Private/Level.cpp
--- a/Engine/Source/Runtime/Engine/Private/Level.cpp
+++ b/Engine/Source/Runtime/Engine/Private/Level.cpp
@@ -2807,13 +2807,26 @@ void ULevel::FixupActorFolders()
 	{
 		TGuardValue<bool> FixingActorFolders(GIsFixingActorFolders, true);
 
-		// At this point, LoadedExternalActorFolders are fully loaded, transfer them to the ActorFolders list.
-		for (UActorFolder* LoadedActorFolder : LoadedExternalActorFolders)
+		if (IsUsingExternalObjects())
 		{
-			check(LoadedActorFolder->GetGuid().IsValid());
-			AddActorFolder(LoadedActorFolder);
+			// At this point, LoadedExternalActorFolders are fully loaded, transfer them to the ActorFolders list.
+			for (UActorFolder* LoadedActorFolder : LoadedExternalActorFolders)
+			{
+				check(LoadedActorFolder->GetGuid().IsValid());
+				AddActorFolder(LoadedActorFolder);
+			}
+			LoadedExternalActorFolders.Empty();
+		}
+		else
+		{
+			// Update FolderLabelToActorFolders for non-externalized ActorFolders
+			ForEachActorFolder([this](UActorFolder* ActorFolder)
+			{
+				FActorFolderSet& Folders = FolderLabelToActorFolders.FindOrAdd(ActorFolder->GetLabel());
+				Folders.Add(ActorFolder);
+				return true;
+			}, /*bSkipDeleted*/ true);
 		}
-		LoadedExternalActorFolders.Empty();
 
 		ForEachActorFolder([](UActorFolder* ActorFolder)
 		{
