World Partition: expose data layer assets to actor descriptor blueprint library.

diff --git a/Engine/Source/Runtime/Engine/Private/WorldPartitionBlueprintLibrary.cpp b/Engine/Source/Runtime/Engine/Private/WorldPartitionBlueprintLibrary.cpp
--- a/Engine/Source/Runtime/Engine/Private/WorldPartitionBlueprintLibrary.cpp
+++ b/Engine/Source/Runtime/Engine/Private/WorldPartitionBlueprintLibrary.cpp
@@ -43,6 +43,14 @@ FActorDesc::FActorDesc(const FWorldPartitionActorDesc& InActorDesc, const FTrans
 	bIsSpatiallyLoaded = InActorDesc.GetIsSpatiallyLoaded();
 	bActorIsEditorOnly = InActorDesc.GetActorIsEditorOnly();
 	ActorPackage = InActorDesc.GetActorPackage();
+	if (InActorDesc.IsUsingDataLayerAsset())
+	{
+		DataLayerAssets.Reserve(InActorDesc.GetDataLayers().Num());
+		for (FName DataLayerAssetPath : InActorDesc.GetDataLayers())
+		{
+			DataLayerAssets.Add(FSoftObjectPath(DataLayerAssetPath.ToString()));
+		}
+	}
 	
 	PRAGMA_DISABLE_DEPRECATION_WARNINGS
 	ActorPath = InActorDesc.GetActorPath();
diff --git a/Engine/Source/Runtime/Engine/Public/WorldPartition/WorldPartitionBlueprintLibrary.h b/Engine/Source/Runtime/Engine/Public/WorldPartition/WorldPartitionBlueprintLibrary.h
--- a/Engine/Source/Runtime/Engine/Public/WorldPartition/WorldPartitionBlueprintLibrary.h
+++ b/Engine/Source/Runtime/Engine/Public/WorldPartition/WorldPartitionBlueprintLibrary.h
@@ -70,6 +70,10 @@ struct FActorDesc
 	/** Actor's path name. */
 	UPROPERTY(BlueprintReadOnly, VisibleAnywhere, Category=Actor)
 	FName ActorPath;
+
+	/** Actor's data layer assets. */
+	UPROPERTY(BlueprintReadOnly, VisibleAnywhere, Category=Actor)
+	TArray<FSoftObjectPath> DataLayerAssets;
 };
 
 UCLASS(MinimalAPI)
