WorldPartitionActorFilter: made recursing into subcontainers generic.

diff --git a/Engine/Source/Runtime/Engine/Private/WorldPartition/LevelInstance/LevelInstanceActorDesc.cpp b/Engine/Source/Runtime/Engine/Private/WorldPartition/LevelInstance/LevelInstanceActorDesc.cpp
--- a/Engine/Source/Runtime/Engine/Private/WorldPartition/LevelInstance/LevelInstanceActorDesc.cpp
+++ b/Engine/Source/Runtime/Engine/Private/WorldPartition/LevelInstance/LevelInstanceActorDesc.cpp
@@ -223,12 +223,11 @@ bool FLevelInstanceActorDesc::GetContainerInstance(FContainerInstance& OutContai
 						}
 					}
 
-					if (ActorDescIt->GetActorNativeClass()->ImplementsInterface(ULevelInstanceInterface::StaticClass()))
+					if (ActorDescIt->IsContainerInstance())
 					{
-						const FLevelInstanceActorDesc* LevelInstanceActorDesc = static_cast<const FLevelInstanceActorDesc*>(*ActorDescIt);
 						FWorldPartitionActorDesc::FContainerInstance ChildContainerInstance;
 						const bool bBuildFilter = false;
-						if (LevelInstanceActorDesc->GetContainerInstance(ChildContainerInstance, bBuildFilter))
+						if (ActorDescIt->GetContainerInstance(ChildContainerInstance, bBuildFilter))
 						{
 							ProcessContainers(FActorContainerID(InContainerID, ActorDescIt->GetGuid()), ChildContainerInstance.Container);
 						}
