Make UEngine latency function static. UBSan complain that the GEngine variable is null when we start the engine. It was not crashing because no member of UEngine is used in the functions.

diff --git a/Engine/Source/Runtime/Engine/Classes/Engine/Engine.h b/Engine/Source/Runtime/Engine/Classes/Engine/Engine.h
--- a/Engine/Source/Runtime/Engine/Classes/Engine/Engine.h
+++ b/Engine/Source/Runtime/Engine/Classes/Engine/Engine.h
@@ -2290,18 +2290,18 @@ public:
 	 */
 	virtual void UpdateTimeAndHandleMaxTickRate();
 
-	void SetInputSampleLatencyMarker(uint64 FrameNumber);
+	static void SetInputSampleLatencyMarker(uint64 FrameNumber);
 
-	void SetSimulationLatencyMarkerStart(uint64 FrameNumber);
-	void SetSimulationLatencyMarkerEnd(uint64 FrameNumber);
+	static void SetSimulationLatencyMarkerStart(uint64 FrameNumber);
+	static void SetSimulationLatencyMarkerEnd(uint64 FrameNumber);
 
-	void SetPresentLatencyMarkerStart(uint64 FrameNumber);
-	void SetPresentLatencyMarkerEnd(uint64 FrameNumber);
+	static void SetPresentLatencyMarkerStart(uint64 FrameNumber);
+	static void SetPresentLatencyMarkerEnd(uint64 FrameNumber);
 
-	void SetRenderSubmitLatencyMarkerStart(uint64 FrameNumber);
-	void SetRenderSubmitLatencyMarkerEnd(uint64 FrameNumber);
+	static void SetRenderSubmitLatencyMarkerStart(uint64 FrameNumber);
+	static void SetRenderSubmitLatencyMarkerEnd(uint64 FrameNumber);
 
-	void SetFlashIndicatorLatencyMarker(uint64 FrameNumber);
+	static void SetFlashIndicatorLatencyMarker(uint64 FrameNumber);
 
 	/**
 	 * Allows games to correct the negative delta
diff --git a/Engine/Source/Runtime/Engine/Private/UnrealClient.cpp b/Engine/Source/Runtime/Engine/Private/UnrealClient.cpp
--- a/Engine/Source/Runtime/Engine/Private/UnrealClient.cpp
+++ b/Engine/Source/Runtime/Engine/Private/UnrealClient.cpp
@@ -1679,7 +1679,7 @@ void FViewport::EndRenderFrame(FRHICommandListImmediate& RHICmdList, bool bPrese
 
 	RHICmdList.EnqueueLambda([CurrentFrameCounter = GFrameCounterRenderThread](FRHICommandListImmediate& InRHICmdList)
 	{
-		GEngine->SetPresentLatencyMarkerStart(CurrentFrameCounter);
+		UEngine::SetPresentLatencyMarkerStart(CurrentFrameCounter);
 	});
 
 	{
@@ -1689,7 +1689,7 @@ void FViewport::EndRenderFrame(FRHICommandListImmediate& RHICmdList, bool bPrese
 
 	RHICmdList.EnqueueLambda([CurrentFrameCounter = GFrameCounterRenderThread](FRHICommandListImmediate& InRHICmdList)
 	{
-		GEngine->SetPresentLatencyMarkerEnd(CurrentFrameCounter);
+		UEngine::SetPresentLatencyMarkerEnd(CurrentFrameCounter);
 	});
 }
 
diff --git a/Engine/Source/Runtime/Launch/Private/LaunchEngineLoop.cpp b/Engine/Source/Runtime/Launch/Private/LaunchEngineLoop.cpp
--- a/Engine/Source/Runtime/Launch/Private/LaunchEngineLoop.cpp
+++ b/Engine/Source/Runtime/Launch/Private/LaunchEngineLoop.cpp
@@ -5470,7 +5470,7 @@ static inline void BeginFrameRenderThread(FRHICommandListImmediate& RHICmdList,
 
 	RHICmdList.EnqueueLambda([CurrentFrameCounter](FRHICommandListImmediate& InRHICmdList)
 	{
-		GEngine->SetRenderSubmitLatencyMarkerStart(CurrentFrameCounter);
+		UEngine::SetRenderSubmitLatencyMarkerStart(CurrentFrameCounter);
 	});
 
 #if CSV_PROFILER
@@ -5496,7 +5496,7 @@ static inline void EndFrameRenderThread(FRHICommandListImmediate& RHICmdList, ui
 
 	RHICmdList.EnqueueLambda([CurrentFrameCounter](FRHICommandListImmediate& InRHICmdList)
 	{
-		GEngine->SetRenderSubmitLatencyMarkerEnd(CurrentFrameCounter);
+		UEngine::SetRenderSubmitLatencyMarkerEnd(CurrentFrameCounter);
 	});
 
 	FCoreDelegates::OnEndFrameRT.Broadcast();
diff --git a/Engine/Source/Runtime/SlateRHIRenderer/Private/SlateRHIRenderer.cpp b/Engine/Source/Runtime/SlateRHIRenderer/Private/SlateRHIRenderer.cpp
--- a/Engine/Source/Runtime/SlateRHIRenderer/Private/SlateRHIRenderer.cpp
+++ b/Engine/Source/Runtime/SlateRHIRenderer/Private/SlateRHIRenderer.cpp
@@ -1499,14 +1499,14 @@ void FSlateRHIRenderer::DrawWindow_RenderThread(FRHICommandListImmediate& RHICmd
 
 	RHICmdList.EnqueueLambda([CurrentFrameCounter = GFrameCounterRenderThread](FRHICommandListImmediate& InRHICmdList)
 	{
-		GEngine->SetPresentLatencyMarkerStart(CurrentFrameCounter);
+		UEngine::SetPresentLatencyMarkerStart(CurrentFrameCounter);
 	});
 
 	RHICmdList.EndDrawingViewport(ViewportInfo.ViewportRHI, true, DrawCommandParams.bLockToVsync);
 
 	RHICmdList.EnqueueLambda([CurrentFrameCounter = GFrameCounterRenderThread](FRHICommandListImmediate& InRHICmdList)
 	{
-		GEngine->SetPresentLatencyMarkerEnd(CurrentFrameCounter);
+		UEngine::SetPresentLatencyMarkerEnd(CurrentFrameCounter);
 	});
 
 	uint32 EndTime = FPlatformTime::Cycles();
