Prevent Dirtying the ExternalActor Packages when an actor unload sequence calls Modify the pacakge
Fix crash occuring when the last refcount on an WorldPartitionHandle unloads an actor which gets dirtied during its unload sequence.
The dirty actors gets caught in UWorldPartition::OnPackageDirtyStateChanged which increment the refcount and reload the unloading actor

diff --git a/Engine/Source/Runtime/Engine/Private/WorldPartition/WorldPartitionHandle.cpp b/Engine/Source/Runtime/Engine/Private/WorldPartition/WorldPartitionHandle.cpp
--- a/Engine/Source/Runtime/Engine/Private/WorldPartition/WorldPartitionHandle.cpp
+++ b/Engine/Source/Runtime/Engine/Private/WorldPartition/WorldPartitionHandle.cpp
@@ -86,6 +86,9 @@ void FWorldPartitionLoadingContext::FImmediate::RegisterActor(FWorldPartitionAct
 
 void FWorldPartitionLoadingContext::FImmediate::UnregisterActor(FWorldPartitionActorDesc* ActorDesc)
 {
+	// Set GIsEditorLoadingPackage to avoid dirtying the Actor package if Modify() is called during the unload sequence
+	TGuardValue<bool> IsEditorLoadingPackageGuard(GIsEditorLoadingPackage, true);
+
 	// When cleaning up worlds, actors are already marked as garbage at this point, so no need to remove them from the world
 	if (AActor* Actor = ActorDesc->GetActor(); IsValid(Actor))
 	{
@@ -149,6 +152,9 @@ FWorldPartitionLoadingContext::FDeferred::~FDeferred()
 
 			if (ContainerOp.Unregistrations.Num())
 			{
+				// Set GIsEditorLoadingPackage to avoid dirtying the Actor package if Modify() is called during the unload sequence
+				TGuardValue<bool> IsEditorLoadingPackageGuard(GIsEditorLoadingPackage, true);
+
 				TArray<AActor*> ActorList;
 				if (ULevel* Level = CreateActorList(ActorList, ContainerOp.Unregistrations))
 				{
