Fix an issue where if you specify an input device type override and the duration of that override curve was longer then the default one, it wouldn't play the whole effect.

Make this a CVar just in case it causes any live issues in 5.2, or a game wants to explictly turn it off for minut performance savings. In 5.3, we need to replace this cvar with an actual float property on the FForceFeedbackEffectOverridenChannelDetails instead so that it can be cached during the cook and not need to be queried from the curve at all during runtime.

#rnx

diff --git a/Engine/Source/Runtime/Engine/Private/GameFramework/ForceFeedbackEffect.cpp b/Engine/Source/Runtime/Engine/Private/GameFramework/ForceFeedbackEffect.cpp
--- a/Engine/Source/Runtime/Engine/Private/GameFramework/ForceFeedbackEffect.cpp
+++ b/Engine/Source/Runtime/Engine/Private/GameFramework/ForceFeedbackEffect.cpp
@@ -4,9 +4,23 @@
 #include "GameFramework/ForceFeedbackParameters.h"
 #include "Misc/App.h"
 #include "GameFramework/InputDeviceProperties.h"
+#include "HAL/IConsoleManager.h"	// For FAutoConsoleVariableRef
 
 #include UE_INLINE_GENERATED_CPP_BY_NAME(ForceFeedbackEffect)
 
+namespace UE::Input::Private
+{
+	// This is a flag that will force the re-evaluation of a force feedback effect's duration at runtime.
+	// This fixes an issue where if you specify an input device type override and the duration of that override
+	// curve was longer then the default one, it wouldn't play the whole effect.
+	// see UE-178719
+	static bool bShouldAlwaysEvaluateForceFeedbackDuration = true;
+	static FAutoConsoleVariableRef CVarShouldAlwaysEvaluateForceFeedbackDuration(TEXT("Input.ShouldAlwaysEvaluateForceFeedbackDuration"),
+		bShouldAlwaysEvaluateForceFeedbackDuration,
+		TEXT("Should the duration of a force feedback effect be evaluated every time it is called?"));
+}
+
+
 UForceFeedbackEffect::UForceFeedbackEffect(const FObjectInitializer& ObjectInitializer)
 	: Super(ObjectInitializer)
 {
@@ -33,7 +47,7 @@ void UForceFeedbackEffect::PostEditChangeChainProperty(struct FPropertyChangedCh
 float UForceFeedbackEffect::GetDuration()
 {
 	// Always recalc the duration when in the editor as it could change
-	if( GIsEditor || ( Duration < UE_SMALL_NUMBER ) )
+	if (GIsEditor || ( Duration < UE_SMALL_NUMBER ) || UE::Input::Private::bShouldAlwaysEvaluateForceFeedbackDuration)
 	{
 		Duration = 0.f;
 
