Fixed VDB importer missing tile data on import. Implemented SVT tile deduplication by hashing and redirecting page table entries to identical tiles, potentially saving a good chunk of memory on suitable assets.

#rnx

diff --git a/Engine/Source/Editor/SparseVolumeTexture/Private/OpenVDBGridAdapter.h b/Engine/Source/Editor/SparseVolumeTexture/Private/OpenVDBGridAdapter.h
--- a/Engine/Source/Editor/SparseVolumeTexture/Private/OpenVDBGridAdapter.h
+++ b/Engine/Source/Editor/SparseVolumeTexture/Private/OpenVDBGridAdapter.h
@@ -116,21 +116,47 @@ public:
 	{
 		constexpr uint32 NumComponents = GetOpenVDBValueNumComponents<ValueType>();
 		static_assert(NumComponents <= 4);
-		for (auto LeafIt = Grid->constTree().cbeginLeaf(); LeafIt; ++LeafIt)
+
+		// Iterate over both voxels and tiles.
+		for (auto ValueIt = Grid->cbeginValueOn(); ValueIt; ++ValueIt)
 		{
-			for (auto ValueIt = LeafIt->cbeginValueOn(); ValueIt; ++ValueIt)
+			const ValueType VoxelValue = ValueIt.getValue();
+			const openvdb::Coord CoordVDB = ValueIt.getCoord();
+			
+			float VoxelValueComponents[4]{};
+			for (uint32 ComponentIdx = 0; ComponentIdx < NumComponents; ++ComponentIdx)
 			{
-				const ValueType VoxelValue = ValueIt.getValue();
-				float VoxelValueComponents[4]{};
-				for (uint32 ComponentIdx = 0; ComponentIdx < NumComponents; ++ComponentIdx)
-				{
-					VoxelValueComponents[ComponentIdx] = GetOpenVDBValueComponent(VoxelValue, ComponentIdx);
-				}
+				VoxelValueComponents[ComponentIdx] = GetOpenVDBValueComponent(VoxelValue, ComponentIdx);
+			}
+			const FIntVector3 Coord(CoordVDB[0], CoordVDB[1], CoordVDB[2]);
 
-				const openvdb::Coord CoordVDB = ValueIt.getCoord();
-				const FIntVector3 Coord(CoordVDB[0], CoordVDB[1], CoordVDB[2]);
+			if (ValueIt.isVoxelValue())
+			{
+				// A single voxel
 				OnVisit(Coord, NumComponents, VoxelValueComponents);
 			}
+			else
+			{
+				// Tile, a higher level node specifying a single value for all finer level nodes/children without actually allocating memory for them.
+				openvdb::CoordBBox TileAABB;
+				ValueIt.getBoundingBox(TileAABB);
+				check(TileAABB.min().x() == Coord.X && TileAABB.min().y() == Coord.Y && TileAABB.min().z() == Coord.Z);
+				const openvdb::Coord AABBDimVDB = TileAABB.dim();
+				const FIntVector3 AABBDim = FIntVector3(AABBDimVDB.x(), AABBDimVDB.y(), AABBDimVDB.z());
+				
+				// Iterate over tile AABB and invoke callback for every (virtual) voxel. AABB.max is inclusive, so we need <= instead of <.
+				for (int32 Z = 0; Z <= AABBDim.Z; ++Z)
+				{
+					for (int32 Y = 0; Y <= AABBDim.Y; ++Y)
+					{
+						for (int32 X = 0; X <= AABBDim.X; ++X)
+						{
+							const FIntVector3 TileLocalCoord(Coord.X + X, Coord.Y + Y, Coord.Z + Z);
+							OnVisit(TileLocalCoord, NumComponents, VoxelValueComponents);
+						}
+					}
+				}
+			}
 		}
 	}
 
diff --git a/Engine/Source/Runtime/Engine/Private/SparseVolumeTexture/SparseVolumeTexture.cpp b/Engine/Source/Runtime/Engine/Private/SparseVolumeTexture/SparseVolumeTexture.cpp
--- a/Engine/Source/Runtime/Engine/Private/SparseVolumeTexture/SparseVolumeTexture.cpp
+++ b/Engine/Source/Runtime/Engine/Private/SparseVolumeTexture/SparseVolumeTexture.cpp
@@ -94,9 +94,12 @@ bool FSparseVolumeTextureFrame::BuildDerivedData(const FIntVector3& VolumeResolu
 		AddressingInfo.AddressY = AddressY;
 		AddressingInfo.AddressZ = AddressZ;
 
-		TextureData.BuildDerivedData(AddressingInfo);
-
-		*OutMippedTextureData = MoveTemp(TextureData);
+		const int32 NumMipLevels = 1; // generate entire mip chain
+		const bool bMoveMip0FromSource = true; // we have no need to keep TextureData around
+		if (!TextureData.BuildDerivedData(AddressingInfo, NumMipLevels, bMoveMip0FromSource, *OutMippedTextureData))
+		{
+			return false;
+		}
 
 		// Now unload the raw data
 		RawData.UnloadData();
diff --git a/Engine/Source/Runtime/Engine/Private/SparseVolumeTexture/SparseVolumeTextureData.cpp b/Engine/Source/Runtime/Engine/Private/SparseVolumeTexture/SparseVolumeTextureData.cpp
--- a/Engine/Source/Runtime/Engine/Private/SparseVolumeTexture/SparseVolumeTextureData.cpp
+++ b/Engine/Source/Runtime/Engine/Private/SparseVolumeTexture/SparseVolumeTextureData.cpp
@@ -3,6 +3,7 @@
 #include "SparseVolumeTexture/SparseVolumeTextureData.h"
 #include "SparseVolumeTexture/SparseVolumeTextureUtility.h"
 #include "Async/ParallelFor.h"
+#include "Hash/CityHash.h"
 
 DEFINE_LOG_CATEGORY_STATIC(LogSparseVolumeTextureData, Log, All);
 
@@ -179,8 +180,8 @@ bool FSparseVolumeTextureData::Construct(const ISparseVolumeTextureDataConstruct
 	const int32 FormatSize[] = { GPixelFormats[Header.AttributesFormats[0]].BlockBytes, GPixelFormats[Header.AttributesFormats[1]].BlockBytes };
 	PageTable.SetNum(1);
 	PageTable[0].SetNumZeroed(Header.PageTableVolumeResolution.X * Header.PageTableVolumeResolution.Y * Header.PageTableVolumeResolution.Z);
-	PhysicalTileDataA.SetNumZeroed(EffectivelyAllocatedTiles * SVTNumVoxelsPerPaddedTile * FormatSize[0]);
-	PhysicalTileDataB.SetNumZeroed(EffectivelyAllocatedTiles * SVTNumVoxelsPerPaddedTile * FormatSize[1]);
+	PhysicalTileDataA.SetNumZeroed((int64)EffectivelyAllocatedTiles * (int64)SVTNumVoxelsPerPaddedTile * (int64)FormatSize[0]);
+	PhysicalTileDataB.SetNumZeroed((int64)EffectivelyAllocatedTiles * (int64)SVTNumVoxelsPerPaddedTile * (int64)FormatSize[1]);
 
 	// Generate page table and tile volume data by splatting the data
 	{
@@ -245,10 +246,10 @@ bool FSparseVolumeTextureData::Construct(const ISparseVolumeTextureDataConstruct
 		{
 			const FIntVector3 GridCoord = Coord;
 			check(GridCoord.X >= 0 && GridCoord.Y >= 0 && GridCoord.Z >= 0);
-			check(GridCoord.X < Header.VirtualVolumeAABBMax.X&& GridCoord.Y < Header.VirtualVolumeAABBMax.Y&& GridCoord.Z < Header.VirtualVolumeAABBMax.Z);
+			check(GridCoord.X < Header.VirtualVolumeAABBMax.X && GridCoord.Y < Header.VirtualVolumeAABBMax.Y && GridCoord.Z < Header.VirtualVolumeAABBMax.Z);
 			const FIntVector3 PageCoord = (GridCoord / SPARSE_VOLUME_TILE_RES) - Header.PageTableVolumeAABBMin;
 			check(PageCoord.X >= 0 && PageCoord.Y >= 0 && PageCoord.Z >= 0);
-			check(PageCoord.X < Header.PageTableVolumeResolution.X&& PageCoord.Y < Header.PageTableVolumeResolution.Y&& PageCoord.Z < Header.PageTableVolumeResolution.Z);
+			check(PageCoord.X < Header.PageTableVolumeResolution.X && PageCoord.Y < Header.PageTableVolumeResolution.Y && PageCoord.Z < Header.PageTableVolumeResolution.Z);
 
 			const int32 PageIndex = PageCoord.Z * (Header.PageTableVolumeResolution.Y * Header.PageTableVolumeResolution.X) + PageCoord.Y * Header.PageTableVolumeResolution.X + PageCoord.X;
 
@@ -315,8 +316,8 @@ FVector4f FSparseVolumeTextureData::ReadTileDataVoxel(int32 TileIndex, const FIn
 		return FVector4f();
 	}
 
-	const FIntVector3 CoordPadded = TileDataCoord + FIntVector3(SPARSE_VOLUME_TILE_BORDER);
-	const int32 VoxelIndex = TileIndex * UE::SVT::SVTNumVoxelsPerPaddedTile + CoordPadded.Z * (SPARSE_VOLUME_TILE_RES_PADDED * SPARSE_VOLUME_TILE_RES_PADDED) + CoordPadded.Y * SPARSE_VOLUME_TILE_RES_PADDED + CoordPadded.X;
+	const FInt64Vector3 CoordPadded = FInt64Vector3(TileDataCoord) + FInt64Vector3(SPARSE_VOLUME_TILE_BORDER);
+	const int64 VoxelIndex = (int64)TileIndex * (int64)UE::SVT::SVTNumVoxelsPerPaddedTile + CoordPadded.Z * (SPARSE_VOLUME_TILE_RES_PADDED * SPARSE_VOLUME_TILE_RES_PADDED) + CoordPadded.Y * SPARSE_VOLUME_TILE_RES_PADDED + CoordPadded.X;
 	const uint8* TileData = AttributesIdx == 0 ? PhysicalTileDataA.GetData() : PhysicalTileDataB.GetData();
 	const EPixelFormat Format = Header.AttributesFormats[AttributesIdx];
 	return UE::SVT::ReadVoxel(VoxelIndex, TileData, Format);
@@ -376,19 +377,28 @@ void FSparseVolumeTextureData::WriteTileDataVoxel(int32 TileIndex, const FIntVec
 		return;
 	}
 
-	const FIntVector3 CoordPadded = TileDataCoord + FIntVector3(SPARSE_VOLUME_TILE_BORDER);
-	const int32 VoxelIndex = TileIndex * UE::SVT::SVTNumVoxelsPerPaddedTile + CoordPadded.Z * (SPARSE_VOLUME_TILE_RES_PADDED * SPARSE_VOLUME_TILE_RES_PADDED) + CoordPadded.Y * SPARSE_VOLUME_TILE_RES_PADDED + CoordPadded.X;
+	const FInt64Vector3 CoordPadded = FInt64Vector3(TileDataCoord) + FInt64Vector3(SPARSE_VOLUME_TILE_BORDER);
+	const int64 VoxelIndex = (int64)TileIndex * (int64)UE::SVT::SVTNumVoxelsPerPaddedTile + CoordPadded.Z * (SPARSE_VOLUME_TILE_RES_PADDED * SPARSE_VOLUME_TILE_RES_PADDED) + CoordPadded.Y * SPARSE_VOLUME_TILE_RES_PADDED + CoordPadded.X;
 	uint8* TileData = AttributesIdx == 0 ? PhysicalTileDataA.GetData() : PhysicalTileDataB.GetData();
 	const EPixelFormat Format = Header.AttributesFormats[AttributesIdx];
 	UE::SVT::WriteVoxel(VoxelIndex, TileData, Format, Value, DstComponent);
 }
 
-void FSparseVolumeTextureData::GenerateMipMaps(const FSparseVolumeTextureDataAddressingInfo& AddressingInfo, int32 NumMipLevels)
+bool FSparseVolumeTextureData::GenerateMipMaps(const FSparseVolumeTextureDataAddressingInfo& AddressingInfo, int32 NumMipLevels)
 {
+	const int32 FormatSize[] = { GPixelFormats[Header.AttributesFormats[0]].BlockBytes, GPixelFormats[Header.AttributesFormats[1]].BlockBytes };
+
 	check(!Header.MipInfo.IsEmpty());
-	check(FMath::IsPowerOfTwo(Header.PageTableVolumeResolution.X) 
-		&& FMath::IsPowerOfTwo(Header.PageTableVolumeResolution.Y) 
+	check(!PageTable.IsEmpty());
+	check(FMath::IsPowerOfTwo(Header.PageTableVolumeResolution.X)
+		&& FMath::IsPowerOfTwo(Header.PageTableVolumeResolution.Y)
 		&& FMath::IsPowerOfTwo(Header.PageTableVolumeResolution.Z));
+	check(Header.MipInfo[0].TileOffset == 1);
+	const int32 NumMip0Tiles = Header.MipInfo[0].TileOffset + Header.MipInfo[0].TileCount;
+	check(PhysicalTileDataA.Num() >= ((int64)NumMip0Tiles * (int64)UE::SVT::SVTNumVoxelsPerPaddedTile * (int64)FormatSize[0]));
+	check(PhysicalTileDataB.Num() >= ((int64)NumMip0Tiles * (int64)UE::SVT::SVTNumVoxelsPerPaddedTile * (int64)FormatSize[1]));
+	check(PageTable[0].Num() == (Header.PageTableVolumeResolution.X * Header.PageTableVolumeResolution.Y * Header.PageTableVolumeResolution.Z));
+	
 	if (NumMipLevels == -1)
 	{
 		NumMipLevels = 1;
@@ -401,6 +411,8 @@ void FSparseVolumeTextureData::GenerateMipMaps(const FSparseVolumeTextureDataAdd
 	}
 	Header.MipInfo.SetNum(1);
 	PageTable.SetNum(1);
+	PhysicalTileDataA.SetNum((int64)NumMip0Tiles * (int64)UE::SVT::SVTNumVoxelsPerPaddedTile * (int64)FormatSize[0]);
+	PhysicalTileDataB.SetNum((int64)NumMip0Tiles * (int64)UE::SVT::SVTNumVoxelsPerPaddedTile * (int64)FormatSize[1]);
 
 	for (int32 MipLevel = 1; MipLevel < NumMipLevels; ++MipLevel)
 	{
@@ -456,11 +468,10 @@ void FSparseVolumeTextureData::GenerateMipMaps(const FSparseVolumeTextureDataAdd
 		LinearAllocatedPages.SetNum(NumAllocatedPages);
 
 		// Initialize the SparseVolumeTexture page and tile.
-		const int32 FormatSize[] = { GPixelFormats[Header.AttributesFormats[0]].BlockBytes, GPixelFormats[Header.AttributesFormats[1]].BlockBytes };
 		PageTable.AddDefaulted();
 		PageTable[MipLevel].SetNumZeroed(PageTableVolumeResolution.X * PageTableVolumeResolution.Y * PageTableVolumeResolution.Z);
-		PhysicalTileDataA.AddZeroed(NumAllocatedPages * UE::SVT::SVTNumVoxelsPerPaddedTile * FormatSize[0]);
-		PhysicalTileDataB.AddZeroed(NumAllocatedPages * UE::SVT::SVTNumVoxelsPerPaddedTile * FormatSize[1]);
+		PhysicalTileDataA.AddZeroed((int64)NumAllocatedPages * (int64)UE::SVT::SVTNumVoxelsPerPaddedTile * (int64)FormatSize[0]);
+		PhysicalTileDataB.AddZeroed((int64)NumAllocatedPages * (int64)UE::SVT::SVTNumVoxelsPerPaddedTile * (int64)FormatSize[1]);
 
 		// Generate page table and tile volume data by splatting the data
 		{
@@ -513,11 +524,16 @@ void FSparseVolumeTextureData::GenerateMipMaps(const FSparseVolumeTextureDataAdd
 			MipInfo.TileCount = NumAllocatedPages;
 		}
 
-		GenerateBorderVoxels(AddressingInfo, MipLevel, LinearAllocatedPages);
+		if (!GenerateBorderVoxels(AddressingInfo, MipLevel, LinearAllocatedPages))
+		{
+			return false;
+		}
 	}
+
+	return true;
 }
 
-void FSparseVolumeTextureData::GenerateBorderVoxels(const FSparseVolumeTextureDataAddressingInfo& AddressingInfo, int32 MipLevel, const TArray<FIntVector3>& PageCoords)
+bool FSparseVolumeTextureData::GenerateBorderVoxels(const FSparseVolumeTextureDataAddressingInfo& AddressingInfo, int32 MipLevel, const TArray<FIntVector3>& PageCoords)
 {
 	const FIntVector3 PageTableVolumeResolution = FIntVector3(
 		FMath::Max(1, Header.PageTableVolumeResolution.X >> MipLevel), 
@@ -553,10 +569,134 @@ void FSparseVolumeTextureData::GenerateBorderVoxels(const FSparseVolumeTextureDa
 				}
 			}
 		});
+
+	return true;
+}
+
+bool FSparseVolumeTextureData::DeduplicateTiles()
+{
+	const int32 NumMipLevels = Header.MipInfo.Num();
+	const int32 FormatSize[] = { GPixelFormats[Header.AttributesFormats[0]].BlockBytes, GPixelFormats[Header.AttributesFormats[1]].BlockBytes };
+	const int32 TileByteSize[] = { UE::SVT::SVTNumVoxelsPerPaddedTile * FormatSize[0], UE::SVT::SVTNumVoxelsPerPaddedTile * FormatSize[1] };
+	uint8* PhysicalTileData[] = { PhysicalTileDataA.GetData(), PhysicalTileDataB.GetData() };
+
+	TArray<uint32> TileRemap;
+	TileRemap.Empty(Header.MipInfo.Last().TileOffset + Header.MipInfo.Last().TileCount);
+	TileRemap.Add(0); // null tile
+	int32 TileWritePos = 1; // null tile
+	int32 NumProcessedTiles = 1;
+
+	TArray<uint64> TileHashes; // Reused between mip level iterations
+
+	// Deduplicate each mip level individually
+	for (int32 MipLevel = 0; MipLevel < NumMipLevels; ++MipLevel)
+	{
+		const FSparseVolumeTextureMipInfo& MipInfo = Header.MipInfo[MipLevel];
+		const int32 MipTileOffset = TileWritePos;
+		TileHashes.Reset();
+
+		// Compute a hash per tile and compare against previous hashes
+		for (int32 TileIdx = 0; TileIdx < MipInfo.TileCount; ++TileIdx)
+		{
+			const int32 GlobalTileIdx = TileIdx + MipInfo.TileOffset;
+			
+			uint64 Hash = 0;
+			if (FormatSize[0])
+			{
+				Hash = CityHash64((const char*)PhysicalTileData[0] + (SIZE_T)GlobalTileIdx * (SIZE_T)TileByteSize[0], (uint32)TileByteSize[0]);
+			}
+			if (FormatSize[1])
+			{
+				Hash = CityHash64WithSeed((const char*)PhysicalTileData[1] + (SIZE_T)GlobalTileIdx * (SIZE_T)TileByteSize[1], (uint32)TileByteSize[1], Hash);
+			}
+
+			const int32 OldNumUniqueTiles = TileHashes.Num();
+			const uint32 RemapIndex = TileHashes.AddUnique(Hash);
+			TileRemap.Add(MipTileOffset + RemapIndex);
+			
+			const bool bAddedNewUniqueTile = TileHashes.Num() > OldNumUniqueTiles;
+			
+			// Move tile to compacted position
+			if (bAddedNewUniqueTile)
+			{
+				if (TileWritePos < GlobalTileIdx)
+				{
+					if (FormatSize[0])
+					{
+						FMemory::Memcpy(PhysicalTileData[0] + (SIZE_T)TileWritePos * (SIZE_T)TileByteSize[0], PhysicalTileData[0] + (SIZE_T)GlobalTileIdx * (SIZE_T)TileByteSize[0], (SIZE_T)TileByteSize[0]);
+					}
+					if (FormatSize[1])
+					{
+						FMemory::Memcpy(PhysicalTileData[1] + (SIZE_T)TileWritePos * (SIZE_T)TileByteSize[1], PhysicalTileData[1] + (SIZE_T)GlobalTileIdx * (SIZE_T)TileByteSize[1], (SIZE_T)TileByteSize[1]);
+					}
+				}
+				else
+				{
+					check(TileWritePos == GlobalTileIdx);
+				}
+				++TileWritePos;
+			}
+
+			++NumProcessedTiles;
+		}
+
+		// Update mip info
+		Header.MipInfo[MipLevel].TileOffset = MipTileOffset;
+		Header.MipInfo[MipLevel].TileCount = TileWritePos - MipTileOffset;
+
+		// Fix up page table
+		for (uint32& PageTableEntry : PageTable[MipLevel])
+		{
+			PageTableEntry = TileRemap[PageTableEntry];
+		}
+	}
+
+	const uint64 DeduplicatedTileSizeBytes[] = { (uint64)TileWritePos * (uint64)TileByteSize[0], (uint64)TileWritePos * (uint64)TileByteSize[1] };
+	if (DeduplicatedTileSizeBytes[0] > INT32_MAX || DeduplicatedTileSizeBytes[1] > INT32_MAX)
+	{
+		UE_LOG(LogSparseVolumeTextureData, Error, TEXT("SparseVolumeTexture still requires too much memory after tile deduplication! Physical tile data A: %llu bytes. Physical tile data B: %llu bytes"), 
+			(unsigned long long)DeduplicatedTileSizeBytes[0], (unsigned long long)DeduplicatedTileSizeBytes[1]);
+		return false;
+	}
+
+	if (!PhysicalTileDataA.IsEmpty())
+	{
+		PhysicalTileDataA.SetNum(TileWritePos * TileByteSize[0]);
+	}
+	if (!PhysicalTileDataB.IsEmpty())
+	{
+		PhysicalTileDataB.SetNum(TileWritePos * TileByteSize[1]);
+	}
+
+	const int32 NumRemovedTiles = NumProcessedTiles - TileWritePos;
+	const float RemovedTilesPercentage = ((float)NumRemovedTiles / (float)NumProcessedTiles) * 100.0f;
+	UE_LOG(LogSparseVolumeTextureData, Display, TEXT("SparseVolumeTexture tile deduplication removed %i tiles out of %i (%f%%)"), NumRemovedTiles, NumProcessedTiles, RemovedTilesPercentage);
+
+	return true;
 }
 
-void FSparseVolumeTextureData::BuildDerivedData(const FSparseVolumeTextureDataAddressingInfo& AddressingInfo, int32 NumMipLevels)
+bool FSparseVolumeTextureData::BuildDerivedData(const FSparseVolumeTextureDataAddressingInfo& AddressingInfo, int32 NumMipLevels, bool bMoveMip0FromThis, FSparseVolumeTextureData& OutDerivedData)
 {
+	OutDerivedData = FSparseVolumeTextureData{};
+	OutDerivedData.Header = Header;
+	OutDerivedData.Header.MipInfo.SetNum(1);
+	OutDerivedData.PageTable.SetNum(1);
+	if (bMoveMip0FromThis)
+	{
+		OutDerivedData.PageTable[0] = MoveTemp(PageTable[0]);
+		OutDerivedData.PhysicalTileDataA = MoveTemp(PhysicalTileDataA);
+		OutDerivedData.PhysicalTileDataB = MoveTemp(PhysicalTileDataB);
+		PageTable.Empty();
+		PhysicalTileDataA.Empty();
+		PhysicalTileDataA.Empty();
+	}
+	else
+	{
+		OutDerivedData.PageTable[0] = PageTable[0];
+		OutDerivedData.PhysicalTileDataA = PhysicalTileDataA;
+		OutDerivedData.PhysicalTileDataB = PhysicalTileDataB;
+	}
+
 	// Generate border voxels of mip0
 	{
 		// Collect the page table coordinates of all non-zero pages
@@ -568,7 +708,7 @@ void FSparseVolumeTextureData::BuildDerivedData(const FSparseVolumeTextureDataAd
 				for (int32 PageX = 0; PageX < Header.PageTableVolumeResolution.X; ++PageX)
 				{
 					const int32 PageIndex = PageZ * (Header.PageTableVolumeResolution.Y * Header.PageTableVolumeResolution.X) + PageY * Header.PageTableVolumeResolution.X + PageX;
-					const uint32 PageTableEntry = PageTable[0][PageIndex];
+					const uint32 PageTableEntry = OutDerivedData.PageTable[0][PageIndex];
 					if (PageTableEntry != 0)
 					{
 						PageCoords.Add(FIntVector3(PageX, PageY, PageZ));
@@ -577,11 +717,25 @@ void FSparseVolumeTextureData::BuildDerivedData(const FSparseVolumeTextureDataAd
 			}
 		}
 
-		GenerateBorderVoxels(AddressingInfo, 0, PageCoords);
+		if (!OutDerivedData.GenerateBorderVoxels(AddressingInfo, 0, PageCoords))
+		{
+			return false;
+		}
 	}
 
 	// Generate all remaining mips. Also generates border voxels.
-	GenerateMipMaps(AddressingInfo, NumMipLevels);
+	if (!OutDerivedData.GenerateMipMaps(AddressingInfo, NumMipLevels))
+	{
+		return false;
+	}
+
+	// Deduplicate tiles by potentially having multiple pages point to the same physical tile
+	if (!OutDerivedData.DeduplicateTiles())
+	{
+		return false;
+	}
+
+	return true;
 }
 
 ////////////////////////////////////////////////////////////////////////////////////////////////
diff --git a/Engine/Source/Runtime/Engine/Private/SparseVolumeTexture/SparseVolumeTextureUtility.cpp b/Engine/Source/Runtime/Engine/Private/SparseVolumeTexture/SparseVolumeTextureUtility.cpp
--- a/Engine/Source/Runtime/Engine/Private/SparseVolumeTexture/SparseVolumeTextureUtility.cpp
+++ b/Engine/Source/Runtime/Engine/Private/SparseVolumeTexture/SparseVolumeTextureUtility.cpp
@@ -36,7 +36,7 @@ FIntVector3 UE::SVT::UnpackPageTableEntry(uint32 Packed)
 	return Result;
 }
 
-FVector4f UE::SVT::ReadVoxel(int32 VoxelIndex, const uint8* TileData, EPixelFormat Format)
+FVector4f UE::SVT::ReadVoxel(int64 VoxelIndex, const uint8* TileData, EPixelFormat Format)
 {
 	using namespace UE::SVT::Private;
 	if (Format == PF_Unknown)
@@ -69,7 +69,7 @@ FVector4f UE::SVT::ReadVoxel(int32 VoxelIndex, const uint8* TileData, EPixelForm
 	}
 }
 
-void UE::SVT::WriteVoxel(int32 VoxelIndex, uint8* TileData, EPixelFormat Format, const FVector4f& Value, int32 DstComponent)
+void UE::SVT::WriteVoxel(int64 VoxelIndex, uint8* TileData, EPixelFormat Format, const FVector4f& Value, int32 DstComponent)
 {
 	if (Format == PF_Unknown)
 	{
diff --git a/Engine/Source/Runtime/Engine/Private/SparseVolumeTexture/SparseVolumeTextureUtility.h b/Engine/Source/Runtime/Engine/Private/SparseVolumeTexture/SparseVolumeTextureUtility.h
--- a/Engine/Source/Runtime/Engine/Private/SparseVolumeTexture/SparseVolumeTextureUtility.h
+++ b/Engine/Source/Runtime/Engine/Private/SparseVolumeTexture/SparseVolumeTextureUtility.h
@@ -17,8 +17,8 @@ namespace SVT
 
 	uint32 PackPageTableEntry(const FIntVector3& Coord);
 	FIntVector3 UnpackPageTableEntry(uint32 Packed);
-	FVector4f ReadVoxel(int32 VoxelIndex, const uint8* TileData, EPixelFormat Format);
-	void WriteVoxel(int32 VoxelIndex, uint8* TileData, EPixelFormat Format, const FVector4f& Value, int32 DstComponent = -1);
+	FVector4f ReadVoxel(int64 VoxelIndex, const uint8* TileData, EPixelFormat Format);
+	void WriteVoxel(int64 VoxelIndex, uint8* TileData, EPixelFormat Format, const FVector4f& Value, int32 DstComponent = -1);
 
 namespace Private
 {
diff --git a/Engine/Source/Runtime/Engine/Public/SparseVolumeTexture/SparseVolumeTextureData.h b/Engine/Source/Runtime/Engine/Public/SparseVolumeTexture/SparseVolumeTextureData.h
--- a/Engine/Source/Runtime/Engine/Public/SparseVolumeTexture/SparseVolumeTextureData.h
+++ b/Engine/Source/Runtime/Engine/Public/SparseVolumeTexture/SparseVolumeTextureData.h
@@ -58,8 +58,8 @@ struct ENGINE_API FSparseVolumeTextureData
 
 	FSparseVolumeTextureDataHeader Header = {};
 	TArray<TArray<uint32>> PageTable; // PageTable[MipLevel][PageCoord]
-	TArray<uint8> PhysicalTileDataA;
-	TArray<uint8> PhysicalTileDataB;
+	TArray64<uint8> PhysicalTileDataA;
+	TArray64<uint8> PhysicalTileDataB;
 
 	void Serialize(FArchive& Ar);
 	bool Construct(const ISparseVolumeTextureDataConstructionAdapter& Adapter);
@@ -67,7 +67,8 @@ struct ENGINE_API FSparseVolumeTextureData
 	FVector4f ReadTileDataVoxel(int32 TileIndex, const FIntVector3& TileDataCoord, int32 AttributesIdx) const;
 	FVector4f Load(const FIntVector3& VolumeCoord, int32 MipLevel, int32 AttributesIdx, const FSparseVolumeTextureDataAddressingInfo& AddressingInfo) const;
 	void WriteTileDataVoxel(int32 TileIndex, const FIntVector3& TileDataCoord, int32 AttributesIdx, const FVector4f& Value, int32 DstComponent = -1);
-	void GenerateMipMaps(const FSparseVolumeTextureDataAddressingInfo& AddressingInfo, int32 NumMipLevels = -1);
-	void GenerateBorderVoxels(const FSparseVolumeTextureDataAddressingInfo& AddressingInfo, int32 MipLevel, const TArray<FIntVector3>& PageCoords);
-	void BuildDerivedData(const FSparseVolumeTextureDataAddressingInfo& AddressingInfo, int32 NumMipLevels = -1);
+	bool GenerateMipMaps(const FSparseVolumeTextureDataAddressingInfo& AddressingInfo, int32 NumMipLevels = -1);
+	bool GenerateBorderVoxels(const FSparseVolumeTextureDataAddressingInfo& AddressingInfo, int32 MipLevel, const TArray<FIntVector3>& PageCoords);
+	bool DeduplicateTiles();
+	bool BuildDerivedData(const FSparseVolumeTextureDataAddressingInfo& AddressingInfo, int32 NumMipLevels, bool bMoveMip0FromThis, FSparseVolumeTextureData& OutDerivedData);
 };
