Localization conflict fix.

diff --git a/Engine/Source/Runtime/Engine/Private/LumenVisualizationData.cpp b/Engine/Source/Runtime/Engine/Private/LumenVisualizationData.cpp
--- a/Engine/Source/Runtime/Engine/Private/LumenVisualizationData.cpp
+++ b/Engine/Source/Runtime/Engine/Private/LumenVisualizationData.cpp
@@ -38,7 +38,7 @@ void FLumenVisualizationData::Initialize()
 		AddVisualizationMode(
 			TEXT("GeometryNormals"),
 			LOCTEXT("GeometryNormals", "Geometry Normals"),
-			LOCTEXT("LumenSceneDesc", "Visualizes Geometry Normals with current Lumen settings."),
+			LOCTEXT("GeometryNormalsDesc", "Visualizes Geometry Normals with current Lumen settings."),
 			FModeType::Standard,
 			LUMEN_VISUALIZE_MODE_GEOMETRY_NORMALS,
 			true);
