Migrate shared linear encoding controls to Texture Encode settings.

diff --git a/Engine/Source/Editor/UnrealEd/Private/Commandlets/DerivedDataCacheCommandlet.cpp b/Engine/Source/Editor/UnrealEd/Private/Commandlets/DerivedDataCacheCommandlet.cpp
--- a/Engine/Source/Editor/UnrealEd/Private/Commandlets/DerivedDataCacheCommandlet.cpp
+++ b/Engine/Source/Editor/UnrealEd/Private/Commandlets/DerivedDataCacheCommandlet.cpp
@@ -30,6 +30,7 @@ DerivedDataCacheCommandlet.cpp: Commandlet for DDC maintenence
 #include "PackageHelperFunctions.h"
 #include "Settings/ProjectPackagingSettings.h"
 #include "ShaderCompiler.h"
+#include "TextureEncodingSettings.h"
 #include "UObject/CoreRedirects.h"
 #include "UObject/Package.h"
 #include "UObject/UObjectHash.h"
@@ -464,8 +465,7 @@ int32 UDerivedDataCacheCommandlet::Main( const FString& Params )
 	FParse::Value(*Params, TEXT("SubsetTarget="), SubsetTarget);
 	bool bDoSubset = SubsetMod > 0 && SubsetTarget < SubsetMod;
 
-	static auto CVarSharedLinearTextureEncoding = IConsoleManager::Get().FindTConsoleVariableDataInt(TEXT("r.SharedLinearTextureEncoding"));
-	if (CVarSharedLinearTextureEncoding->GetValueOnAnyThread())
+	if (GetDefault<UTextureEncodingProjectSettings>()->bSharedLinearTextureEncoding)
 	{
 		bSharedLinearTextureEncodingEnabled = true;
 	}
diff --git a/Engine/Source/Runtime/Engine/Private/TextureDerivedData.cpp b/Engine/Source/Runtime/Engine/Private/TextureDerivedData.cpp
--- a/Engine/Source/Runtime/Engine/Private/TextureDerivedData.cpp
+++ b/Engine/Source/Runtime/Engine/Private/TextureDerivedData.cpp
@@ -529,6 +529,29 @@ static void GetTextureDerivedDataKey(
 
 #if WITH_EDITOR
 
+static bool GetSharedLinearEncodingEnabled()
+{
+	// We have to cache this because we are hitting the options on a worker thread, and it'll
+	// crash if we use GetDefault while someone edits the project settings.
+	// At the moment there's no guaranteed game thread place to do this as jobs can be kicked
+	// off from worker threads (async encodes shader/light map).
+	static struct ThreadSafeInitCSO
+	{
+		bool bSharedLinearTextureEncoding = false;
+		ThreadSafeInitCSO()
+		{
+			const UTextureEncodingProjectSettings* Settings = GetDefault<UTextureEncodingProjectSettings>();
+			bSharedLinearTextureEncoding = Settings->bSharedLinearTextureEncoding;
+
+			UE_LOG(LogTexture, Display, TEXT("Shared Linear Texture Encoding: %s"), bSharedLinearTextureEncoding ? TEXT("Enabled") : TEXT("Disabled"));
+		}
+
+	} SLE;
+
+	return SLE.bSharedLinearTextureEncoding;
+}
+
+
 struct FTextureEncodeSpeedOptions
 {
 	ETextureEncodeEffort Effort = ETextureEncodeEffort::Default;
@@ -697,10 +720,8 @@ static void FinalizeBuildSettingsForLayer(
 					OutBuildResultMetadata->bSupportsEncodeSpeed = bSupportsEncodeSpeed;
 				}
 			
-				
 				{
-					static auto CVarSharedLinearTextureEncoding = IConsoleManager::Get().FindTConsoleVariableDataInt(TEXT("r.SharedLinearTextureEncoding"));
-					if (CVarSharedLinearTextureEncoding->GetValueOnAnyThread())
+					if (GetSharedLinearEncodingEnabled())
 					{
 						//
 						// We want to separate out textures involved in shared linear encoding in order to facilitate
diff --git a/Engine/Source/Runtime/Engine/Private/TextureDerivedDataTask.cpp b/Engine/Source/Runtime/Engine/Private/TextureDerivedDataTask.cpp
--- a/Engine/Source/Runtime/Engine/Private/TextureDerivedDataTask.cpp
+++ b/Engine/Source/Runtime/Engine/Private/TextureDerivedDataTask.cpp
@@ -53,16 +53,10 @@ static TAutoConsoleVariable<int32> CVarVTValidateCompressionOnSave(
 static TAutoConsoleVariable<int32> CVarForceRetileTextures(
 	TEXT("r.ForceRetileTextures"),
 	0,
-	TEXT("If r.SharedLinearTextureEncoding is enabled, this will force the tiling build step to rebuild,")
+	TEXT("If Shared Linear Texture Encoding is enabled in project settings, this will force the tiling build step to rebuild,")
 	TEXT("however the linear texture is allowed to fetch from cache.")
 );
 
-static TAutoConsoleVariable<int32> CVarSharedLinearTextureEncoding(
-	TEXT("r.SharedLinearTextureEncoding"),
-	0,
-	TEXT("If set to 1, textures for platforms that tile will reuse a host linear texture instead of reencoding.")
-);
-
 
 void GetTextureDerivedDataKeyFromSuffix(const FString& KeySuffix, FString& OutKey);
 UE::DerivedData::FCacheKey GetTextureDerivedMetadataKeyFromSuffix(const FString& KeySuffix);
diff --git a/Engine/Source/Runtime/Engine/Private/TextureEncodingSettings.cpp b/Engine/Source/Runtime/Engine/Private/TextureEncodingSettings.cpp
--- a/Engine/Source/Runtime/Engine/Private/TextureEncodingSettings.cpp
+++ b/Engine/Source/Runtime/Engine/Private/TextureEncodingSettings.cpp
@@ -7,6 +7,7 @@
 
 UTextureEncodingProjectSettings::UTextureEncodingProjectSettings(const FObjectInitializer& ObjectInitializer)
 	: Super(ObjectInitializer),
+	bSharedLinearTextureEncoding(false),
 	bFinalUsesRDO(false),
 	FinalRDOLambda(30), /* OodleTex_RDOLagrangeLambda_Default */
 	FinalEffortLevel(ETextureEncodeEffort::Normal),
diff --git a/Engine/Source/Runtime/Engine/Public/TextureEncodingSettings.h b/Engine/Source/Runtime/Engine/Public/TextureEncodingSettings.h
--- a/Engine/Source/Runtime/Engine/Public/TextureEncodingSettings.h
+++ b/Engine/Source/Runtime/Engine/Public/TextureEncodingSettings.h
@@ -42,13 +42,31 @@ enum class ETextureEncodeSpeedOverride : uint8
 
 // 
 // Encoding can either use the "Final" or "Fast" speeds, for supported encoders (e.g. Oodle)
-// These settings have no effect on encoders that don't support encode speed
+// Encode speed settings have no effect on encoders that don't support encode speed, currently limited to Oodle.
 //
 UCLASS(config = Engine, defaultconfig, meta = (DisplayName = "Texture Encoding"))
 class ENGINE_API UTextureEncodingProjectSettings : public UDeveloperSettings
 {
 	GENERATED_UCLASS_BODY()
 
+	// If true, platforms that want to take a linearly encoded texture and then tile them
+	// will try to reuse the linear texture rather than encode it for every platform. This can result in
+	// massive speedups for texture building as tiling is very fast compared to encoding. So instead of:
+	//
+	//	Host Platform: Linear encode
+	//	Console 1: Linear encode + platform specific tile
+	//	Console 2: Linear encode + platform specific tile
+	//
+	// you instead get:
+	//	Host platform: Linear encode
+	//	Console 1: fetch linear + platform specific tile
+	//	Console 2: fetch linear + platform specific tile
+	//
+	// Note that this has no effect on cook time, only build time - once the texture is in the DDC this has no
+	// effect.
+	UPROPERTY(EditAnywhere, config, Category=EncodeSettings, meta = (ConfigRestartRequired = true))
+	uint32 bSharedLinearTextureEncoding : 1;
+
 	// If true, Final encode speed enables rate-distortion optimization on supported encoders to
 	// decrease *on disc* size of textures in compressed package files.
 	// This rate-distortion tradeoff is controlled via "Lambda". The "LossyCompressionAmount" parameter on
