Migrate texture encoding configuration to a thread safe access model.

diff --git a/Engine/Source/Editor/TextureEditor/Private/TextureEditorToolkit.cpp b/Engine/Source/Editor/TextureEditor/Private/TextureEditorToolkit.cpp
--- a/Engine/Source/Editor/TextureEditor/Private/TextureEditorToolkit.cpp
+++ b/Engine/Source/Editor/TextureEditor/Private/TextureEditorToolkit.cpp
@@ -653,8 +653,8 @@ void FTextureEditorToolkit::PopulateQuickInfo( )
 
 				if (ResultMetadata.OodleRDO == 0)
 				{
-					const UTextureEncodingProjectSettings* Settings = GetDefault<UTextureEncodingProjectSettings>();
-					const bool bDisabledGlobally = ResultMetadata.EncodeSpeed == (uint8)ETextureEncodeSpeed::Fast ? !Settings->bFastUsesRDO : !Settings->bFinalUsesRDO;
+					const FResolvedTextureEncodingSettings& Settings = FResolvedTextureEncodingSettings::Get();
+					const bool bDisabledGlobally = ResultMetadata.EncodeSpeed == (uint8)ETextureEncodeSpeed::Fast ? !Settings.Project.bFastUsesRDO : !Settings.Project.bFinalUsesRDO;
 
 					OodleRDOText->SetText(NSLOCTEXT("TextureEditor", "QuickInfo_Oodle_RDODisable", "Disabled"));
 					if (ResultMetadata.bWasEditorCustomEncoding)
diff --git a/Engine/Source/Editor/UnrealEd/Private/Commandlets/DerivedDataCacheCommandlet.cpp b/Engine/Source/Editor/UnrealEd/Private/Commandlets/DerivedDataCacheCommandlet.cpp
--- a/Engine/Source/Editor/UnrealEd/Private/Commandlets/DerivedDataCacheCommandlet.cpp
+++ b/Engine/Source/Editor/UnrealEd/Private/Commandlets/DerivedDataCacheCommandlet.cpp
@@ -465,7 +465,7 @@ int32 UDerivedDataCacheCommandlet::Main( const FString& Params )
 	FParse::Value(*Params, TEXT("SubsetTarget="), SubsetTarget);
 	bool bDoSubset = SubsetMod > 0 && SubsetTarget < SubsetMod;
 
-	if (GetDefault<UTextureEncodingProjectSettings>()->bSharedLinearTextureEncoding)
+	if (FResolvedTextureEncodingSettings::Get().Project.bSharedLinearTextureEncoding)
 	{
 		bSharedLinearTextureEncodingEnabled = true;
 	}
diff --git a/Engine/Source/Runtime/Engine/Private/TextureDerivedData.cpp b/Engine/Source/Runtime/Engine/Private/TextureDerivedData.cpp
--- a/Engine/Source/Runtime/Engine/Private/TextureDerivedData.cpp
+++ b/Engine/Source/Runtime/Engine/Private/TextureDerivedData.cpp
@@ -529,29 +529,6 @@ static void GetTextureDerivedDataKey(
 
 #if WITH_EDITOR
 
-static bool GetSharedLinearEncodingEnabled()
-{
-	// We have to cache this because we are hitting the options on a worker thread, and it'll
-	// crash if we use GetDefault while someone edits the project settings.
-	// At the moment there's no guaranteed game thread place to do this as jobs can be kicked
-	// off from worker threads (async encodes shader/light map).
-	static struct ThreadSafeInitCSO
-	{
-		bool bSharedLinearTextureEncoding = false;
-		ThreadSafeInitCSO()
-		{
-			const UTextureEncodingProjectSettings* Settings = GetDefault<UTextureEncodingProjectSettings>();
-			bSharedLinearTextureEncoding = Settings->bSharedLinearTextureEncoding;
-
-			UE_LOG(LogTexture, Display, TEXT("Shared Linear Texture Encoding: %s"), bSharedLinearTextureEncoding ? TEXT("Enabled") : TEXT("Disabled"));
-		}
-
-	} SLE;
-
-	return SLE.bSharedLinearTextureEncoding;
-}
-
-
 struct FTextureEncodeSpeedOptions
 {
 	ETextureEncodeEffort Effort = ETextureEncodeEffort::Default;
@@ -563,77 +540,20 @@ struct FTextureEncodeSpeedOptions
 // InEncodeSpeed must be fast or final.
 static void GetEncodeSpeedOptions(ETextureEncodeSpeed InEncodeSpeed, FTextureEncodeSpeedOptions* OutOptions)
 {
-	// We have to cache this because we are hitting the options on a worker thread, and it'll
-	// crash if we use GetDefault while someone edits the project settings.
-	// At the moment there's no guaranteed game thread place to do this as jobs can be kicked
-	// off from worker threads (async encodes shader/light map).
-	static struct ThreadSafeInitCSO
-	{
-		FTextureEncodeSpeedOptions Fast, Final;
-		ThreadSafeInitCSO()
-		{
-
-			const UTextureEncodingProjectSettings* Settings = GetDefault<UTextureEncodingProjectSettings>();
-			Fast.Effort = Settings->FastEffortLevel;
-			Fast.Tiling = Settings->FastUniversalTiling;
-			Fast.bUsesRDO = Settings->bFastUsesRDO;
-			Fast.RDOLambda = Settings->FastRDOLambda;
-
-			Final.Effort = Settings->FinalEffortLevel;
-			Final.Tiling = Settings->FinalUniversalTiling;
-			Final.bUsesRDO = Settings->bFinalUsesRDO;
-			Final.RDOLambda = Settings->FinalRDOLambda;
-			
-			// log settings once at startup
-			UEnum* EncodeEffortEnum = StaticEnum<ETextureEncodeEffort>();
-			
-			UEnum* UniversalTilingEnum = StaticEnum<ETextureUniversalTiling>();
-
-			FString FastRDOString;
-			if ( Fast.bUsesRDO )
-			{
-				FastRDOString = FString(TEXT("On"));
-				if ( Fast.Tiling != ETextureUniversalTiling::Disabled )
-				{
-					FastRDOString += TEXT(" UT=");
-					FastRDOString += UniversalTilingEnum->GetNameStringByValue((int64)Fast.Tiling);
-				}
-			}
-			else
-			{
-				FastRDOString = FString(TEXT("Off"));
-			}
-			
-			FString FinalRDOString;
-			if ( Final.bUsesRDO )
-			{
-				FinalRDOString = FString(TEXT("On"));
-				if ( Final.Tiling != ETextureUniversalTiling::Disabled )
-				{
-					FinalRDOString += TEXT(" UT=");
-					FinalRDOString += UniversalTilingEnum->GetNameStringByValue((int64)Final.Tiling);
-				}
-			}
-			else
-			{
-				FinalRDOString = FString(TEXT("Off"));
-			}
-
-			UE_LOG(LogTexture, Display, TEXT("Oodle Texture Encode Speed settings: Fast: RDO %s Lambda=%d, Effort=%s Final: RDO %s Lambda=%d, Effort=%s"), \
-				*FastRDOString, Fast.bUsesRDO ? Fast.RDOLambda : 0,  *(EncodeEffortEnum->GetNameStringByValue((int64)Fast.Effort)), \
-				*FinalRDOString, Final.bUsesRDO ? Final.RDOLambda : 0,  *(EncodeEffortEnum->GetNameStringByValue((int64)Final.Effort)) );
-
-
-		}
-	} EncodeSpeedOptions;
-
+	FResolvedTextureEncodingSettings const& EncodeSettings = FResolvedTextureEncodingSettings::Get();
 	if (InEncodeSpeed == ETextureEncodeSpeed::Final)
 	{
-		*OutOptions = EncodeSpeedOptions.Final;
+		OutOptions->bUsesRDO = EncodeSettings.Project.bFinalUsesRDO;
+		OutOptions->Effort = EncodeSettings.Project.FinalEffortLevel;
+		OutOptions->Tiling = EncodeSettings.Project.FinalUniversalTiling;
+		OutOptions->RDOLambda = EncodeSettings.Project.FinalRDOLambda;
 	}
 	else
 	{
-		*OutOptions = EncodeSpeedOptions.Fast;
+		OutOptions->bUsesRDO = EncodeSettings.Project.bFastUsesRDO;
+		OutOptions->Effort = EncodeSettings.Project.FastEffortLevel;
+		OutOptions->Tiling = EncodeSettings.Project.FastUniversalTiling;
+		OutOptions->RDOLambda = EncodeSettings.Project.FastRDOLambda;
 	}
 }
 
@@ -721,7 +641,7 @@ static void FinalizeBuildSettingsForLayer(
 				}
 			
 				{
-					if (GetSharedLinearEncodingEnabled())
+					if (FResolvedTextureEncodingSettings::Get().Project.bSharedLinearTextureEncoding)
 					{
 						//
 						// We want to separate out textures involved in shared linear encoding in order to facilitate
@@ -801,62 +721,7 @@ ENGINE_API ETextureEncodeSpeed UTexture::GetDesiredEncodeSpeed() const
 		return ETextureEncodeSpeed::Final;
 	}
 
-	// Get the command line and config options with a one-time static init :
-	static struct FThreadSafeInitializer
-	{
-		ETextureEncodeSpeed CachedEncodeSpeedOption;
-
-		// For thread safety, we do all the work in a constructor that the compiler
-		// guarantees will be called only once.
-		FThreadSafeInitializer()
-		{
-			const UEnum* EncodeSpeedEnum = StaticEnum<ETextureEncodeSpeed>();
-
-			// Overridden by command line?
-			FString CmdLineSpeed;
-			if (FParse::Value(FCommandLine::Get(), TEXT("-ForceTextureEncodeSpeed="), CmdLineSpeed))
-			{
-				int64 Value = EncodeSpeedEnum->GetValueByNameString(CmdLineSpeed);
-				if (Value == INDEX_NONE)
-				{
-					UE_LOG(LogTexture, Error, TEXT("Invalid value for ForceTextureEncodeSpeed, ignoring. Valid values are the ETextureEncodeSpeed enum (Final, FinalIfAvailable, Fast)"));
-				}
-				else
-				{
-					CachedEncodeSpeedOption = (ETextureEncodeSpeed)Value;
-					UE_LOG(LogTexture, Display, TEXT("Texture Encode Speed forced to %s via command line."), *EncodeSpeedEnum->GetNameStringByValue(Value));
-					return;
-				}
-			}
-
-			// Overridden by user settings?
-			const UTextureEncodingUserSettings* UserSettings = GetDefault<UTextureEncodingUserSettings>();
-			if (UserSettings->ForceEncodeSpeed != ETextureEncodeSpeedOverride::Disabled)
-			{
-				// enums have same values for payload.
-				CachedEncodeSpeedOption = (ETextureEncodeSpeed)UserSettings->ForceEncodeSpeed;
-				UE_LOG(LogTexture, Display, TEXT("Texture Encode Speed forced to %s via user settings."), *EncodeSpeedEnum->GetNameStringByValue((int64)CachedEncodeSpeedOption));
-				return;
-			}
-
-			// Use project settings
-			const UTextureEncodingProjectSettings* Settings = GetDefault<UTextureEncodingProjectSettings>();
-			if (GIsEditor && !IsRunningCommandlet())
-			{
-				// Interactive editor
-				CachedEncodeSpeedOption = Settings->EditorUsesSpeed;
-				UE_LOG(LogTexture, Display, TEXT("Texture Encode Speed: %s (editor)."), *EncodeSpeedEnum->GetNameStringByValue((int64)CachedEncodeSpeedOption));
-			}
-			else
-			{
-				CachedEncodeSpeedOption = Settings->CookUsesSpeed;
-				UE_LOG(LogTexture, Display, TEXT("Texture Encode Speed: %s (cook)."), *EncodeSpeedEnum->GetNameStringByValue((int64)CachedEncodeSpeedOption));
-			}
-			
-		}
-	} FThreadSafeInitializer;
-
-	return FThreadSafeInitializer.CachedEncodeSpeedOption;
+	return FResolvedTextureEncodingSettings::Get().EncodeSpeed;
 }
 
 
diff --git a/Engine/Source/Runtime/Engine/Private/TextureEncodingSettings.cpp b/Engine/Source/Runtime/Engine/Private/TextureEncodingSettings.cpp
--- a/Engine/Source/Runtime/Engine/Private/TextureEncodingSettings.cpp
+++ b/Engine/Source/Runtime/Engine/Private/TextureEncodingSettings.cpp
@@ -1,10 +1,14 @@
 // Copyright Epic Games, Inc. All Rights Reserved.
 
 #include "TextureEncodingSettings.h"
+#include "TextureEncodingSettingsPrivate.h"
 #include "Engine/TextureDefines.h"
+#include "Misc/CommandLine.h"
 
 #include UE_INLINE_GENERATED_CPP_BY_NAME(TextureEncodingSettings)
 
+DEFINE_LOG_CATEGORY_STATIC(LogTextureEncodingSettings, Log, All);
+
 UTextureEncodingProjectSettings::UTextureEncodingProjectSettings(const FObjectInitializer& ObjectInitializer)
 	: Super(ObjectInitializer),
 	bSharedLinearTextureEncoding(false),
@@ -27,3 +31,201 @@ UTextureEncodingUserSettings::UTextureEncodingUserSettings(const FObjectInitiali
 {
 }
 
+enum class ReadyMask : uint32
+{
+	Query = 0,
+	Project = 0x1,
+	User = 0x2,
+	Everyone = 0x3
+};
+
+static void ConstructResolvedSettings(ReadyMask InReadyMask, UDeveloperSettings* InSettings, FResolvedTextureEncodingSettings* OutResolvedSettings)
+{
+	// Settings is the last one to get initialized - the others have all had PostInitProperties
+	// called, so we can safely retrieve their CDO. This is all so we don't cause a cycle
+	// calling GetDefault() - we only call it for classes we know are complete.
+	UTextureEncodingProjectSettings const* ProjectSettings = nullptr;
+	UTextureEncodingUserSettings const* UserSettings = nullptr;
+
+	if (InReadyMask == ReadyMask::Project)
+	{
+		ProjectSettings = (UTextureEncodingProjectSettings*)InSettings;
+	}
+	else if (InReadyMask == ReadyMask::User)
+	{
+		UserSettings = (UTextureEncodingUserSettings*)InSettings;
+	}
+
+	if (ProjectSettings == nullptr)
+	{
+		ProjectSettings = GetDefault<UTextureEncodingProjectSettings>();
+	}
+	if (UserSettings == nullptr)
+	{
+		UserSettings = GetDefault<UTextureEncodingUserSettings>();
+	}
+
+	OutResolvedSettings->User.ForceEncodeSpeed = UserSettings->ForceEncodeSpeed;
+
+	OutResolvedSettings->Project.bSharedLinearTextureEncoding = ProjectSettings->bSharedLinearTextureEncoding;
+	OutResolvedSettings->Project.bFinalUsesRDO = ProjectSettings->bFinalUsesRDO;
+	OutResolvedSettings->Project.bFastUsesRDO = ProjectSettings->bFastUsesRDO;
+	OutResolvedSettings->Project.FinalRDOLambda = ProjectSettings->FinalRDOLambda;
+	OutResolvedSettings->Project.FastRDOLambda = ProjectSettings->FastRDOLambda;
+	OutResolvedSettings->Project.FinalEffortLevel = ProjectSettings->FinalEffortLevel;
+	OutResolvedSettings->Project.FinalUniversalTiling = ProjectSettings->FinalUniversalTiling;
+	OutResolvedSettings->Project.FastEffortLevel = ProjectSettings->FastEffortLevel;
+	OutResolvedSettings->Project.FastUniversalTiling = ProjectSettings->FastUniversalTiling;
+	OutResolvedSettings->Project.CookUsesSpeed = ProjectSettings->CookUsesSpeed;
+	OutResolvedSettings->Project.EditorUsesSpeed = ProjectSettings->EditorUsesSpeed;
+
+	// Determine what encode speed to use
+	{
+		const UEnum* EncodeSpeedEnum = StaticEnum<ETextureEncodeSpeed>();
+
+		// Overridden by command line?
+		FString CmdLineSpeed;
+		if (FParse::Value(FCommandLine::Get(), TEXT("-ForceTextureEncodeSpeed="), CmdLineSpeed))
+		{
+			int64 Value = EncodeSpeedEnum->GetValueByNameString(CmdLineSpeed);
+			if (Value == INDEX_NONE)
+			{
+				UE_LOG(LogTextureEncodingSettings, Error, TEXT("Invalid value for ForceTextureEncodeSpeed, ignoring. Valid values are the ETextureEncodeSpeed enum (Final, FinalIfAvailable, Fast)"));
+			}
+			else
+			{
+				OutResolvedSettings->EncodeSpeed = (ETextureEncodeSpeed)Value;
+				UE_LOG(LogTextureEncodingSettings, Display, TEXT("Texture Encode Speed forced to %s via command line."), *EncodeSpeedEnum->GetNameStringByValue(Value));
+				return;
+			}
+		}
+
+		// Overridden by user settings?
+		if (OutResolvedSettings->User.ForceEncodeSpeed != ETextureEncodeSpeedOverride::Disabled)
+		{
+			// enums have same values for payload.
+			OutResolvedSettings->EncodeSpeed = (ETextureEncodeSpeed)(OutResolvedSettings->User.ForceEncodeSpeed);
+			UE_LOG(LogTextureEncodingSettings, Display, TEXT("Texture Encode Speed forced to %s via user settings."), *EncodeSpeedEnum->GetNameStringByValue((int64)OutResolvedSettings->EncodeSpeed));
+			return;
+		}
+
+		// Use project settings
+		if (GIsEditor && !IsRunningCommandlet())
+		{
+			// Interactive editor
+			OutResolvedSettings->EncodeSpeed = OutResolvedSettings->Project.EditorUsesSpeed;
+			UE_LOG(LogTextureEncodingSettings, Display, TEXT("Texture Encode Speed: %s (editor)."), *EncodeSpeedEnum->GetNameStringByValue((int64)OutResolvedSettings->EncodeSpeed));
+		}
+		else
+		{
+			OutResolvedSettings->EncodeSpeed = OutResolvedSettings->Project.CookUsesSpeed;
+			UE_LOG(LogTextureEncodingSettings, Display, TEXT("Texture Encode Speed: %s (cook)."), *EncodeSpeedEnum->GetNameStringByValue((int64)OutResolvedSettings->EncodeSpeed));
+		}
+	}
+
+	// Log encode speeds
+	{
+		UEnum* EncodeEffortEnum = StaticEnum<ETextureEncodeEffort>();
+
+		UEnum* UniversalTilingEnum = StaticEnum<ETextureUniversalTiling>();
+
+		FString FastRDOString;
+		if (OutResolvedSettings->Project.bFastUsesRDO)
+		{
+			FastRDOString = FString(TEXT("On"));
+			if (OutResolvedSettings->Project.FastUniversalTiling != ETextureUniversalTiling::Disabled)
+			{
+				FastRDOString += TEXT(" UT=");
+				FastRDOString += UniversalTilingEnum->GetNameStringByValue((int64)OutResolvedSettings->Project.FastUniversalTiling);
+			}
+		}
+		else
+		{
+			FastRDOString = FString(TEXT("Off"));
+		}
+
+		FString FinalRDOString;
+		if (OutResolvedSettings->Project.bFinalUsesRDO)
+		{
+			FinalRDOString = FString(TEXT("On"));
+			if (OutResolvedSettings->Project.FinalUniversalTiling != ETextureUniversalTiling::Disabled)
+			{
+				FinalRDOString += TEXT(" UT=");
+				FinalRDOString += UniversalTilingEnum->GetNameStringByValue((int64)OutResolvedSettings->Project.FinalUniversalTiling);
+			}
+		}
+		else
+		{
+			FinalRDOString = FString(TEXT("Off"));
+		}
+
+		UE_LOG(LogTextureEncodingSettings, Display, TEXT("Oodle Texture Encode Speed settings: Fast: RDO %s Lambda=%d, Effort=%s Final: RDO %s Lambda=%d, Effort=%s"), \
+			*FastRDOString, OutResolvedSettings->Project.bFastUsesRDO ? OutResolvedSettings->Project.FastRDOLambda : 0, *(EncodeEffortEnum->GetNameStringByValue((int64)OutResolvedSettings->Project.FastEffortLevel)), \
+			*FinalRDOString, OutResolvedSettings->Project.bFinalUsesRDO ? OutResolvedSettings->Project.FinalRDOLambda : 0, *(EncodeEffortEnum->GetNameStringByValue((int64)OutResolvedSettings->Project.FinalEffortLevel)));
+	}
+
+	UE_LOG(LogTextureEncodingSettings, Display, TEXT("Shared linear texture encoding: %s"), OutResolvedSettings->Project.bSharedLinearTextureEncoding ? TEXT("Enabled") : TEXT("Disabled"));
+}
+
+static void TryToResolveSettings(ReadyMask InReadyMask, UDeveloperSettings* InSettings, FResolvedTextureEncodingSettings** OutResolvedSettings)
+{
+	static FResolvedTextureEncodingSettings ResolvedSettings;
+
+
+	// If we are called with anything other than ::Query, then we are in PIP for the settings classes,
+	// which a) only happens from the main thread during preinit and b) is guaranteed to happen. We don't know the order,
+	// so we just track who we've seen and copy everything once we've got everyone.
+	// 
+	// Once we've completed engine init, then we could get hit from other threads, however WhoIsReady
+	// at that point is guaranteed to be Everyone, so we don't need to worry about atomicity of read-modify-write.
+	static uint32 WhoIsReady = 0;
+	if (InReadyMask != ReadyMask::Query)
+	{		
+		WhoIsReady |= (uint32)InReadyMask;
+		if (WhoIsReady == (uint32)ReadyMask::Everyone)
+		{
+			ConstructResolvedSettings(InReadyMask, InSettings, &ResolvedSettings);
+		}
+	}
+	else
+	{
+		// Two cases:
+		// 1 -- We are requesting the properties before engine class init has completed,
+		//		meaning we are game thread, so WhoIsReady is NOT everyone, so we need to
+		//		do initialization by requesting our CDOs.
+		// 2 -- We are requesting the properties _after_ engine class init has completed,
+		//		meaning WhoIsReady is Everyone and we can safely return the value. Since
+		//		we _know_ WhoIsReady will be Everyone before we go async, we can safely
+		//		query it.
+		if (WhoIsReady != (uint32)ReadyMask::Everyone)
+		{
+			GetDefault<UTextureEncodingUserSettings>();
+			GetDefault<UTextureEncodingProjectSettings>();
+			check(WhoIsReady == (uint32)ReadyMask::Everyone);
+		}
+	}
+
+	if (OutResolvedSettings)
+	{
+		*OutResolvedSettings = &ResolvedSettings;
+	}
+}
+
+FResolvedTextureEncodingSettings const& FResolvedTextureEncodingSettings::Get()
+{
+	FResolvedTextureEncodingSettings* ResolvedSettings = nullptr;
+	TryToResolveSettings(ReadyMask::Query, nullptr, &ResolvedSettings);
+	return *ResolvedSettings;
+}
+
+void UTextureEncodingUserSettings::PostInitProperties()
+{
+	Super::PostInitProperties();
+	TryToResolveSettings(ReadyMask::User, this, nullptr);
+}
+
+void UTextureEncodingProjectSettings::PostInitProperties()
+{
+	Super::PostInitProperties();
+	TryToResolveSettings(ReadyMask::Project, this, nullptr);
+}
diff --git a/Engine/Source/Runtime/Engine/Private/TextureEncodingSettingsPrivate.h b/Engine/Source/Runtime/Engine/Private/TextureEncodingSettingsPrivate.h
new file mode 100644
--- /dev/null
+++ b/Engine/Source/Runtime/Engine/Private/TextureEncodingSettingsPrivate.h
@@ -0,0 +1,139 @@
+// Copyright Epic Games, Inc. All Rights Reserved.
+
+#pragma once
+
+#include "TextureEncodingSettings.h"
+#include "TextureEncodingSettingsPrivate.generated.h"
+
+// 
+// Encoding can either use the "Final" or "Fast" speeds, for supported encoders (e.g. Oodle)
+// Encode speed settings have no effect on encoders that don't support encode speed, currently limited to Oodle.
+//
+UCLASS(config = Engine, defaultconfig, meta = (DisplayName = "Texture Encoding"))
+class ENGINE_API UTextureEncodingProjectSettings : public UDeveloperSettings
+{
+	//
+	// Anything added here should be added to FResolvedTextureEncodingSettings!
+	//
+	GENERATED_UCLASS_BODY()
+
+	virtual void PostInitProperties() override;
+
+	// If true, platforms that want to take a linearly encoded texture and then tile them
+	// will try to reuse the linear texture rather than encode it for every platform. This can result in
+	// massive speedups for texture building as tiling is very fast compared to encoding. So instead of:
+	//
+	//	Host Platform: Linear encode
+	//	Console 1: Linear encode + platform specific tile
+	//	Console 2: Linear encode + platform specific tile
+	//
+	// you instead get:
+	//	Host platform: Linear encode
+	//	Console 1: fetch linear + platform specific tile
+	//	Console 2: fetch linear + platform specific tile
+	//
+	// Note that this has no effect on cook time, only build time - once the texture is in the DDC this has no
+	// effect.
+	UPROPERTY(EditAnywhere, config, Category=EncodeSettings, meta = (ConfigRestartRequired = true))
+	uint32 bSharedLinearTextureEncoding : 1;
+
+	// If true, Final encode speed enables rate-distortion optimization on supported encoders to
+	// decrease *on disc* size of textures in compressed package files.
+	// This rate-distortion tradeoff is controlled via "Lambda". The "LossyCompressionAmount" parameter on
+	// textures is used to control it. Specific LossyCompressionAmount values correspond to
+	// to RDO lambdas of:
+	// 
+	//	None - Disable RDO for this texture.
+	//	Lowest - 1 (Least distortion)
+	//	Low - 10
+	//	Medium - 20
+	//	High - 30
+	//	Highest - 40
+	// 
+	// If set to Default, then the LossyCompressionAmount in the LODGroup for the texture is
+	// used. If that is also Default, then the RDOLambda in these settings is used.
+	//
+	// Note that any distortion introduced is on top of, and likely less than, any introduced by the format itself.
+	UPROPERTY(config, EditAnywhere, Category = EncodeSpeedSettings, meta = (ConfigRestartRequired = true))
+	uint32 bFinalUsesRDO : 1;
+
+	// Ignored if UsesRDO is false. This value is used if a given texture is using "Default" LossyCompressionAmount.
+	// Otherwise, the value of LossyCompressionAmount is translated in to a fixed lambda (see UsesRDO tooltip).
+	// 
+	// Low values (1) represent highest quality (least distortion) results.
+	UPROPERTY(config, EditAnywhere, Category = EncodeSpeedSettings, meta = (DisplayName = "Final RDO Lambda", UIMin = 1, UIMax = 100, ClampMin = 1, ClampMax = 100, EditCondition = bFinalUsesRDO, ConfigRestartRequired = true))
+	int8 FinalRDOLambda;
+
+	// Specifies how much time to take trying for better encoding results.
+	UPROPERTY(config, EditAnywhere, Category = EncodeSpeedSettings, meta = (ConfigRestartRequired = true))
+	ETextureEncodeEffort FinalEffortLevel;	
+
+	// Specifies how to assume textures are laid out on disc. This only applies to Oodle with RDO
+	// enabled. 256 KB is a good middle ground. Enabling this will decrease the on-disc
+	// sizes of textures for platforms with exposed texture tiling (i.e. consoles), but will slightly increase
+	// sizes of textures for platforms with opaque tiling (i.e. desktop).
+	UPROPERTY(config, EditAnywhere, Category = EncodeSpeedSettings, meta = (ConfigRestartRequired = true))
+	ETextureUniversalTiling FinalUniversalTiling;
+
+	// If true, Final encode speed enables rate-distortion optimization on supported encoders to
+	// decrease *on disc* size of textures in compressed package files.
+	// This rate-distortion tradeoff is controlled via "Lambda". The "LossyCompressionAmount" parameter on
+	// textures is used to control it. Specific LossyCompressionAmount values correspond to
+	// to RDO lambdas of:
+	// 
+	//	None - Disable RDO for this texture.
+	//	Lowest - 1 (Least distortion)
+	//	Low - 10
+	//	Medium - 20
+	//	High - 30
+	//	Highest - 40
+	// 
+	// If set to Default, then the LossyCompressionAmount in the LODGroup for the texture is
+	// used. If that is also Default, then the RDOLambda in these settings is used.
+	//
+	// Note that any distortion introduced is on top of, and likely less than, any introduced by the format itself.
+	UPROPERTY(config, EditAnywhere, Category = EncodeSpeedSettings, meta = (ConfigRestartRequired = true))
+	uint32 bFastUsesRDO : 1;
+
+	// Ignored if UsesRDO is false. This value is used if a given texture is using "Default" LossyCompressionAmount.
+	// Otherwise, the value of LossyCompressionAmount is translated in to a fixed lambda (see UsesRDO tooltip).
+	// 
+	// Low values (1) represent highest quality (least distortion) results.
+	UPROPERTY(config, EditAnywhere, Category = EncodeSpeedSettings, meta = (DisplayName = "Fast RDO Lambda", UIMin = 1, UIMax = 100, ClampMin = 1, ClampMax = 100, EditCondition = bFastUsesRDO, ConfigRestartRequired = true))
+	int8 FastRDOLambda;
+
+	// Specifies how much time to take trying for better encode results.
+	UPROPERTY(config, EditAnywhere, Category = EncodeSpeedSettings, meta = (ConfigRestartRequired = true))
+	ETextureEncodeEffort FastEffortLevel;
+
+	// Specifies how to assume textures are laid out on disc. This only applies to Oodle with RDO
+	// enabled. 256 KB is a good middle ground. Enabling this will decrease the on-disc
+	// sizes of textures for platforms with exposed texture tiling (i.e. consoles), but will slightly increase
+	// sizes of textures for platforms with opaque tiling (i.e. desktop).
+	UPROPERTY(config, EditAnywhere, Category = EncodeSpeedSettings, meta = (ConfigRestartRequired = true))
+	ETextureUniversalTiling FastUniversalTiling;
+
+	// Which encode speed non interactive editor sessions will use (i.e. commandlets)
+	UPROPERTY(config, EditAnywhere, Category = EncodeSpeeds, meta = (ConfigRestartRequired = true))
+	ETextureEncodeSpeed CookUsesSpeed;
+
+	// Which encode speed everything else uses.
+	UPROPERTY(config, EditAnywhere, Category = EncodeSpeeds, meta = (ConfigRestartRequired = true))
+	ETextureEncodeSpeed EditorUsesSpeed;
+};
+
+UCLASS(config = EditorPerProjectUserSettings, defaultconfig, meta = (DisplayName = "Texture Encoding"))
+class ENGINE_API UTextureEncodingUserSettings : public UDeveloperSettings
+{
+	//
+	// Anything added here should be added to FResolvedTextureEncodingSettings!
+	//
+	GENERATED_UCLASS_BODY()
+
+	virtual void PostInitProperties() override;
+
+	// Local machine/project setting to force an encode speed, if desired.
+	// See the Engine "Texture Encoding" section for details.
+	UPROPERTY(config, EditAnywhere, Category = EncodeSpeeds, meta = (ConfigRestartRequired = true))
+	ETextureEncodeSpeedOverride ForceEncodeSpeed;
+};
diff --git a/Engine/Source/Runtime/Engine/Public/TextureEncodingSettings.h b/Engine/Source/Runtime/Engine/Public/TextureEncodingSettings.h
--- a/Engine/Source/Runtime/Engine/Public/TextureEncodingSettings.h
+++ b/Engine/Source/Runtime/Engine/Public/TextureEncodingSettings.h
@@ -40,127 +40,41 @@ enum class ETextureEncodeSpeedOverride : uint8
 	Fast = 2
 };
 
-// 
-// Encoding can either use the "Final" or "Fast" speeds, for supported encoders (e.g. Oodle)
-// Encode speed settings have no effect on encoders that don't support encode speed, currently limited to Oodle.
 //
-UCLASS(config = Engine, defaultconfig, meta = (DisplayName = "Texture Encoding"))
-class ENGINE_API UTextureEncodingProjectSettings : public UDeveloperSettings
-{
-	GENERATED_UCLASS_BODY()
-
-	// If true, platforms that want to take a linearly encoded texture and then tile them
-	// will try to reuse the linear texture rather than encode it for every platform. This can result in
-	// massive speedups for texture building as tiling is very fast compared to encoding. So instead of:
-	//
-	//	Host Platform: Linear encode
-	//	Console 1: Linear encode + platform specific tile
-	//	Console 2: Linear encode + platform specific tile
-	//
-	// you instead get:
-	//	Host platform: Linear encode
-	//	Console 1: fetch linear + platform specific tile
-	//	Console 2: fetch linear + platform specific tile
-	//
-	// Note that this has no effect on cook time, only build time - once the texture is in the DDC this has no
-	// effect.
-	UPROPERTY(EditAnywhere, config, Category=EncodeSettings, meta = (ConfigRestartRequired = true))
-	uint32 bSharedLinearTextureEncoding : 1;
-
-	// If true, Final encode speed enables rate-distortion optimization on supported encoders to
-	// decrease *on disc* size of textures in compressed package files.
-	// This rate-distortion tradeoff is controlled via "Lambda". The "LossyCompressionAmount" parameter on
-	// textures is used to control it. Specific LossyCompressionAmount values correspond to
-	// to RDO lambdas of:
-	// 
-	//	None - Disable RDO for this texture.
-	//	Lowest - 1 (Least distortion)
-	//	Low - 10
-	//	Medium - 20
-	//	High - 30
-	//	Highest - 40
-	// 
-	// If set to Default, then the LossyCompressionAmount in the LODGroup for the texture is
-	// used. If that is also Default, then the RDOLambda in these settings is used.
-	//
-	// Note that any distortion introduced is on top of, and likely less than, any introduced by the format itself.
-	UPROPERTY(config, EditAnywhere, Category = EncodeSpeedSettings, meta = (ConfigRestartRequired = true))
-	uint32 bFinalUsesRDO : 1;
-
-	// Ignored if UsesRDO is false. This value is used if a given texture is using "Default" LossyCompressionAmount.
-	// Otherwise, the value of LossyCompressionAmount is translated in to a fixed lambda (see UsesRDO tooltip).
-	// 
-	// Low values (1) represent highest quality (least distortion) results.
-	UPROPERTY(config, EditAnywhere, Category = EncodeSpeedSettings, meta = (DisplayName = "Final RDO Lambda", UIMin = 1, UIMax = 100, ClampMin = 1, ClampMax = 100, EditCondition = bFinalUsesRDO, ConfigRestartRequired = true))
-	int8 FinalRDOLambda;
-
-	// Specifies how much time to take trying for better encoding results.
-	UPROPERTY(config, EditAnywhere, Category = EncodeSpeedSettings, meta = (ConfigRestartRequired = true))
-	ETextureEncodeEffort FinalEffortLevel;	
-
-	// Specifies how to assume textures are laid out on disc. This only applies to Oodle with RDO
-	// enabled. 256 KB is a good middle ground. Enabling this will decrease the on-disc
-	// sizes of textures for platforms with exposed texture tiling (i.e. consoles), but will slightly increase
-	// sizes of textures for platforms with opaque tiling (i.e. desktop).
-	UPROPERTY(config, EditAnywhere, Category = EncodeSpeedSettings, meta = (ConfigRestartRequired = true))
-	ETextureUniversalTiling FinalUniversalTiling;
-
-	// If true, Final encode speed enables rate-distortion optimization on supported encoders to
-	// decrease *on disc* size of textures in compressed package files.
-	// This rate-distortion tradeoff is controlled via "Lambda". The "LossyCompressionAmount" parameter on
-	// textures is used to control it. Specific LossyCompressionAmount values correspond to
-	// to RDO lambdas of:
-	// 
-	//	None - Disable RDO for this texture.
-	//	Lowest - 1 (Least distortion)
-	//	Low - 10
-	//	Medium - 20
-	//	High - 30
-	//	Highest - 40
-	// 
-	// If set to Default, then the LossyCompressionAmount in the LODGroup for the texture is
-	// used. If that is also Default, then the RDOLambda in these settings is used.
-	//
-	// Note that any distortion introduced is on top of, and likely less than, any introduced by the format itself.
-	UPROPERTY(config, EditAnywhere, Category = EncodeSpeedSettings, meta = (ConfigRestartRequired = true))
-	uint32 bFastUsesRDO : 1;
-
-	// Ignored if UsesRDO is false. This value is used if a given texture is using "Default" LossyCompressionAmount.
-	// Otherwise, the value of LossyCompressionAmount is translated in to a fixed lambda (see UsesRDO tooltip).
-	// 
-	// Low values (1) represent highest quality (least distortion) results.
-	UPROPERTY(config, EditAnywhere, Category = EncodeSpeedSettings, meta = (DisplayName = "Fast RDO Lambda", UIMin = 1, UIMax = 100, ClampMin = 1, ClampMax = 100, EditCondition = bFastUsesRDO, ConfigRestartRequired = true))
-	int8 FastRDOLambda;
-
-	// Specifies how much time to take trying for better encode results.
-	UPROPERTY(config, EditAnywhere, Category = EncodeSpeedSettings, meta = (ConfigRestartRequired = true))
-	ETextureEncodeEffort FastEffortLevel;
-
-	// Specifies how to assume textures are laid out on disc. This only applies to Oodle with RDO
-	// enabled. 256 KB is a good middle ground. Enabling this will decrease the on-disc
-	// sizes of textures for platforms with exposed texture tiling (i.e. consoles), but will slightly increase
-	// sizes of textures for platforms with opaque tiling (i.e. desktop).
-	UPROPERTY(config, EditAnywhere, Category = EncodeSpeedSettings, meta = (ConfigRestartRequired = true))
-	ETextureUniversalTiling FastUniversalTiling;
-
-	// Which encode speed non interactive editor sessions will use (i.e. commandlets)
-	UPROPERTY(config, EditAnywhere, Category = EncodeSpeeds, meta = (ConfigRestartRequired = true))
-	ETextureEncodeSpeed CookUsesSpeed;
-
-	// Which encode speed everything else uses.
-	UPROPERTY(config, EditAnywhere, Category = EncodeSpeeds, meta = (ConfigRestartRequired = true))
-	ETextureEncodeSpeed EditorUsesSpeed;
-};
-
-UCLASS(config = EditorPerProjectUserSettings, defaultconfig, meta = (DisplayName = "Texture Encoding"))
-class ENGINE_API UTextureEncodingUserSettings : public UDeveloperSettings
+// This is the public, thread safe class for accessing the encoding settings. They are
+// filled out as part of the engine class default object init loop during PreInit, and
+// are safe to access at any point during CDO construction as well.
+//
+struct FResolvedTextureEncodingSettings
 {
-	GENERATED_UCLASS_BODY()
-
-	// Local machine/project setting to force an encode speed, if desired.
-	// See the Engine "Texture Encoding" section for details.
-	UPROPERTY(config, EditAnywhere, Category = EncodeSpeeds, meta = (ConfigRestartRequired = true))
-	ETextureEncodeSpeedOverride ForceEncodeSpeed;
+	// Properties mirrored from UTextureEncodingUserSettings - look there for 
+	// documentation
+	struct 
+	{
+		ETextureEncodeSpeedOverride ForceEncodeSpeed;
+	} User;
+
+	// Properties mirrored from UTextureEncodingProjectSettings - look there for 
+	// documentation
+	struct
+	{
+		uint32 bSharedLinearTextureEncoding : 1;
+		uint32 bFinalUsesRDO : 1;
+		uint32 bFastUsesRDO : 1;
+		int8 FinalRDOLambda;
+		int8 FastRDOLambda;
+		ETextureEncodeEffort FinalEffortLevel;
+		ETextureUniversalTiling FinalUniversalTiling;
+		ETextureEncodeEffort FastEffortLevel;
+		ETextureUniversalTiling FastUniversalTiling;
+		ETextureEncodeSpeed CookUsesSpeed;
+		ETextureEncodeSpeed EditorUsesSpeed;
+	} Project;
+
+	// The resolved EncodeSpeed to use for this instance, taking in to account overrides.
+	ETextureEncodeSpeed EncodeSpeed;
+
+	static ENGINE_API FResolvedTextureEncodingSettings const& Get();
 };
 
 //
