Update tooltip comment on Set Modulation Routing Function

diff --git a/Engine/Source/Runtime/Engine/Classes/Components/AudioComponent.h b/Engine/Source/Runtime/Engine/Classes/Components/AudioComponent.h
--- a/Engine/Source/Runtime/Engine/Classes/Components/AudioComponent.h
+++ b/Engine/Source/Runtime/Engine/Classes/Components/AudioComponent.h
@@ -732,10 +732,9 @@ public:
 
 	/**
 	* Sets the routing for one of the given Audio component's Modulation Destinations.
-	* The changes do not apply to any currently active sounds, but will apply to future sounds.
 	* @param Modulators The set of modulators to apply to the given destination on the component.
 	* @param Destination The destination to assign the modulators to.
-	* @param RoutingMethod The routing method to use for the given volume modulator.
+	* @param RoutingMethod The routing method to use for the given modulator.
 	*/
 	UFUNCTION(BlueprintCallable, Category = "Audio|Components|Audio", DisplayName = "Set Modulation Routing")
 	void SetModulationRouting(const TSet<USoundModulatorBase*>& Modulators, const EModulationDestination Destination, const EModulationRouting RoutingMethod = EModulationRouting::Inherit);
