Actor Class Descriptors
- Actors derived Blueprints will now store a class descriptor from the CDO object.
- Actor descriptors will now do delta serialization from their class descriptor.

diff --git a/Engine/Plugins/Editor/DataValidation/Source/DataValidation/Private/WorldPartitionChangelistValidator.cpp b/Engine/Plugins/Editor/DataValidation/Source/DataValidation/Private/WorldPartitionChangelistValidator.cpp
--- a/Engine/Plugins/Editor/DataValidation/Source/DataValidation/Private/WorldPartitionChangelistValidator.cpp
+++ b/Engine/Plugins/Editor/DataValidation/Source/DataValidation/Private/WorldPartitionChangelistValidator.cpp
@@ -208,10 +208,8 @@ EDataValidationResult UWorldPartitionChangelistValidator::ValidateActorsAndDataL
 
 		for (const FAssetData& ActorData : ActorsData)
 		{
-			// Get the FWorldPartitionActor			
-			const FWorldPartitionActorDesc* ActorDesc = ContainersToValidate.GetActorDesc(ActorData.AssetName.ToString());
-
-			if (ActorDesc != nullptr)
+			// Get the actor descriptor
+			if (const FWorldPartitionActorDesc* ActorDesc = ContainersToValidate.GetActorDescByName(ActorData.AssetName.ToString()))
 			{
 				RelevantActorGuids.Add(ActorDesc->GetGuid());
 			}
diff --git a/Engine/Source/Editor/PropertyEditor/Private/UserInterface/PropertyEditor/SPropertyEditorAsset.cpp b/Engine/Source/Editor/PropertyEditor/Private/UserInterface/PropertyEditor/SPropertyEditorAsset.cpp
--- a/Engine/Source/Editor/PropertyEditor/Private/UserInterface/PropertyEditor/SPropertyEditorAsset.cpp
+++ b/Engine/Source/Editor/PropertyEditor/Private/UserInterface/PropertyEditor/SPropertyEditorAsset.cpp
@@ -1115,7 +1115,7 @@ void SPropertyEditorAsset::OnBrowse()
 				{
 					if (UWorld* World = Cast<UWorld>(MapObject); World && World->IsPartitionedWorld())
 					{
-						if (const FWorldPartitionActorDesc* ActorDesc = World->GetWorldPartition()->GetActorDesc(Value.ObjectPath))
+						if (const FWorldPartitionActorDesc* ActorDesc = World->GetWorldPartition()->GetActorDescByName(Value.ObjectPath))
 						{
 							World->GetWorldPartition()->PinActors({ ActorDesc->GetGuid() });
 							GetValue(Value);
diff --git a/Engine/Source/Editor/WorldPartitionEditor/Private/WorldPartitionEditorModule.cpp b/Engine/Source/Editor/WorldPartitionEditor/Private/WorldPartitionEditorModule.cpp
--- a/Engine/Source/Editor/WorldPartitionEditor/Private/WorldPartitionEditorModule.cpp
+++ b/Engine/Source/Editor/WorldPartitionEditor/Private/WorldPartitionEditorModule.cpp
@@ -16,6 +16,7 @@
 #include "WorldPartition/WorldPartitionConvertOptions.h"
 #include "WorldPartition/WorldPartitionEditorSettings.h"
 #include "WorldPartition/HLOD/SWorldPartitionBuildHLODsDialog.h"
+#include "WorldPartition/WorldPartitionClassDescRegistry.h"
 
 #include "LevelEditor.h"
 #include "Framework/Application/SlateApplication.h"
@@ -144,10 +145,14 @@ void FWorldPartitionEditorModule::StartupModule()
 	PropertyEditor.RegisterCustomClassLayout("WorldPartition", FOnGetDetailCustomizationInstance::CreateStatic(&FWorldPartitionDetails::MakeInstance));
 	PropertyEditor.RegisterCustomClassLayout("WorldPartitionRuntimeSpatialHash", FOnGetDetailCustomizationInstance::CreateStatic(&FWorldPartitionRuntimeSpatialHashDetails::MakeInstance));
 	PropertyEditor.RegisterCustomClassLayout("WorldPartitionHLOD", FOnGetDetailCustomizationInstance::CreateStatic(&FWorldPartitionHLODDetailsCustomization::MakeInstance));
+
+	FWorldPartitionClassDescRegistry().Get().Initialize();
 }
 
 void FWorldPartitionEditorModule::ShutdownModule()
 {
+	FWorldPartitionClassDescRegistry().Get().Uninitialize();
+
 	if (!IsRunningGame())
 	{
 		if (FLevelEditorModule* LevelEditorModule = FModuleManager::Get().GetModulePtr<FLevelEditorModule>("LevelEditor"))
diff --git a/Engine/Source/Runtime/Core/Public/UObject/FortniteMainBranchObjectVersion.h b/Engine/Source/Runtime/Core/Public/UObject/FortniteMainBranchObjectVersion.h
--- a/Engine/Source/Runtime/Core/Public/UObject/FortniteMainBranchObjectVersion.h
+++ b/Engine/Source/Runtime/Core/Public/UObject/FortniteMainBranchObjectVersion.h
@@ -277,6 +277,9 @@ struct CORE_API FFortniteMainBranchObjectVersion
 		// Change the non-spatialized radius to blend to a pure 2D spatialized sound vs omnidirectional
 		AudioAttenuationNonSpatializedRadiusBlend,
 
+		// Serialize actor class descriptors
+		WorldPartitionActorClassDescSerialize,
+
 		// -----<new versions can be added above this line>-------------------------------------------------
 		VersionPlusOne,
 		LatestVersion = VersionPlusOne - 1
diff --git a/Engine/Source/Runtime/Engine/Private/ActorEditor.cpp b/Engine/Source/Runtime/Engine/Private/ActorEditor.cpp
--- a/Engine/Source/Runtime/Engine/Private/ActorEditor.cpp
+++ b/Engine/Source/Runtime/Engine/Private/ActorEditor.cpp
@@ -1079,8 +1079,6 @@ TUniquePtr<FWorldPartitionActorDesc> AActor::CreateClassActorDesc() const
 
 TUniquePtr<FWorldPartitionActorDesc> AActor::CreateActorDesc() const
 {
-	check(!HasAnyFlags(RF_ArchetypeObject | RF_ClassDefaultObject));
-	
 	TUniquePtr<FWorldPartitionActorDesc> ActorDesc(CreateClassActorDesc());
 		
 	ActorDesc->Init(this);
diff --git a/Engine/Source/Runtime/Engine/Private/Blueprint.cpp b/Engine/Source/Runtime/Engine/Private/Blueprint.cpp
--- a/Engine/Source/Runtime/Engine/Private/Blueprint.cpp
+++ b/Engine/Source/Runtime/Engine/Private/Blueprint.cpp
@@ -32,6 +32,7 @@
 #include "UObject/MetaData.h"
 #include "Blueprint/BlueprintExtension.h"
 #include "UObject/TextProperty.h"
+#include "WorldPartition/WorldPartitionActorDescUtils.h"
 #endif
 
 #include "Engine/InheritableComponentHandler.h"
@@ -1068,6 +1069,11 @@ void UBlueprint::GetExtendedAssetRegistryTagsForSave(const ITargetPlatform* Targ
 			OutTags.Add( FAssetRegistryTag(FBlueprintTags::FindInBlueprintsData, Value, FAssetRegistryTag::TT_Hidden) );
 		}
 	}
+
+	if (AActor* BlueprintCDO = GeneratedClass ? Cast<AActor>(GeneratedClass->ClassDefaultObject) : nullptr)
+	{
+		FWorldPartitionActorDescUtils::AppendAssetDataTagsFromActor(BlueprintCDO, OutTags);
+	}
 }
 #endif
 
diff --git a/Engine/Source/Runtime/Engine/Private/UnrealEngine.cpp b/Engine/Source/Runtime/Engine/Private/UnrealEngine.cpp
--- a/Engine/Source/Runtime/Engine/Private/UnrealEngine.cpp
+++ b/Engine/Source/Runtime/Engine/Private/UnrealEngine.cpp
@@ -2087,6 +2087,7 @@ void UEngine::Init(IEngineLoop* InEngineLoop)
 		FModuleManager::Get().LoadModuleChecked("MovieSceneTracks");
 		FModuleManager::Get().LoadModule("LevelSequence");
 #if WITH_EDITOR
+		FModuleManager::Get().LoadModule("WorldPartitionEditor");
 		// The SparseVolumeTexture module containing the importer is only loaded and used in the editor.
 		FModuleManager::Get().LoadModuleChecked("SparseVolumeTexture");
 #endif
diff --git a/Engine/Source/Runtime/Engine/Private/WorldPartition/ActorDescContainer.cpp b/Engine/Source/Runtime/Engine/Private/WorldPartition/ActorDescContainer.cpp
--- a/Engine/Source/Runtime/Engine/Private/WorldPartition/ActorDescContainer.cpp
+++ b/Engine/Source/Runtime/Engine/Private/WorldPartition/ActorDescContainer.cpp
@@ -13,6 +13,7 @@
 #include "Modules/ModuleManager.h"
 #include "WorldPartition/WorldPartition.h"
 #include "WorldPartition/WorldPartitionActorDescUtils.h"
+#include "WorldPartition/WorldPartitionClassDescRegistry.h"
 #include "WorldPartition/DataLayer/DataLayerManager.h"
 
 UActorDescContainer::FActorDescContainerInitializeDelegate UActorDescContainer::OnActorDescContainerInitialized;
@@ -35,7 +36,7 @@ void UActorDescContainer::Initialize(const FInitializeParams& InitParams)
 {
 	TRACE_CPUPROFILER_EVENT_SCOPE(UActorDescContainer::Initialize);
 
-	// @todo_ow: We need to pass the world context to AddActorDescriptor for FLevelInstanceActorDesc::RegisterContainerInstance to resolve the ULevelInstanceSubsystem.
+	// @todo_ow: We need to pass the world context to RegisterActorDescriptor for FLevelInstanceActorDesc::RegisterContainerInstance to resolve the ULevelInstanceSubsystem.
 	// A better solution would be for ActorDescContainers to always be outered to an OwningWorldPartition (with the downside of not sharing between 2 instanced WorldPartition).
 	// With this, we could always find the owning WorldPartition and of course the owning world (GetOwningWorldPartition()->GetWorld()).
 	UWorld* OwningWorld = InitParams.World;
@@ -64,19 +65,36 @@ void UActorDescContainer::Initialize(const FInitializeParams& InitParams)
 		AssetRegistry.GetAssets(Filter, Assets);
 	}
 
+	FWorldPartitionClassDescRegistry& ClassDescRegistry = FWorldPartitionClassDescRegistry::Get();
+	//{
+		TRACE_CPUPROFILER_EVENT_SCOPE(GatherDescriptorsClass);
+		
+		TSet<FTopLevelAssetPath> ClassPaths;
+		for (const FAssetData& Asset : Assets)
+		{
+			ClassPaths.Add(Asset.AssetClassPath);
+		}
+
+		ClassDescRegistry.PrefetchClassDescs(ClassPaths.Array());
+	//}
+
 	{
 		TRACE_CPUPROFILER_EVENT_SCOPE(RegisterDescriptors);
 		for (const FAssetData& Asset : Assets)
 		{
 			TUniquePtr<FWorldPartitionActorDesc> ActorDesc = FWorldPartitionActorDescUtils::GetActorDescriptorFromAssetData(Asset);
 
-			if (!ActorDesc.IsValid() || !FWorldPartitionActorDescUtils::ValidateActorDescClass(ActorDesc.Get()) || (InitParams.FilterActorDesc && !InitParams.FilterActorDesc(ActorDesc.Get())))
+			if (!ActorDesc.IsValid() || 
+				!ActorDesc->GetNativeClass().IsValid() || 
+				(ActorDesc->GetBaseClass().IsValid() && !ClassDescRegistry.IsRegisteredClass(ActorDesc->GetBaseClass())) || 
+				FActorDescList::GetActorDesc(ActorDesc->GetGuid()) ||
+				(InitParams.FilterActorDesc && !InitParams.FilterActorDesc(ActorDesc.Get())))
 			{
 				InvalidActors.Emplace(MoveTemp(ActorDesc));
 				continue;
 			}
 
-			AddActorDescriptor(ActorDesc.Release(), OwningWorld);
+			RegisterActorDescriptor(ActorDesc.Release(), OwningWorld);
 		}
 	}
 
@@ -132,7 +150,7 @@ void UActorDescContainer::Update()
 			else
 			{
 				FWorldPartitionActorDesc* ActorDescPtr = NewActorDesc.Release();
-				AddActorDescriptor(ActorDescPtr, OwningWorld);
+				RegisterActorDescriptor(ActorDescPtr, OwningWorld);
 				OnActorDescAdded(ActorDescPtr);
 			}
 		}
@@ -168,7 +186,7 @@ void UActorDescContainer::Uninitialize()
 	{
 		if (FWorldPartitionActorDesc* ActorDesc = ActorDescPtr.Get())
 		{
-			RemoveActorDescriptor(ActorDesc);
+			UnregisterActorDescriptor(ActorDesc);
 		}
 		ActorDescPtr.Reset();
 	}
@@ -216,30 +234,53 @@ UWorldPartition* UActorDescContainer::GetWorldPartition() const
 	return OuterWorldPartition;
 }
 
-void UActorDescContainer::AddActorDescriptor(FWorldPartitionActorDesc* ActorDesc, UWorld* InWorldContext)
+void UActorDescContainer::RegisterActorDescriptor(FWorldPartitionActorDesc* ActorDesc, UWorld* InWorldContext)
 {
-	FActorDescList::AddActorDescriptor(ActorDesc, InWorldContext);
+	FActorDescList::AddActorDescriptor(ActorDesc);
+
 	ActorDesc->SetContainer(this, InWorldContext);
+		
+	//@odo_ow get rid of this with a delegate
 	if (const UDataLayerManager* DataLayerManager = UDataLayerManager::GetDataLayerManager(InWorldContext))
 	{
 		DataLayerManager->ResolveActorDescDataLayers(ActorDesc);
 	}
+
+	ActorsByName.Add(ActorDesc->GetActorName(), ActorsByGuid.FindChecked(ActorDesc->GetGuid()));
 }
 
-void UActorDescContainer::RemoveActorDescriptor(FWorldPartitionActorDesc* ActorDesc)
+void UActorDescContainer::UnregisterActorDescriptor(FWorldPartitionActorDesc* ActorDesc)
 {
-	ActorDesc->SetContainer(nullptr, nullptr);
 	FActorDescList::RemoveActorDescriptor(ActorDesc);
+	ActorDesc->SetContainer(nullptr, nullptr);
+	verify(ActorsByName.Remove(ActorDesc->GetActorName()));	
 }
 
 bool UActorDescContainer::ShouldHandleActorEvent(const AActor* Actor)
 {
-	if (Actor != nullptr)
+	return Actor && IsActorDescHandled(Actor) && Actor->IsMainPackageActor() && Actor->GetLevel();
+}
+
+const FWorldPartitionActorDesc* UActorDescContainer::GetActorDescByName(const FString& ActorPath) const
+{
+	FString ActorName;
+	FString ActorContext;
+	if (!ActorPath.Split(TEXT("."), &ActorContext, &ActorName, ESearchCase::CaseSensitive, ESearchDir::FromEnd))
 	{
-		return IsActorDescHandled(Actor) && Actor->IsMainPackageActor() && (Actor->GetLevel() != nullptr);
+		ActorName = ActorPath;
 	}
-	
-	return false;
+
+	if (const TUniquePtr<FWorldPartitionActorDesc>* const* ActorDesc = ActorsByName.Find(*ActorName))
+	{
+		return (*ActorDesc)->Get();
+	}
+
+	return nullptr;
+}
+
+const FWorldPartitionActorDesc* UActorDescContainer::GetActorDescByName(const FSoftObjectPath& ActorPath) const
+{
+	return GetActorDescByName(ActorPath.ToString());
 }
 
 void UActorDescContainer::OnObjectPreSave(UObject* Object, FObjectPreSaveContext SaveContext)
@@ -261,7 +302,8 @@ void UActorDescContainer::OnObjectPreSave(UObject* Object, FObjectPreSaveContext
 				else
 				{
 					// New actor
-					FWorldPartitionActorDesc* const AddedActorDesc = AddActor(Actor);
+					FWorldPartitionActorDesc* AddedActorDesc = Actor->CreateActorDesc().Release();
+					RegisterActorDescriptor(AddedActorDesc, Actor->GetWorld());
 					OnActorDescAdded(AddedActorDesc);
 				}
 			}
@@ -303,7 +345,7 @@ bool UActorDescContainer::RemoveActor(const FGuid& ActorGuid)
 	if (TUniquePtr<FWorldPartitionActorDesc>* ExistingActorDesc = GetActorDescriptor(ActorGuid))
 	{
 		OnActorDescRemoved(ExistingActorDesc->Get());
-		RemoveActorDescriptor(ExistingActorDesc->Get());
+		UnregisterActorDescriptor(ExistingActorDesc->Get());
 		ExistingActorDesc->Reset();
 		return true;
 	}
diff --git a/Engine/Source/Runtime/Engine/Private/WorldPartition/ActorDescContainerCollection.cpp b/Engine/Source/Runtime/Engine/Private/WorldPartition/ActorDescContainerCollection.cpp
--- a/Engine/Source/Runtime/Engine/Private/WorldPartition/ActorDescContainerCollection.cpp
+++ b/Engine/Source/Runtime/Engine/Private/WorldPartition/ActorDescContainerCollection.cpp
@@ -88,24 +88,24 @@ FWorldPartitionActorDesc& FActorDescContainerCollection::GetActorDescChecked(con
 	return const_cast<FWorldPartitionActorDesc&>(const_cast<const FActorDescContainerCollection*>(this)->GetActorDescChecked(Guid));
 }
 
-const FWorldPartitionActorDesc* FActorDescContainerCollection::GetActorDesc(const FString& ActorPath) const
+const FWorldPartitionActorDesc* FActorDescContainerCollection::GetActorDescByName(const FString& ActorPath) const
 {
 	const FWorldPartitionActorDesc* ActorDesc = nullptr;
 	ForEachActorDescContainerBreakable([&ActorPath, &ActorDesc](const UActorDescContainer* ActorDescContainer)
 	{
-		ActorDesc = ActorDescContainer->GetActorDesc(ActorPath);
+		ActorDesc = ActorDescContainer->GetActorDescByName(ActorPath);
 		return ActorDesc == nullptr;
 	});
 
 	return ActorDesc;
 }
 
-const FWorldPartitionActorDesc* FActorDescContainerCollection::GetActorDesc(const FSoftObjectPath& InActorPath) const
+const FWorldPartitionActorDesc* FActorDescContainerCollection::GetActorDescByName(const FSoftObjectPath& ActorPath) const
 {
 	const FWorldPartitionActorDesc* ActorDesc = nullptr;
-	ForEachActorDescContainerBreakable([&InActorPath, &ActorDesc](const UActorDescContainer* ActorDescContainer)
+	ForEachActorDescContainerBreakable([&ActorPath, &ActorDesc](const UActorDescContainer* ActorDescContainer)
 	{
-		ActorDesc = ActorDescContainer->GetActorDesc(InActorPath);
+		ActorDesc = ActorDescContainer->GetActorDescByName(ActorPath);
 		return ActorDesc == nullptr;
 	});
 
diff --git a/Engine/Source/Runtime/Engine/Private/WorldPartition/ActorDescList.cpp b/Engine/Source/Runtime/Engine/Private/WorldPartition/ActorDescList.cpp
--- a/Engine/Source/Runtime/Engine/Private/WorldPartition/ActorDescList.cpp
+++ b/Engine/Source/Runtime/Engine/Private/WorldPartition/ActorDescList.cpp
@@ -9,7 +9,7 @@ FWorldPartitionActorDesc* FActorDescList::AddActor(const AActor* InActor)
 	FWorldPartitionActorDesc* NewActorDesc = InActor->CreateActorDesc().Release();
 	check(NewActorDesc);
 
-	AddActorDescriptor(NewActorDesc, InActor->GetWorld());
+	AddActorDescriptor(NewActorDesc);
 
 	return NewActorDesc;
 }
@@ -38,56 +38,25 @@ FWorldPartitionActorDesc& FActorDescList::GetActorDescChecked(const FGuid& Guid)
 	return *ActorDesc->Get();
 }
 
-const FWorldPartitionActorDesc* FActorDescList::GetActorDesc(const FString& ActorPath) const
-{
-	FString ActorName;
-	FString ActorContext;
-	if (!ActorPath.Split(TEXT("."), &ActorContext, &ActorName, ESearchCase::CaseSensitive, ESearchDir::FromEnd))
-	{
-		ActorName = ActorPath;
-	}
-
-	if (const TUniquePtr<FWorldPartitionActorDesc>* const* ActorDesc = ActorsByName.Find(*ActorName))
-	{
-		return (*ActorDesc)->Get();
-	}
-
-	return nullptr;
-}
-
-const FWorldPartitionActorDesc* FActorDescList::GetActorDesc(const FSoftObjectPath& ActorPath) const
-{
-	return GetActorDesc(ActorPath.ToString());
-}
-
 void FActorDescList::Empty()
 {
 	ActorsByGuid.Empty();
-	ActorsByName.Empty();
 	ActorDescList.Empty();
 }
 
-void FActorDescList::AddActorDescriptor(FWorldPartitionActorDesc* ActorDesc, UWorld* InWorldContext)
+void FActorDescList::AddActorDescriptor(FWorldPartitionActorDesc* ActorDesc)
 {
 	check(ActorDesc);
+	checkf(!ActorsByGuid.Contains(ActorDesc->GetGuid()), TEXT("Duplicated actor descriptor guid '%s' detected: `%s`"), *ActorDesc->GetGuid().ToString(), *ActorDesc->GetActorName().ToString());
 
-	if (const TUniquePtr<FWorldPartitionActorDesc>* const* ExistingActorDesc = ActorsByGuid.Find(ActorDesc->GetGuid()))
-	{
-		UE_LOG(LogWorldPartition, Warning, TEXT("Duplicated actor descriptor guid '%s' detected: `%s` (existing: `%s`)"), *ActorDesc->GetGuid().ToString(), *ActorDesc->GetActorName().ToString(), *(*ExistingActorDesc)->Get()->GetActorName().ToString());
-	}
-	else
-	{
-		TUniquePtr<FWorldPartitionActorDesc>* NewActorDesc = new(ActorDescList) TUniquePtr<FWorldPartitionActorDesc>(ActorDesc);
-		ActorsByGuid.Add(ActorDesc->GetGuid(), NewActorDesc);
-		ActorsByName.Add(ActorDesc->GetActorName(), NewActorDesc);
-	}
+	TUniquePtr<FWorldPartitionActorDesc>* NewActorDesc = new(ActorDescList) TUniquePtr<FWorldPartitionActorDesc>(ActorDesc);
+	ActorsByGuid.Add(ActorDesc->GetGuid(), NewActorDesc);
 }
 
 void FActorDescList::RemoveActorDescriptor(FWorldPartitionActorDesc* ActorDesc)
 {
 	check(ActorDesc);
 	verify(ActorsByGuid.Remove(ActorDesc->GetGuid()));
-	verify(ActorsByName.Remove(ActorDesc->GetActorName()));
 }
 
 TUniquePtr<FWorldPartitionActorDesc>* FActorDescList::GetActorDescriptor(const FGuid& ActorGuid)
diff --git a/Engine/Source/Runtime/Engine/Private/WorldPartition/ActorPartition/PartitionActorDesc.cpp b/Engine/Source/Runtime/Engine/Private/WorldPartition/ActorPartition/PartitionActorDesc.cpp
--- a/Engine/Source/Runtime/Engine/Private/WorldPartition/ActorPartition/PartitionActorDesc.cpp
+++ b/Engine/Source/Runtime/Engine/Private/WorldPartition/ActorPartition/PartitionActorDesc.cpp
@@ -13,12 +13,16 @@ void FPartitionActorDesc::Init(const AActor* InActor)
 	const APartitionActor* PartitionActor = CastChecked<APartitionActor>(InActor);
 
 	GridSize = PartitionActor->GridSize;
-	GridGuid = PartitionActor->GetGridGuid();
+
+	if (!bIsDefaultActorDesc)
+	{
+		GridGuid = PartitionActor->GetGridGuid();
 	
-	const FVector ActorLocation = InActor->GetActorLocation();
-	GridIndexX = FMath::FloorToInt(ActorLocation.X / GridSize);
-	GridIndexY = FMath::FloorToInt(ActorLocation.Y / GridSize);
-	GridIndexZ = FMath::FloorToInt(ActorLocation.Z / GridSize);
+		const FVector ActorLocation = InActor->GetActorLocation();
+		GridIndexX = FMath::FloorToInt(ActorLocation.X / GridSize);
+		GridIndexY = FMath::FloorToInt(ActorLocation.Y / GridSize);
+		GridIndexZ = FMath::FloorToInt(ActorLocation.Z / GridSize);
+	}
 }
 
 void FPartitionActorDesc::Serialize(FArchive& Ar)
@@ -26,12 +30,15 @@ void FPartitionActorDesc::Serialize(FArchive& Ar)
 	Ar.UsingCustomVersion(FUE5MainStreamObjectVersion::GUID);
 
 	FWorldPartitionActorDesc::Serialize(Ar);
-
-	Ar << GridSize << GridIndexX << GridIndexY << GridIndexZ;
-
-	if (Ar.CustomVer(FUE5MainStreamObjectVersion::GUID) >= FUE5MainStreamObjectVersion::PartitionActorDescSerializeGridGuid)
+	
+	if (!bIsDefaultActorDesc)
 	{
-		Ar << GridGuid;
+		Ar << GridSize << GridIndexX << GridIndexY << GridIndexZ;
+
+		if (Ar.CustomVer(FUE5MainStreamObjectVersion::GUID) >= FUE5MainStreamObjectVersion::PartitionActorDescSerializeGridGuid)
+		{
+			Ar << GridGuid;
+		}
 	}
 }
 
diff --git a/Engine/Source/Runtime/Engine/Private/WorldPartition/ContentBundle/ContentBundleEditor.cpp b/Engine/Source/Runtime/Engine/Private/WorldPartition/ContentBundle/ContentBundleEditor.cpp
--- a/Engine/Source/Runtime/Engine/Private/WorldPartition/ContentBundle/ContentBundleEditor.cpp
+++ b/Engine/Source/Runtime/Engine/Private/WorldPartition/ContentBundle/ContentBundleEditor.cpp
@@ -178,7 +178,7 @@ bool FContentBundleEditor::ContainsActor(const AActor* InActor) const
 {
 	if (InActor != nullptr)
 	{
-		return ActorDescContainer->GetActorDesc(InActor) != nullptr || UnsavedActorMonitor->IsMonitoring(InActor);
+		return ActorDescContainer->GetActorDesc(InActor->GetActorGuid()) != nullptr || UnsavedActorMonitor->IsMonitoring(InActor);
 	}
 
 	return false;
diff --git a/Engine/Source/Runtime/Engine/Private/WorldPartition/DataLayer/WorldDataLayersActorDesc.cpp b/Engine/Source/Runtime/Engine/Private/WorldPartition/DataLayer/WorldDataLayersActorDesc.cpp
--- a/Engine/Source/Runtime/Engine/Private/WorldPartition/DataLayer/WorldDataLayersActorDesc.cpp
+++ b/Engine/Source/Runtime/Engine/Private/WorldPartition/DataLayer/WorldDataLayersActorDesc.cpp
@@ -185,10 +185,13 @@ void FWorldDataLayersActorDesc::Serialize(FArchive& Ar)
 
 	FWorldPartitionActorDesc::Serialize(Ar);
 
-	if (Ar.CustomVer(FFortniteNCBranchObjectVersion::GUID) >= FFortniteNCBranchObjectVersion::AddedWorldDataLayersActorDesc)
+	if (!bIsDefaultActorDesc)
 	{
-		Ar << DataLayerInstances;
-		bIsValid = true;
+		if (Ar.CustomVer(FFortniteNCBranchObjectVersion::GUID) >= FFortniteNCBranchObjectVersion::AddedWorldDataLayersActorDesc)
+		{
+			Ar << DataLayerInstances;
+			bIsValid = true;
+		}
 	}
 }
 
diff --git a/Engine/Source/Runtime/Engine/Private/WorldPartition/HLOD/HLODActorDesc.cpp b/Engine/Source/Runtime/Engine/Private/WorldPartition/HLOD/HLODActorDesc.cpp
--- a/Engine/Source/Runtime/Engine/Private/WorldPartition/HLOD/HLODActorDesc.cpp
+++ b/Engine/Source/Runtime/Engine/Private/WorldPartition/HLOD/HLODActorDesc.cpp
@@ -39,50 +39,53 @@ void FHLODActorDesc::Serialize(FArchive& Ar)
 
 	FWorldPartitionActorDesc::Serialize(Ar);
 
-	if (Ar.CustomVer(FUE5ReleaseStreamObjectVersion::GUID) < FUE5ReleaseStreamObjectVersion::WorldPartitionHLODActorDescSerializeHLODSubActors)
+	if (!bIsDefaultActorDesc)
 	{
-		TArray<FGuid> SubActors;
-		Ar << SubActors;
-	}
-	else
-	{
-		Ar << HLODSubActors;
-	}
+		if (Ar.CustomVer(FUE5ReleaseStreamObjectVersion::GUID) < FUE5ReleaseStreamObjectVersion::WorldPartitionHLODActorDescSerializeHLODSubActors)
+		{
+			TArray<FGuid> SubActors;
+			Ar << SubActors;
+		}
+		else
+		{
+			Ar << HLODSubActors;
+		}
 	
-	if (Ar.CustomVer(FUE5MainStreamObjectVersion::GUID) < FUE5MainStreamObjectVersion::WorldPartitionHLODActorDescSerializeHLODLayer)
-	{
-		FString HLODLayer_Deprecated;
-		Ar << HLODLayer_Deprecated;
-	}
-
-	const bool bSerializeCellHash = Ar.CustomVer(FUE5MainStreamObjectVersion::GUID) >= FUE5MainStreamObjectVersion::WorldPartitionHLODActorDescSerializeCellHash &&
-									Ar.CustomVer(FFortniteMainBranchObjectVersion::GUID) < FFortniteMainBranchObjectVersion::WorldPartitionHLODActorDescSerializeStats;
-	if (Ar.CustomVer(FUE5MainStreamObjectVersion::GUID) >= FUE5MainStreamObjectVersion::WorldPartitionHLODActorDescSerializeCellHash)
-	{
-		uint64 CellHash_Deprecated;
-		Ar << CellHash_Deprecated;
-	}
-
-	if (Ar.CustomVer(FFortniteMainBranchObjectVersion::GUID) < FFortniteMainBranchObjectVersion::WorldPartitionActorDescSerializeActorIsRuntimeOnly)
-	{
-		bActorIsRuntimeOnly = true;
-	}
+		if (Ar.CustomVer(FUE5MainStreamObjectVersion::GUID) < FUE5MainStreamObjectVersion::WorldPartitionHLODActorDescSerializeHLODLayer)
+		{
+			FString HLODLayer_Deprecated;
+			Ar << HLODLayer_Deprecated;
+		}
 
-	if (Ar.CustomVer(FFortniteMainBranchObjectVersion::GUID) >= FFortniteMainBranchObjectVersion::WorldPartitionHLODActorDescSerializeStats)
-	{
-		if (Ar.CustomVer(FFortniteMainBranchObjectVersion::GUID) < FFortniteMainBranchObjectVersion::WorldPartitionHLODActorUseSourceCellGuid)
+		const bool bSerializeCellHash = Ar.CustomVer(FUE5MainStreamObjectVersion::GUID) >= FUE5MainStreamObjectVersion::WorldPartitionHLODActorDescSerializeCellHash &&
+										Ar.CustomVer(FFortniteMainBranchObjectVersion::GUID) < FFortniteMainBranchObjectVersion::WorldPartitionHLODActorDescSerializeStats;
+		if (Ar.CustomVer(FUE5MainStreamObjectVersion::GUID) >= FUE5MainStreamObjectVersion::WorldPartitionHLODActorDescSerializeCellHash)
 		{
-			FName SourceCellName;
-			Ar << SourceCellName;
+			uint64 CellHash_Deprecated;
+			Ar << CellHash_Deprecated;
 		}
 
-		Ar << SourceHLODLayerName;
-		Ar << HLODStats;
+		if (Ar.CustomVer(FFortniteMainBranchObjectVersion::GUID) < FFortniteMainBranchObjectVersion::WorldPartitionActorDescSerializeActorIsRuntimeOnly)
+		{
+			bActorIsRuntimeOnly = true;
+		}
 
-		// Update package size stat on load
-		if (Ar.IsLoading())
+		if (Ar.CustomVer(FFortniteMainBranchObjectVersion::GUID) >= FFortniteMainBranchObjectVersion::WorldPartitionHLODActorDescSerializeStats)
 		{
-			HLODStats.Add(FWorldPartitionHLODStats::MemoryDiskSizeBytes, GetPackageSize());
+			if (Ar.CustomVer(FFortniteMainBranchObjectVersion::GUID) < FFortniteMainBranchObjectVersion::WorldPartitionHLODActorUseSourceCellGuid)
+			{
+				FName SourceCellName;
+				Ar << SourceCellName;
+			}
+
+			Ar << SourceHLODLayerName;
+			Ar << HLODStats;
+
+			// Update package size stat on load
+			if (Ar.IsLoading())
+			{
+				HLODStats.Add(FWorldPartitionHLODStats::MemoryDiskSizeBytes, GetPackageSize());
+			}
 		}
 	}
 }
diff --git a/Engine/Source/Runtime/Engine/Private/WorldPartition/Landscape/LandscapeActorDesc.cpp b/Engine/Source/Runtime/Engine/Private/WorldPartition/Landscape/LandscapeActorDesc.cpp
--- a/Engine/Source/Runtime/Engine/Private/WorldPartition/Landscape/LandscapeActorDesc.cpp
+++ b/Engine/Source/Runtime/Engine/Private/WorldPartition/Landscape/LandscapeActorDesc.cpp
@@ -17,10 +17,13 @@ void FLandscapeActorDesc::Init(const AActor* InActor)
 	GridIndexY = LandscapeProxy->LandscapeSectionOffset.Y / (int32)LandscapeProxy->GridSize;
 	GridIndexZ = 0;
 
-	const ALandscape* LandscapeActor = LandscapeProxy->GetLandscapeActor();
-	if (LandscapeActor)
+	if (!bIsDefaultActorDesc)
 	{
-		LandscapeActorGuid = LandscapeActor->GetActorGuid();
+		const ALandscape* LandscapeActor = LandscapeProxy->GetLandscapeActor();
+		if (LandscapeActor)
+		{
+			LandscapeActorGuid = LandscapeActor->GetActorGuid();
+		}
 	}
 }
 
@@ -31,15 +34,18 @@ void FLandscapeActorDesc::Serialize(FArchive& Ar)
 	Ar.UsingCustomVersion(FUE5MainStreamObjectVersion::GUID);
 	Ar.UsingCustomVersion(FUE5ReleaseStreamObjectVersion::GUID);
 
-	if (Ar.CustomVer(FUE5MainStreamObjectVersion::GUID) < FUE5MainStreamObjectVersion::FLandscapeActorDescFixupGridIndices)
+	if (!bIsDefaultActorDesc)
 	{
-		GridIndexX = (int32)(GridIndexX * GridSize) / (int32)GridSize;
-		GridIndexY = (int32)(GridIndexY * GridSize) / (int32)GridSize;
-	}
-
-	if (Ar.CustomVer(FUE5ReleaseStreamObjectVersion::GUID) >= FUE5ReleaseStreamObjectVersion::WorldPartitionLandscapeActorDescSerializeLandscapeActorGuid)
-	{
-		Ar << LandscapeActorGuid;
+		if (Ar.CustomVer(FUE5MainStreamObjectVersion::GUID) < FUE5MainStreamObjectVersion::FLandscapeActorDescFixupGridIndices)
+		{
+			GridIndexX = (int32)(GridIndexX * GridSize) / (int32)GridSize;
+			GridIndexY = (int32)(GridIndexY * GridSize) / (int32)GridSize;
+		}
+
+		if (Ar.CustomVer(FUE5ReleaseStreamObjectVersion::GUID) >= FUE5ReleaseStreamObjectVersion::WorldPartitionLandscapeActorDescSerializeLandscapeActorGuid)
+		{
+			Ar << LandscapeActorGuid;
+		}
 	}
 }
 
diff --git a/Engine/Source/Runtime/Engine/Private/WorldPartition/LevelInstance/LevelInstanceActorDesc.cpp b/Engine/Source/Runtime/Engine/Private/WorldPartition/LevelInstance/LevelInstanceActorDesc.cpp
--- a/Engine/Source/Runtime/Engine/Private/WorldPartition/LevelInstance/LevelInstanceActorDesc.cpp
+++ b/Engine/Source/Runtime/Engine/Private/WorldPartition/LevelInstance/LevelInstanceActorDesc.cpp
@@ -12,6 +12,7 @@
 #include "Misc/PackageName.h"
 #include "WorldPartition/ActorDescContainer.h"
 #include "WorldPartition/ErrorHandling/WorldPartitionStreamingGenerationErrorHandler.h"
+#include "WorldPartition/WorldPartitionActorDescArchive.h"
 #include "UObject/UE5ReleaseStreamObjectVersion.h"
 #include "UObject/UE5MainStreamObjectVersion.h"
 #include "UObject/FortniteMainBranchObjectVersion.h"
@@ -286,26 +287,30 @@ void FLevelInstanceActorDesc::TransferFrom(const FWorldPartitionActorDesc* From)
 void FLevelInstanceActorDesc::Serialize(FArchive& Ar)
 {
 	FWorldPartitionActorDesc::Serialize(Ar);
+
 	Ar.UsingCustomVersion(FUE5ReleaseStreamObjectVersion::GUID);
 	Ar.UsingCustomVersion(FUE5MainStreamObjectVersion::GUID);
 	Ar.UsingCustomVersion(FFortniteMainBranchObjectVersion::GUID);
 
-	Ar << LevelPackage;
+	Ar << TDeltaSerialize<FName>(LevelPackage);
 
-	if (Ar.CustomVer(FUE5ReleaseStreamObjectVersion::GUID) < FUE5ReleaseStreamObjectVersion::LargeWorldCoordinates)
-	{
-		FTransform3f LevelInstanceTransformFlt;
-		Ar << LevelInstanceTransformFlt;
-		LevelInstanceTransform = FTransform(LevelInstanceTransformFlt);
-	}
-	else
+	if (!bIsDefaultActorDesc)
 	{
-		Ar << LevelInstanceTransform;
+		if (Ar.CustomVer(FUE5ReleaseStreamObjectVersion::GUID) < FUE5ReleaseStreamObjectVersion::LargeWorldCoordinates)
+		{
+			FTransform3f LevelInstanceTransformFlt;
+			Ar << LevelInstanceTransformFlt;
+			LevelInstanceTransform = FTransform(LevelInstanceTransformFlt);
+		}
+		else
+		{
+			Ar << LevelInstanceTransform;
+		}
 	}
 
 	if (Ar.CustomVer(FUE5ReleaseStreamObjectVersion::GUID) >= FUE5ReleaseStreamObjectVersion::LevelInstanceSerializeRuntimeBehavior)
 	{
-		Ar << DesiredRuntimeBehavior;
+		Ar << TDeltaSerialize<ELevelInstanceRuntimeBehavior>(DesiredRuntimeBehavior);
 
 		if (Ar.IsLoading() && DesiredRuntimeBehavior == ELevelInstanceRuntimeBehavior::Embedded_Deprecated)
 		{
@@ -313,21 +318,24 @@ void FLevelInstanceActorDesc::Serialize(FArchive& Ar)
 		}
 	}
 
-	if (Ar.IsLoading())
+	if (!bIsDefaultActorDesc)
 	{
-		const bool bFixupOldVersion = (Ar.CustomVer(FUE5MainStreamObjectVersion::GUID) < FUE5MainStreamObjectVersion::PackedLevelInstanceBoundsFix) && 
-									  (Ar.CustomVer(FUE5ReleaseStreamObjectVersion::GUID) < FUE5ReleaseStreamObjectVersion::PackedLevelInstanceBoundsFix);
-
-		const AActor* CDO = GetActorNativeClass()->GetDefaultObject<AActor>();
-		const ILevelInstanceInterface* LevelInstanceCDO = CastChecked<ILevelInstanceInterface>(CDO);
-		if (!LevelPackage.IsNone() && (LevelInstanceCDO->IsLoadingEnabled() || bFixupOldVersion))
+		if (Ar.IsLoading())
 		{
-			if (!IsContainerInstance())
+			const bool bFixupOldVersion = (Ar.CustomVer(FUE5MainStreamObjectVersion::GUID) < FUE5MainStreamObjectVersion::PackedLevelInstanceBoundsFix) && 
+										  (Ar.CustomVer(FUE5ReleaseStreamObjectVersion::GUID) < FUE5ReleaseStreamObjectVersion::PackedLevelInstanceBoundsFix);
+
+			const AActor* CDO = GetActorNativeClass()->GetDefaultObject<AActor>();
+			const ILevelInstanceInterface* LevelInstanceCDO = CastChecked<ILevelInstanceInterface>(CDO);
+			if (!LevelPackage.IsNone() && (LevelInstanceCDO->IsLoadingEnabled() || bFixupOldVersion))
 			{
-				FBox OutBounds;
-				if (ULevelInstanceSubsystem::GetLevelInstanceBoundsFromPackage(LevelInstanceTransform, LevelPackage, OutBounds))
+				if (!IsContainerInstance())
 				{
-					OutBounds.GetCenterAndExtents(BoundsLocation, BoundsExtent);
+					FBox OutBounds;
+					if (ULevelInstanceSubsystem::GetLevelInstanceBoundsFromPackage(LevelInstanceTransform, LevelPackage, OutBounds))
+					{
+						OutBounds.GetCenterAndExtents(BoundsLocation, BoundsExtent);
+					}
 				}
 			}
 		}
diff --git a/Engine/Source/Runtime/Engine/Private/WorldPartition/WorldPartition.cpp b/Engine/Source/Runtime/Engine/Private/WorldPartition/WorldPartition.cpp
--- a/Engine/Source/Runtime/Engine/Private/WorldPartition/WorldPartition.cpp
+++ b/Engine/Source/Runtime/Engine/Private/WorldPartition/WorldPartition.cpp
@@ -193,7 +193,7 @@ static FAutoConsoleCommand DumpActorDesc(
 						TMap<FName, FString> DataLayersDumpString = GetDataLayersDumpString(WorldPartition);
 						for (const FString& ActorPath : ActorPaths)
 						{
-							if (const FWorldPartitionActorDesc* ActorDesc = WorldPartition->GetActorDesc(ActorPath))
+							if (const FWorldPartitionActorDesc* ActorDesc = WorldPartition->GetActorDescByName(ActorPath))
 							{
 								UE_LOG(LogWorldPartition, Log, TEXT("%s"), *GetActorDescDumpString(ActorDesc, DataLayersDumpString));
 							}
@@ -1306,7 +1306,7 @@ bool UWorldPartition::ResolveSubobject(const TCHAR* SubObjectPath, UObject*& Out
 				SubObjectName = SubObjectPath;
 			}
 
-			if (const FWorldPartitionActorDesc* ActorDesc = GetActorDesc(SubObjectName))
+			if (const FWorldPartitionActorDesc* ActorDesc = GetActorDescByName(SubObjectName))
 			{
 				if (bLoadIfExists)
 				{
diff --git a/Engine/Source/Runtime/Engine/Private/WorldPartition/WorldPartitionActorDesc.cpp b/Engine/Source/Runtime/Engine/Private/WorldPartition/WorldPartitionActorDesc.cpp
--- a/Engine/Source/Runtime/Engine/Private/WorldPartition/WorldPartitionActorDesc.cpp
+++ b/Engine/Source/Runtime/Engine/Private/WorldPartition/WorldPartitionActorDesc.cpp
@@ -2,10 +2,11 @@
 
 #if WITH_EDITOR
 #include "WorldPartition/WorldPartitionActorDesc.h"
-#include "Misc/Paths.h"
 #include "WorldPartition/WorldPartitionActorDescArchive.h"
 #include "WorldPartition/WorldPartitionActorDescUtils.h"
 
+#include "Misc/Paths.h"
+#include "Misc/ArchiveMD5.h"
 #include "UObject/MetaData.h"
 #include "AssetRegistry/AssetData.h"
 #include "Serialization/MemoryReader.h"
@@ -37,16 +38,13 @@ FWorldPartitionActorDesc::FWorldPartitionActorDesc()
 	, HardRefCount(0)
 	, Container(nullptr)
 	, bIsForcedNonSpatiallyLoaded(false)
+	, bIsDefaultActorDesc(false)
 	, UnloadedReason(nullptr)
 {}
 
 void FWorldPartitionActorDesc::Init(const AActor* InActor)
 {	
 	check(IsValid(InActor));
-	check(InActor->IsPackageExternal());
-
-	Guid = InActor->GetActorGuid();
-	check(Guid.IsValid());
 
 	UClass* ActorClass = InActor->GetClass();
 
@@ -60,6 +58,31 @@ void FWorldPartitionActorDesc::Init(const AActor* InActor)
 		BaseClass = *InActor->GetClass()->GetPathName();
 	}
 
+	if (InActor->HasAnyFlags(RF_ArchetypeObject | RF_ClassDefaultObject))
+	{
+		check(!InActor->IsPackageExternal());
+		check(!InActor->GetActorGuid().IsValid());
+
+		FArchiveMD5 MD5Ar;
+		FString ClassPath = BaseClass.IsValid() ? BaseClass.ToString() : NativeClass.ToString();
+		MD5Ar << ClassPath;
+		
+		FMD5Hash MD5Hash;
+		MD5Ar.GetHash(MD5Hash);
+
+		Guid = MD5HashToGuid(MD5Hash);
+
+		bIsDefaultActorDesc = true;
+	}
+	else
+	{
+		check(InActor->IsPackageExternal());
+		check(InActor->GetActorGuid().IsValid());
+		Guid = InActor->GetActorGuid();
+	}
+
+	check(Guid.IsValid());
+
 	const FBox StreamingBounds = InActor->GetStreamingBounds();
 	StreamingBounds.GetCenterAndExtents(BoundsLocation, BoundsExtent);
 
@@ -138,30 +161,34 @@ void FWorldPartitionActorDesc::Init(const AActor* InActor)
 	InActor->GetActorDescProperties(Properties);
 
 	ActorPackage = InActor->GetPackage()->GetFName();
-	ActorPath = *InActor->GetPathName();
-	FolderPath = InActor->GetFolderPath();
-	FolderGuid = InActor->GetFolderGuid();
-
-	const AActor* AttachParentActor = InActor->GetAttachParentActor();
-	if (AttachParentActor)
-	{
-		ParentActor = AttachParentActor->GetActorGuid();
-	}
+	ActorPath = bIsDefaultActorDesc ? *InActor->GetClass()->GetPathName() : *InActor->GetPathName();
 
 	ContentBundleGuid = InActor->GetContentBundleGuid();
-	
-	TArray<AActor*> ActorReferences = ActorsReferencesUtils::GetExternalActorReferences(const_cast<AActor*>(InActor));
 
-	if (ActorReferences.Num())
+	if (!bIsDefaultActorDesc)
 	{
-		References.Empty(ActorReferences.Num());
-		for(AActor* ActorReference: ActorReferences)
+		FolderPath = InActor->GetFolderPath();
+		FolderGuid = InActor->GetFolderGuid();
+
+		const AActor* AttachParentActor = InActor->GetAttachParentActor();
+		if (AttachParentActor)
 		{
-			References.Add(ActorReference->GetActorGuid());
+			ParentActor = AttachParentActor->GetActorGuid();
 		}
-	}
 
-	ActorLabel = *InActor->GetActorLabel(false);
+		TArray<AActor*> ActorReferences = ActorsReferencesUtils::GetExternalActorReferences(const_cast<AActor*>(InActor));
+
+		if (ActorReferences.Num())
+		{
+			References.Empty(ActorReferences.Num());
+			for(AActor* ActorReference: ActorReferences)
+			{
+				References.Add(ActorReference->GetActorGuid());
+			}
+		}
+
+		ActorLabel = *InActor->GetActorLabel(false);
+	}
 
 	Container = nullptr;
 	ActorPtr = const_cast<AActor*>(InActor);
@@ -183,7 +210,7 @@ void FWorldPartitionActorDesc::Init(const FWorldPartitionActorDescInitData& Desc
 	MetadataAr.SetCustomVersions(CustomVersions);
 	
 	// Serialize metadata payload
-	FActorDescArchive ActorDescAr(MetadataAr);
+	FActorDescArchive ActorDescAr(MetadataAr, this);
 	Serialize(ActorDescAr);
 
 	// Call registered deprecator
@@ -232,16 +259,14 @@ void FWorldPartitionActorDesc::SerializeTo(TArray<uint8>& OutData)
 	// Serialize to archive and gather custom versions
 	TArray<uint8> PayloadData;
 	FMemoryWriter PayloadAr(PayloadData, true);
-	FActorDescArchive ActorDescAr(PayloadAr);
-
-	// Serialize actor descriptor
+	FActorDescArchive ActorDescAr(PayloadAr, this);
 	Serialize(ActorDescAr);
 
 	// Serialize custom versions
 	TArray<uint8> HeaderData;
 	FMemoryWriter HeaderAr(HeaderData);
 	FCustomVersionContainer CustomVersions = ActorDescAr.GetCustomVersions();
-	CustomVersions.Serialize(HeaderAr);
+	CustomVersions.Serialize(HeaderAr); 
 
 	// Append data
 	OutData = MoveTemp(HeaderData);
@@ -372,35 +397,6 @@ void FWorldPartitionActorDesc::Serialize(FArchive& Ar)
 	Ar.UsingCustomVersion(FFortniteNCBranchObjectVersion::GUID);
 	Ar.UsingCustomVersion(FFortniteMainBranchObjectVersion::GUID);
 
-	if (Ar.CustomVer(FFortniteNCBranchObjectVersion::GUID) >= FFortniteNCBranchObjectVersion::WorldPartitionActorDescNativeBaseClassSerialization)
-	{
-		if (Ar.CustomVer(FUE5MainStreamObjectVersion::GUID) < FUE5MainStreamObjectVersion::WorldPartitionActorDescActorAndClassPaths)
-		{
-			FName BaseClassPathName;
-			Ar << BaseClassPathName;
-PRAGMA_DISABLE_DEPRECATION_WARNINGS
-			BaseClass = FAssetData::TryConvertShortClassNameToPathName(BaseClassPathName, ELogVerbosity::NoLogging);
-PRAGMA_ENABLE_DEPRECATION_WARNINGS
-		}
-		else
-		{
-			Ar << BaseClass;
-		}
-	}
-
-	if (Ar.CustomVer(FUE5MainStreamObjectVersion::GUID) < FUE5MainStreamObjectVersion::WorldPartitionActorDescActorAndClassPaths)
-	{
-		FName NativeClassPathName;
-		Ar << NativeClassPathName;
-PRAGMA_DISABLE_DEPRECATION_WARNINGS
-		NativeClass = FAssetData::TryConvertShortClassNameToPathName(NativeClassPathName, ELogVerbosity::NoLogging);
-PRAGMA_ENABLE_DEPRECATION_WARNINGS
-	}
-	else
-	{
-		Ar << NativeClass;
-	}
-
 	Ar << Guid;
 
 	if(Ar.CustomVer(FUE5ReleaseStreamObjectVersion::GUID) < FUE5ReleaseStreamObjectVersion::LargeWorldCoordinates)
@@ -410,7 +406,7 @@ PRAGMA_ENABLE_DEPRECATION_WARNINGS
 		BoundsLocation = FVector(BoundsLocationFlt);
 		BoundsExtent = FVector(BoundsExtentFlt);
 	}
-	else
+	else if (!bIsDefaultActorDesc)
 	{
 		Ar << BoundsLocation << BoundsExtent;
 	}
@@ -425,14 +421,15 @@ PRAGMA_ENABLE_DEPRECATION_WARNINGS
 	}
 	else
 	{
-		Ar << bIsSpatiallyLoaded;
+		Ar << TDeltaSerialize<bool>(bIsSpatiallyLoaded);
 	}
-		
-	Ar << RuntimeGrid << bActorIsEditorOnly;
+
+	Ar << TDeltaSerialize<FName>(RuntimeGrid);
+	Ar << TDeltaSerialize<bool>(bActorIsEditorOnly);
 
 	if (Ar.CustomVer(FFortniteMainBranchObjectVersion::GUID) >= FFortniteMainBranchObjectVersion::WorldPartitionActorDescSerializeActorIsRuntimeOnly)
 	{
-		Ar << bActorIsRuntimeOnly;
+		Ar << TDeltaSerialize<bool>(bActorIsRuntimeOnly);
 	}
 	else
 	{
@@ -455,7 +452,7 @@ PRAGMA_ENABLE_DEPRECATION_WARNINGS
 
 	if (Ar.CustomVer(FFortniteNCBranchObjectVersion::GUID) >= FFortniteNCBranchObjectVersion::WorldPartitionActorDescTagsSerialization)
 	{
-		Ar << Tags;
+		Ar << TDeltaSerialize<TArray<FName>>(Tags);
 	}
 
 	if (Ar.CustomVer(FUE5MainStreamObjectVersion::GUID) < FUE5MainStreamObjectVersion::WorldPartitionActorDescSerializeArchivePersistent)
@@ -477,24 +474,24 @@ PRAGMA_ENABLE_DEPRECATION_WARNINGS
 
 	if (Ar.CustomVer(FUE5MainStreamObjectVersion::GUID) >= FUE5MainStreamObjectVersion::WorldPartitionActorDescSerializeDataLayers)
 	{
-		Ar << DataLayers;
+		Ar << TDeltaSerialize<TArray<FName>>(DataLayers);
 	}
 
 	if (Ar.CustomVer(FFortniteNCBranchObjectVersion::GUID) >= FFortniteNCBranchObjectVersion::WorldPartitionActorDescSerializeDataLayerAssets)
 	{
-		Ar << bIsUsingDataLayerAsset;
+		Ar << TDeltaSerialize<bool>(bIsUsingDataLayerAsset);
 	}
 
 	if (Ar.CustomVer(FUE5MainStreamObjectVersion::GUID) >= FUE5MainStreamObjectVersion::WorldPartitionActorDescSerializeActorLabel)
 	{
-		Ar << ActorLabel;
+		Ar << TDeltaSerialize<FName>(ActorLabel);
 	}
 
 	if ((Ar.CustomVer(FUE5MainStreamObjectVersion::GUID) >= FUE5MainStreamObjectVersion::WorldPartitionActorDescSerializeHLODInfo) ||
 		(Ar.CustomVer(FUE5ReleaseStreamObjectVersion::GUID) >= FUE5ReleaseStreamObjectVersion::WorldPartitionActorDescSerializeHLODInfo))
 	{
-		Ar << bActorIsHLODRelevant;
-		Ar << HLODLayer;
+		Ar << TDeltaSerialize<bool>(bActorIsHLODRelevant);
+		Ar << TDeltaSerialize<FName>(HLODLayer);
 	}
 	else
 	{
@@ -502,24 +499,27 @@ PRAGMA_ENABLE_DEPRECATION_WARNINGS
 		HLODLayer = FName();
 	}
 
-	if (Ar.CustomVer(FUE5ReleaseStreamObjectVersion::GUID) >= FUE5ReleaseStreamObjectVersion::WorldPartitionActorDescSerializeActorFolderPath)
+	if (!bIsDefaultActorDesc)
 	{
-		Ar << FolderPath;
-	}
+		if (Ar.CustomVer(FUE5ReleaseStreamObjectVersion::GUID) >= FUE5ReleaseStreamObjectVersion::WorldPartitionActorDescSerializeActorFolderPath)
+		{
+			Ar << FolderPath;
+		}
 
-	if (Ar.CustomVer(FUE5ReleaseStreamObjectVersion::GUID) >= FUE5ReleaseStreamObjectVersion::WorldPartitionActorDescSerializeAttachParent)
-	{
-		Ar << ParentActor;
-	}
+		if (Ar.CustomVer(FUE5ReleaseStreamObjectVersion::GUID) >= FUE5ReleaseStreamObjectVersion::WorldPartitionActorDescSerializeAttachParent)
+		{
+			Ar << ParentActor;
+		}
 
-	if (Ar.CustomVer(FUE5ReleaseStreamObjectVersion::GUID) >= FUE5ReleaseStreamObjectVersion::AddLevelActorFolders)
-	{
-		Ar << FolderGuid;
-	}
+		if (Ar.CustomVer(FUE5ReleaseStreamObjectVersion::GUID) >= FUE5ReleaseStreamObjectVersion::AddLevelActorFolders)
+		{
+			Ar << FolderGuid;
+		}
 
-	if (Ar.CustomVer(FFortniteNCBranchObjectVersion::GUID) >= FFortniteNCBranchObjectVersion::WorldPartitionActorDescPropertyMapSerialization)
-	{
-		Ar << Properties;
+		if (Ar.CustomVer(FFortniteNCBranchObjectVersion::GUID) >= FFortniteNCBranchObjectVersion::WorldPartitionActorDescPropertyMapSerialization)
+		{
+			Ar << Properties;
+		}
 	}
 
 	if (Ar.CustomVer(FFortniteMainBranchObjectVersion::GUID) < FFortniteMainBranchObjectVersion::WorldPartitionActorDescSerializeContentBundleGuid)
@@ -528,7 +528,8 @@ PRAGMA_ENABLE_DEPRECATION_WARNINGS
 	}
 	else
 	{
-		Ar << ContentBundleGuid;
+		Ar << TDeltaSerialize<FGuid>(ContentBundleGuid);
+
 		// @todo_ow: remove once we find why some actors end up with invalid ContentBundleGuids
 		if (Ar.IsLoading())
 		{
diff --git a/Engine/Source/Runtime/Engine/Private/WorldPartition/WorldPartitionActorDescArchive.cpp b/Engine/Source/Runtime/Engine/Private/WorldPartition/WorldPartitionActorDescArchive.cpp
--- a/Engine/Source/Runtime/Engine/Private/WorldPartition/WorldPartitionActorDescArchive.cpp
+++ b/Engine/Source/Runtime/Engine/Private/WorldPartition/WorldPartitionActorDescArchive.cpp
@@ -2,13 +2,70 @@
 
 #if WITH_EDITOR
 #include "WorldPartition/WorldPartitionActorDescArchive.h"
+#include "WorldPartition/WorldPartitionClassDescRegistry.h"
+#include "WorldPartition/WorldPartitionActorDesc.h"
+#include "WorldPartition/WorldPartitionLog.h"
+#include "UObject/UE5MainStreamObjectVersion.h"
+#include "UObject/FortniteNCBranchObjectVersion.h"
+#include "UObject/FortniteMainBranchObjectVersion.h"
+#include "AssetRegistry/AssetData.h"
 
-FActorDescArchive::FActorDescArchive(FArchive& InArchive)
+FActorDescArchive::FActorDescArchive(FArchive& InArchive, FWorldPartitionActorDesc* InActorDesc)
 	: FArchiveProxy(InArchive)
+	, ActorDesc(InActorDesc)
 {
-	check(InArchive.IsPersistent());		
+	check(InArchive.IsPersistent());
 
 	SetIsPersistent(true);
-	SetIsLoading(InArchive.IsLoading());		
+	SetIsLoading(InArchive.IsLoading());
+
+	UsingCustomVersion(FUE5MainStreamObjectVersion::GUID);
+	UsingCustomVersion(FFortniteNCBranchObjectVersion::GUID);
+	UsingCustomVersion(FFortniteMainBranchObjectVersion::GUID);
+
+	if (CustomVer(FFortniteMainBranchObjectVersion::GUID) >= FFortniteMainBranchObjectVersion::WorldPartitionActorClassDescSerialize)
+	{
+		InnerArchive << InActorDesc->bIsDefaultActorDesc;
+	}
+
+	if (CustomVer(FFortniteNCBranchObjectVersion::GUID) >= FFortniteNCBranchObjectVersion::WorldPartitionActorDescNativeBaseClassSerialization)
+	{
+		if (CustomVer(FUE5MainStreamObjectVersion::GUID) < FUE5MainStreamObjectVersion::WorldPartitionActorDescActorAndClassPaths)
+		{
+			FName BaseClassPathName;
+			InnerArchive << BaseClassPathName;
+PRAGMA_DISABLE_DEPRECATION_WARNINGS
+			InActorDesc->BaseClass = FAssetData::TryConvertShortClassNameToPathName(BaseClassPathName);
+PRAGMA_ENABLE_DEPRECATION_WARNINGS
+		}
+		else
+		{
+			InnerArchive << InActorDesc->BaseClass;
+		}
+	}
+
+	if (CustomVer(FUE5MainStreamObjectVersion::GUID) < FUE5MainStreamObjectVersion::WorldPartitionActorDescActorAndClassPaths)
+	{
+		FName NativeClassPathName;
+		InnerArchive << NativeClassPathName;
+PRAGMA_DISABLE_DEPRECATION_WARNINGS
+		InActorDesc->NativeClass = FAssetData::TryConvertShortClassNameToPathName(NativeClassPathName);
+PRAGMA_ENABLE_DEPRECATION_WARNINGS
+	}
+	else
+	{
+		InnerArchive << InActorDesc->NativeClass;
+	}
+
+	// Get the class descriptor to do delta serialization
+	FWorldPartitionClassDescRegistry& ClassDescRegistry = FWorldPartitionClassDescRegistry::Get();
+	const FTopLevelAssetPath ClassPath = InActorDesc->BaseClass.IsValid() ? InActorDesc->BaseClass : InActorDesc->NativeClass;
+	ClassDesc = InActorDesc->bIsDefaultActorDesc ? ClassDescRegistry.GetClassDescDefaultForClass(ClassPath) : ClassDescRegistry.GetClassDescDefaultForActor(ClassPath);
+	if (!ClassDesc)
+	{
+		ClassDesc = ClassDescRegistry.GetClassDescDefault(FTopLevelAssetPath(TEXT("/Script/Engine.Actor")));
+		UE_LOG(LogWorldPartition, Warning, TEXT("Can't find clas descriptor '%s' for '%s', using '%s'"), *ClassPath.ToString(), *InActorDesc->GetActorSoftPath().ToString(), *ClassDesc->GetActorSoftPath().ToString());
+	}
+	check(ClassDesc);
 }
 #endif
\ No newline at end of file
diff --git a/Engine/Source/Runtime/Engine/Private/WorldPartition/WorldPartitionActorDescUtils.cpp b/Engine/Source/Runtime/Engine/Private/WorldPartition/WorldPartitionActorDescUtils.cpp
--- a/Engine/Source/Runtime/Engine/Private/WorldPartition/WorldPartitionActorDescUtils.cpp
+++ b/Engine/Source/Runtime/Engine/Private/WorldPartition/WorldPartitionActorDescUtils.cpp
@@ -53,6 +53,10 @@ UClass* FWorldPartitionActorDescUtils::GetActorNativeClassFromAssetData(const FA
 	FString ActorMetaDataClass;
 	if (InAssetData.GetTagValue(NAME_ActorMetaDataClass, ActorMetaDataClass))
 	{
+		// Avoid an assert when calling StaticFindObject during save to retrieve the actor's class.
+		// Since we are only looking for a native class, the call to StaticFindObject is legit.
+		TGuardValue<bool> GIsSavingPackageGuard(GIsSavingPackage, false);
+
 		// Look for a class redirectors
 		const FString ActorNativeClassName = ResolveClassRedirector(ActorMetaDataClass);
 		
@@ -96,31 +100,32 @@ TUniquePtr<FWorldPartitionActorDesc> FWorldPartitionActorDescUtils::GetActorDesc
 
 void FWorldPartitionActorDescUtils::AppendAssetDataTagsFromActor(const AActor* InActor, TArray<UObject::FAssetRegistryTag>& OutTags)
 {
-	check(InActor->IsPackageExternal());
-	
 	TUniquePtr<FWorldPartitionActorDesc> ActorDesc(InActor->CreateActorDesc());
 
-	// If the actor is not added to a world, we can't retrieve its bounding volume, so try to get the existing one
-	if (ULevel* Level = InActor->GetLevel(); !Level || !Level->Actors.Contains(InActor))
+	if (!InActor->HasAnyFlags(RF_ArchetypeObject | RF_ClassDefaultObject))
 	{
-		// Avoid an assert when calling StaticFindObject during save to retrieve the actor's class.
-		// Since we are only looking for a native class, the call to StaticFindObject is legit.
-		TGuardValue<bool> GIsSavingPackageGuard(GIsSavingPackage, false);
+		// If the actor is not added to a world, we can't retrieve its bounding volume, so try to get the existing one
+		if (ULevel* Level = InActor->GetLevel(); !Level || !Level->Actors.Contains(InActor))
+		{
+			// Avoid an assert when calling StaticFindObject during save to retrieve the actor's class.
+			// Since we are only looking for a native class, the call to StaticFindObject is legit.
+			TGuardValue<bool> GIsSavingPackageGuard(GIsSavingPackage, false);
 
-		IAssetRegistry& AssetRegistry = FModuleManager::LoadModuleChecked<FAssetRegistryModule>(TEXT("AssetRegistry")).Get();
+			IAssetRegistry& AssetRegistry = FModuleManager::LoadModuleChecked<FAssetRegistryModule>(TEXT("AssetRegistry")).Get();
 
-		FARFilter Filter;
-		Filter.bIncludeOnlyOnDiskAssets = true;
-		Filter.PackageNames.Add(InActor->GetPackage()->GetFName());
+			FARFilter Filter;
+			Filter.bIncludeOnlyOnDiskAssets = true;
+			Filter.PackageNames.Add(InActor->GetPackage()->GetFName());
 
-		TArray<FAssetData> Assets;
-		AssetRegistry.GetAssets(Filter, Assets);
+			TArray<FAssetData> Assets;
+			AssetRegistry.GetAssets(Filter, Assets);
 
-		if (Assets.Num() == 1)
-		{
-			if (TUniquePtr<FWorldPartitionActorDesc> NewActorDesc = FWorldPartitionActorDescUtils::GetActorDescriptorFromAssetData(Assets[0]))
+			if (Assets.Num() == 1)
 			{
-				ActorDesc->TransferWorldData(NewActorDesc.Get());
+				if (TUniquePtr<FWorldPartitionActorDesc> NewActorDesc = FWorldPartitionActorDescUtils::GetActorDescriptorFromAssetData(Assets[0]))
+				{
+					ActorDesc->TransferWorldData(NewActorDesc.Get());
+				}
 			}
 		}
 	}
@@ -147,6 +152,9 @@ void FWorldPartitionActorDescUtils::UpdateActorDescriptorFromActor(const AActor*
 
 void FWorldPartitionActorDescUtils::UpdateActorDescriptorFromActorDescriptor(TUniquePtr<FWorldPartitionActorDesc>& InActorDesc, TUniquePtr<FWorldPartitionActorDesc>& OutActorDesc)
 {
+	check(InActorDesc->GetActorPackage() == OutActorDesc->GetActorPackage());
+	check(InActorDesc->GetActorSoftPath() == OutActorDesc->GetActorSoftPath());
+	check(InActorDesc->GetGuid() == OutActorDesc->GetGuid());
 	InActorDesc->TransferFrom(OutActorDesc.Get());
 	OutActorDesc = MoveTemp(InActorDesc);
 }
@@ -159,84 +167,6 @@ void FWorldPartitionActorDescUtils::ReplaceActorDescriptorPointerFromActor(const
 	InActorDesc->ActorPtr = InNewActor;
 }
 
-bool FWorldPartitionActorDescUtils::ValidateActorDescClass(FWorldPartitionActorDesc* InActorDesc)
-{
-	// If the native class in invalid (potentially deleted), it means we parsed the actor descriptor with AActor::StaticClass and explicitly
-	// marked the class as invalid.
-	if (!InActorDesc->GetNativeClass().IsValid())
-	{
-		UE_LOG(LogWorldPartition, Warning, TEXT("Failed to find native class for actor '%s"), *InActorDesc->GetActorSoftPath().ToString());
-		return false;
-	}
-
-	// If the base class is invalid, it means the actor is from a native class.
-	if (!InActorDesc->GetBaseClass().IsValid())
-	{
-		return true;
-	}
-
-	// Lookup the Bp class, if it's already loaded we don't need to validate anything.
-	if (UClass* BaseClass = FindObject<UClass>(InActorDesc->GetBaseClass()))
-	{
-		return true;
-	}
-
-	// The actor is from a BP class which isn't loaded. To avoid loading the class, go through the asset registry to validate the class.
-	const FAssetData* ClassData;
-	FString ActorDescBaseClass = InActorDesc->GetBaseClass().ToString();
-
-	IAssetRegistry& AssetRegistry = FModuleManager::LoadModuleChecked<FAssetRegistryModule>(TEXT("AssetRegistry")).Get();
-
-	for (;;)
-	{
-		// Look for a class redirectors
-		const FString ActorClassName = ResolveClassRedirector(ActorDescBaseClass);
-		FTopLevelAssetPath AssetClassPath(*ActorClassName);
-		
-		ActorDescBaseClass.RemoveFromEnd(TEXT("_C"), ESearchCase::CaseSensitive);
-		const FString ActorClassNameAlt = ResolveClassRedirector(ActorDescBaseClass);
-		FTopLevelAssetPath AssetClassPathAlt(*ActorClassNameAlt);
-
-		TArray<FAssetData> BlueprintAssets;							
-		AssetRegistry.ScanFilesSynchronous({ AssetClassPath.GetPackageName().ToString() }, /*bForceRescan*/false);
-		AssetRegistry.GetAssetsByPackageName(AssetClassPath.GetPackageName(), BlueprintAssets, /*bIncludeOnlyOnDiskAssets*/true);
-
-		if (!BlueprintAssets.Num())
-		{
-			UE_LOG(LogWorldPartition, Warning, TEXT("Failed to find assets for class '%s' for actor '%s"), *AssetClassPath.ToString(), *InActorDesc->GetActorSoftPath().ToString());
-			return false;
-		}
-
-		ClassData = BlueprintAssets.FindByPredicate([&AssetClassPath, &AssetClassPathAlt](const FAssetData& AssetData)
-		{ 
-			return (AssetData.ToSoftObjectPath().GetAssetPath() == AssetClassPath) || (AssetData.ToSoftObjectPath().GetAssetPath() == AssetClassPathAlt);
-		});
-
-		if (!ClassData)
-		{
-			UE_LOG(LogWorldPartition, Display, TEXT("Failed to find class asset '%s' for actor '%s"), *AssetClassPath.ToString(), *InActorDesc->GetActorSoftPath().ToString());
-			return false;
-		}
-
-		if (!ClassData->IsRedirector())
-		{
-			break;
-		}
-
-		FString DestinationObjectPath;
-		if (!ClassData->GetTagValue(TEXT("DestinationObject"), DestinationObjectPath))
-		{
-			UE_LOG(LogWorldPartition, Warning, TEXT("Failed to follow class redirector for '%s' for actor '%s"), *AssetClassPath.ToString(), *InActorDesc->GetActorSoftPath().ToString());
-			return false;
-		}
-
-		// Cleanup the destination object path from prefixed class type, if any
-		ActorDescBaseClass = FTopLevelAssetPath(DestinationObjectPath).ToString();
-	}
-
-	return true;
-}
-
 bool FWorldPartitionActorDescUtils::FixupRedirectedAssetPath(FName& InOutAssetPath)
 {
 	const FAssetData* AssetData;
diff --git a/Engine/Source/Runtime/Engine/Private/WorldPartition/WorldPartitionClassDescRegistry.cpp b/Engine/Source/Runtime/Engine/Private/WorldPartition/WorldPartitionClassDescRegistry.cpp
new file mode 100644
--- /dev/null
+++ b/Engine/Source/Runtime/Engine/Private/WorldPartition/WorldPartitionClassDescRegistry.cpp
@@ -0,0 +1,557 @@
+// Copyright Epic Games, Inc. All Rights Reserved.
+// Copyright Epic Games, Inc. All Rights Reserved.
+
+#include "WorldPartition/WorldPartitionClassDescRegistry.h"
+
+#if WITH_EDITOR
+#include "Editor.h"
+#include "Algo/ForEach.h"
+#include "Algo/RemoveIf.h"
+#include "Algo/Transform.h"
+#include "UObject/CoreRedirects.h"
+#include "UObject/ObjectSaveContext.h"
+#include "UObject/UObjectIterator.h"
+#include "AssetRegistry/AssetData.h"
+#include "AssetRegistry/AssetRegistryModule.h"
+#include "Interfaces/IPluginManager.h"
+#include "Engine/Blueprint.h"
+#include "Engine/BlueprintGeneratedClass.h"
+#include "ProfilingDebugging/ScopedTimers.h"
+#include "WorldPartition/WorldPartitionLog.h"
+#include "WorldPartition/WorldPartitionActorDesc.h"
+#include "WorldPartition/WorldPartitionActorDescUtils.h"
+
+void FWorldPartitionClassDescRegistry::Initialize()
+{
+	UE_SCOPED_TIMER(TEXT("FWorldPartitionClassDescRegistry::Initialize"), LogWorldPartition, Display);
+	TRACE_CPUPROFILER_EVENT_SCOPE(FWorldPartitionClassDescRegistry::Initialize);
+
+	check(!IsInitialized());
+
+	RegisterClasses();
+
+	FCoreUObjectDelegates::OnAssetLoaded.AddRaw(this, &FWorldPartitionClassDescRegistry::OnAssetLoaded);
+	FCoreUObjectDelegates::OnObjectPreSave.AddRaw(this, &FWorldPartitionClassDescRegistry::OnObjectPreSave);
+	FCoreUObjectDelegates::OnObjectPropertyChanged.AddRaw(this, &FWorldPartitionClassDescRegistry::OnObjectPropertyChanged);
+
+	IPluginManager& PluginManager = IPluginManager::Get();
+	ELoadingPhase::Type LoadingPhase = PluginManager.GetLastCompletedLoadingPhase();
+	if (LoadingPhase == ELoadingPhase::None || LoadingPhase < ELoadingPhase::PostEngineInit)
+	{
+		PluginManager.OnLoadingPhaseComplete().AddRaw(this, &FWorldPartitionClassDescRegistry::OnPluginLoadingPhaseComplete);
+	}
+
+	IAssetRegistry& AssetRegistry = FModuleManager::LoadModuleChecked<FAssetRegistryModule>(TEXT("AssetRegistry")).Get();
+	AssetRegistry.OnInMemoryAssetCreated().AddRaw(this, &FWorldPartitionClassDescRegistry::OnAssetCreated);
+	AssetRegistry.OnAssetRemoved().AddRaw(this, &FWorldPartitionClassDescRegistry::OnAssetRemoved);
+	AssetRegistry.OnAssetRenamed().AddRaw(this, &FWorldPartitionClassDescRegistry::OnAssetRenamed);
+
+	check(IsInitialized());
+}
+
+void FWorldPartitionClassDescRegistry::Uninitialize()
+{
+	check(IsInitialized());
+	FActorDescList::Empty();
+
+	FCoreUObjectDelegates::OnAssetLoaded.RemoveAll(this);
+	FCoreUObjectDelegates::OnObjectPreSave.RemoveAll(this);
+	FCoreUObjectDelegates::OnObjectPropertyChanged.RemoveAll(this);
+	
+	if (!IsEngineExitRequested())
+	{
+		IAssetRegistry& AssetRegistry = FModuleManager::LoadModuleChecked<FAssetRegistryModule>(TEXT("AssetRegistry")).Get();
+		AssetRegistry.OnInMemoryAssetCreated().RemoveAll(this);
+		AssetRegistry.OnAssetRemoved().RemoveAll(this);
+		AssetRegistry.OnAssetRenamed().RemoveAll(this);
+	}
+
+	ClassByPath.Empty();
+	ParentClassMap.Empty();
+
+	check(!IsInitialized());
+}
+
+bool FWorldPartitionClassDescRegistry::IsInitialized() const
+{
+	return !FActorDescList::IsEmpty();
+}
+
+void FWorldPartitionClassDescRegistry::PrefetchClassDescs(const TArray<FTopLevelAssetPath>& InClassPaths)
+{
+	TRACE_CPUPROFILER_EVENT_SCOPE(FWorldPartitionClassDescRegistry::PrefetchClassDescs);
+
+	check(IsInitialized());
+
+	IAssetRegistry& AssetRegistry = FModuleManager::LoadModuleChecked<FAssetRegistryModule>(TEXT("AssetRegistry")).Get();
+
+	// First, remap classes and remove already known ones.
+	TArray<FTopLevelAssetPath> ClassPaths;
+	ClassPaths.Reserve(InClassPaths.Num());
+	Algo::Transform(InClassPaths, ClassPaths, [this](const FTopLevelAssetPath& ClassPath) { return RedirectClassPath(ClassPath); });
+	ClassPaths.SetNum(Algo::RemoveIf(ClassPaths, [this](const FTopLevelAssetPath& ClassPath) { return ClassByPath.Contains(ClassPath) || ClassPath.ToString().StartsWith(TEXT("/Script/")); }));
+
+	if (ClassPaths.IsEmpty())
+	{
+		return;
+	}
+
+	// Make sure the asset registry has all paths scanned
+	TArray<FString> FilePaths;
+	FilePaths.Reserve(ClassPaths.Num());
+	Algo::Transform(ClassPaths, FilePaths, [this](const FTopLevelAssetPath& ClassPath) { return ClassPath.GetPackageName().ToString(); });
+
+	// In the editor, the asset registry will not return blueprint generated classes or their default object but will during cooks (see AssetRegistry::FFiltering::ShouldSkipAsset),
+	// so we must filter our results by UBlueprint and UObjectRedirector. For redirectors, they will show up for the blueprint, generated class and default objects in both situations,
+	// so we must take special care to filter out everything except blueprint objects. The current way to find the blueprint asset redirector is based on finding the redirector with
+	// the shortest name.
+	auto SortBlueprintAssetsByNameLength = [](const FAssetData& Lhs, const FAssetData& Rhs) { return Lhs.AssetName.GetStringLength() < Rhs.AssetName.GetStringLength(); };
+
+	auto GetBlueprintAssets = [&AssetRegistry](const TArray<FString>& FilePaths, TArray<FAssetData>& Assets)
+	{
+		FARFilter Filter;
+		Filter.ClassPaths.Add(UBlueprint::StaticClass()->GetClassPathName());
+		Filter.ClassPaths.Add(UObjectRedirector::StaticClass()->GetClassPathName());
+		Filter.bRecursiveClasses = true;
+		Filter.PackageNames.Reserve(FilePaths.Num());
+		Algo::Transform(FilePaths, Filter.PackageNames, [](const FString& ClassPath) { return *ClassPath; });
+	
+		AssetRegistry.ScanFilesSynchronous(FilePaths, /*bForceRescan*/false);
+		AssetRegistry.GetAssets(Filter, Assets);
+	};
+
+	TArray<FAssetData> Assets;
+	GetBlueprintAssets(FilePaths, Assets);
+
+	// We are only interested in redirectors for blueprint objects, not their generated class (ClassName_C) or their default object (Default__ClassName_C), as per previous comment.
+	{
+		TMap<FName, TArray<FAssetData*>> PackageRedirectors;
+		Algo::ForEach(Assets, [&PackageRedirectors](FAssetData& AssetData) { PackageRedirectors.FindOrAdd(AssetData.PackageName).Add(&AssetData); });
+		Algo::ForEach(PackageRedirectors, [&SortBlueprintAssetsByNameLength](TPair<FName, TArray<FAssetData*>>& Pair)
+		{
+			// Find the blueprint asset or the proper blueprint redirector
+			Pair.Value.Sort(SortBlueprintAssetsByNameLength);
+			for (int32 AssetIndex = 1; AssetIndex < Pair.Value.Num(); AssetIndex++)	{ *Pair.Value[AssetIndex] = FAssetData(); }
+		});
+		Assets.SetNum(Algo::RemoveIf(Assets, [](const FAssetData& AssetData) { return !AssetData.IsValid(); }));
+	}	
+
+	TMap<FTopLevelAssetPath, FTopLevelAssetPath> ParentClassPathsMap;
+	Algo::ForEach(Assets, [this, &AssetRegistry, &ParentClassPathsMap, &SortBlueprintAssetsByNameLength, &GetBlueprintAssets](FAssetData& AssetData)
+	{
+		check(AssetData.IsValid());
+
+		TArray<FTopLevelAssetPath> ClassRedirects;
+		while (AssetData.IsRedirector())
+		{
+			// Folow the redirector to the destination object
+			FString DestinationObjectPath;
+			if (!AssetData.GetTagValue(TEXT("DestinationObject"), DestinationObjectPath))
+			{
+				AssetData = FAssetData();
+				UE_LOG(LogWorldPartition, Warning, TEXT("Failed to follow class redirector for '%s'"), *AssetData.ToSoftObjectPath().ToString());
+				break;
+			}
+
+			const FTopLevelAssetPath AssetClassPath = FTopLevelAssetPath(DestinationObjectPath);
+
+			TArray<FAssetData> RedirectAssets;
+			GetBlueprintAssets({ AssetClassPath.GetPackageName().ToString() }, RedirectAssets);
+
+			if (RedirectAssets.IsEmpty())
+			{
+				AssetData = FAssetData();
+				UE_LOG(LogWorldPartition, Warning, TEXT("Failed to find redirected assets for '%s'"), *AssetData.ToSoftObjectPath().ToString());
+				break;
+			}
+
+			ClassRedirects.Add(FTopLevelAssetPath(AssetData.ToSoftObjectPath().ToString()));
+
+			// Find the blueprint asset or the proper blueprint redirector
+			RedirectAssets.Sort(SortBlueprintAssetsByNameLength);
+			AssetData = MoveTemp(RedirectAssets[0]);
+		}
+
+		// Register redirects
+		for (const FTopLevelAssetPath& ClassRedirect : ClassRedirects)
+		{
+			FTopLevelAssetPath SourceClassPath(ClassRedirect.ToString() + TEXT("_C"));
+			FTopLevelAssetPath RedirectedClassPath(AssetData.ToSoftObjectPath().ToString() + TEXT("_C"));
+			RedirectClassMap.Add(SourceClassPath, RedirectedClassPath);
+		}
+
+		if (AssetData.IsValid())
+		{
+			FString ParentClassName;
+			if (AssetData.GetTagValue(FBlueprintTags::ParentClassPath, ParentClassName))
+			{
+				ParentClassPathsMap.Add(FTopLevelAssetPath(AssetData.ToSoftObjectPath().ToString() + TEXT("_C")), FTopLevelAssetPath(ParentClassName));
+			}
+		}
+	});
+
+	// Register parent class descriptors
+	if (!ParentClassPathsMap.IsEmpty())
+	{
+		TArray<FTopLevelAssetPath> ParentClassPaths;
+		ParentClassPathsMap.GenerateValueArray(ParentClassPaths);
+		PrefetchClassDescs(ParentClassPaths);
+
+		for (auto& [ClassPath, ParentClassPath] : ParentClassPathsMap)
+		{
+			ParentClassMap.Add(ClassPath, RedirectClassPath(ParentClassPath));
+		}
+	}
+
+	// Register current class descriptors
+	Algo::ForEach(Assets, [this](FAssetData& AssetData)
+	{
+		// Create the class descriptor
+		AssetData.AssetName = *FString(AssetData.AssetName.ToString() + TEXT("_C"));
+
+		// Lookup for an already registered class
+		if (!ClassByPath.Contains(FTopLevelAssetPath(AssetData.ToSoftObjectPath().ToString())))
+		{
+			if (TUniquePtr<FWorldPartitionActorDesc> ClassDesc = FWorldPartitionActorDescUtils::GetActorDescriptorFromAssetData(AssetData); ClassDesc.IsValid())
+			{
+				FWorldPartitionActorDesc* Result = ClassDesc.Release();
+				RegisterClassDescriptor(Result);
+			}
+			else
+			{
+				// For missing class descriptors, register a null entry so we won't be trying to register them again for future objects
+				ClassByPath.Add(FTopLevelAssetPath(FTopLevelAssetPath(AssetData.ToSoftObjectPath().ToString())), nullptr);
+			}
+		}
+	});
+
+	ValidateInternalState();
+}
+
+bool FWorldPartitionClassDescRegistry::IsRegisteredClass(const FTopLevelAssetPath& InClassPath) const
+{
+	check(IsInitialized());
+	const FTopLevelAssetPath ClassPath = RedirectClassPath(InClassPath);
+	return ClassByPath.Contains(ClassPath);
+}
+
+const FWorldPartitionActorDesc* FWorldPartitionClassDescRegistry::GetClassDescDefault(const FTopLevelAssetPath& InClassPath) const
+{
+	check(IsInitialized());
+
+	const FWorldPartitionActorDesc* Result = nullptr;
+	FTopLevelAssetPath ClassPath = InClassPath;
+
+	while (ClassPath.IsValid())
+	{
+		// Lookup for an already registered class
+		if (const TUniquePtr<FWorldPartitionActorDesc>* const* ClassDesc = ClassByPath.Find(ClassPath); ClassDesc && *ClassDesc)
+		{
+			Result = (*ClassDesc)->Get();
+			break;
+		}
+
+		ClassPath = ParentClassMap.FindRef(ClassPath);
+	}
+
+	return Result;
+}
+
+const FWorldPartitionActorDesc* FWorldPartitionClassDescRegistry::GetClassDescDefaultForActor(const FTopLevelAssetPath& InClassPath) const
+{
+	return GetClassDescDefault(RedirectClassPath(InClassPath));
+}
+
+const FWorldPartitionActorDesc* FWorldPartitionClassDescRegistry::GetClassDescDefaultForClass(const FTopLevelAssetPath& InClassPath) const
+{
+	const FTopLevelAssetPath ParentClassPath = ParentClassMap.FindRef(RedirectClassPath(InClassPath));
+	return GetClassDescDefault(ParentClassPath);
+}
+
+void FWorldPartitionClassDescRegistry::RegisterClassDescriptor(FWorldPartitionActorDesc* InClassDesc)
+{
+	FActorDescList::AddActorDescriptor(InClassDesc);
+
+	// Validate that the class isn't registered yet
+	const FTopLevelAssetPath ClassPath(InClassDesc->GetActorSoftPath().ToString());
+	if (TUniquePtr<FWorldPartitionActorDesc>** ExistingClassDesc = ClassByPath.Find(ClassPath))
+	{
+		check(!(*ExistingClassDesc));
+	}
+
+	ClassByPath.Add(ClassPath, ActorsByGuid.FindChecked(InClassDesc->GetGuid()));
+}
+
+void FWorldPartitionClassDescRegistry::UnregisterClassDescriptor(FWorldPartitionActorDesc* InClassDesc)
+{
+	verify(ClassByPath.Remove(FTopLevelAssetPath(InClassDesc->GetActorSoftPath().ToString())));
+	FActorDescList::RemoveActorDescriptor(InClassDesc);
+}
+
+void FWorldPartitionClassDescRegistry::RegisterClassDescriptorFromAssetData(const FAssetData& InAssetData)
+{
+	const FTopLevelAssetPath ClassPath = FTopLevelAssetPath(InAssetData.ToSoftObjectPath().ToString() + TEXT("_C"));
+	check(!ClassByPath.Contains(ClassPath));
+
+	UBlueprint* Blueprint = CastChecked<UBlueprint>(InAssetData.GetAsset());
+	RegisterClassDescriptorFromActorClass(Blueprint->GeneratedClass);
+}
+
+void FWorldPartitionClassDescRegistry::RegisterClassDescriptorFromActorClass(const UClass* InActorClass)
+{
+	ParentClassMap.Add(FTopLevelAssetPath(InActorClass->GetPathName()), FTopLevelAssetPath(InActorClass->GetSuperClass()->GetPathName()));
+
+	const AActor* ActorCDO = CastChecked<AActor>(InActorClass->GetDefaultObject());
+	FWorldPartitionActorDesc* NewActorDesc = ActorCDO->CreateActorDesc().Release();
+	RegisterClassDescriptor(NewActorDesc);
+}
+
+void FWorldPartitionClassDescRegistry::OnAssetCreated(UObject* InAssetCreated)
+{
+	OnAssetLoaded(InAssetCreated);
+}
+
+void FWorldPartitionClassDescRegistry::OnAssetLoaded(UObject* InAssetLoaded)
+{
+	check(IsInitialized());
+
+	if (UBlueprint* Blueprint = Cast<UBlueprint>(InAssetLoaded))
+	{
+		if (Blueprint->GeneratedClass && Blueprint->GeneratedClass->IsChildOf<AActor>())
+		{
+			if (!ClassByPath.Contains(FTopLevelAssetPath(Blueprint->GeneratedClass->GetPathName())))
+			{
+				ValidateInternalState();
+				RegisterClassDescriptorFromActorClass(Blueprint->GeneratedClass);
+				ValidateInternalState();
+			}
+		}		
+	}
+}
+
+void FWorldPartitionClassDescRegistry::OnObjectPreSave(UObject* InObject, FObjectPreSaveContext InSaveContext)
+{
+	if (!InSaveContext.IsProceduralSave() && !(InSaveContext.GetSaveFlags() & SAVE_FromAutosave))
+	{
+		ValidateInternalState();
+
+		if (AActor* Actor = Actor = Cast<AActor>(InObject))
+		{
+			if (!Actor->HasAnyFlags(RF_ArchetypeObject | RF_ClassDefaultObject))
+			{
+				// We are saving an actor
+				if (UBlueprintGeneratedClass* BlueprintGeneratedClass = Cast<UBlueprintGeneratedClass>(Actor->GetClass()))
+				{
+					if (UBlueprint* Blueprint = Cast<UBlueprint>(BlueprintGeneratedClass->ClassGeneratedBy))
+					{
+						UpdateClassDescriptor(Blueprint);
+					}
+				}
+			}
+		}
+		else if (UBlueprint* Blueprint = Cast<UBlueprint>(InObject))
+		{
+			// We are saving a blueprint
+			if (UBlueprintGeneratedClass* BlueprintGeneratedClass = Cast<UBlueprintGeneratedClass>(Blueprint->GeneratedClass))
+			{
+				if (BlueprintGeneratedClass->IsChildOf<AActor>())
+				{
+					UpdateClassDescriptor(Blueprint);
+				}
+			}
+		}
+
+		ValidateInternalState();
+	}
+}
+
+void FWorldPartitionClassDescRegistry::OnObjectPropertyChanged(UObject* InObject, FPropertyChangedEvent& InPropertyChangedEvent)
+{
+	if (UBlueprint* Blueprint = Cast<UBlueprint>(InObject))
+	{
+		// The generated class is invalid in some situations, like renaming a blueprint, etc.
+		if (Blueprint->GeneratedClass && Blueprint->GeneratedClass->IsChildOf<AActor>())
+		{
+			ValidateInternalState();
+			UpdateClassDescriptor(InObject);
+			ValidateInternalState();
+		}
+	}
+}
+
+void FWorldPartitionClassDescRegistry::OnPluginLoadingPhaseComplete(ELoadingPhase::Type InLoadingPhase, bool bInPhaseSuccessful)
+{
+	if (InLoadingPhase == ELoadingPhase::PostEngineInit)
+	{
+		RegisterClasses();
+		IPluginManager::Get().OnLoadingPhaseComplete().RemoveAll(this);
+	}
+}
+
+void FWorldPartitionClassDescRegistry::OnAssetRemoved(const FAssetData& InAssetData)
+{
+	check(IsInitialized());
+
+	if (InAssetData.GetClass() && InAssetData.GetClass()->IsChildOf(UBlueprint::StaticClass()))
+	{
+		const FTopLevelAssetPath ClassPath = FTopLevelAssetPath(InAssetData.ToSoftObjectPath().ToString() + TEXT("_C"));
+
+		if (TUniquePtr<FWorldPartitionActorDesc>** ExistingClassDesc = ClassByPath.Find(ClassPath))
+		{
+			ValidateInternalState();
+			if (*ExistingClassDesc)
+			{
+				UnregisterClassDescriptor((*ExistingClassDesc)->Get());
+			}
+			else
+			{
+				ClassByPath.Remove(ClassPath);
+			}
+			ValidateInternalState();
+		}
+	}
+}
+
+void FWorldPartitionClassDescRegistry::OnAssetRenamed(const FAssetData& InAssetData, const FString& InOldObjectPath)
+{
+	check(IsInitialized());
+
+	if (InAssetData.GetClass() && InAssetData.GetClass()->IsChildOf(UBlueprint::StaticClass()))
+	{
+		ValidateInternalState();
+
+		const FTopLevelAssetPath OldClassPath = FTopLevelAssetPath(InOldObjectPath + TEXT("_C"));
+		const FTopLevelAssetPath NewClassPath = FTopLevelAssetPath(InAssetData.ToSoftObjectPath().ToString() + TEXT("_C"));
+
+		if (TUniquePtr<FWorldPartitionActorDesc>** ExistingClassDesc = ClassByPath.Find(OldClassPath))
+		{
+			if (*ExistingClassDesc)
+			{
+				UnregisterClassDescriptor((*ExistingClassDesc)->Get());
+				RegisterClassDescriptorFromAssetData(InAssetData);
+			}
+			else
+			{
+				ClassByPath.Remove(OldClassPath);
+				ClassByPath.Add(NewClassPath, nullptr);
+			}
+		
+			// Update renamed class parent
+			const FTopLevelAssetPath ParentClassPath = ParentClassMap.FindAndRemoveChecked(OldClassPath);
+			ParentClassMap.Add(NewClassPath, ParentClassPath);
+
+			// Update renamed class child classes
+			TArray<FTopLevelAssetPath> ChildsToUpdate;
+			for (auto& [ChildClassPath, ChildParentClassPath] : ParentClassMap)
+			{
+				if (ChildParentClassPath == OldClassPath)
+				{
+					ChildsToUpdate.Add(ChildClassPath);
+				}
+			}
+
+			for (const FTopLevelAssetPath& ChildClassPath : ChildsToUpdate)
+			{
+				ParentClassMap.FindChecked(ChildClassPath) = NewClassPath;
+			}
+		}
+
+		ValidateInternalState();
+	}
+}
+
+void FWorldPartitionClassDescRegistry::RegisterClass(UClass* Class)
+{
+	if (!ClassByPath.Contains(FTopLevelAssetPath(Class->GetPathName())))
+	{
+		if (Class != AActor::StaticClass())
+		{
+			RegisterClass(Class->GetSuperClass());
+		}
+
+		RegisterClassDescriptorFromActorClass(Class);
+	}
+}
+
+void FWorldPartitionClassDescRegistry::RegisterClasses()
+{
+	ValidateInternalState();
+
+	for (TObjectIterator<UClass> It; It; ++It)
+	{
+		UClass* Class = *It;
+		if (Class->IsChildOf<AActor>() && (Class->GetPackage() != GetTransientPackage()))
+		{
+			RegisterClass(Class);
+		}
+	}
+
+	ValidateInternalState();
+}
+
+void FWorldPartitionClassDescRegistry::UpdateClassDescriptor(UObject* InObject)
+{
+	check(IsInitialized());
+
+	UBlueprint* Blueprint = CastChecked<UBlueprint>(InObject);
+	check(Blueprint->GeneratedClass);
+
+	AActor* ActorCDO = CastChecked<AActor>(Blueprint->GeneratedClass->GetDefaultObject());
+
+	const FTopLevelAssetPath ClassPath(ActorCDO->GetClass()->GetPathName());
+	const FTopLevelAssetPath CurrentParentClassPath = FTopLevelAssetPath(ActorCDO->GetClass()->GetSuperClass()->GetPathName());
+
+	if (TUniquePtr<FWorldPartitionActorDesc>** ExistingClassDesc = ClassByPath.Find(ClassPath))
+	{
+		const FTopLevelAssetPath PreviousParentClassPath = ParentClassMap.FindChecked(ClassPath);
+
+		if (!*ExistingClassDesc)
+		{
+			FWorldPartitionActorDesc* NewActorDesc = ActorCDO->CreateActorDesc().Release();
+			RegisterClassDescriptor(NewActorDesc);
+		}
+		else
+		{
+			FWorldPartitionActorDescUtils::UpdateActorDescriptorFromActor(ActorCDO, **ExistingClassDesc);
+		}
+
+		if (PreviousParentClassPath != CurrentParentClassPath)
+		{
+			// We are reparenting a blueprint, update our parent map
+			ParentClassMap.Add(ClassPath, CurrentParentClassPath);
+		}
+	}
+}
+
+void FWorldPartitionClassDescRegistry::ValidateInternalState()
+{
+#if DO_GUARD_SLOW
+	for (auto& [ClassPath, ParentClassPath] : ParentClassMap)
+	{
+		// Validate that the class parents chain is valid
+		static const FTopLevelAssetPath ActorClassPath(TEXT("/Script/Engine.Actor"));
+		for (FTopLevelAssetPath CurrentClassPath = ClassPath; CurrentClassPath != ActorClassPath; CurrentClassPath = ParentClassMap.FindChecked(CurrentClassPath))
+		{
+			check(CurrentClassPath.IsValid());
+			ClassByPath.FindChecked(CurrentClassPath);
+		}
+	}
+#endif
+}
+
+FTopLevelAssetPath FWorldPartitionClassDescRegistry::RedirectClassPath(const FTopLevelAssetPath& InClassPath) const
+{
+	// Resolve native class redirectors
+	const FCoreRedirectObjectName OldClassName = FCoreRedirectObjectName(InClassPath.GetAssetName(), NAME_None, InClassPath.GetPackageName());
+	const FCoreRedirectObjectName NewClassName = FCoreRedirects::GetRedirectedName(ECoreRedirectFlags::Type_Class, OldClassName);
+	const FTopLevelAssetPath ResolvedClassPath = FTopLevelAssetPath(NewClassName.ToString());
+
+	// Resolve blueprint object redirectors
+	if (const FTopLevelAssetPath* RedirectedClass = RedirectClassMap.Find(ResolvedClassPath))
+	{
+		return *RedirectedClass;
+	}
+
+	return ResolvedClassPath;
+}
+#endif
\ No newline at end of file
diff --git a/Engine/Source/Runtime/Engine/Public/WorldPartition/ActorDescContainer.h b/Engine/Source/Runtime/Engine/Public/WorldPartition/ActorDescContainer.h
--- a/Engine/Source/Runtime/Engine/Public/WorldPartition/ActorDescContainer.h
+++ b/Engine/Source/Runtime/Engine/Public/WorldPartition/ActorDescContainer.h
@@ -22,6 +22,8 @@ class ENGINE_API UActorDescContainer : public UObject, public FActorDescList
 	friend struct FWorldPartitionHandleUtils;
 	friend class FWorldPartitionActorDesc;
 
+	using TNameActorDescMap = TMap<FName, TUniquePtr<FWorldPartitionActorDesc>*>;
+
 public:
 	/* Struct of parameters passed to Initialize function. */
 	struct ENGINE_API FInitializeParams
@@ -88,10 +90,8 @@ public:
 	const TArray<TUniquePtr<FWorldPartitionActorDesc>>& GetInvalidActors() const { return InvalidActors; }
 	void ClearInvalidActors() { InvalidActors.Empty(); }
 
-	//~ Begin FActorDescList Interface
-	virtual void AddActorDescriptor(FWorldPartitionActorDesc* ActorDesc, UWorld* InWorldContext) override;
-	virtual void RemoveActorDescriptor(FWorldPartitionActorDesc* ActorDesc) override;
-	//~ End FActorDescList Interface
+	void RegisterActorDescriptor(FWorldPartitionActorDesc* ActorDesc, UWorld* InWorldContext);
+	void UnregisterActorDescriptor(FWorldPartitionActorDesc* ActorDesc);
 
 	void OnActorDescAdded(FWorldPartitionActorDesc* NewActorDesc);
 	void OnActorDescRemoved(FWorldPartitionActorDesc* ActorDesc);
@@ -100,6 +100,9 @@ public:
 
 	bool ShouldHandleActorEvent(const AActor* Actor);
 
+	const FWorldPartitionActorDesc* GetActorDescByName(const FString& ActorPath) const;
+	const FWorldPartitionActorDesc* GetActorDescByName(const FSoftObjectPath& InActorPath) const;
+
 	bool bContainerInitialized;
 
 	FName ContainerPackageName;
@@ -108,6 +111,8 @@ public:
 	TArray<TUniquePtr<FWorldPartitionActorDesc>> InvalidActors;
 
 protected:
+	TNameActorDescMap ActorsByName;
+
 	//~ Begin UObject Interface
 	virtual void BeginDestroy() override;
 	//~ End UObject Interface
diff --git a/Engine/Source/Runtime/Engine/Public/WorldPartition/ActorDescContainerCollection.h b/Engine/Source/Runtime/Engine/Public/WorldPartition/ActorDescContainerCollection.h
--- a/Engine/Source/Runtime/Engine/Public/WorldPartition/ActorDescContainerCollection.h
+++ b/Engine/Source/Runtime/Engine/Public/WorldPartition/ActorDescContainerCollection.h
@@ -29,8 +29,8 @@ public:
 	FWorldPartitionActorDesc& GetActorDescChecked(const FGuid& Guid);
 	const FWorldPartitionActorDesc& GetActorDescChecked(const FGuid& Guid) const;
 
-	const FWorldPartitionActorDesc* GetActorDesc(const FString& PackageName) const;
-	const FWorldPartitionActorDesc* GetActorDesc(const FSoftObjectPath& InActorPath) const;
+	const FWorldPartitionActorDesc* GetActorDescByName(const FString& ActorPath) const;
+	const FWorldPartitionActorDesc* GetActorDescByName(const FSoftObjectPath& ActorPath) const;
 
 	bool RemoveActor(const FGuid& ActorGuid);
 
diff --git a/Engine/Source/Runtime/Engine/Public/WorldPartition/ActorDescList.h b/Engine/Source/Runtime/Engine/Public/WorldPartition/ActorDescList.h
--- a/Engine/Source/Runtime/Engine/Public/WorldPartition/ActorDescList.h
+++ b/Engine/Source/Runtime/Engine/Public/WorldPartition/ActorDescList.h
@@ -11,6 +11,16 @@ class ENGINE_API FActorDescList
 #if WITH_EDITOR
 	friend struct FWorldPartitionImplBase;
 
+	struct FActorGuidKeyFuncs : TDefaultMapKeyFuncs<FGuid, TUniquePtr<FWorldPartitionActorDesc>*, false>
+	{
+		static FORCEINLINE uint32 GetKeyHash(const FGuid& Key)
+		{
+			return HashCombineFast(HashCombineFast(Key.A, Key.B), HashCombineFast(Key.C, Key.D));
+		}
+	};
+
+	using TGuidActorDescMap = TMap<FGuid, TUniquePtr<FWorldPartitionActorDesc>*, FDefaultSetAllocator, FActorGuidKeyFuncs>;
+
 public:
 	FActorDescList() {}
 	virtual ~FActorDescList() {}
@@ -27,9 +37,6 @@ public:
 	FWorldPartitionActorDesc& GetActorDescChecked(const FGuid& Guid);
 	const FWorldPartitionActorDesc& GetActorDescChecked(const FGuid& Guid) const;
 
-	const FWorldPartitionActorDesc* GetActorDesc(const FString& PackageName) const;
-	const FWorldPartitionActorDesc* GetActorDesc(const FSoftObjectPath& InActorPath) const;
-
 	int32 GetActorDescCount() const { return ActorsByGuid.Num(); }
 
 	bool IsEmpty() const { return GetActorDescCount() == 0; }
@@ -41,11 +48,11 @@ public:
 		static_assert(TIsDerivedFrom<ActorType, AActor>::IsDerived, "Type is not derived from AActor.");
 
 	protected:
-		typedef TMap<FGuid, TUniquePtr<FWorldPartitionActorDesc>*> MapType;
-		typedef typename FWorldPartitionActorDescType<ActorType>::Type ValueType;
-		typedef typename TChooseClass<bConst, MapType::TConstIterator, MapType::TIterator>::Result IteratorType;
-		typedef typename TChooseClass<bConst, const FActorDescList*, FActorDescList*>::Result ListType;		
-		typedef typename TChooseClass<bConst, const ValueType*, ValueType*>::Result ReturnType;
+		using MapType = TGuidActorDescMap;
+		using ValueType = typename FWorldPartitionActorDescType<ActorType>::Type;
+		using IteratorType = typename TChooseClass<bConst, MapType::TConstIterator, MapType::TIterator>::Result;
+		using ListType = typename TChooseClass<bConst, const FActorDescList*, FActorDescList*>::Result;
+		using ReturnType = typename TChooseClass<bConst, const ValueType*, ValueType*>::Result;
 
 	public:
 		TBaseIterator(ListType InActorDescList, UClass* InActorClass)
@@ -91,6 +98,7 @@ public:
 		{
 			return StaticCast<ReturnType>(ActorsIterator->Value->Get());
 		}
+
 		/**
 		 * Returns whether the iterator has reached the end and no longer points
 		 * to a suitable actor desc.
@@ -131,7 +139,7 @@ public:
 	template <class ActorType = AActor>
 	class TIterator : public TBaseIterator<false, ActorType>
 	{
-		typedef TBaseIterator<false, ActorType> BaseType;
+		using BaseType = TBaseIterator<false, ActorType>;
 
 	public:
 		TIterator(typename BaseType::ListType InActorDescList, UClass* InActorClass = nullptr)
@@ -142,7 +150,7 @@ public:
 	template <class ActorType = AActor>
 	class TConstIterator : public TBaseIterator<true, ActorType>
 	{
-		typedef TBaseIterator<true, ActorType> BaseType;
+		using BaseType = TBaseIterator<true, ActorType>;
 
 	public:
 		TConstIterator(typename BaseType::ListType InActorDescList, UClass* InActorClass = nullptr)
@@ -150,15 +158,15 @@ public:
 		{}
 	};
 
+	void AddActorDescriptor(FWorldPartitionActorDesc* ActorDesc);
+	void RemoveActorDescriptor(FWorldPartitionActorDesc* ActorDesc);
+
 protected:
-	virtual void AddActorDescriptor(FWorldPartitionActorDesc* ActorDesc, UWorld* InWorldContext);
-	virtual void RemoveActorDescriptor(FWorldPartitionActorDesc* ActorDesc);
+
 	TUniquePtr<FWorldPartitionActorDesc>* GetActorDescriptor(const FGuid& ActorGuid);
 
 	TChunkedArray<TUniquePtr<FWorldPartitionActorDesc>> ActorDescList;
 
-private:
-	TMap<FGuid, TUniquePtr<FWorldPartitionActorDesc>*> ActorsByGuid;
-	TMap<FName, TUniquePtr<FWorldPartitionActorDesc>*> ActorsByName;
+	TGuidActorDescMap ActorsByGuid;
 #endif
 };
\ No newline at end of file
diff --git a/Engine/Source/Runtime/Engine/Public/WorldPartition/WorldPartitionActorDesc.h b/Engine/Source/Runtime/Engine/Public/WorldPartition/WorldPartitionActorDesc.h
--- a/Engine/Source/Runtime/Engine/Public/WorldPartition/WorldPartitionActorDesc.h
+++ b/Engine/Source/Runtime/Engine/Public/WorldPartition/WorldPartitionActorDesc.h
@@ -71,6 +71,7 @@ class ENGINE_API FWorldPartitionActorDesc
 	friend struct FWorldPartitionActorDescUtils;
 	friend struct FWorldPartitionActorDescUnitTestAcccessor;
 	friend class FAssetRootPackagePatcher;
+	friend class FActorDescArchive;
 
 public:
 	struct FContainerInstance
@@ -308,6 +309,7 @@ protected:
 	UActorDescContainer*			Container;
 	TOptional<TArray<FName>>		ResolvedDataLayerInstanceNames; // Can only resolve in ActorDesc if Container is not used as a template
 	bool							bIsForcedNonSpatiallyLoaded;
+	bool							bIsDefaultActorDesc;
 	mutable FText*					UnloadedReason;
 
 	static TMap<TSubclassOf<AActor>, FActorDescDeprecator> Deprecators;
diff --git a/Engine/Source/Runtime/Engine/Public/WorldPartition/WorldPartitionActorDescArchive.h b/Engine/Source/Runtime/Engine/Public/WorldPartition/WorldPartitionActorDescArchive.h
--- a/Engine/Source/Runtime/Engine/Public/WorldPartition/WorldPartitionActorDescArchive.h
+++ b/Engine/Source/Runtime/Engine/Public/WorldPartition/WorldPartitionActorDescArchive.h
@@ -4,11 +4,14 @@
 #if WITH_EDITOR
 #include "CoreMinimal.h"
 #include "Serialization/Archive.h"
+#include "UObject/FortniteMainBranchObjectVersion.h"
+
+class FWorldPartitionActorDesc;
 
 class FActorDescArchive : public FArchiveProxy
 {
 public:
-	FActorDescArchive(FArchive& InArchive);
+	FActorDescArchive(FArchive& InArchive, FWorldPartitionActorDesc* InActorDesc);
 
 	//~ Begin FArchive Interface
 	virtual FArchive& operator<<(FText& Value) override { unimplemented(); return *this; }
@@ -19,5 +22,61 @@ public:
 	virtual FArchive& operator<<(FSoftObjectPath& Value) override { unimplemented(); return *this; }
 	virtual FArchive& operator<<(FWeakObjectPtr& Value) override { unimplemented(); return *this; }
 	//~ End FArchive Interface
+
+	template <typename T>
+	struct TDeltaSerializer
+	{
+		explicit TDeltaSerializer(T& InValue)
+			: Value(InValue)
+		{}
+
+		friend FArchive& operator<<(FArchive& Ar, const TDeltaSerializer<T>& V)
+		{
+			FActorDescArchive& ActorDescAr = (FActorDescArchive&)Ar;
+			
+			check(ActorDescAr.ClassDesc);
+			check(ActorDescAr.ActorDesc);
+
+			Ar.UsingCustomVersion(FFortniteMainBranchObjectVersion::GUID);
+
+			auto GetClassDefaultValue = [&V, &ActorDescAr]() -> const T&
+			{
+				const UPTRINT PropertyOffset = (UPTRINT)&V.Value - *(UPTRINT*)&ActorDescAr.ActorDesc;
+				const T& RefValue = *(const T*)(*(UPTRINT*)&ActorDescAr.ClassDesc + PropertyOffset);
+				return RefValue;
+			};
+
+			uint8 bSerialize = 1;
+
+			if (Ar.CustomVer(FFortniteMainBranchObjectVersion::GUID) >= FFortniteMainBranchObjectVersion::WorldPartitionActorClassDescSerialize)
+			{
+				if (Ar.IsSaving())
+				{
+					bSerialize = (V.Value != GetClassDefaultValue()) ? 1 : 0;
+				}
+
+				Ar << bSerialize;
+			}
+
+			if (bSerialize)
+			{
+				Ar << V.Value;
+			}
+			else if (Ar.IsLoading())
+			{
+				V.Value = GetClassDefaultValue();
+			}
+
+			return Ar;
+		}
+
+		T& Value;
+	};
+
+	FWorldPartitionActorDesc* ActorDesc;
+	const FWorldPartitionActorDesc* ClassDesc;
 };
+
+template <typename T>
+using TDeltaSerialize = FActorDescArchive::TDeltaSerializer<T>;
 #endif
\ No newline at end of file
diff --git a/Engine/Source/Runtime/Engine/Public/WorldPartition/WorldPartitionActorDescUtils.h b/Engine/Source/Runtime/Engine/Public/WorldPartition/WorldPartitionActorDescUtils.h
--- a/Engine/Source/Runtime/Engine/Public/WorldPartition/WorldPartitionActorDescUtils.h
+++ b/Engine/Source/Runtime/Engine/Public/WorldPartition/WorldPartitionActorDescUtils.h
@@ -73,13 +73,6 @@ struct ENGINE_API FWorldPartitionActorDescUtils
 	 */
 	static void ReplaceActorDescriptorPointerFromActor(const AActor* InOldActor, AActor* InNewActor, FWorldPartitionActorDesc* InActorDesc);
 
-	/** 
-	 * Validate if the actor descriptor's class exists without loading anything by going through the asset registry, walking the redirector chain if needed.
-	 * @param InActorDesc		The actor descriptor to resolve the class from.
-	 * @return					Whether the actor descriptor's class is valid or not.
-	 */
-	static bool ValidateActorDescClass(FWorldPartitionActorDesc* InActorDesc);
-
 	/**
 	 * Resolves the provided asset path using asset redirectors.
 	 * @param	InOutAssetPath	Asset path to resolve [In/Out].
diff --git a/Engine/Source/Runtime/Engine/Public/WorldPartition/WorldPartitionClassDescRegistry.h b/Engine/Source/Runtime/Engine/Public/WorldPartition/WorldPartitionClassDescRegistry.h
new file mode 100644
--- /dev/null
+++ b/Engine/Source/Runtime/Engine/Public/WorldPartition/WorldPartitionClassDescRegistry.h
@@ -0,0 +1,71 @@
+// Copyright Epic Games, Inc. All Rights Reserved.
+#pragma once
+
+#if WITH_EDITOR
+#include "CoreMinimal.h"
+#include "Misc/LazySingleton.h"
+#include "UObject/ObjectMacros.h"
+#include "WorldPartition/ActorDescList.h"
+#include "ModuleDescriptor.h"
+
+class UActorDescContainer;
+
+class ENGINE_API FWorldPartitionClassDescRegistry : FActorDescList
+{
+	using TNameClassDescMap = TMap<FTopLevelAssetPath, TUniquePtr<FWorldPartitionActorDesc>*>;
+	using TParentClassMap = TMap<FTopLevelAssetPath, FTopLevelAssetPath>;
+	using TRedirectClassMap = TMap<FTopLevelAssetPath, FTopLevelAssetPath>;
+
+public:
+	static FWorldPartitionClassDescRegistry& Get()
+	{
+		return TLazySingleton<FWorldPartitionClassDescRegistry>::Get();
+	}
+
+	void Initialize();
+	void Uninitialize();
+	
+	bool IsInitialized() const;
+
+	void PrefetchClassDescs(const TArray<FTopLevelAssetPath>& InClassPaths);
+
+	bool IsRegisteredClass(const FTopLevelAssetPath& InClassPath) const;
+
+	const TParentClassMap& GetParentClassMap() const { check(IsInitialized()); return ParentClassMap; }
+
+private:
+	void RegisterClassDescriptor(FWorldPartitionActorDesc* InClassDesc);
+	void UnregisterClassDescriptor(FWorldPartitionActorDesc* InClassDesc);
+
+	void RegisterClassDescriptorFromAssetData(const FAssetData& InAssetData);
+	void RegisterClassDescriptorFromActorClass(const UClass* InActorClass);
+
+	friend class FActorDescArchive;
+	const FWorldPartitionActorDesc* GetClassDescDefault(const FTopLevelAssetPath& InClassPath) const;
+	const FWorldPartitionActorDesc* GetClassDescDefaultForActor(const FTopLevelAssetPath& InClassPath) const;
+	const FWorldPartitionActorDesc* GetClassDescDefaultForClass(const FTopLevelAssetPath& InClassPath) const;
+
+	void OnAssetCreated(UObject* InAssetCreated);
+	void OnAssetLoaded(UObject* InAssetLoaded);
+	void OnObjectPreSave(UObject* InObject, FObjectPreSaveContext InSaveContext);
+	void OnObjectPropertyChanged(UObject* InObject, FPropertyChangedEvent& InPropertyChangedEvent);
+
+	void OnPluginLoadingPhaseComplete(ELoadingPhase::Type InLoadingPhase, bool bInPhaseSuccessful);
+
+	void OnAssetRemoved(const FAssetData& InAssetData);
+	void OnAssetRenamed(const FAssetData& InAssetData, const FString& InOldObjectPath);
+
+	void RegisterClass(UClass* Class);
+	void RegisterClasses();
+
+	void UpdateClassDescriptor(UObject* InObject);
+
+	void ValidateInternalState();
+
+	FTopLevelAssetPath RedirectClassPath(const FTopLevelAssetPath& InClassPath) const;
+
+	TNameClassDescMap ClassByPath;
+	TParentClassMap ParentClassMap;
+	TRedirectClassMap RedirectClassMap;
+};
+#endif
\ No newline at end of file
