Actor Class Descriptors: revert debug code submitted by error.

diff --git a/Engine/Source/Runtime/Engine/Private/WorldPartition/ActorDescContainer.cpp b/Engine/Source/Runtime/Engine/Private/WorldPartition/ActorDescContainer.cpp
--- a/Engine/Source/Runtime/Engine/Private/WorldPartition/ActorDescContainer.cpp
+++ b/Engine/Source/Runtime/Engine/Private/WorldPartition/ActorDescContainer.cpp
@@ -66,7 +66,7 @@ void UActorDescContainer::Initialize(const FInitializeParams& InitParams)
 	}
 
 	FWorldPartitionClassDescRegistry& ClassDescRegistry = FWorldPartitionClassDescRegistry::Get();
-	//{
+	{
 		TRACE_CPUPROFILER_EVENT_SCOPE(GatherDescriptorsClass);
 		
 		TSet<FTopLevelAssetPath> ClassPaths;
@@ -76,7 +76,7 @@ void UActorDescContainer::Initialize(const FInitializeParams& InitParams)
 		}
 
 		ClassDescRegistry.PrefetchClassDescs(ClassPaths.Array());
-	//}
+	}
 
 	{
 		TRACE_CPUPROFILER_EVENT_SCOPE(RegisterDescriptors);
