Allow level streaming profiler to be used in shipping.
Add LevelStreaming.Profiling.Enabled, LevelStreaming.Profiling.Enabled.Shipping and LevelStreaming.Profiling.StartAutomatically cvars to control behavior.
Automatically emit report when subsystem is deinitialized.
Do not create base-class subsystem if a subclass exists.
Call UpdateTrackingData when level is first added with its initial state.

diff --git a/Engine/Source/Runtime/Engine/Private/ProfilingDebugging/LevelStreamingProfilingSubsystem.cpp b/Engine/Source/Runtime/Engine/Private/ProfilingDebugging/LevelStreamingProfilingSubsystem.cpp
--- a/Engine/Source/Runtime/Engine/Private/ProfilingDebugging/LevelStreamingProfilingSubsystem.cpp
+++ b/Engine/Source/Runtime/Engine/Private/ProfilingDebugging/LevelStreamingProfilingSubsystem.cpp
@@ -20,6 +20,26 @@
 DEFINE_LOG_CATEGORY_STATIC(LogLevelStreamingProfiling, Log, All);
 CSV_DEFINE_CATEGORY_MODULE(ENGINE_API, LevelStreaming, true);
 
+static bool GLevelStreamingProfilingEnabled = !UE_BUILD_SHIPPING;
+static FAutoConsoleVariableRef CVarLevelStreamingProfilingEnabled(
+#if UE_BUILD_SHIPPING
+	TEXT("LevelStreaming.Profiling.Enabled.Shipping"),
+#else
+	TEXT("LevelStreaming.Profiling.Enabled"),
+#endif
+	GLevelStreamingProfilingEnabled,
+	TEXT("Whether to enable the LevelStreamingProfilingSubsystem automatically."),
+	ECVF_Default
+);
+
+static bool GStartProfilingAutomatically = false;
+static FAutoConsoleVariableRef CVarStartProfilingAutomatically(
+	TEXT("LevelStreaming.Profiling.StartAutomatically"),
+	GStartProfilingAutomatically,
+	TEXT("Whether to start recording level streaminge events as soon as the subsystem is created."),
+	ECVF_Default
+);
+
 static float GLateStreamingDistanceSquared = 0.0f;
 static FAutoConsoleVariableRef CVarLateStreamingDistanceSquared(
 	TEXT("LevelStreaming.Profiling.LateStreamingDistanceSquared"),
@@ -48,16 +68,40 @@ bool ULevelStreamingProfilingSubsystem::DoesSupportWorldType(const EWorldType::T
 
 bool ULevelStreamingProfilingSubsystem::ShouldCreateSubsystem(UObject* Outer) const
 {
-	return Super::ShouldCreateSubsystem(Outer);
+	// If there are any loaded child classes of this subsystem and this is not the CDO of that child class, skip creation of this subsystem.
+	TArray<UClass*> DerivedClasses; 
+	GetDerivedClasses(ULevelStreamingProfilingSubsystem::StaticClass(), DerivedClasses, true);
+	if (DerivedClasses.Num() > 0)
+	{
+		if (!Algo::AnyOf(DerivedClasses, [Class=GetClass()](UClass* DerivedClass) { return Class == DerivedClass; }))
+		{
+			return false;
+		}
+	}
+	return GLevelStreamingProfilingEnabled && Super::ShouldCreateSubsystem(Outer);
 }
 
-#if !UE_BUILD_SHIPPING
 void ULevelStreamingProfilingSubsystem::PostInitialize()
 {
 	Handle_OnLevelStreamingTargetStateChanged = FLevelStreamingDelegates::OnLevelStreamingTargetStateChanged.AddUObject(this, &ULevelStreamingProfilingSubsystem::OnLevelStreamingTargetStateChanged);
 	Handle_OnLevelStreamingStateChanged = FLevelStreamingDelegates::OnLevelStreamingStateChanged.AddUObject(this, &ULevelStreamingProfilingSubsystem::OnLevelStreamingStateChanged);
 	Handle_OnLevelBeginAddToWorld = FLevelStreamingDelegates::OnLevelBeginMakingVisible.AddUObject(this, &ULevelStreamingProfilingSubsystem::OnLevelStartedAddToWorld);
 	Handle_OnLevelBeginRemoveFromWorld = FLevelStreamingDelegates::OnLevelBeginMakingInvisible.AddUObject(this, &ULevelStreamingProfilingSubsystem::OnLevelStartedRemoveFromWorld);
+	
+	if (GStartProfilingAutomatically)
+	{
+		StartTracking();
+	}
+}
+
+void ULevelStreamingProfilingSubsystem::Deinitialize() 
+{
+	if (IsTracking())
+	{
+		StopTrackingAndReport();
+	}
+
+	Super::Deinitialize();
 }
 
 const TCHAR* ULevelStreamingProfilingSubsystem::EnumToString(ULevelStreamingProfilingSubsystem::ELevelState State)
@@ -89,7 +133,7 @@ TConstArrayView<ULevelStreamingProfilingSubsystem::FLevelStats> ULevelStreamingP
 TUniquePtr<ULevelStreamingProfilingSubsystem::FActiveLevel> ULevelStreamingProfilingSubsystem::MakeActiveLevel(const ULevelStreaming* StreamingLevel, ELevelState InitialState, ULevel* LoadedLevel)
 {
 	TUniquePtr<FActiveLevel> Level = MakeUnique<FActiveLevel>(LevelStats.AddDefaulted());
-	Level->State = InitialState;
+	Level->State = ELevelState::None;
 	Level->StateStartTime = FPlatformTime::Seconds();
 	if (LoadedLevel)
 	{
@@ -124,8 +168,10 @@ TUniquePtr<ULevelStreamingProfilingSubsystem::FActiveLevel> ULevelStreamingProfi
 		{
 			LevelStats[Level->StatsIndex].CellBounds = Cell->GetCellBounds();
 			LevelStats[Level->StatsIndex].ContentBounds = Cell->GetContentBounds();
+			LevelStats[Level->StatsIndex].bIsHLOD = Cell->GetIsHLOD();
 		}
 	}
+	UpdateLevelState(*Level, StreamingLevel, InitialState, Level->StateStartTime);
 	return MoveTemp(Level);
 }
 
@@ -215,7 +261,7 @@ void ULevelStreamingProfilingSubsystem::StopTrackingAndReport()
 				}
 			};
 
-			Builder << "NameOnDisk\tNameInMemory\tTimeQueuedForLoadingMS\tTimeLoadingMS\tTimeQueuedForAddToWorldMS\tTimeAddingToWorldMS\tTimeInWorldS\tTimeQueuedForRemoveFromWorldMS\tTimeRemovingFromWorldMS\tStreamInDistance_Cell\tStreamInDistance_Content\tStreamInSourceLocation\tCellBounds\tContentBounds";
+			Builder << "NameOnDisk\tNameInMemory\tTimeQueuedForLoadingMS\tTimeLoadingMS\tTimeQueuedForAddToWorldMS\tTimeAddingToWorldMS\tTimeInWorldS\tTimeQueuedForRemoveFromWorldMS\tTimeRemovingFromWorldMS\tStreamInDistance_Cell\tStreamInDistance_Content\tStreamInSourceLocation\tCellBounds\tContentBounds\tbIsHLOD";
 			AugmentReportHeader(Builder);
 			Write();
 			for (int32 i=0; i < LevelStats.Num(); ++i)
@@ -236,6 +282,7 @@ void ULevelStreamingProfilingSubsystem::StopTrackingAndReport()
 				AddVector(Stats.FinalStreamInLocation);
 				AddBounds(Stats.CellBounds);
 				AddBounds(Stats.ContentBounds);
+				Builder << '\t' << (int32)Stats.bIsHLOD;
 				AugmentReportRow(Builder, i);
 				Write();
 			}
@@ -815,4 +862,3 @@ void ULevelStreamingProfilingSubsystem::OnLevelFinishedRemoveFromWorld(UWorld* W
 	LevelStats[Level->StatsIndex].TimeRemovingFromWorld = Time - Level->StateStartTime; 
 	UpdateLevelState(*Level, StreamingLevel, ELevelState::RemovedFromWorld, Time);
 }
-#endif // !UE_BUILD_SHIPPING
diff --git a/Engine/Source/Runtime/Engine/Public/ProfilingDebugging/LevelStreamingProfilingSubsystem.h b/Engine/Source/Runtime/Engine/Public/ProfilingDebugging/LevelStreamingProfilingSubsystem.h
--- a/Engine/Source/Runtime/Engine/Public/ProfilingDebugging/LevelStreamingProfilingSubsystem.h
+++ b/Engine/Source/Runtime/Engine/Public/ProfilingDebugging/LevelStreamingProfilingSubsystem.h
@@ -34,14 +34,12 @@ class ULevelStreamingProfilingSubsystem : public UWorldSubsystem
 public:
 	// WorldSubsystem interface
 	ENGINE_API virtual bool ShouldCreateSubsystem(UObject* Outer) const override;
-#if !UE_BUILD_SHIPPING
 	ENGINE_API virtual void PostInitialize() override;
-#endif
+	ENGINE_API virtual void Deinitialize() override;
 
 	ENGINE_API ULevelStreamingProfilingSubsystem(const FObjectInitializer&);
 	ENGINE_API ~ULevelStreamingProfilingSubsystem();
 
-#if !UE_BUILD_SHIPPING
 	// Begin recording timings for level streaming events.
 	ENGINE_API void StartTracking();
 	// Top recording timings for level streaming events and output a .tsv (tab separated values) file to the Profiling directory.
@@ -51,7 +49,6 @@ public:
 
 	// Gives child classes an opportunity to clean up after a report is produced.
 	ENGINE_API virtual void PostReport() { }
-#endif
 
 	// Access to tuning values set by cvars for other systems
 	/* Returns the squared distance (e.g. from world partition cell bounds) at which a level is considered to have streamed in too late. */
@@ -61,7 +58,6 @@ protected:
 	// WorldSubsystem interface
 	ENGINE_API virtual bool DoesSupportWorldType(const EWorldType::Type WorldType) const;
 
-#if !UE_BUILD_SHIPPING
 	enum class ELevelState;
 	struct FLevelStats;
 
@@ -99,7 +95,7 @@ protected:
 	ENGINE_API static const TCHAR* EnumToString(ULevelStreamingProfilingSubsystem::ELevelState State);
 
 	ENGINE_API TConstArrayView<FLevelStats> GetLevelStats() const;
-
+	
 	enum class ELevelState
 	{
 		None,
@@ -155,6 +151,8 @@ protected:
 		// Location of streaming source (e.g. player) when level was fully streamed in 
 		TOptional<FVector> FinalStreamInLocation;
 
+		// Is this level a hierarchical LOD representation of more detailed content
+		bool bIsHLOD = false;
 		// If a level is tracked but never starts loading before being removed, we don't include it in the results 
 		bool bValid = false;
 	};
@@ -198,7 +196,6 @@ private:
 
 	// Possibly-executing task dependend on our recorded stats. Cannot start recording again until this task is complete.
 	UE::Tasks::TTask<void> ReportWritingTask;
-#endif // !UE_BUILD_SHIPPING 
 };
 
 
