#UE Silence cook warning, for now

diff --git a/Engine/Source/Runtime/Engine/Private/StaticMeshBuild.cpp b/Engine/Source/Runtime/Engine/Private/StaticMeshBuild.cpp
--- a/Engine/Source/Runtime/Engine/Private/StaticMeshBuild.cpp
+++ b/Engine/Source/Runtime/Engine/Private/StaticMeshBuild.cpp
@@ -631,7 +631,7 @@ void UStaticMesh::FinishBuildInternal(const TArray<UStaticMeshComponent*>& InAff
 		// Calculate extended bounds
 		CalculateExtendedBounds();
 	}
-	else
+	else if (!IsNaniteLandscape())
 	{
 		// We don't care about minor differences but might want to highlight if big differences are noted.
 		const FBoxSphereBounds RenderBounds = GetRenderData()->Bounds;
