Significant optimization to CheckMaterialUsage_Concurrent  - keep track of cache misses on the render/worker threads, and only send a single check packet back to the game thread when a miss for a material occurs (instead of each request instance spamming packets). As soon as a cache hit occurs, clear the cache miss tracking.

diff --git a/Engine/Source/Runtime/Engine/Classes/Materials/Material.h b/Engine/Source/Runtime/Engine/Classes/Materials/Material.h
--- a/Engine/Source/Runtime/Engine/Classes/Materials/Material.h
+++ b/Engine/Source/Runtime/Engine/Classes/Materials/Material.h
@@ -997,6 +997,8 @@ public:
 	UPROPERTY(transient, duplicatetransient)
 	uint32 UsageFlagWarnings;
 
+	mutable std::atomic<uint32> UsageFlagCacheMiss{0u};
+
 	/** This is the refraction depth bias, larger values offset distortion to prevent closer objects from rendering into the distorted surface at acute viewing angles but increases the disconnect between surface and where the refraction starts. */
 	UPROPERTY(EditAnywhere, Category=Refraction)
 	float RefractionDepthBias;
diff --git a/Engine/Source/Runtime/Engine/Private/Materials/Material.cpp b/Engine/Source/Runtime/Engine/Private/Materials/Material.cpp
--- a/Engine/Source/Runtime/Engine/Private/Materials/Material.cpp
+++ b/Engine/Source/Runtime/Engine/Private/Materials/Material.cpp
@@ -1672,6 +1672,8 @@ bool UMaterial::CheckMaterialUsage(EMaterialUsage Usage)
 
 bool UMaterial::CheckMaterialUsage_Concurrent(EMaterialUsage Usage) const 
 {
+	const uint32 UsageFlagBit = (1u << (uint32)Usage);
+
 	bool bUsageSetSuccessfully = false;
 	if (NeedsSetMaterialUsage_Concurrent(bUsageSetSuccessfully, Usage))
 	{
@@ -1681,37 +1683,45 @@ bool UMaterial::CheckMaterialUsage_Concurrent(EMaterialUsage Usage) const
 		}	
 		else
 		{
-			struct FCallSMU
+			const uint32 PrevCacheMissMask = std::atomic_fetch_or(&UsageFlagCacheMiss, UsageFlagBit);
+			if ((PrevCacheMissMask & UsageFlagBit) == 0u)
 			{
-				UMaterial* Material;
-				EMaterialUsage Usage;
-
-				FCallSMU(UMaterial* InMaterial, EMaterialUsage InUsage)
-					: Material(InMaterial)
-					, Usage(InUsage)
+				struct FCallSMU
 				{
-				}
+					UMaterial* Material;
+					EMaterialUsage Usage;
 
-				void Task()
-				{
-					Material->CheckMaterialUsage(Usage);
-				}
-			};
-			UE_LOG(LogMaterial, Log, TEXT("Had to pass SMU back to game thread. Please ensure correct material usage flags."));
+					FCallSMU(UMaterial* InMaterial, EMaterialUsage InUsage)
+						: Material(InMaterial)
+						, Usage(InUsage)
+					{
+					}
+
+					void Task()
+					{
+						Material->CheckMaterialUsage(Usage);
+					}
+				};
+				UE_LOG(LogMaterial, Log, TEXT("Had to pass SMU back to game thread. Please ensure correct material usage flags."));
 
-			TSharedRef<FCallSMU, ESPMode::ThreadSafe> CallSMU = MakeShareable(new FCallSMU(const_cast<UMaterial*>(this), Usage));
-			bUsageSetSuccessfully = false;
+				TSharedRef<FCallSMU, ESPMode::ThreadSafe> CallSMU = MakeShareable(new FCallSMU(const_cast<UMaterial*>(this), Usage));
 
-			DECLARE_CYCLE_STAT(TEXT("FSimpleDelegateGraphTask.CheckMaterialUsage"),
-				STAT_FSimpleDelegateGraphTask_CheckMaterialUsage,
-				STATGROUP_TaskGraphTasks);
+				DECLARE_CYCLE_STAT(TEXT("FSimpleDelegateGraphTask.CheckMaterialUsage"),
+					STAT_FSimpleDelegateGraphTask_CheckMaterialUsage,
+					STATGROUP_TaskGraphTasks);
 
-			FSimpleDelegateGraphTask::CreateAndDispatchWhenReady(
-				FSimpleDelegateGraphTask::FDelegate::CreateThreadSafeSP(CallSMU, &FCallSMU::Task),
-				GET_STATID(STAT_FSimpleDelegateGraphTask_CheckMaterialUsage), NULL, ENamedThreads::GameThread_Local
-			);
+				FSimpleDelegateGraphTask::CreateAndDispatchWhenReady(
+					FSimpleDelegateGraphTask::FDelegate::CreateThreadSafeSP(CallSMU, &FCallSMU::Task),
+					GET_STATID(STAT_FSimpleDelegateGraphTask_CheckMaterialUsage), NULL, ENamedThreads::GameThread_Local
+				);
+			}
 		}
 	}
+	else if ((UsageFlagCacheMiss.load(std::memory_order_relaxed) & UsageFlagBit) != 0u)
+	{
+		std::atomic_fetch_and(&UsageFlagCacheMiss, ~UsageFlagBit);
+	}
+
 	return bUsageSetSuccessfully;
 }
 
