Fix out of bounds uv coord array in FMeshBuildVertexData::Empty()

diff --git a/Engine/Source/Runtime/Engine/Public/Components.h b/Engine/Source/Runtime/Engine/Public/Components.h
--- a/Engine/Source/Runtime/Engine/Public/Components.h
+++ b/Engine/Source/Runtime/Engine/Public/Components.h
@@ -106,7 +106,7 @@ struct FMeshBuildVertexData
 		TangentX.Empty(Slack);
 		TangentY.Empty(Slack);
 		TangentZ.Empty(Slack);
-		UVs.Empty(NumTexCoords);
+		UVs.SetNum(NumTexCoords);
 		for (int32 TexCoord = 0; TexCoord < NumTexCoords; ++TexCoord)
 		{
 			UVs[TexCoord].Empty(Slack);
