Added new type of ISM component that doesn't store direct references to assets, but rather point to a template class / component to perform fixup on load

diff --git a/Engine/Plugins/Editor/WorldPartitionHLODUtilities/Source/Private/WorldPartition/HLOD/Utilities/WorldPartitionHLODUtilities.cpp b/Engine/Plugins/Editor/WorldPartitionHLODUtilities/Source/Private/WorldPartition/HLOD/Utilities/WorldPartitionHLODUtilities.cpp
--- a/Engine/Plugins/Editor/WorldPartitionHLODUtilities/Source/Private/WorldPartition/HLOD/Utilities/WorldPartitionHLODUtilities.cpp
+++ b/Engine/Plugins/Editor/WorldPartitionHLODUtilities/Source/Private/WorldPartition/HLOD/Utilities/WorldPartitionHLODUtilities.cpp
@@ -125,6 +125,15 @@ static uint32 ComputeHLODHash(AWorldPartitionHLOD* InHLODActor, const TArray<AAc
 	UE_LOG(LogHLODBuilder, VeryVerbose, TEXT(" - HLOD Min Visible Distance (%.02f) = %x"), InHLODActor->GetMinVisibleDistance(), HLODMinVisibleDistanceHash);
 	Ar << HLODMinVisibleDistanceHash;
 
+	// ISM Component Class
+	TSubclassOf<UInstancedStaticMeshComponent> HLODISMComponentClass = UHLODBuilder::GetInstancedStaticMeshComponentClass();
+	if (HLODISMComponentClass != UInstancedStaticMeshComponent::StaticClass())
+	{
+		uint32 HLODISMComponentClassHash = GetTypeHash(HLODISMComponentClass);
+		UE_LOG(LogHLODBuilder, VeryVerbose, TEXT(" - HLOD ISM Component Class (%s) = %x"), *HLODISMComponentClass->GetName(), HLODISMComponentClassHash);
+		Ar << HLODISMComponentClassHash;
+	}
+
 	// Append all components CRCs
 	uint32 HLODComponentsHash = UHLODBuilder::ComputeHLODHash(InActors);
 	UE_LOG(LogHLODBuilder, VeryVerbose, TEXT(" - HLOD Source Components = %x"), HLODComponentsHash);
diff --git a/Engine/Source/Runtime/Engine/Private/WorldPartition/HLOD/HLODBuilder.cpp b/Engine/Source/Runtime/Engine/Private/WorldPartition/HLOD/HLODBuilder.cpp
--- a/Engine/Source/Runtime/Engine/Private/WorldPartition/HLOD/HLODBuilder.cpp
+++ b/Engine/Source/Runtime/Engine/Private/WorldPartition/HLOD/HLODBuilder.cpp
@@ -11,6 +11,7 @@
 #include "ISMPartition/ISMComponentDescriptor.h"
 #include "Materials/MaterialInterface.h"
 #include "UObject/Package.h"
+#include "WorldPartition/HLOD/HLODTemplatedInstancedStaticMeshComponent.h"
 
 #include UE_INLINE_GENERATED_CPP_BY_NAME(HLODBuilder)
 
@@ -20,6 +21,9 @@ DEFINE_LOG_CATEGORY(LogHLODBuilder);
 
 UHLODBuilder::UHLODBuilder(const FObjectInitializer& ObjectInitializer)
 	: Super(ObjectInitializer)
+#if WITH_EDITORONLY_DATA
+	, HLODInstancedStaticMeshComponentClass(UInstancedStaticMeshComponent::StaticClass())
+#endif
 {
 }
 
@@ -151,12 +155,23 @@ namespace
 	// Instance batcher class based on FISMComponentDescriptor
 	struct FCustomISMComponentDescriptor : public FISMComponentDescriptor
 	{
+		TSubclassOf<AActor> TemplateActorClass;
+		FName				TemplateComponentName;
+
 		FCustomISMComponentDescriptor(UStaticMeshComponent* SMC)
 		{
 			InitFrom(SMC, false);
 
 			// We'll always want to spawn ISMC, even if our source components are all SMC
-			ComponentClass = UInstancedStaticMeshComponent::StaticClass();
+			ComponentClass = UHLODBuilder::GetInstancedStaticMeshComponentClass();
+
+			// Extra work for templated ISMC
+			// This code should be reorganized
+			if (ComponentClass->IsChildOf<UHLODTemplatedInstancedStaticMeshComponent>())
+			{
+				TemplateActorClass = SMC->GetOwner()->GetClass();
+				TemplateComponentName = SMC->GetFName();
+			}			
 
 			// Stationnary can be considered as static for the purpose of HLODs
 			if (Mobility == EComponentMobility::Stationary)
@@ -165,6 +180,14 @@ namespace
 			}
 
 			ComputeHash();
+
+			// Extra work for templated ISMC
+			// This code should be reorganized
+			if (ComponentClass->IsChildOf<UHLODTemplatedInstancedStaticMeshComponent>())
+			{
+				Hash = HashCombine(Hash, GetTypeHash(TemplateActorClass));
+				Hash = HashCombine(Hash, GetTypeHash(TemplateComponentName));
+			}
 		}
 	};
 }
@@ -176,7 +199,7 @@ TArray<UActorComponent*> UHLODBuilder::BatchInstances(const TArray<UActorCompone
 	TArray<UStaticMeshComponent*> SourceStaticMeshComponents = FilterComponents<UStaticMeshComponent>(InSourceComponents);
 
 	// Prepare instance batches
-	TMap<FISMComponentDescriptor, FISMComponentBatcher> InstancesData;
+	TMap<FCustomISMComponentDescriptor, FISMComponentBatcher> InstancesData;
 	for (UStaticMeshComponent* SMC : SourceStaticMeshComponents)
 	{
 		FCustomISMComponentDescriptor ISMComponentDescriptor(SMC);
@@ -188,11 +211,19 @@ TArray<UActorComponent*> UHLODBuilder::BatchInstances(const TArray<UActorCompone
 	TArray<UActorComponent*> HLODComponents;
 	for (auto& Entry : InstancesData)
 	{
-		const FISMComponentDescriptor& ISMComponentDescriptor = Entry.Key;
+		const FCustomISMComponentDescriptor& ISMComponentDescriptor = Entry.Key;
 		const FISMComponentBatcher& ISMComponentBatcher = Entry.Value;
 
 		UInstancedStaticMeshComponent* ISMComponent = ISMComponentDescriptor.CreateComponent(GetTransientPackage());
 		ISMComponentBatcher.InitComponent(ISMComponent);
+		
+		// Extra work for templated ISMC
+		// This code should be reorganized
+		if (UHLODTemplatedInstancedStaticMeshComponent* TemplatedISMC = Cast<UHLODTemplatedInstancedStaticMeshComponent>(ISMComponent))
+		{
+			TemplatedISMC->SetTemplateActorClass(ISMComponentDescriptor.TemplateActorClass);
+			TemplatedISMC->SetTemplateComponentName(ISMComponentDescriptor.TemplateComponentName);
+		}
 
 		ISMComponent->SetForcedLodModel(ISMComponent->GetStaticMesh()->GetNumLODs());
 
diff --git a/Engine/Source/Runtime/Engine/Private/WorldPartition/HLOD/HLODTemplatedInstancedStaticMeshComponent.cpp b/Engine/Source/Runtime/Engine/Private/WorldPartition/HLOD/HLODTemplatedInstancedStaticMeshComponent.cpp
new file mode 100644
--- /dev/null
+++ b/Engine/Source/Runtime/Engine/Private/WorldPartition/HLOD/HLODTemplatedInstancedStaticMeshComponent.cpp
@@ -0,0 +1,173 @@
+// Copyright Epic Games, Inc. All Rights Reserved.
+
+#include "WorldPartition/HLOD/HLODTemplatedInstancedStaticMeshComponent.h"
+
+#include "GameFramework/Actor.h"
+#include "Engine/BlueprintGeneratedClass.h"
+#include "Engine/InheritableComponentHandler.h"
+#include "Engine/SCS_Node.h"
+#include "Engine/StaticMesh.h"
+#include "UObject/ObjectSaveContext.h"
+#include "UObject/ObjectSaveOverride.h"
+
+UHLODTemplatedInstancedStaticMeshComponent::UHLODTemplatedInstancedStaticMeshComponent(const FObjectInitializer& ObjectInitializer)
+	: Super(ObjectInitializer)
+{
+}
+
+void UHLODTemplatedInstancedStaticMeshComponent::PreSave(FObjectPreSaveContext ObjectSaveContext)
+{
+#if WITH_EDITOR
+	// Ensure transient overrides are applied before the base implementation of PreSave
+	if (!HasAnyFlags(RF_ArchetypeObject | RF_ClassDefaultObject))
+	{
+		auto GetTransientPropertyOverride = [](FName PropertyName) -> FPropertySaveOverride
+		{
+			FProperty* OverrideProperty = FindFProperty<FProperty>(UStaticMeshComponent::StaticClass(), PropertyName);
+			check(OverrideProperty);
+
+			FPropertySaveOverride PropOverride;
+			PropOverride.PropertyPath = FFieldPath(OverrideProperty);
+			PropOverride.bMarkTransient = true;
+			return PropOverride;
+		};
+
+		FObjectSaveOverride ObjectSaveOverride;
+		ObjectSaveOverride.PropOverrides.Add(GetTransientPropertyOverride(GetMemberNameChecked_StaticMesh()));
+		ObjectSaveOverride.PropOverrides.Add(GetTransientPropertyOverride(GET_MEMBER_NAME_CHECKED(UMeshComponent, OverrideMaterials)));
+		ObjectSaveOverride.PropOverrides.Add(GetTransientPropertyOverride(GET_MEMBER_NAME_CHECKED(UMeshComponent, OverlayMaterial)));
+
+		ObjectSaveContext.AddSaveOverride(this, ObjectSaveOverride);
+	}
+#endif
+
+	Super::PreSave(ObjectSaveContext);
+}
+
+void UHLODTemplatedInstancedStaticMeshComponent::PostLoad()
+{
+	if (!HasAnyFlags(RF_ArchetypeObject | RF_ClassDefaultObject))
+	{
+		check(GetStaticMesh() == nullptr);
+		check(!OverrideMaterials.ContainsByPredicate([](const UMaterialInterface* MaterialInterface) { return MaterialInterface != nullptr; }));
+		check(OverlayMaterial == nullptr);
+
+		RestoreAssetsFromActorTemplate();
+	}
+
+	Super::PostLoad();
+}
+
+void UHLODTemplatedInstancedStaticMeshComponent::SetTemplateActorClass(const TSubclassOf<AActor>& InTemplateActorClass)
+{
+	TemplateActorClass = InTemplateActorClass;
+}
+
+void UHLODTemplatedInstancedStaticMeshComponent::SetTemplateComponentName(const FName& InTemplatePropertyName)
+{
+	TemplateComponentName = InTemplatePropertyName;
+}
+
+// Note: this is greatly inspired by a similar function in AIHelpers.h/.cpp
+template <class TComp>
+static const TComp* GetActorClassDefaultComponentByName(const TSubclassOf<AActor>& ActorClass, FName InComponentName)
+{
+	auto IsMatchingComponent = [InComponentName](const UActorComponent* InComponent, bool bStripGenSuffix = false)
+	{
+		if (const TComp* TypedComponent = Cast<TComp>(InComponent))
+		{
+			if (bStripGenSuffix)
+			{
+				FString StrippedName = InComponent->GetName();
+				StrippedName.RemoveFromEnd(TEXT("_GEN_VARIABLE"));
+
+				if (StrippedName == InComponentName.ToString())
+				{
+					return true;
+				}
+			}
+			else if (TypedComponent->GetFName() == InComponentName)
+			{
+				return true;
+			}
+		}
+
+		return false;
+	};
+
+	if (!ActorClass.Get())
+	{
+		return nullptr;
+	}
+
+	// Test the components defined on the native class.
+	const AActor* CDO = ActorClass->GetDefaultObject<AActor>();
+	check(CDO);
+
+	for (const UActorComponent* Component : CDO->GetComponents())
+	{
+		if (IsMatchingComponent(Component))
+		{
+			return CastChecked<TComp>(Component);
+		}
+	}
+
+	// Try to get the components off the BP class.
+	UBlueprintGeneratedClass* BPClass = Cast<UBlueprintGeneratedClass>(*ActorClass);
+	if (BPClass)
+	{
+		// A BlueprintGeneratedClass has a USimpleConstructionScript member. This member has an array of RootNodes
+		// which contains the SCSNode for the root SceneComponent and non-SceneComponents. For the SceneComponent
+		// hierarchy, each SCSNode knows its children SCSNodes. Each SCSNode stores the component template that will
+		// be created when the Actor is spawned.
+		//
+		// WARNING: This may change in future engine versions!
+
+		// Check added components.
+		USimpleConstructionScript* ConstructionScript = BPClass->SimpleConstructionScript;
+		if (ConstructionScript)
+		{
+			for (const USCS_Node* Node : ConstructionScript->GetAllNodes())
+			{
+				if (IsMatchingComponent(Node->ComponentTemplate, true))
+				{
+					return CastChecked<TComp>(Node->ComponentTemplate);
+				}
+			}
+		}
+		// Check modified inherited components.
+		UInheritableComponentHandler* InheritableComponentHandler = BPClass->InheritableComponentHandler;
+		if (InheritableComponentHandler)
+		{
+			for (TArray<FComponentOverrideRecord>::TIterator It = InheritableComponentHandler->CreateRecordIterator(); It; ++It)
+			{
+				if (IsMatchingComponent(It->ComponentTemplate, true))
+				{
+					return CastChecked<TComp>(It->ComponentTemplate);
+				}
+			}
+		}
+	}
+
+	return nullptr;
+}
+
+void UHLODTemplatedInstancedStaticMeshComponent::RestoreAssetsFromActorTemplate()
+{
+	const UStaticMeshComponent* TemplateSMC = GetActorClassDefaultComponentByName<UStaticMeshComponent>(TemplateActorClass, TemplateComponentName);
+	if (TemplateSMC)
+	{
+		// StaticMesh
+		if (UStaticMesh* TemplateStaticMesh = TemplateSMC->GetStaticMesh())
+		{
+			SetStaticMesh(TemplateStaticMesh);
+			SetForcedLodModel(TemplateStaticMesh->GetNumLODs());
+		}
+
+		// OverrideMaterials
+		OverrideMaterials = TemplateSMC->OverrideMaterials;
+
+		// OverlayMaterial
+		OverlayMaterial = TemplateSMC->OverlayMaterial;
+	}
+}
diff --git a/Engine/Source/Runtime/Engine/Public/WorldPartition/HLOD/HLODBuilder.h b/Engine/Source/Runtime/Engine/Public/WorldPartition/HLOD/HLODBuilder.h
--- a/Engine/Source/Runtime/Engine/Public/WorldPartition/HLOD/HLODBuilder.h
+++ b/Engine/Source/Runtime/Engine/Public/WorldPartition/HLOD/HLODBuilder.h
@@ -10,6 +10,7 @@
 
 class AActor;
 class UActorComponent;
+class UInstancedStaticMeshComponent;
 
 ENGINE_API DECLARE_LOG_CATEGORY_EXTERN(LogHLODBuilder, Log, All);
 
@@ -61,7 +62,7 @@ struct FHLODBuildContext
  * Base class for all HLOD builders
  * This class takes as input a group of components, and should return component(s) that will be included in the HLOD actor.
  */
-UCLASS(Abstract)
+UCLASS(Abstract, Config = Editor)
 class ENGINE_API UHLODBuilder : public UObject
 {
 	 GENERATED_UCLASS_BODY()
@@ -114,6 +115,14 @@ public:
 	 */
 	static uint32 ComputeHLODHash(const TArray<AActor*>& InSourceActors);
 
+	/** 
+	 * Get the InstancedStaticMeshComponent subclass that should be used when creating instanced HLODs.
+	 */
+	static TSubclassOf<UInstancedStaticMeshComponent> GetInstancedStaticMeshComponentClass()
+	{
+		return StaticClass()->GetDefaultObject<UHLODBuilder>()->HLODInstancedStaticMeshComponentClass;
+	}
+
 protected:
 	virtual bool ShouldIgnoreBatchingPolicy() const { return false; }
 
@@ -131,13 +140,15 @@ protected:
 		}
 		return FilteredComponents;
 	}
-
 #endif // WITH_EDITOR
 
 #if WITH_EDITORONLY_DATA
 protected:
 	UPROPERTY()
 	TObjectPtr<const UHLODBuilderSettings> HLODBuilderSettings;
+
+	UPROPERTY(Config)
+	TSubclassOf<UInstancedStaticMeshComponent> HLODInstancedStaticMeshComponentClass;
 #endif
 };
 
diff --git a/Engine/Source/Runtime/Engine/Public/WorldPartition/HLOD/HLODTemplatedInstancedStaticMeshComponent.h b/Engine/Source/Runtime/Engine/Public/WorldPartition/HLOD/HLODTemplatedInstancedStaticMeshComponent.h
new file mode 100644
--- /dev/null
+++ b/Engine/Source/Runtime/Engine/Public/WorldPartition/HLOD/HLODTemplatedInstancedStaticMeshComponent.h
@@ -0,0 +1,35 @@
+// Copyright Epic Games, Inc. All Rights Reserved.
+
+#pragma once
+
+#include "CoreMinimal.h"
+#include "Components/InstancedStaticMeshComponent.h"
+#include "Templates/SubclassOf.h"
+#include "HLODTemplatedInstancedStaticMeshComponent.generated.h"
+
+class AActor;
+
+UCLASS(Hidden, NotPlaceable)
+class ENGINE_API UHLODTemplatedInstancedStaticMeshComponent : public UInstancedStaticMeshComponent
+{
+	GENERATED_UCLASS_BODY()
+
+public: 
+	//~ Begin UObject Interface.
+	virtual void PreSave(FObjectPreSaveContext ObjectSaveContext) override;
+	virtual void PostLoad() override;
+	//~ End UObject Interface.
+	
+	void SetTemplateActorClass(const TSubclassOf<AActor>& InTemplateActorClass);
+	void SetTemplateComponentName(const FName& InTemplateComponentName);
+	
+private:
+	void RestoreAssetsFromActorTemplate();
+
+private:
+	UPROPERTY()
+	TSubclassOf<AActor> TemplateActorClass;
+
+	UPROPERTY()
+	FName TemplateComponentName;
+};
