Add an option to the render asset streamer that fully streams in LODs of all loaded meshes

diff --git a/Engine/Source/Runtime/Engine/Private/Streaming/AsyncTextureStreaming.cpp b/Engine/Source/Runtime/Engine/Private/Streaming/AsyncTextureStreaming.cpp
--- a/Engine/Source/Runtime/Engine/Private/Streaming/AsyncTextureStreaming.cpp
+++ b/Engine/Source/Runtime/Engine/Private/Streaming/AsyncTextureStreaming.cpp
@@ -131,7 +131,11 @@ void FAsyncRenderAssetStreamingData::UpdatePerfectWantedMips_Async(FStreamingRen
 	}
 	else
 #endif
-	if (Settings.bFullyLoadUsedTextures)
+	if (Settings.bFullyLoadMeshes && StreamingRenderAsset.IsMesh())
+	{
+		MaxSize_VisibleOnly = MaxSize = FLT_MAX;
+	}
+	else if (Settings.bFullyLoadUsedTextures)
 	{
 		if (StreamingRenderAsset.LastRenderTime < 300 || StreamingRenderAsset.bForceFullyLoad)
 		{
@@ -630,6 +634,7 @@ void FRenderAssetStreamingMipCalcTask::UpdateBudgetedMips_Async()
 
 			if (StreamingRenderAsset.BudgetMipBias > 0
 				&& (bResetMipBias
+					|| (Settings.bFullyLoadMeshes && StreamingRenderAsset.IsMesh())
 					|| FMath::Max<int32>(
 						StreamingRenderAsset.VisibleWantedMips,
 						StreamingRenderAsset.HiddenWantedMips + StreamingRenderAsset.NumMissingMips) < StreamingRenderAsset.MaxAllowedMips))
@@ -660,7 +665,10 @@ void FRenderAssetStreamingMipCalcTask::UpdateBudgetedMips_Async()
 			if (!StreamingRenderAsset.RenderAsset) continue;
 
 			// Ignore textures/meshes for which we are not allowed to reduce resolution.
-			if (!StreamingRenderAsset.IsMaxResolutionAffectedByGlobalBias()) continue;
+			if (!StreamingRenderAsset.IsMaxResolutionAffectedByGlobalBias() || (Settings.bFullyLoadMeshes && StreamingRenderAsset.IsMesh()))
+			{
+				continue;
+			}
 
 			// Ignore texture/mesh that can't drop any mips
 			const int32 MinAllowedMips = FMath::Max(StreamingRenderAsset.MinAllowedMips, StreamingRenderAsset.NumForcedMips);
diff --git a/Engine/Source/Runtime/Engine/Private/Streaming/StreamingManagerTexture.cpp b/Engine/Source/Runtime/Engine/Private/Streaming/StreamingManagerTexture.cpp
--- a/Engine/Source/Runtime/Engine/Private/Streaming/StreamingManagerTexture.cpp
+++ b/Engine/Source/Runtime/Engine/Private/Streaming/StreamingManagerTexture.cpp
@@ -343,7 +343,7 @@ bool FRenderAssetStreamingManager::StreamOutRenderAssetData( int64 RequiredMemor
 	{
 		FStreamingRenderAsset& StreamingRenderAsset = StreamingRenderAssets[Idx];
 		// Only texture for which we can drop mips.
-		if (StreamingRenderAsset.IsMaxResolutionAffectedByGlobalBias())
+		if (StreamingRenderAsset.IsMaxResolutionAffectedByGlobalBias() && (!Settings.bFullyLoadMeshes || !StreamingRenderAsset.IsMesh()))
 		{
 			PrioritizedRenderAssets.Add(Idx);
 		}
diff --git a/Engine/Source/Runtime/Engine/Private/Streaming/TextureStreamingHelpers.cpp b/Engine/Source/Runtime/Engine/Private/Streaming/TextureStreamingHelpers.cpp
--- a/Engine/Source/Runtime/Engine/Private/Streaming/TextureStreamingHelpers.cpp
+++ b/Engine/Source/Runtime/Engine/Private/Streaming/TextureStreamingHelpers.cpp
@@ -193,6 +193,12 @@ TAutoConsoleVariable<int32> CVarStreamingFullyLoadUsedTextures(
 	TEXT("If non-zero, all used texture will be fully streamed in as fast as possible"),
 	ECVF_Default);
 
+static TAutoConsoleVariable<int32> CVarStreamingFullyLoadMeshes(
+	TEXT("r.Streaming.FullyLoadMeshes"),
+	0,
+	TEXT("If non-zero, stream in all mesh LODs. This allows semi-disabling mesh LOD streaming without recook."),
+	ECVF_Default);
+
 TAutoConsoleVariable<int32> CVarStreamingUseAllMips(
 	TEXT("r.Streaming.UseAllMips"),
 	0,
@@ -305,6 +311,7 @@ void FRenderAssetStreamingSettings::Update()
 	bUseNewMetrics = CVarStreamingUseNewMetrics.GetValueOnAnyThread() != 0;
 	bLimitPoolSizeToVRAM = !GIsEditor && CVarStreamingLimitPoolSizeToVRAM.GetValueOnAnyThread() != 0;
 	bFullyLoadUsedTextures = CVarStreamingFullyLoadUsedTextures.GetValueOnAnyThread() != 0;
+	bFullyLoadMeshes = CVarStreamingFullyLoadMeshes.GetValueOnAnyThread() != 0;
 	bUseAllMips = CVarStreamingUseAllMips.GetValueOnAnyThread() != 0;
 	MinMipForSplitRequest = CVarStreamingMinMipForSplitRequest.GetValueOnAnyThread();
 	PerTextureBiasViewBoostThreshold = CVarStreamingPerTextureBiasViewBoostThreshold.GetValueOnAnyThread();
diff --git a/Engine/Source/Runtime/Engine/Private/Streaming/TextureStreamingHelpers.h b/Engine/Source/Runtime/Engine/Private/Streaming/TextureStreamingHelpers.h
--- a/Engine/Source/Runtime/Engine/Private/Streaming/TextureStreamingHelpers.h
+++ b/Engine/Source/Runtime/Engine/Private/Streaming/TextureStreamingHelpers.h
@@ -91,6 +91,7 @@ struct FRenderAssetStreamingSettings
 	bool bLimitPoolSizeToVRAM;
 	bool bUseNewMetrics;
 	bool bFullyLoadUsedTextures;
+	bool bFullyLoadMeshes;
 	bool bUseAllMips;
 	bool bUsePerTextureBias;
 	bool bUseMaterialData;
