Fix Warning in Static Analysis Win64.

diff --git a/Engine/Source/Runtime/Engine/Private/WorldPartition/WorldPartitionClassDescRegistry.cpp b/Engine/Source/Runtime/Engine/Private/WorldPartition/WorldPartitionClassDescRegistry.cpp
--- a/Engine/Source/Runtime/Engine/Private/WorldPartition/WorldPartitionClassDescRegistry.cpp
+++ b/Engine/Source/Runtime/Engine/Private/WorldPartition/WorldPartitionClassDescRegistry.cpp
@@ -335,7 +335,7 @@ void FWorldPartitionClassDescRegistry::OnObjectPreSave(UObject* InObject, FObjec
 	{
 		ValidateInternalState();
 
-		if (AActor* Actor = Actor = Cast<AActor>(InObject))
+		if (AActor* Actor = Cast<AActor>(InObject))
 		{
 			if (!Actor->HasAnyFlags(RF_ArchetypeObject | RF_ClassDefaultObject))
 			{
