World Partition Class Registry: fix warning when unable to follow redirector.

diff --git a/Engine/Source/Runtime/Engine/Private/WorldPartition/WorldPartitionActorDescArchive.cpp b/Engine/Source/Runtime/Engine/Private/WorldPartition/WorldPartitionActorDescArchive.cpp
--- a/Engine/Source/Runtime/Engine/Private/WorldPartition/WorldPartitionActorDescArchive.cpp
+++ b/Engine/Source/Runtime/Engine/Private/WorldPartition/WorldPartitionActorDescArchive.cpp
@@ -64,7 +64,7 @@ PRAGMA_ENABLE_DEPRECATION_WARNINGS
 	if (!ClassDesc)
 	{
 		ClassDesc = ClassDescRegistry.GetClassDescDefault(FTopLevelAssetPath(TEXT("/Script/Engine.Actor")));
-		UE_LOG(LogWorldPartition, Warning, TEXT("Can't find clas descriptor '%s' for '%s', using '%s'"), *ClassPath.ToString(), *InActorDesc->GetActorSoftPath().ToString(), *ClassDesc->GetActorSoftPath().ToString());
+		UE_LOG(LogWorldPartition, Warning, TEXT("Can't find class descriptor '%s' for '%s', using '%s'"), *ClassPath.ToString(), *InActorDesc->GetActorSoftPath().ToString(), *ClassDesc->GetActorSoftPath().ToString());
 	}
 	check(ClassDesc);
 }
