Actor Class Descriptors: remove overly cautious (and bad assumption for save as) assert.

diff --git a/Engine/Source/Runtime/Engine/Private/WorldPartition/WorldPartitionActorDescUtils.cpp b/Engine/Source/Runtime/Engine/Private/WorldPartition/WorldPartitionActorDescUtils.cpp
--- a/Engine/Source/Runtime/Engine/Private/WorldPartition/WorldPartitionActorDescUtils.cpp
+++ b/Engine/Source/Runtime/Engine/Private/WorldPartition/WorldPartitionActorDescUtils.cpp
@@ -152,9 +152,6 @@ void FWorldPartitionActorDescUtils::UpdateActorDescriptorFromActor(const AActor*
 
 void FWorldPartitionActorDescUtils::UpdateActorDescriptorFromActorDescriptor(TUniquePtr<FWorldPartitionActorDesc>& InActorDesc, TUniquePtr<FWorldPartitionActorDesc>& OutActorDesc)
 {
-	check(InActorDesc->GetActorPackage() == OutActorDesc->GetActorPackage());
-	check(InActorDesc->GetActorSoftPath() == OutActorDesc->GetActorSoftPath());
-	check(InActorDesc->GetGuid() == OutActorDesc->GetGuid());
 	InActorDesc->TransferFrom(OutActorDesc.Get());
 	OutActorDesc = MoveTemp(InActorDesc);
 }
