Add triangle to root binding visualization visualization.

diff --git a/Engine/Plugins/Runtime/HairStrands/Source/HairStrandsCore/Private/GroomDebug.cpp b/Engine/Plugins/Runtime/HairStrands/Source/HairStrandsCore/Private/GroomDebug.cpp
--- a/Engine/Plugins/Runtime/HairStrands/Source/HairStrandsCore/Private/GroomDebug.cpp
+++ b/Engine/Plugins/Runtime/HairStrands/Source/HairStrandsCore/Private/GroomDebug.cpp
@@ -312,8 +312,8 @@ static void AddDebugProjectionHairPass(
 		return;
 	}
 
-	ShaderPrint::RequestSpaceForLines(PrimitiveCount);
-	ShaderPrint::RequestSpaceForTriangles(PrimitiveCount);
+	ShaderPrint::RequestSpaceForLines(PrimitiveCount * 4 /* 1 hair root + 3 triangles */);
+	ShaderPrint::RequestSpaceForTriangles(PrimitiveCount * 3);
 
 	if (EDebugProjectionHairType::HairFrame == GeometryType &&
 		!RestRootResources->LODs[MeshLODIndex].RootBarycentricBuffer.Buffer)
diff --git a/Engine/Shaders/Private/HairStrands/HairStrandsDebug.usf b/Engine/Shaders/Private/HairStrands/HairStrandsDebug.usf
--- a/Engine/Shaders/Private/HairStrands/HairStrandsDebug.usf
+++ b/Engine/Shaders/Private/HairStrands/HairStrandsDebug.usf
@@ -361,6 +361,15 @@ void MainCS(uint3 DispatchThreadId : SV_DispatchThreadID)
 		const float4 Color = float4(UVs, 0, 1);
 		//Color = float4(N* 0.5f + 0.5f, 1);
 		AddLineWS(P0, P1, Color);
+
+		// Add triangle
+		const float4 TriColor = ColorWhite;
+		const float3 WP0 = mul(float4(Tri.P0, 1), RootLocalToWorld).xyz;
+		const float3 WP1 = mul(float4(Tri.P1, 1), RootLocalToWorld).xyz;
+		const float3 WP2 = mul(float4(Tri.P2, 1), RootLocalToWorld).xyz;
+		AddLineWS(WP0, WP1, TriColor);
+		AddLineWS(WP1, WP2, TriColor);
+		AddLineWS(WP2, WP0, TriColor);
 	}
 #endif // INPUT_ROOT_FRAME
 
diff --git a/Engine/Source/Runtime/Engine/Private/GroomVisualizationData.cpp b/Engine/Source/Runtime/Engine/Private/GroomVisualizationData.cpp
--- a/Engine/Source/Runtime/Engine/Private/GroomVisualizationData.cpp
+++ b/Engine/Source/Runtime/Engine/Private/GroomVisualizationData.cpp
@@ -184,7 +184,7 @@ void FGroomVisualizationData::Initialize()
 		AddVisualizationMode(ModeMap, true, EGroomViewMode::None,						LOCTEXT("NoneDebug", "None"),								LOCTEXT("NoneDebugDesc", "No debug mode"));
 		AddVisualizationMode(ModeMap, true, EGroomViewMode::MacroGroups,				LOCTEXT("MacroGroups", "Instances"),						LOCTEXT("MacroGroupsDesc", "Instances info"));
 		AddVisualizationMode(ModeMap, true, EGroomViewMode::LightBounds,				LOCTEXT("LightBounds", "Light Bound"),						LOCTEXT("LightBoundsDesc", "All DOMs light bounds"));
-		AddVisualizationMode(ModeMap, true, EGroomViewMode::MacroGroupScreenRect,		LOCTEXT("MacroGroupScreenRect", "screen Bounds"),			LOCTEXT("MacroGroupScreenRectDesc", "Screen projected instances"));
+		AddVisualizationMode(ModeMap, true, EGroomViewMode::MacroGroupScreenRect,		LOCTEXT("MacroGroupScreenRect", "Screen Bounds"),			LOCTEXT("MacroGroupScreenRectDesc", "Screen projected instances"));
 		AddVisualizationMode(ModeMap, true, EGroomViewMode::DeepOpacityMaps,			LOCTEXT("DeepOpacityMaps", "Deep Shadows"),					LOCTEXT("DeepOpacityMapsDesc", "Deep opacity maps"));
 		AddVisualizationMode(ModeMap, true, EGroomViewMode::SamplePerPixel,				LOCTEXT("SamplePerPixel", "Sample Per Pixel"),				LOCTEXT("SamplePerPixelDesc", "Sub-pixel sample count"));
 		AddVisualizationMode(ModeMap, true, EGroomViewMode::TAAResolveType,				LOCTEXT("TAAResolveType", "AA Type"),						LOCTEXT("TAAResolveTypeDesc", "TAA resolve type (regular/responsive)"));
