Moved stereo layer destruction logic to OnUnregister.

When the actor a stereo layer was attached to was destroyed the stereo layer wasn't destroyed. This is because the BeginDestroy callback is called by the GC at a future point in time which does not coincide with the actor destruction. On the other hand, components receive the OnUnregister callback when the actor is destroyed.

diff --git a/Engine/Source/Runtime/Engine/Classes/Components/StereoLayerComponent.h b/Engine/Source/Runtime/Engine/Classes/Components/StereoLayerComponent.h
--- a/Engine/Source/Runtime/Engine/Classes/Components/StereoLayerComponent.h
+++ b/Engine/Source/Runtime/Engine/Classes/Components/StereoLayerComponent.h
@@ -249,7 +249,7 @@ class ENGINE_API UStereoLayerComponent : public USceneComponent
 public:
 
 	//~ Begin UObject Interface
-	void BeginDestroy() override;
+	virtual void OnUnregister() override;
 	virtual void Serialize(FArchive& Ar) override;
 	virtual void PostLoad() override;
 	//~ End UObject Interface
diff --git a/Engine/Source/Runtime/Engine/Private/Components/StereoLayerComponent.cpp b/Engine/Source/Runtime/Engine/Private/Components/StereoLayerComponent.cpp
--- a/Engine/Source/Runtime/Engine/Private/Components/StereoLayerComponent.cpp
+++ b/Engine/Source/Runtime/Engine/Private/Components/StereoLayerComponent.cpp
@@ -26,7 +26,7 @@ UStereoLayerComponent::UStereoLayerComponent(const FObjectInitializer& ObjectIni
 	, Priority(0)
 	, bIsDirty(true)
 	, bTextureNeedsUpdate(false)
-	, LayerId(0)
+	, LayerId(IStereoLayers::FLayerDesc::INVALID_LAYER_ID)
 	, LastTransform(FTransform::Identity)
 	, bLastVisible(false)
 	, bNeedsPostLoadFixup(false)
@@ -37,15 +37,15 @@ UStereoLayerComponent::UStereoLayerComponent(const FObjectInitializer& ObjectIni
 	//Shape = ObjectInitializer.CreateDefaultSubobject<UStereoLayerShapeQuad>(this, TEXT("Shape"));
 }
 
-void UStereoLayerComponent::BeginDestroy()
+void UStereoLayerComponent::OnUnregister()
 {
-	Super::BeginDestroy();
+	Super::OnUnregister();
 
 	IStereoLayers* StereoLayers;
 	if (LayerId && GEngine->StereoRenderingDevice.IsValid() && (StereoLayers = GEngine->StereoRenderingDevice->GetStereoLayers()) != nullptr)
 	{
 		StereoLayers->DestroyLayer(LayerId);
-		LayerId = 0;
+		LayerId = IStereoLayers::FLayerDesc::INVALID_LAYER_ID;
 	}
 }
 
