Fix TSAN warning in TickTaskManager by using proper atomic for boolean which also helps getting rid of memory barrier

#rnx

diff --git a/Engine/Source/Runtime/Engine/Classes/Engine/EngineBaseTypes.h b/Engine/Source/Runtime/Engine/Classes/Engine/EngineBaseTypes.h
--- a/Engine/Source/Runtime/Engine/Classes/Engine/EngineBaseTypes.h
+++ b/Engine/Source/Runtime/Engine/Classes/Engine/EngineBaseTypes.h
@@ -14,6 +14,7 @@
 #include "UObject/WeakObjectPtr.h"
 #include "Misc/CoreMisc.h"
 #include "Net/Core/Connection/NetEnums.h"
+#include <atomic>
 
 #include "EngineBaseTypes.generated.h"
 
@@ -261,8 +262,8 @@ private:
 		/** Internal data to track if we have started visiting this tick function yet this frame **/
 		int32 TickVisitedGFrameCounter;
 
-		/** Internal data to track if we have finshed visiting this tick function yet this frame **/
-		int32 TickQueuedGFrameCounter;
+		/** Internal data to track if we have finished visiting this tick function yet this frame **/
+		std::atomic<int32> TickQueuedGFrameCounter;
 
 		/** Pointer to the task, only used during setup. This is often stale. **/
 		FBaseGraphTask* TaskPointer;
diff --git a/Engine/Source/Runtime/Engine/Private/TickTaskManager.cpp b/Engine/Source/Runtime/Engine/Private/TickTaskManager.cpp
--- a/Engine/Source/Runtime/Engine/Private/TickTaskManager.cpp
+++ b/Engine/Source/Runtime/Engine/Private/TickTaskManager.cpp
@@ -2045,7 +2045,7 @@ void FTickFunction::QueueTickFunctionParallel(const struct FTickContext& TickCon
 {
 	bool bProcessTick;
 
-	int32 OldValue = *(volatile int32*)&InternalData->TickVisitedGFrameCounter;
+	int32 OldValue = FPlatformAtomics::AtomicRead_Relaxed(&InternalData->TickVisitedGFrameCounter);
 	if (OldValue != GFrameCounter)
 	{
 		OldValue = FPlatformAtomics::InterlockedCompareExchange(&InternalData->TickVisitedGFrameCounter , GFrameCounter, OldValue);
@@ -2066,9 +2066,7 @@ void FTickFunction::QueueTickFunctionParallel(const struct FTickContext& TickCon
 				for (int32 PrereqIndex = 0; PrereqIndex < Prerequisites.Num(); PrereqIndex++)
 				{
 					FTickFunction* Prereq = Prerequisites[PrereqIndex].Get();
-#if USING_THREAD_SANITISER
-					if (Prereq) { TSAN_AFTER(&Prereq->InternalData->TickQueuedGFrameCounter); }
-#endif
+
 					if (!Prereq)
 					{
 						// stale prereq, delete it
@@ -2132,21 +2130,16 @@ void FTickFunction::QueueTickFunctionParallel(const struct FTickContext& TickCon
 			}
 		}
 		
-		TSAN_BEFORE(&InternalData->TickQueuedGFrameCounter);
-		FPlatformMisc::MemoryBarrier();
-		
-		// MSVC enforces acq/rel semantics on volatile values, but clang cannot (supports more backend architectures).
-		// consequently on ARM64 you would end up racing
-		FPlatformAtomics::InterlockedExchange(&InternalData->TickQueuedGFrameCounter, GFrameCounter);
+		InternalData->TickQueuedGFrameCounter = GFrameCounter;
 	}
 	else
 	{
 		// if we are not going to process it, we need to at least wait until the other thread finishes it
-		volatile int32* TickQueuedGFrameCounterPtr = &InternalData->TickQueuedGFrameCounter;
-		if (*TickQueuedGFrameCounterPtr != GFrameCounter)
+		std::atomic<int32>& TickQueuedGFrameCounter = InternalData->TickQueuedGFrameCounter;
+		if (TickQueuedGFrameCounter != GFrameCounter)
 		{
 			QUICK_SCOPE_CYCLE_COUNTER(STAT_FTickFunction_QueueTickFunctionParallel_Spin);
-			while (*TickQueuedGFrameCounterPtr != GFrameCounter)
+			while (TickQueuedGFrameCounter != GFrameCounter)
 			{
 				FPlatformProcess::YieldThread();
 			}
