- Made DrawOrientedWireBox to use DrawTranslucentLine()
- Changed DrawWireCapsuleto use DrawTranslucentLine() consistently

diff --git a/Engine/Source/Runtime/Engine/Private/PrimitiveDrawingUtils.cpp b/Engine/Source/Runtime/Engine/Private/PrimitiveDrawingUtils.cpp
--- a/Engine/Source/Runtime/Engine/Private/PrimitiveDrawingUtils.cpp
+++ b/Engine/Source/Runtime/Engine/Private/PrimitiveDrawingUtils.cpp
@@ -1174,7 +1174,7 @@ static void DrawHalfCircle(FPrimitiveDrawInterface* PDI, const FVector& Base, co
 	for(int32 SideIndex = 0; SideIndex < NumSides; SideIndex++)
 	{
 		const FVector	Vertex = Base + (X * FMath::Cos(AngleDelta * (SideIndex + 1)) + Y * FMath::Sin(AngleDelta * (SideIndex + 1))) * Radius;
-		PDI->DrawLine(LastVertex, Vertex, Color, SDPG_World, Thickness, DepthBias, bScreenSpace);
+		PDI->DrawTranslucentLine(LastVertex, Vertex, Color, SDPG_World, Thickness, DepthBias, bScreenSpace);
 		LastVertex = Vertex;
 	}	
 }
@@ -1223,7 +1223,7 @@ void DrawWireCapsule(FPrimitiveDrawInterface* PDI, const FVector& Base, const FV
 	{
 		const FVector Vertex = Base + (XAxis * FMath::Cos(AngleDelta * (SideIndex + 1)) + YAxis * FMath::Sin(AngleDelta * (SideIndex + 1))) * CapsuleRadius;
 
-		PDI->DrawLine(LastVertex - ZAxis * HalfHeight, LastVertex + ZAxis * HalfHeight, Color, DepthPriority, Thickness, DepthBias, bScreenSpace);
+		PDI->DrawTranslucentLine(LastVertex - ZAxis * HalfHeight, LastVertex + ZAxis * HalfHeight, Color, DepthPriority, Thickness, DepthBias, bScreenSpace);
 
 		LastVertex = Vertex;
 	}
@@ -1336,21 +1336,21 @@ void DrawOrientedWireBox(FPrimitiveDrawInterface* PDI, const FVector& Base, cons
 			P.Z = B[0].Z; Q.Z = B[1].Z;
 			P = M.TransformPosition(P);
 			Q = M.TransformPosition(Q);
-			PDI->DrawLine(P, Q, Color, DepthPriority, Thickness, DepthBias, bScreenSpace);
+			PDI->DrawTranslucentLine(P, Q, Color, DepthPriority, Thickness, DepthBias, bScreenSpace);
 
 			P.Y = B[i].Y; Q.Y = B[i].Y;
 			P.Z = B[j].Z; Q.Z = B[j].Z;
 			P.X = B[0].X; Q.X = B[1].X;
 			P = M.TransformPosition(P);
 			Q = M.TransformPosition(Q);
-			PDI->DrawLine(P, Q, Color, DepthPriority, Thickness, DepthBias, bScreenSpace);
+			PDI->DrawTranslucentLine(P, Q, Color, DepthPriority, Thickness, DepthBias, bScreenSpace);
 
 			P.Z = B[i].Z; Q.Z = B[i].Z;
 			P.X = B[j].X; Q.X = B[j].X;
 			P.Y = B[0].Y; Q.Y = B[1].Y;
 			P = M.TransformPosition(P);
 			Q = M.TransformPosition(Q);
-			PDI->DrawLine(P, Q, Color, DepthPriority, Thickness, DepthBias, bScreenSpace);
+			PDI->DrawTranslucentLine(P, Q, Color, DepthPriority, Thickness, DepthBias, bScreenSpace);
 		}
 	}
 }
