Fix disabling world partition for a level in the editor.

diff --git a/Engine/Source/Runtime/Engine/Private/WorldPartition/WorldPartition.cpp b/Engine/Source/Runtime/Engine/Private/WorldPartition/WorldPartition.cpp
--- a/Engine/Source/Runtime/Engine/Private/WorldPartition/WorldPartition.cpp
+++ b/Engine/Source/Runtime/Engine/Private/WorldPartition/WorldPartition.cpp
@@ -992,6 +992,7 @@ bool UWorldPartition::RemoveWorldPartition(AWorldSettings* WorldSettings)
 
 			WorldPartition->Uninitialize();
 			WorldSettings->SetWorldPartition(nullptr);
+			PersistentLevel->bIsPartitioned = false;
 
 			if (WorldPartition->WorldPartitionEditor)
 			{
