Prevent triggering an invalid ensure when saving a partitioned map.

diff --git a/Engine/Source/Runtime/Engine/Private/WorldPartition/WorldPartitionActorDesc.cpp b/Engine/Source/Runtime/Engine/Private/WorldPartition/WorldPartitionActorDesc.cpp
--- a/Engine/Source/Runtime/Engine/Private/WorldPartition/WorldPartitionActorDesc.cpp
+++ b/Engine/Source/Runtime/Engine/Private/WorldPartition/WorldPartitionActorDesc.cpp
@@ -108,20 +108,21 @@ void FWorldPartitionActorDesc::Init(const AActor* InActor)
 					LocalDataLayerAssetPaths.Add(*DataLayerAsset->GetPathName());
 				}
 			}
-			bIsUsingDataLayerAsset = LocalDataLayerAssetPaths.Num() > 0;
-
+			
+			PRAGMA_DISABLE_DEPRECATION_WARNINGS
+			// If the deprected ActorDataLayers is empty, consider the ActorDesc to be is using DataLayerAssets (with an empty array)
+			bIsUsingDataLayerAsset = (LocalDataLayerAssetPaths.Num() > 0) || (InActor->GetActorDataLayers().IsEmpty());
 			if (!bIsUsingDataLayerAsset)
 			{
 				// Use Actor's DataLayerManager since the fixup is relative to this level
 				const UDataLayerManager* DataLayerManager = ActorWorldPartition->GetDataLayerManager();
 				if (ensure(DataLayerManager))
 				{
-					PRAGMA_DISABLE_DEPRECATION_WARNINGS
 					// Pass Actor Level when resolving the DataLayerInstance as FWorldPartitionActorDesc always represents the state of the actor local to its outer level
 					LocalDataLayerInstanceNames = DataLayerManager->GetDataLayerInstanceNames(InActor->GetActorDataLayers());
-					PRAGMA_ENABLE_DEPRECATION_WARNINGS
 				}
 			}
+			PRAGMA_ENABLE_DEPRECATION_WARNINGS
 
 			// Init DataLayers persistent info
 			DataLayers = bIsUsingDataLayerAsset ? MoveTemp(LocalDataLayerAssetPaths) : MoveTemp(LocalDataLayerInstanceNames);
