Add guard when registering and unregistering an ActorDesc's Actor to a World.
Guard checks for register flows that trigger unregister and vice versa

diff --git a/Engine/Source/Runtime/Engine/Private/WorldPartition/WorldPartitionHandle.cpp b/Engine/Source/Runtime/Engine/Private/WorldPartition/WorldPartitionHandle.cpp
--- a/Engine/Source/Runtime/Engine/Private/WorldPartition/WorldPartitionHandle.cpp
+++ b/Engine/Source/Runtime/Engine/Private/WorldPartition/WorldPartitionHandle.cpp
@@ -42,11 +42,19 @@ FWorldPartitionLoadingContext::IContext* FWorldPartitionLoadingContext::ActiveCo
 
 void FWorldPartitionLoadingContext::LoadAndRegisterActor(FWorldPartitionActorDesc* ActorDesc)
 {
+#if DO_CHECK
+	FWorldPartitionActorDesc::FRegisteringUnregisteringGuard Guard(ActorDesc);
+#endif
+
 	ActiveContext->RegisterActor(ActorDesc);
 }
 
 void FWorldPartitionLoadingContext::UnloadAndUnregisterActor(FWorldPartitionActorDesc* ActorDesc)
 {
+#if DO_CHECK
+	FWorldPartitionActorDesc::FRegisteringUnregisteringGuard Guard(ActorDesc);
+#endif
+
 	ActiveContext->UnregisterActor(ActorDesc);
 }
 
diff --git a/Engine/Source/Runtime/Engine/Public/WorldPartition/WorldPartitionActorDesc.h b/Engine/Source/Runtime/Engine/Public/WorldPartition/WorldPartitionActorDesc.h
--- a/Engine/Source/Runtime/Engine/Public/WorldPartition/WorldPartitionActorDesc.h
+++ b/Engine/Source/Runtime/Engine/Public/WorldPartition/WorldPartitionActorDesc.h
@@ -313,5 +313,32 @@ protected:
 	mutable FText*					UnloadedReason;
 
 	static TMap<TSubclassOf<AActor>, FActorDescDeprecator> Deprecators;
+
+#if DO_CHECK
+public:
+	struct FRegisteringUnregisteringGuard
+	{
+		FRegisteringUnregisteringGuard(FWorldPartitionActorDesc* InActorDesc)
+			:ActorDesc(InActorDesc)
+		{
+			check(!ActorDesc->bIsRegisteringOrUnregistering);
+			ActorDesc->bIsRegisteringOrUnregistering = true;
+		}
+
+		~FRegisteringUnregisteringGuard()
+		{
+			check(ActorDesc->bIsRegisteringOrUnregistering);
+			ActorDesc->bIsRegisteringOrUnregistering = false;
+		}
+
+		FRegisteringUnregisteringGuard(const FRegisteringUnregisteringGuard&) = delete;
+		FRegisteringUnregisteringGuard& operator=(const FRegisteringUnregisteringGuard&) = delete;
+
+	private:
+		FWorldPartitionActorDesc* ActorDesc;
+	};
+private:
+	bool bIsRegisteringOrUnregistering = false;
+#endif
 };
 #endif
