Level Duplication changes to prevent failures with World Partition

diff --git a/Engine/Source/Runtime/Engine/Private/World.cpp b/Engine/Source/Runtime/Engine/Private/World.cpp
--- a/Engine/Source/Runtime/Engine/Private/World.cpp
+++ b/Engine/Source/Runtime/Engine/Private/World.cpp
@@ -8468,20 +8468,34 @@ void UWorld::DuplicateRequestedLevels(const FName MapName)
 		DuplicateLevels.SetPersistentLevel(DuplicatePersistentLevel);
 		DuplicateLevels.AddLevel(DuplicatePersistentLevel);
 
-		for (ULevelStreaming* StreamingLevel : StreamingLevels)
+		if (!IsPartitionedWorld())
 		{
-			if (StreamingLevel && !StreamingLevel->bIsStatic)
+			for (ULevelStreaming* StreamingLevel : StreamingLevels)
 			{
-				ULevel* DuplicatedLevel = DuplicateLevelWithPrefix(StreamingLevel->GetLoadedLevel(), 1);
-				if (!DuplicatedLevel)
+				if (StreamingLevel && !StreamingLevel->bIsStatic)
 				{
-					UE_LOG(LogWorld, Warning, TEXT("UWorld::DuplicateRequestedLevels: failed to duplicate streaming level %s. No duplicate level collection will be created."),
-						*GetFullNameSafe(StreamingLevel->GetLoadedLevel()));
-					return;
+					ULevel* DuplicatedLevel = DuplicateLevelWithPrefix(StreamingLevel->GetLoadedLevel(), 1);
+					if (!DuplicatedLevel)
+					{
+						UE_LOG(LogWorld, Warning, TEXT("UWorld::DuplicateRequestedLevels: failed to duplicate streaming level %s. No duplicate level collection will be created."),
+							*GetFullNameSafe(StreamingLevel->GetLoadedLevel()));
+						return;
+					}
+					// Don't tell the server about these levels
+					DuplicatedLevel->bClientOnlyVisible = true;
+					DuplicateLevels.AddLevel(DuplicatedLevel);
 				}
-				// Don't tell the server about these levels
-				DuplicatedLevel->bClientOnlyVisible = true;
-				DuplicateLevels.AddLevel(DuplicatedLevel);
+			}
+		}
+		else
+		{
+			UE_LOG(LogWorld, Log, TEXT("UWorld::DuplicateRequestedLevels: Attempting to duplicate levels for partitioned world. The persistent Level will be duplicated, but all streaming cells will be considered static."));
+
+			// If duplicating a partitioned world, remove and destroy the duplicated persistent level's UWorldPartition
+			UWorldPartition* DuplicateWP = DuplicatePersistentLevel->GetWorldPartition();
+			if (DuplicateWP)
+			{
+				DuplicateWP->MarkAsGarbage();
 			}
 		}
 
diff --git a/Engine/Source/Runtime/Engine/Private/WorldPartition/WorldPartitionLevelStreamingDynamic.cpp b/Engine/Source/Runtime/Engine/Private/WorldPartition/WorldPartitionLevelStreamingDynamic.cpp
--- a/Engine/Source/Runtime/Engine/Private/WorldPartition/WorldPartitionLevelStreamingDynamic.cpp
+++ b/Engine/Source/Runtime/Engine/Private/WorldPartition/WorldPartitionLevelStreamingDynamic.cpp
@@ -50,6 +50,10 @@ void UWorldPartitionLevelStreamingDynamic::Initialize(const UWorldPartitionRunti
 	bShouldBeAlwaysLoaded = InCell.IsAlwaysLoaded();
 	StreamingPriority = 0;
 
+	// Default WP streaming levels to static, so they're not hidden during instant replay playback
+	// This means only the persistent level will be duplicated/hidden
+	bIsStatic = true;
+
 	UWorld* CellOuterWorld = InCell.GetOuterWorld();
 #if WITH_EDITOR
 	check(ChildPackages.Num() == 0);
