Make sure FWorldPartitionClassDescRegistry::Get is properly declared as dllexport.

diff --git a/Engine/Source/Editor/WorldPartitionEditor/Private/WorldPartitionEditorModule.cpp b/Engine/Source/Editor/WorldPartitionEditor/Private/WorldPartitionEditorModule.cpp
--- a/Engine/Source/Editor/WorldPartitionEditor/Private/WorldPartitionEditorModule.cpp
+++ b/Engine/Source/Editor/WorldPartitionEditor/Private/WorldPartitionEditorModule.cpp
@@ -152,6 +152,7 @@ void FWorldPartitionEditorModule::StartupModule()
 void FWorldPartitionEditorModule::ShutdownModule()
 {
 	FWorldPartitionClassDescRegistry().Get().Uninitialize();
+	FWorldPartitionClassDescRegistry().Get().TearDown();
 
 	if (!IsRunningGame())
 	{
diff --git a/Engine/Source/Runtime/Engine/Private/WorldPartition/WorldPartitionClassDescRegistry.cpp b/Engine/Source/Runtime/Engine/Private/WorldPartition/WorldPartitionClassDescRegistry.cpp
--- a/Engine/Source/Runtime/Engine/Private/WorldPartition/WorldPartitionClassDescRegistry.cpp
+++ b/Engine/Source/Runtime/Engine/Private/WorldPartition/WorldPartitionClassDescRegistry.cpp
@@ -21,6 +21,16 @@
 #include "WorldPartition/WorldPartitionActorDesc.h"
 #include "WorldPartition/WorldPartitionActorDescUtils.h"
 
+FWorldPartitionClassDescRegistry& FWorldPartitionClassDescRegistry::Get()
+{
+	return TLazySingleton<FWorldPartitionClassDescRegistry>::Get();
+}
+
+void FWorldPartitionClassDescRegistry::TearDown()
+{
+	TLazySingleton<FWorldPartitionClassDescRegistry>::TearDown();
+}
+
 void FWorldPartitionClassDescRegistry::Initialize()
 {
 	UE_SCOPED_TIMER(TEXT("FWorldPartitionClassDescRegistry::Initialize"), LogWorldPartition, Display);
@@ -299,6 +309,14 @@ void FWorldPartitionClassDescRegistry::RegisterClassDescriptorFromAssetData(cons
 
 void FWorldPartitionClassDescRegistry::RegisterClassDescriptorFromActorClass(const UClass* InActorClass)
 {
+	if (UClass* ParentClass = InActorClass->GetSuperClass(); ParentClass && ParentClass->IsChildOf<AActor>())
+	{
+		if (!ClassByPath.Contains(FTopLevelAssetPath(ParentClass->GetPathName())))
+		{
+			RegisterClassDescriptorFromActorClass(ParentClass);
+		}
+	}
+
 	ParentClassMap.Add(FTopLevelAssetPath(InActorClass->GetPathName()), FTopLevelAssetPath(InActorClass->GetSuperClass()->GetPathName()));
 
 	const AActor* ActorCDO = CastChecked<AActor>(InActorClass->GetDefaultObject());
diff --git a/Engine/Source/Runtime/Engine/Public/WorldPartition/WorldPartitionClassDescRegistry.h b/Engine/Source/Runtime/Engine/Public/WorldPartition/WorldPartitionClassDescRegistry.h
--- a/Engine/Source/Runtime/Engine/Public/WorldPartition/WorldPartitionClassDescRegistry.h
+++ b/Engine/Source/Runtime/Engine/Public/WorldPartition/WorldPartitionClassDescRegistry.h
@@ -17,10 +17,8 @@ class ENGINE_API FWorldPartitionClassDescRegistry : FActorDescList
 	using TRedirectClassMap = TMap<FTopLevelAssetPath, FTopLevelAssetPath>;
 
 public:
-	static FWorldPartitionClassDescRegistry& Get()
-	{
-		return TLazySingleton<FWorldPartitionClassDescRegistry>::Get();
-	}
+	static FWorldPartitionClassDescRegistry& Get();
+	static void TearDown();
 
 	void Initialize();
 	void Uninitialize();
