WorldPartitionActorFilter: Fix crash caused by bad cast

#rnx

diff --git a/Engine/Source/Runtime/Engine/Private/LevelInstance/LevelInstanceSubsystem.cpp b/Engine/Source/Runtime/Engine/Private/LevelInstance/LevelInstanceSubsystem.cpp
--- a/Engine/Source/Runtime/Engine/Private/LevelInstance/LevelInstanceSubsystem.cpp
+++ b/Engine/Source/Runtime/Engine/Private/LevelInstance/LevelInstanceSubsystem.cpp
@@ -2479,6 +2479,7 @@ FWorldPartitionActorFilter ULevelInstanceSubsystem::GetLevelInstanceFilter(const
 	return GetLevelInstanceFilterInternal(LevelPackage, VisitedPackages);
 }
 
+// @todo_ow: Move this code out of ULevelInstanceSubsystem into UWorldPartitionSubsystem now that it is generic (not relying on LevelInstances)
 FWorldPartitionActorFilter ULevelInstanceSubsystem::GetLevelInstanceFilterInternal(const FString& LevelPackage, TSet<FString>& VisitedPackages) const
 {
 	if (VisitedPackages.Contains(LevelPackage))
@@ -2495,7 +2496,7 @@ FWorldPartitionActorFilter ULevelInstanceSubsystem::GetLevelInstanceFilterIntern
 	ON_SCOPE_EXIT{ ActorDescContainerInstanceManager.UnregisterContainer(LevelContainer); };
 
 	// Lazy create filter for now
-	TArray<const FLevelInstanceActorDesc*> LevelInstanceActorDescs;
+	TArray<const FWorldPartitionActorDesc*> ContainerActorDescs;
 	const FWorldDataLayersActorDesc* WorldDataLayersActorDesc = nullptr;
 
 	for (FActorDescList::TConstIterator<> ActorDescIt(LevelContainer); ActorDescIt; ++ActorDescIt)
@@ -2505,9 +2506,9 @@ FWorldPartitionActorFilter ULevelInstanceSubsystem::GetLevelInstanceFilterIntern
 			check(!WorldDataLayersActorDesc);
 			WorldDataLayersActorDesc = static_cast<const FWorldDataLayersActorDesc*>(*ActorDescIt);
 		}
-		else if (ActorDescIt->GetActorNativeClass()->ImplementsInterface(ULevelInstanceInterface::StaticClass()))
+		else if (ActorDescIt->IsContainerInstance())
 		{
-			LevelInstanceActorDescs.Add(static_cast<const FLevelInstanceActorDesc*>(*ActorDescIt));
+			ContainerActorDescs.Add(*ActorDescIt);
 		}
 	}
 	
@@ -2525,18 +2526,21 @@ FWorldPartitionActorFilter ULevelInstanceSubsystem::GetLevelInstanceFilterIntern
 		}
 	}
 
-	for (const FLevelInstanceActorDesc* LevelInstanceActorDesc : LevelInstanceActorDescs)
+	for (const FWorldPartitionActorDesc* ContainerActorDesc : ContainerActorDescs)
 	{
 		TSet<FString> VisitedPackagesCopy(VisitedPackages);
 
 		// Get World Default Filter
-		FWorldPartitionActorFilter* ChildFilter = new FWorldPartitionActorFilter(GetLevelInstanceFilterInternal(LevelInstanceActorDesc->GetLevelPackage().ToString(), VisitedPackagesCopy));
-		ChildFilter->DisplayName = LevelInstanceActorDesc->GetActorLabelOrName().ToString();
+		FWorldPartitionActorFilter* ChildFilter = new FWorldPartitionActorFilter(GetLevelInstanceFilterInternal(ContainerActorDesc->GetLevelPackage().ToString(), VisitedPackagesCopy));
+		ChildFilter->DisplayName = ContainerActorDesc->GetActorLabelOrName().ToString();
 
-		// Apply LevelInstanceActorDesc Filter to Default
-		ChildFilter->Override(LevelInstanceActorDesc->GetFilter());
+		// Apply Filter to Default
+		if (const FWorldPartitionActorFilter* ContainerFilter = ContainerActorDesc->GetContainerFilter())
+		{
+			ChildFilter->Override(*ContainerFilter);
+		}
 
-		Filter.AddChildFilter(LevelInstanceActorDesc->GetGuid(), ChildFilter);
+		Filter.AddChildFilter(ContainerActorDesc->GetGuid(), ChildFilter);
 	}
 
 	return Filter;
diff --git a/Engine/Source/Runtime/Engine/Public/WorldPartition/LevelInstance/LevelInstanceActorDesc.h b/Engine/Source/Runtime/Engine/Public/WorldPartition/LevelInstance/LevelInstanceActorDesc.h
--- a/Engine/Source/Runtime/Engine/Public/WorldPartition/LevelInstance/LevelInstanceActorDesc.h
+++ b/Engine/Source/Runtime/Engine/Public/WorldPartition/LevelInstance/LevelInstanceActorDesc.h
@@ -28,8 +28,8 @@ public:
 	virtual bool IsContainerInstance() const override;
 	virtual FName GetLevelPackage() const override { return LevelPackage; }
 	virtual bool GetContainerInstance(FContainerInstance& OutContainerInstance, bool bInBuildFilter) const override;
+	virtual const FWorldPartitionActorFilter* GetContainerFilter() const override { return &Filter; }
 	virtual void CheckForErrors(IStreamingGenerationErrorHandler* ErrorHandler) const override;
-	const FWorldPartitionActorFilter& GetFilter() const { return Filter; }
 
 protected:
 	virtual void Init(const AActor* InActor) override;
diff --git a/Engine/Source/Runtime/Engine/Public/WorldPartition/WorldPartitionActorDesc.h b/Engine/Source/Runtime/Engine/Public/WorldPartition/WorldPartitionActorDesc.h
--- a/Engine/Source/Runtime/Engine/Public/WorldPartition/WorldPartitionActorDesc.h
+++ b/Engine/Source/Runtime/Engine/Public/WorldPartition/WorldPartitionActorDesc.h
@@ -26,6 +26,7 @@ struct FWorldPartitionActorDescInitData
 class AActor;
 class UActorDescContainer;
 class IStreamingGenerationErrorHandler;
+struct FWorldPartitionActorFilter;
 
 enum class EContainerClusterMode : uint8
 {
@@ -154,6 +155,7 @@ PRAGMA_ENABLE_DEPRECATION_WARNINGS
 
 	virtual bool IsContainerInstance() const { return false; }
 	virtual FName GetLevelPackage() const { return NAME_None; }
+	virtual const FWorldPartitionActorFilter* GetContainerFilter() const { return nullptr; }
 	virtual bool GetContainerInstance(FContainerInstance& OutContainerInstance, bool bInBuildFilter = false) const { return false; }
 
 	FGuid GetContentBundleGuid() const;
