Make sure deferred actor construction scripts are done before saving out a level

diff --git a/Engine/Source/Runtime/Engine/Private/World.cpp b/Engine/Source/Runtime/Engine/Private/World.cpp
--- a/Engine/Source/Runtime/Engine/Private/World.cpp
+++ b/Engine/Source/Runtime/Engine/Private/World.cpp
@@ -91,6 +91,7 @@
 	#include "WorldPartition/DataLayer/WorldDataLayers.h"
 	#include "PieFixupSerializer.h"
 	#include "ActorFolder.h"
+	#include "ActorDeferredScriptManager.h"
 #endif
 
 
@@ -1260,6 +1261,9 @@ void UWorld::PreSaveRoot(FObjectPreSaveRootContext ObjectSaveContext)
 
 	// Flush outstanding static mesh compilation to ensure that construction scripts are properly ran and not deferred prior to saving
 	FStaticMeshCompilingManager::Get().FinishAllCompilation();
+
+	// Execute all pending actor construction scripts
+	FActorDeferredScriptManager::Get().FinishAllCompilation();
 #endif
 
 	// Update components and keep track off whether we need to clean them up afterwards.
diff --git a/Engine/Source/Runtime/Engine/Public/ActorDeferredScriptManager.h b/Engine/Source/Runtime/Engine/Public/ActorDeferredScriptManager.h
--- a/Engine/Source/Runtime/Engine/Public/ActorDeferredScriptManager.h
+++ b/Engine/Source/Runtime/Engine/Public/ActorDeferredScriptManager.h
@@ -26,7 +26,10 @@ public:
 	/**
 	 * empty implementation since deferred actor scripts are not actually async
 	 */
-	ENGINE_API void FinishAllCompilation() override {}
+	ENGINE_API void FinishAllCompilation() override
+	{
+		ProcessAsyncTasks(false);
+	}
 
 	/** 
 	 * Add an actor for the manager to run the construction of once dependent asset compilation are done (i.e Static Mesh)
