Chaos - Add actor name to particles

diff --git a/Engine/Source/Runtime/Engine/Private/PhysicsEngine/BodyInstance.cpp b/Engine/Source/Runtime/Engine/Private/PhysicsEngine/BodyInstance.cpp
--- a/Engine/Source/Runtime/Engine/Private/PhysicsEngine/BodyInstance.cpp
+++ b/Engine/Source/Runtime/Engine/Private/PhysicsEngine/BodyInstance.cpp
@@ -7,6 +7,7 @@
 #include "Engine/Engine.h"
 #include "Components/SkeletalMeshComponent.h"
 #include "Components/ShapeComponent.h"
+#include "GameFramework/Actor.h"
 #include "Engine/CollisionProfile.h"
 #include "Engine/World.h"
 #include "SceneManagement.h"
@@ -1080,6 +1081,12 @@ TSharedPtr<TArray<ANSICHAR>> GetDebugDebugName(const UPrimitiveComponent* Primit
 
 	if (PrimitiveComp)
 	{
+#if WITH_EDITOR
+		if (PrimitiveComp->GetOwner())
+		{
+			DebugName += FString::Printf(TEXT("Actor: '%s' "), *PrimitiveComp->GetOwner()->GetActorLabel(false));
+		}
+#endif
 		DebugName += FString::Printf(TEXT("Component: '%s' "), *PrimitiveComp->GetPathName());
 	}
 
