Change check() to an ensure() with diagnostic information in UBlendProfile::FillSkeletonBoneDurationsArray().

diff --git a/Engine/Source/Runtime/Engine/Private/Animation/BlendProfile.cpp b/Engine/Source/Runtime/Engine/Private/Animation/BlendProfile.cpp
--- a/Engine/Source/Runtime/Engine/Private/Animation/BlendProfile.cpp
+++ b/Engine/Source/Runtime/Engine/Private/Animation/BlendProfile.cpp
@@ -265,7 +265,11 @@ void UBlendProfile::FillSkeletonBoneDurationsArray(TCustomBoneIndexArrayView<flo
 	check(OwningSkeleton != nullptr);
 	const FReferenceSkeleton& RefSkeleton = OwningSkeleton->GetReferenceSkeleton();
 	const int32 NumSkeletonBones = RefSkeleton.GetNum();
-	check(OutDurationPerBone.Num() == NumSkeletonBones);
+	if (!ensureMsgf(OutDurationPerBone.Num() == NumSkeletonBones, 
+		TEXT("UBlendProfile::FillSkeletonBoneDurationsArray received OutDurationPerBone(%d) vs NumSkeletonBones(%d) for %s"), OutDurationPerBone.Num(), NumSkeletonBones, *GetNameSafe(OwningSkeleton)))
+	{
+		return;
+	}
 
 	for(float& BoneDuration: OutDurationPerBone)
 	{
