Move streaming tasks to normal / background threads and added naming to identify them in Insights.

diff --git a/Engine/Source/Runtime/Engine/Private/Streaming/StreamingManagerTexture.cpp b/Engine/Source/Runtime/Engine/Private/Streaming/StreamingManagerTexture.cpp
--- a/Engine/Source/Runtime/Engine/Private/Streaming/StreamingManagerTexture.cpp
+++ b/Engine/Source/Runtime/Engine/Private/Streaming/StreamingManagerTexture.cpp
@@ -1334,7 +1334,7 @@ void FRenderAssetStreamingManager::UpdateStreamingRenderAssets(int32 StageIndex,
 			}
 		}
 
-		ParallelFor(Packets.Num(), [this, &Packets, &bWaitForMipFading, &bAsync](int32 PacketIndex)
+		ParallelFor(TEXT("RenderAssetStreaming"), Packets.Num(), 1, [this, &Packets, &bWaitForMipFading, &bAsync](int32 PacketIndex)
 		{
 			FOptionalTaskTagScope Scope(ETaskTag::EParallelGameThread);
 
@@ -1359,7 +1359,8 @@ void FRenderAssetStreamingManager::UpdateStreamingRenderAssets(int32 StageIndex,
 					Packets[PacketIndex].LocalInflightRenderAssets.Add(Index);
 				}
 			}
-		});
+
+		}, EParallelForFlags::BackgroundPriority);
 
 		for (FPacket Packet : Packets) {
 			InflightRenderAssets.Append(Packet.LocalInflightRenderAssets);
@@ -1715,7 +1716,7 @@ public:
 	}
 	static FORCEINLINE ENamedThreads::Type GetDesiredThread()
 	{
-		return ENamedThreads::AnyHiPriThreadHiPriTask;
+		return ENamedThreads::AnyNormalThreadNormalTask;
 	}
 	static FORCEINLINE ESubsequentsMode::Type GetSubsequentsMode()
 	{
