World Partition Class Descriptors: temporary fix wrongly reported missing classes for cooked editor.

diff --git a/Engine/Source/Runtime/Engine/Private/WorldPartition/ActorDescContainer.cpp b/Engine/Source/Runtime/Engine/Private/WorldPartition/ActorDescContainer.cpp
--- a/Engine/Source/Runtime/Engine/Private/WorldPartition/ActorDescContainer.cpp
+++ b/Engine/Source/Runtime/Engine/Private/WorldPartition/ActorDescContainer.cpp
@@ -86,7 +86,7 @@ void UActorDescContainer::Initialize(const FInitializeParams& InitParams)
 
 			if (!ActorDesc.IsValid() || 
 				!ActorDesc->GetNativeClass().IsValid() || 
-				(ActorDesc->GetBaseClass().IsValid() && !ClassDescRegistry.IsRegisteredClass(ActorDesc->GetBaseClass())) || 
+				//(ActorDesc->GetBaseClass().IsValid() && !ClassDescRegistry.IsRegisteredClass(ActorDesc->GetBaseClass())) || 
 				FActorDescList::GetActorDesc(ActorDesc->GetGuid()) ||
 				(InitParams.FilterActorDesc && !InitParams.FilterActorDesc(ActorDesc.Get())))
 			{
