Initial orthogonal rendering fix for Point and Spot lights with light framework setup for future ortho rendering fixes

- Common.ush change of runtime check for IsOrthoProjection() function calls to minimise cost to 1ALU
- Add 2 new functions, ScreenVectorFromScreenRect and GetScreenPositionForProjectionMatrix which return corrected values depending on projection type
- Converted CommonViewUniformBuffer to read from the new ViewUniformBuffer TanAndInvTanHalfFOV vector. This means that the static projection matrix values used by these functions are assigned in SceneView.cpp. Ortho rendering does not use these values, so these functions only need to return 1.0f for each value for ortho, and moving this to a constant vector means we avoid unnecessary IsOrthoProjection checks in the shader code.
- Added a few helper functions for the above to SceneView.h/.cpp, including a correction for the ScreenToClip matrix which is now reused in multiple locations, including for shadow rendering, rather than being hard coded.
- Utilised these changes to correct the problems with Point and Spot lights+shadows in ortho perspective, and also corrects some minor artifacts that occur with directional light in ortho.

diff --git a/Engine/Shaders/Private/Common.ush b/Engine/Shaders/Private/Common.ush
--- a/Engine/Shaders/Private/Common.ush
+++ b/Engine/Shaders/Private/Common.ush
@@ -1208,9 +1208,16 @@ float ConvertFromDeviceZ(float DeviceZ)
 }
 
 // Returns whether or not the given projection matrix is orthographic
+// Utilises asuint call to reduce compare to 1ALU cost
 bool IsOrthoProjection(float4x4 ViewToClip)
 {
-	return ViewToClip[3][3] >= 1.0f;
+	return asuint(ViewToClip._44);
+}
+
+// Returns whether or not the default view is orthographic
+bool IsOrthoProjection()
+{
+	return IsOrthoProjection(View.ViewToClip);
 }
 
 // inverse operation of ConvertFromDeviceZ()
@@ -1236,6 +1243,21 @@ float2 ScreenPositionToBufferUV(float4 ScreenPosition)
 	return float2(ScreenPosition.xy / ScreenPosition.w * ResolvedView.ScreenPositionScaleBias.xy + ResolvedView.ScreenPositionScaleBias.wz);
 }
 
+float3 ScreenVectorFromScreenRect(float4 Position)
+{
+	//Screen rect does not include scene depth. Perspective projection requires ScreenToTranslated world matrix
+	//Ortho Screen to View in this context is 1:1 so uses the ViewToTranslatedMatrix only
+	return select(IsOrthoProjection(), mul(Position, View.ViewToTranslatedWorld), mul(Position,View.ScreenToTranslatedWorld)).xyz;
+}
+
+// Returns the screen position for projection matrix calculations depending on the type of projection the view is using
+float2 GetScreenPositionForProjectionType(float2 ScreenPosition, float SceneDepth)
+{
+	// For perspective projection matrix, the scene depth is required as part of the clip to view calculations
+	// For orthogonal projection matrix, scene depth should not be used
+	return select(IsOrthoProjection(), ScreenPosition, ScreenPosition * SceneDepth);
+}
+
 float2 SvPositionToBufferUV(float4 SvPosition)
 {
 	return SvPosition.xy * View.BufferSizeAndInvSize.zw;
@@ -1344,7 +1366,7 @@ float3 ScreenToViewPos(float2 ViewportUV, float SceneDepth)
 
 	ProjViewPos.x =  ViewportUV.x * View.ScreenToViewSpace.x  + View.ScreenToViewSpace.z;
 	ProjViewPos.y =  ViewportUV.y * View.ScreenToViewSpace.y  + View.ScreenToViewSpace.w;
-	return float3(ProjViewPos * SceneDepth, SceneDepth);
+	return float3(GetScreenPositionForProjectionType(ProjViewPos, SceneDepth), SceneDepth);
 }
 
 // ----------------------------
diff --git a/Engine/Shaders/Private/CommonViewUniformBuffer.ush b/Engine/Shaders/Private/CommonViewUniformBuffer.ush
--- a/Engine/Shaders/Private/CommonViewUniformBuffer.ush
+++ b/Engine/Shaders/Private/CommonViewUniformBuffer.ush
@@ -11,26 +11,26 @@
 */
 float2 GetTanHalfFieldOfView()
 {
-	return float2(View.ClipToView[0][0], View.ClipToView[1][1]);
+	return View.TanAndInvTanHalfFOV.xy;
 }
 
 float2 GetPrevTanHalfFieldOfView()
 {
-	return float2(View.PrevClipToView[0][0], View.PrevClipToView[1][1]);
+	return View.PrevTanAndInvTanHalfFOV.xy;
 }
 
 // might be used by Custom material expressions (still best to wrap the custom node in a material function)
 // @return 1 / tan(View.FieldOfViewWideAngles * .5)
 float2 GetCotanHalfFieldOfView()
 {
-	return float2(View.ViewToClip[0][0], View.ViewToClip[1][1]);
+	return View.TanAndInvTanHalfFOV.zw;
 }
 
 // might be used by Custom material expressions (still best to wrap the custom node in a material function)
 // @return previous 1 / tan(View.FieldOfViewWideAngles * .5)
 float2 GetPrevCotanHalfFieldOfView()
 {
-	return float2(View.PrevViewToClip[0][0], View.PrevViewToClip[1][1]);
+	return View.PrevTanAndInvTanHalfFOV.zw;
 }
 
 // Return the index of the frame.
diff --git a/Engine/Shaders/Private/DeferredLightPixelShaders.usf b/Engine/Shaders/Private/DeferredLightPixelShaders.usf
--- a/Engine/Shaders/Private/DeferredLightPixelShaders.usf
+++ b/Engine/Shaders/Private/DeferredLightPixelShaders.usf
@@ -81,7 +81,7 @@ FDerivedParams GetDerivedParams(in FInputParams Input, in float SceneDepth)
 #if LIGHT_SOURCE_SHAPE > 0
 	// With a perspective projection, the clip space position is NDC * Clip.w
 	// With an orthographic projection, clip space is the same as NDC
-	float2 ClipPosition = Input.ScreenPosition.xy / Input.ScreenPosition.w * (View.ViewToClip[3][3] < 1.0f ? SceneDepth : 1.0f);
+	float2 ClipPosition = GetScreenPositionForProjectionType(Input.ScreenPosition.xy / Input.ScreenPosition.w, SceneDepth);
 	Out.TranslatedWorldPosition = mul(float4(ClipPosition, SceneDepth, 1), View.ScreenToTranslatedWorld).xyz;
 	Out.CameraVector = normalize(Out.TranslatedWorldPosition - View.TranslatedWorldCameraOrigin);
 #else
diff --git a/Engine/Shaders/Private/DeferredLightVertexShaders.usf b/Engine/Shaders/Private/DeferredLightVertexShaders.usf
--- a/Engine/Shaders/Private/DeferredLightVertexShaders.usf
+++ b/Engine/Shaders/Private/DeferredLightVertexShaders.usf
@@ -41,7 +41,7 @@ void VertexMain(
 	)
 {	
 	DrawFullScreenRect(float4(InPosition.xy, 0, 1), InUV, OutPosition, OutTexCoord);
-	OutScreenVector = mul(float4(OutPosition.xy, 1, 0), View.ScreenToTranslatedWorld).xyz;
+	OutScreenVector = ScreenVectorFromScreenRect(float4(OutPosition.xy, 1, 0));
 }
 #endif
 
diff --git a/Engine/Shaders/Private/DeferredLightingCommon.ush b/Engine/Shaders/Private/DeferredLightingCommon.ush
--- a/Engine/Shaders/Private/DeferredLightingCommon.ush
+++ b/Engine/Shaders/Private/DeferredLightingCommon.ush
@@ -122,7 +122,7 @@ float ShadowRayCast(
 void GetShadowTerms(float SceneDepth, half4 PrecomputedShadowFactors, uint ShadingModelID, float ContactShadowOpacity, FDeferredLightData LightData, float3 TranslatedWorldPosition, half3 L, half4 LightAttenuation, float Dither, inout FShadowTerms Shadow)
 {
 	float ContactShadowLength = 0.0f;
-	const float ContactShadowLengthScreenScale = View.ClipToView[1][1] * SceneDepth;
+	const float ContactShadowLengthScreenScale = GetTanHalfFieldOfView().y * SceneDepth;
 
 	BRANCH
 	if (LightData.ShadowedBits)
diff --git a/Engine/Shaders/Private/VirtualShadowMaps/VirtualShadowMapProjection.usf b/Engine/Shaders/Private/VirtualShadowMaps/VirtualShadowMapProjection.usf
--- a/Engine/Shaders/Private/VirtualShadowMaps/VirtualShadowMapProjection.usf
+++ b/Engine/Shaders/Private/VirtualShadowMaps/VirtualShadowMapProjection.usf
@@ -211,7 +211,7 @@ FVirtualShadowMapSampleResult ProjectLight(
 
 	// NOTE: If PrimaryView.PreViewTranslation is guaranteed to be the camera position we can simplify
 	const float DistanceToCamera = length(TranslatedWorldPosition - View.TranslatedWorldCameraOrigin);
-	const float NormalBiasLength = max(0.02f, NormalBias * DistanceToCamera / View.ViewToClip._11);
+	const float NormalBiasLength = max(0.02f, NormalBias * DistanceToCamera / GetCotanHalfFieldOfView().x);
 
 #if DIRECTIONAL_LIGHT
 	float3 L = Light.Direction;
@@ -589,7 +589,7 @@ void VirtualShadowMapProjection(
 	const float4 SvPosition = float4( float2( PixelPos ) + 0.5, DeviceZ, 1.0f );
 	const float3 TranslatedWorldPosition = SvPositionToTranslatedWorld( SvPosition );
 
-	const float ScreenRayLengthWorld = ScreenRayLength * View.ClipToView[1][1] * SceneDepth;
+	const float ScreenRayLengthWorld = ScreenRayLength * GetTanHalfFieldOfView().y * SceneDepth;
 	const float Noise = InterleavedGradientNoise( SvPosition.xy, View.StateFrameIndexMod8 );
 
 	const FProjectionShadingInfo ShadingInfo = GetProjectionShadingInfo(PixelPos);
diff --git a/Engine/Source/Runtime/Engine/Private/SceneView.cpp b/Engine/Source/Runtime/Engine/Private/SceneView.cpp
--- a/Engine/Source/Runtime/Engine/Private/SceneView.cpp
+++ b/Engine/Source/Runtime/Engine/Private/SceneView.cpp
@@ -574,6 +574,7 @@ void FViewMatrices::Init(const FMinimalInitializer& Initializer)
 	// Adjust the projection matrix for the current RHI.
 	ProjectionMatrix = AdjustProjectionMatrixForRHI(Initializer.ProjectionMatrix);
 	InvProjectionMatrix = InvertProjectionMatrix(ProjectionMatrix);
+	ScreenToClipMatrix = ScreenToClipProjectionMatrix();
 
 	// Compute the view projection matrix and its inverse.
 	ViewProjectionMatrix = GetViewMatrix() * GetProjectionMatrix();
@@ -589,6 +590,7 @@ void FViewMatrices::Init(const FMinimalInitializer& Initializer)
 	FMatrix LocalTranslatedViewMatrix = ViewRotationMatrix;
 	FMatrix LocalInvTranslatedViewMatrix = LocalTranslatedViewMatrix.GetTransposed();
 
+	//TODO(OrthoRendering) - Review bUseFauxOrthoViewPos paths + possibly deprecate
 	if (!IsPerspectiveProjection() && Initializer.bUseFauxOrthoViewPos)
 	{
 		auto DistanceToViewOrigin = UE_OLD_WORLD_MAX;
@@ -614,8 +616,16 @@ void FViewMatrices::Init(const FMinimalInitializer& Initializer)
 	// Stereo renders at half horizontal resolution, but compute shadow resolution based on full resolution.
 	const bool bStereo = IStereoRendering::IsStereoEyePass(Initializer.StereoPass);
 	const float ScreenXScale = bStereo ? 2.0f : 1.0f;
-	ProjectionScale.X = ScreenXScale * FMath::Abs(ProjectionMatrix.M[0][0]);
-	ProjectionScale.Y = FMath::Abs(ProjectionMatrix.M[1][1]);
+	if (IsPerspectiveProjection())
+	{
+		ProjectionScale.X = ScreenXScale * FMath::Abs(ProjectionMatrix.M[0][0]);
+		ProjectionScale.Y = FMath::Abs(ProjectionMatrix.M[1][1]);
+	}
+	else
+	{
+		//No FOV for ortho so do not scale
+		ProjectionScale = FVector2D(ScreenXScale, 1.0f);
+	}
 	ScreenScale = FMath::Max(
 		Initializer.ConstrainedViewRect.Size().X * 0.5f * ProjectionScale.X,
 		Initializer.ConstrainedViewRect.Size().Y * 0.5f * ProjectionScale.Y
@@ -1088,6 +1098,27 @@ void FViewMatrices::UpdateViewMatrix(const FVector& ViewLocation, const FRotator
 	InvTranslatedViewProjectionMatrix = GetInvProjectionMatrix() * GetInvTranslatedViewMatrix();
 }
 
+FMatrix FViewMatrices::ScreenToClipProjectionMatrix() const
+{
+	//Screen to clip matrix should not utilise scene depth for the w component in ortho projections, but is needed for perspective
+	if (IsPerspectiveProjection())
+	{
+		return FMatrix(
+			FPlane(1, 0, 0, 0),
+			FPlane(0, 1, 0, 0),
+			FPlane(0, 0, ProjectionMatrix.M[2][2], 1.0f),
+			FPlane(0, 0, ProjectionMatrix.M[3][2], 0.0f));
+	}
+	else
+	{
+		return FMatrix(
+			FPlane(1, 0, 0, 0),
+			FPlane(0, 1, 0, 0),
+			FPlane(0, 0, ProjectionMatrix.M[2][2], 0.0f),
+			FPlane(0, 0, ProjectionMatrix.M[3][2], 1.0f));
+	}
+}
+
 void FViewMatrices::HackOverrideViewMatrixForShadows(const FMatrix& InViewMatrix)
 {
 	OverriddenTranslatedViewMatrix = ViewMatrix = InViewMatrix;
@@ -2508,19 +2539,26 @@ void FSceneView::SetupViewRectUniformBufferParameters(FViewUniformShaderParamete
 				FPlane(0, 0, 1, 0),
 				FPlane(Ax, Ay, 0, 1)) * InViewMatrices.GetInvTranslatedViewProjectionMatrix());
 	}
+	
+	{
+		/**
+		* Ortho projection does not use FOV calculations, so rather than sourcing the projection matrix values in CommonViewUniformBuffer.ush,
+		* the appropriate values are uploaded in the per view uniform buffer (projection matrix values for perspective, 1.0f for ortho).
+		* Doing this here avoids unnecessarily checking for perspective vs ortho in shaders at runtime.
+		*/
 
-	// Compute coefficients which takes a screen UV and converts to Viewspace.xy / ViewZ
-	float InvTanHalfFov = InViewMatrices.GetProjectionMatrix().M[0][0];
-	float Ratio = UnscaledViewRect.Width() / (float)UnscaledViewRect.Height();
-
-	float InvFovFixX = 1.0f / (InvTanHalfFov);
-	float InvFovFixY = 1.0f / (Ratio * InvTanHalfFov);
+		ViewUniformShaderParameters.TanAndInvTanHalfFOV = InViewMatrices.GetTanAndInvTanHalfFOV();
+		ViewUniformShaderParameters.PrevTanAndInvTanHalfFOV = InPrevViewMatrices.GetTanAndInvTanHalfFOV();
+	}
+	
+	float FovFixX = ViewUniformShaderParameters.TanAndInvTanHalfFOV.X;
+	float FovFixY = ViewUniformShaderParameters.TanAndInvTanHalfFOV.Y;
 
-	ViewUniformShaderParameters.ScreenToViewSpace.X = BufferSize.X * ViewUniformShaderParameters.ViewSizeAndInvSize.Z * 2 * InvFovFixX;
-	ViewUniformShaderParameters.ScreenToViewSpace.Y = BufferSize.Y * ViewUniformShaderParameters.ViewSizeAndInvSize.W  * -2 * InvFovFixY;
+	ViewUniformShaderParameters.ScreenToViewSpace.X = BufferSize.X * ViewUniformShaderParameters.ViewSizeAndInvSize.Z * 2 * FovFixX;
+	ViewUniformShaderParameters.ScreenToViewSpace.Y = BufferSize.Y * ViewUniformShaderParameters.ViewSizeAndInvSize.W  * -2 * FovFixY;
 
-	ViewUniformShaderParameters.ScreenToViewSpace.Z = -((ViewUniformShaderParameters.ViewRectMin.X * ViewUniformShaderParameters.ViewSizeAndInvSize.Z * 2 * InvFovFixX) + InvFovFixX);
-	ViewUniformShaderParameters.ScreenToViewSpace.W = (ViewUniformShaderParameters.ViewRectMin.Y * ViewUniformShaderParameters.ViewSizeAndInvSize.W * 2 * InvFovFixY) + InvFovFixY;
+	ViewUniformShaderParameters.ScreenToViewSpace.Z = -((ViewUniformShaderParameters.ViewRectMin.X * ViewUniformShaderParameters.ViewSizeAndInvSize.Z * 2 * FovFixX) + FovFixX);
+	ViewUniformShaderParameters.ScreenToViewSpace.W = (ViewUniformShaderParameters.ViewRectMin.Y * ViewUniformShaderParameters.ViewSizeAndInvSize.W * 2 * FovFixY) + FovFixY;
 
 	ViewUniformShaderParameters.ViewResolutionFraction = EffectiveViewRect.Width() / (float)UnscaledViewRect.Width();
 
@@ -2637,33 +2675,17 @@ void FSceneView::SetupCommonViewUniformBufferParameters(
 
 	ViewUniformShaderParameters.bCheckerboardSubsurfaceProfileRendering = 0;
 
-	const FMatrix44f ScreenToClip(FPlane4f(1, 0, 0, 0),
-		FPlane4f(0, 1, 0, 0),
-		FPlane4f(0, 0, ProjectionMatrixUnadjustedForRHI.M[2][2], 1),
-		FPlane4f(0, 0, ProjectionMatrixUnadjustedForRHI.M[3][2], 0));
-	ViewUniformShaderParameters.ScreenToRelativeWorld = ScreenToClip * ViewUniformShaderParameters.ClipToRelativeWorld;
+	ViewUniformShaderParameters.ScreenToRelativeWorld = 
+		FMatrix44f(InViewMatrices.GetScreenToClipMatrix()) * ViewUniformShaderParameters.ClipToRelativeWorld;
 
-	ViewUniformShaderParameters.ScreenToTranslatedWorld = FMatrix44f(FMatrix(
-		FPlane(1, 0, 0, 0),
-		FPlane(0, 1, 0, 0),
-		FPlane(0, 0, ProjectionMatrixUnadjustedForRHI.M[2][2], 1),
-		FPlane(0, 0, ProjectionMatrixUnadjustedForRHI.M[3][2], 0))
-		* InViewMatrices.GetInvTranslatedViewProjectionMatrix());
+	ViewUniformShaderParameters.ScreenToTranslatedWorld = 
+		FMatrix44f(InViewMatrices.GetScreenToClipMatrix() * InViewMatrices.GetInvTranslatedViewProjectionMatrix());
 
-	ViewUniformShaderParameters.MobileMultiviewShadowTransform = FMatrix44f(FMatrix(
-		FPlane(1, 0, 0, 0),
-		FPlane(0, 1, 0, 0),
-		FPlane(0, 0, InViewMatrices.GetProjectionMatrix().M[2][2], 1),
-		FPlane(0, 0, InViewMatrices.GetProjectionMatrix().M[3][2], 0)) *
-		InViewMatrices.GetInvTranslatedViewProjectionMatrix() *
-		FTranslationMatrix(-InViewMatrices.GetPreViewTranslation()));
+	ViewUniformShaderParameters.MobileMultiviewShadowTransform = 
+		FMatrix44f(InViewMatrices.GetScreenToClipMatrix() * InViewMatrices.GetInvTranslatedViewProjectionMatrix() * FTranslationMatrix(-InViewMatrices.GetPreViewTranslation()));
 
-	ViewUniformShaderParameters.PrevScreenToTranslatedWorld = FMatrix44f(FMatrix(
-		FPlane(1, 0, 0, 0),
-		FPlane(0, 1, 0, 0),
-		FPlane(0, 0, ProjectionMatrixUnadjustedForRHI.M[2][2], 1),
-		FPlane(0, 0, ProjectionMatrixUnadjustedForRHI.M[3][2], 0))
-		* InPrevViewMatrices.GetInvTranslatedViewProjectionMatrix());
+	ViewUniformShaderParameters.PrevScreenToTranslatedWorld = 
+		FMatrix44f(InPrevViewMatrices.GetScreenToClipMatrix() * InPrevViewMatrices.GetInvTranslatedViewProjectionMatrix());
 
 	{
 		FVector DeltaTranslation = InPrevViewMatrices.GetPreViewTranslation() - InViewMatrices.GetPreViewTranslation();
diff --git a/Engine/Source/Runtime/Engine/Public/SceneView.h b/Engine/Source/Runtime/Engine/Public/SceneView.h
--- a/Engine/Source/Runtime/Engine/Public/SceneView.h
+++ b/Engine/Source/Runtime/Engine/Public/SceneView.h
@@ -278,6 +278,7 @@ struct FViewMatrices
 		OverriddenInvTranslatedViewMatrix.SetIdentity();
 		TranslatedViewProjectionMatrix.SetIdentity();
 		InvTranslatedViewProjectionMatrix.SetIdentity();
+		ScreenToClipMatrix.SetIdentity();
 		PreViewTranslation = FVector::ZeroVector;
 		ViewOrigin = FVector::ZeroVector;
 		ProjectionScale = FVector2D::ZeroVector;
@@ -320,6 +321,8 @@ private:
 	FMatrix		TranslatedViewProjectionMatrix;
 	/** The inverse view-projection transform, ending with world-space points translated by -ViewOrigin. */
 	FMatrix		InvTranslatedViewProjectionMatrix;
+	/** The screen to clip matrix (defined depending on whether this is a perspective or ortho projection view)*/
+	FMatrix		ScreenToClipMatrix;
 	/** The translation to apply to the world before TranslatedViewProjectionMatrix. Usually it is -ViewOrigin but with rereflections this can differ */
 	FVector		PreViewTranslation;
 	/** To support ortho and other modes this is redundant, in world space */
@@ -421,6 +424,11 @@ public:
 		return InvTranslatedViewProjectionMatrix;
 	}
 
+	inline const FMatrix& GetScreenToClipMatrix() const
+	{
+		return ScreenToClipMatrix;
+	}
+
 	inline const FVector& GetPreViewTranslation() const
 	{
 		return PreViewTranslation;
@@ -447,6 +455,42 @@ public:
 		return ProjectionMatrix.M[3][3] < 1.0f;
 	}
 
+	inline FVector2f GetInvTanHalfFov(bool YAxis = false) const
+	{
+		//No concept of FOV for orthographic projection so only return perspective related values or 1.0f
+		if (IsPerspectiveProjection())
+		{
+			return FVector2f(static_cast<float>(ProjectionMatrix.M[0][0]), static_cast<float>(ProjectionMatrix.M[1][1]));
+		}
+		return  FVector2f(1.0f, 1.0f);
+	}
+
+	inline FVector2f GetTanHalfFov() const
+	{
+		//No concept of FOV for orthographic projection so only return perspective related values or 1.0f
+		if (IsPerspectiveProjection())
+		{
+			return FVector2f(static_cast<float>(InvProjectionMatrix.M[0][0]), static_cast<float>(InvProjectionMatrix.M[1][1]));
+		}
+		return  FVector2f(1.0f,1.0f);
+	}
+
+	//Used for initializing the View Uniform Buffer TanAndInvTanHalfFOV variable without repeated calls to IsPerspectiveProjection
+	inline FVector4f GetTanAndInvTanHalfFOV() const
+	{
+		//No concept of FOV for orthographic projection so only return perspective related values or 1.0f
+		if (IsPerspectiveProjection())
+		{
+			FVector4f(static_cast<float>(InvProjectionMatrix.M[0][0]), //ClipToView[0][0] - X axis
+					  static_cast<float>(InvProjectionMatrix.M[1][1]), //ClipToView[1][1] - Y axis
+					  static_cast<float>(ProjectionMatrix.M[0][0]), //ViewToClip[0][0] - X axis
+					  static_cast<float>(ProjectionMatrix.M[1][1])); //ViewToClip[1][1] - Y axis
+		}
+		return FVector4f(1.0f, 1.0f, 1.0f, 1.0f);
+	}
+
+	FMatrix ScreenToClipProjectionMatrix() const;
+
 	ENGINE_API void HackOverrideViewMatrixForShadows(const FMatrix& InViewMatrix);
 
 	void SaveProjectionNoAAMatrix()
@@ -835,6 +879,8 @@ enum ETranslucencyVolumeCascade
 	VIEW_UNIFORM_BUFFER_MEMBER_PER_VIEW(uint32, GPUSceneViewId) \
 	VIEW_UNIFORM_BUFFER_MEMBER(float, ViewResolutionFraction) \
 	VIEW_UNIFORM_BUFFER_MEMBER(float, SubSurfaceColorAsTransmittanceAtDistanceInMeters) \
+	VIEW_UNIFORM_BUFFER_MEMBER_PER_VIEW(FVector4f, TanAndInvTanHalfFOV) \
+	VIEW_UNIFORM_BUFFER_MEMBER_PER_VIEW(FVector4f, PrevTanAndInvTanHalfFOV) \
 
 /** The uniform shader parameters associated with a view. */
 #define VIEW_UNIFORM_BUFFER_MEMBER(type, identifier) \
diff --git a/Engine/Source/Runtime/Renderer/Private/ShadowRendering.cpp b/Engine/Source/Runtime/Renderer/Private/ShadowRendering.cpp
--- a/Engine/Source/Runtime/Renderer/Private/ShadowRendering.cpp
+++ b/Engine/Source/Runtime/Renderer/Private/ShadowRendering.cpp
@@ -1713,11 +1713,7 @@ FMatrix FProjectedShadowInfo::GetScreenToShadowMatrix(const FSceneView& View, ui
 		// Z of the position being transformed is actually view space Z, 
 			// Transform it into post projection space by applying the projection matrix,
 			// Which is the required space before applying View.InvTranslatedViewProjectionMatrix
-		FMatrix(
-			FPlane(1,0,0,0),
-			FPlane(0,1,0,0),
-			FPlane(0,0,View.ViewMatrices.GetProjectionMatrix().M[2][2],1),
-			FPlane(0,0,View.ViewMatrices.GetProjectionMatrix().M[3][2],0)) *
+		View.ViewMatrices.GetScreenToClipMatrix() *
 		// Transform the post projection space position into translated world space
 		// Translated world space is normal world space translated to the view's origin, 
 		// Which prevents floating point imprecision far from the world origin.
