World Partition - HLOD: Log an error if an invalid class is provided for HLODInstancedStaticMeshComponentClass

diff --git a/Engine/Source/Runtime/Engine/Private/WorldPartition/HLOD/HLODBuilder.cpp b/Engine/Source/Runtime/Engine/Private/WorldPartition/HLOD/HLODBuilder.cpp
--- a/Engine/Source/Runtime/Engine/Private/WorldPartition/HLOD/HLODBuilder.cpp
+++ b/Engine/Source/Runtime/Engine/Private/WorldPartition/HLOD/HLODBuilder.cpp
@@ -10,6 +10,7 @@
 #include "ISMPartition/ISMComponentBatcher.h"
 #include "ISMPartition/ISMComponentDescriptor.h"
 #include "Materials/MaterialInterface.h"
+#include "Misc/ConfigCacheIni.h"
 #include "UObject/Package.h"
 #include "WorldPartition/HLOD/HLODTemplatedInstancedStaticMeshComponent.h"
 
@@ -150,6 +151,21 @@ uint32 UHLODBuilder::ComputeHLODHash(const TArray<AActor*>& InSourceActors)
 	return FCrc::MemCrc32(ComponentsCRCs.GetData(), ComponentsCRCs.Num() * ComponentsCRCs.GetTypeSize());
 }
 
+TSubclassOf<UInstancedStaticMeshComponent> UHLODBuilder::GetInstancedStaticMeshComponentClass()
+{
+	TSubclassOf<UInstancedStaticMeshComponent> ISMClass = StaticClass()->GetDefaultObject<UHLODBuilder>()->HLODInstancedStaticMeshComponentClass;
+	if (!ISMClass)
+	{
+		FString ConfigValue;
+		GConfig->GetString(TEXT("/Script/Engine.HLODBuilder"), TEXT("HLODInstancedStaticMeshComponentClass"), ConfigValue, GEditorIni);
+		UE_LOG(LogHLODBuilder, Error, TEXT("Could not resolve the class specified for HLODInstancedStaticMeshComponentClass. Config value was %s"), *ConfigValue);
+
+		// Fallback to standard ISMC
+		ISMClass = UInstancedStaticMeshComponent::StaticClass();
+	}
+	return ISMClass;
+}
+
 namespace
 {
 	// Instance batcher class based on FISMComponentDescriptor
diff --git a/Engine/Source/Runtime/Engine/Public/WorldPartition/HLOD/HLODBuilder.h b/Engine/Source/Runtime/Engine/Public/WorldPartition/HLOD/HLODBuilder.h
--- a/Engine/Source/Runtime/Engine/Public/WorldPartition/HLOD/HLODBuilder.h
+++ b/Engine/Source/Runtime/Engine/Public/WorldPartition/HLOD/HLODBuilder.h
@@ -118,10 +118,7 @@ public:
 	/** 
 	 * Get the InstancedStaticMeshComponent subclass that should be used when creating instanced HLODs.
 	 */
-	static TSubclassOf<UInstancedStaticMeshComponent> GetInstancedStaticMeshComponentClass()
-	{
-		return StaticClass()->GetDefaultObject<UHLODBuilder>()->HLODInstancedStaticMeshComponentClass;
-	}
+	static TSubclassOf<UInstancedStaticMeshComponent> GetInstancedStaticMeshComponentClass();
 
 protected:
 	virtual bool ShouldIgnoreBatchingPolicy() const { return false; }
@@ -147,6 +144,7 @@ protected:
 	UPROPERTY()
 	TObjectPtr<const UHLODBuilderSettings> HLODBuilderSettings;
 
+private:
 	UPROPERTY(Config)
 	TSubclassOf<UInstancedStaticMeshComponent> HLODInstancedStaticMeshComponentClass;
 #endif
