World Partition - HLOD: Guard against potentially deleted HLOD ISM BP classes

diff --git a/Engine/Source/Runtime/Engine/Private/WorldPartition/HLOD/HLODActor.cpp b/Engine/Source/Runtime/Engine/Private/WorldPartition/HLOD/HLODActor.cpp
--- a/Engine/Source/Runtime/Engine/Private/WorldPartition/HLOD/HLODActor.cpp
+++ b/Engine/Source/Runtime/Engine/Private/WorldPartition/HLOD/HLODActor.cpp
@@ -224,7 +224,10 @@ void AWorldPartitionHLOD::SetHLODComponents(const TArray<UActorComponent*>& InHL
 	TArray<UActorComponent*> ComponentsToRemove = GetInstanceComponents();
 	for (UActorComponent* ComponentToRemove : ComponentsToRemove)
 	{
-		ComponentToRemove->DestroyComponent();
+		if (ComponentToRemove)
+		{
+			ComponentToRemove->DestroyComponent();
+		}
 	}
 
 	// We'll turn on replication for this actor only if it contains a replicated component
