TickComponent and PostPhysicsTick in SkeletalMeshComponent can cause simultaneous read/write access

#fix set tick function priority as point of registering rather than during first-tick. The seen concurrent access only ever happened during the first tick, as priorities had not been set - causing misordering of the involved Anim/Physics tick functions
#misc left a check in TickFunction to ensure priority is set as expected

diff --git a/Engine/Source/Runtime/Engine/Private/Components/SkeletalMeshComponent.cpp b/Engine/Source/Runtime/Engine/Private/Components/SkeletalMeshComponent.cpp
--- a/Engine/Source/Runtime/Engine/Private/Components/SkeletalMeshComponent.cpp
+++ b/Engine/Source/Runtime/Engine/Private/Components/SkeletalMeshComponent.cpp
@@ -61,6 +61,11 @@ static TAutoConsoleVariable<float> CVarStallParallelAnimation(
 	0.0f,
 	TEXT("Sleep for the given time in each parallel animation task. Time is given in ms. This is a debug option used for critical path analysis and forcing a change in the critical path."));
 
+static TAutoConsoleVariable<int32> CVarHiPriSkinnedMeshesTicks(
+	TEXT("tick.HiPriSkinnedMeshes"),
+	1,
+	TEXT("If > 0, then schedule the skinned component ticks in a tick group before other ticks."));
+
 // Deprecated. Please switch to ANIM_SKINNED_ASSET_ISPC_ENABLED_DEFAULT and "a.SkinnedAsset.ISPC", see SkinnedAsset.cpp.
 #if !defined(ANIM_SKELETAL_MESH_ISPC_ENABLED_DEFAULT)
 #define ANIM_SKELETAL_MESH_ISPC_ENABLED_DEFAULT 1
@@ -403,6 +408,13 @@ void USkeletalMeshComponent::PostInitProperties()
 void USkeletalMeshComponent::RegisterComponentTickFunctions(bool bRegister)
 {
 	Super::RegisterComponentTickFunctions(bRegister);
+	const bool bDoHiPri = CVarHiPriSkinnedMeshesTicks.GetValueOnGameThread() > 0;
+	if (PrimaryComponentTick.bHighPriority != bDoHiPri)
+	{
+		// Note that if animation is so long that we are blocked in EndPhysics we may want to reduce the priority. However, there is a risk that this function will not go wide early enough.
+		// This requires profiling and is very game dependent so cvar for now makes sense
+		PrimaryComponentTick.SetPriorityIncludingPrerequisites(bDoHiPri);
+	}
 
 	UpdateEndPhysicsTickRegisteredState();
 	UpdateClothTickRegisteredState();
@@ -1531,11 +1543,6 @@ static TAutoConsoleVariable<int32> CVarAnimationDelaysEndGroup(
 	TEXT("tick.AnimationDelaysEndGroup"),
 	1,
 	TEXT("If > 0, then skeletal meshes that do not rely on physics simulation will set their animation end tick group to TG_PostPhysics."));
-static TAutoConsoleVariable<int32> CVarHiPriSkinnedMeshesTicks(
-	TEXT("tick.HiPriSkinnedMeshes"),
-	1,
-	TEXT("If > 0, then schedule the skinned component ticks in a tick group before other ticks."));
-
 
 void USkeletalMeshComponent::TickComponent(float DeltaTime, enum ELevelTick TickType, FActorComponentTickFunction *ThisTickFunction)
 {
@@ -1586,13 +1593,8 @@ void USkeletalMeshComponent::TickComponent(float DeltaTime, enum ELevelTick Tick
 	{
 		ThisTickFunction->EndTickGroup = EndTickGroup;
 
-		// Note that if animation is so long that we are blocked in EndPhysics we may want to reduce the priority. However, there is a risk that this function will not go wide early enough.
-		// This requires profiling and is very game dependent so cvar for now makes sense
-		bool bDoHiPri = CVarHiPriSkinnedMeshesTicks.GetValueOnGameThread() > 0;
-		if (ThisTickFunction->bHighPriority != bDoHiPri)
-		{
-			ThisTickFunction->SetPriorityIncludingPrerequisites(bDoHiPri);
-		}
+		const bool bDoHiPri = CVarHiPriSkinnedMeshesTicks.GetValueOnGameThread() > 0;
+		check(PrimaryComponentTick.bHighPriority == bDoHiPri)
 	}
 
 	// If we are waiting for ParallelEval to complete or if we require Physics, 
