Log package content when failing to find an Actor in ExternalActor package.
Tracking missing actors during cook.

diff --git a/Engine/Source/Runtime/Engine/Private/WorldPartition/WorldPartitionLevelHelper.cpp b/Engine/Source/Runtime/Engine/Private/WorldPartition/WorldPartitionLevelHelper.cpp
--- a/Engine/Source/Runtime/Engine/Private/WorldPartition/WorldPartitionLevelHelper.cpp
+++ b/Engine/Source/Runtime/Engine/Private/WorldPartition/WorldPartitionLevelHelper.cpp
@@ -492,7 +492,20 @@ bool FWorldPartitionLevelHelper::LoadActors(UWorld* InOwningWorld, ULevel* InDes
 			}
 			else
 			{
-				UE_LOG(LogEngine, Warning, TEXT("Failed to load %s"), *LoadedPackageName.ToString());
+				if (LoadedPackage)
+				{
+					UE_LOG(LogEngine, Warning, TEXT("Failed to find actor in package %s. Package Content:"), *LoadedPackageName.ToString());
+					ForEachObjectWithPackage(LoadedPackage, [](UObject* Object)
+					{
+						UE_LOG(LogEngine, Warning, TEXT("\t Object %s, Flags 0x%llx"), *Object->GetName(), static_cast<uint64>(Object->GetFlags()));
+						return true;
+					}, false);
+				}
+				else
+				{
+					UE_LOG(LogEngine, Warning, TEXT("Failed to load actor package %s"), *LoadedPackageName.ToString());
+				}
+
 				//@todo_ow: cumulate and process when NumPendingActorRequests == 0
 				LoadProgress->NumFailedLoadedRequests++;
 			}
