Fix to crash in loading screen when clicking X window button before the end of the Loading Screen. Reported on UDN.

diff --git a/Engine/Source/Runtime/Engine/Private/GameEngine.cpp b/Engine/Source/Runtime/Engine/Private/GameEngine.cpp
--- a/Engine/Source/Runtime/Engine/Private/GameEngine.cpp
+++ b/Engine/Source/Runtime/Engine/Private/GameEngine.cpp
@@ -673,7 +673,7 @@ void UGameEngine::SwitchGameWindowToUseGameViewport()
 				CreateGameViewport(GameViewport);
 			}
 
-			if (GameViewportWidget.IsValid())
+			if (GameViewportWidget.IsValid() && FSlateApplication::IsInitialized())
 			{
 				GameViewportWindowPtr->SetContent(GameViewportWidget.ToSharedRef());
 				GameViewportWindowPtr->SlatePrepass(FSlateApplication::Get().GetApplicationScale() * GameViewportWindowPtr->GetNativeWindow()->GetDPIScaleFactor());
@@ -683,20 +683,17 @@ void UGameEngine::SwitchGameWindowToUseGameViewport()
 				UE_LOG(LogEngine, Error, TEXT("The Game Viewport Widget is invalid."));
 			}
 
-			if (SceneViewport.IsValid())
+			// If Scene Viewport is not valid, the window was closed.
+			if (SceneViewport.IsValid() && FSlateApplication::IsInitialized())
 			{
 				SceneViewport->ResizeFrame((uint32)GSystemResolution.ResX, (uint32)GSystemResolution.ResY, GSystemResolution.WindowMode);
-			}
-
-			// Registration of the game viewport to that messages are correctly received.
-			// Could be a re-register, however it's necessary after the window is set.
-			if (GameViewportWidget.IsValid())
-			{
-				FSlateApplication::Get().RegisterGameViewport(GameViewportWidget.ToSharedRef());
-			}
 
-			if (FSlateApplication::IsInitialized())
-			{
+				// Registration of the game viewport to that messages are correctly received.
+				// Could be a re-register, however it's necessary after the window is set.
+				if (GameViewportWidget.IsValid())
+				{
+					FSlateApplication::Get().RegisterGameViewport(GameViewportWidget.ToSharedRef());
+				}
 				FSlateApplication::Get().SetAllUserFocusToGameViewport(EFocusCause::SetDirectly);
 			}
 		}
