Clarifying MinLODSize and MaxLODSize

diff --git a/Engine/Source/Runtime/Engine/Classes/Engine/TextureLODSettings.h b/Engine/Source/Runtime/Engine/Classes/Engine/TextureLODSettings.h
--- a/Engine/Source/Runtime/Engine/Classes/Engine/TextureLODSettings.h
+++ b/Engine/Source/Runtime/Engine/Classes/Engine/TextureLODSettings.h
@@ -74,18 +74,23 @@ struct ENGINE_API FTextureLODGroup
 	UPROPERTY()
 	TEnumAsByte<TextureMipGenSettings> MipGenSettings;
 
+	/** Prevent LODBias from making the textures smaller than this value. Note that this does _not_ affect the smallest mip level size. */
 	UPROPERTY()
 	int32 MinLODSize;
 
+	/** Cap the number of mips such that the largest mip is this big. Has no effect for textures with no mip chain. Not used for virtual textures. */
 	UPROPERTY()
 	int32 MaxLODSize;
 
+	/** Cap the number of mips such that the largest mip is this big. Has no effect for textures with no mip chain. Used for platforms with the "Smaller" memory bucket. Not used for virtual textures. */
 	UPROPERTY()
 	int32 MaxLODSize_Smaller;
 
+	/** Cap the number of mips such that the largest mip is this big. Has no effect for textures with no mip chain. Used for platforms with the "Smallest" memory bucket. Not used for virtual textures. */
 	UPROPERTY()
 	int32 MaxLODSize_Smallest;
 
+	/** Cap the number of mips such that the largest mip is this big. Has no effect for textures with no mip chain. Used for virtual textures. */
 	UPROPERTY()
 	int32 MaxLODSize_VT;
 
diff --git a/Engine/Source/Runtime/Engine/Private/TextureLODSettings.cpp b/Engine/Source/Runtime/Engine/Private/TextureLODSettings.cpp
--- a/Engine/Source/Runtime/Engine/Private/TextureLODSettings.cpp
+++ b/Engine/Source/Runtime/Engine/Private/TextureLODSettings.cpp
@@ -201,6 +201,10 @@ int32 UTextureLODSettings::CalculateLODBias(int32 Width, int32 Height, int32 Max
 		// -> this looks broken, probably never used
 	}
 
+	// Min/Max LOD Size clamps the *LODBias*, not the actual texture LODs. Meaning:
+	//	MaxLODSize caps the largest mip size.
+	//	MinLODSize prevents LODBias from making the largest mip smaller than the given value. Does *not*
+	//	affect the size of the smallest mip. Almost never used.
 	int32 MinLOD = FMath::CeilLogTwo(LODGroupInfo.MinLODSize);
 	int32 MaxLOD = FMath::CeilLogTwo(LODGroupInfo.MaxLODSize);
 
