Fix for CIS issue : Nanite Landscape static meshes can now have a unique name so the previous name-based hack to detect if the static mesh is a Nanite landscape one was failing.

diff --git a/Engine/Source/Runtime/Engine/Classes/Engine/StaticMesh.h b/Engine/Source/Runtime/Engine/Classes/Engine/StaticMesh.h
--- a/Engine/Source/Runtime/Engine/Classes/Engine/StaticMesh.h
+++ b/Engine/Source/Runtime/Engine/Classes/Engine/StaticMesh.h
@@ -879,7 +879,7 @@ public:
 	// TODO: Temp/deprecated hack - Do not call
 	inline bool IsNaniteLandscape() const
 	{
-		return GetName() == TEXT("LandscapeNaniteMesh");
+		return GetName().StartsWith(TEXT("LandscapeNaniteMesh"));
 	}
 
 private:
