Fix off by one logic in check()

diff --git a/Engine/Source/Runtime/Engine/Private/Streaming/StreamingTexture.h b/Engine/Source/Runtime/Engine/Private/Streaming/StreamingTexture.h
--- a/Engine/Source/Runtime/Engine/Private/Streaming/StreamingTexture.h
+++ b/Engine/Source/Runtime/Engine/Private/Streaming/StreamingTexture.h
@@ -116,7 +116,7 @@ struct FStreamingRenderAsset
 	float GetNormalizedScreenSize(int32 NumMips) const
 	{
 		check(IsMesh());
-		check(NumMips > 0 && (uint32)NumMips < UE_ARRAY_COUNT(LODScreenSizes));
+		check(NumMips > 0 && (uint32)NumMips <= UE_ARRAY_COUNT(LODScreenSizes));
 		return LODScreenSizes[NumMips - 1];
 	}
 
