Only check bStartStalledHandles once.

diff --git a/Engine/Source/Runtime/Engine/Private/StreamableManager.cpp b/Engine/Source/Runtime/Engine/Private/StreamableManager.cpp
--- a/Engine/Source/Runtime/Engine/Private/StreamableManager.cpp
+++ b/Engine/Source/Runtime/Engine/Private/StreamableManager.cpp
@@ -363,26 +363,29 @@ EAsyncPackageState::Type FStreamableHandle::WaitUntilComplete(float Timeout, boo
 	}
 
 	// We need to recursively start any stalled handles
-	TArray<TSharedRef<FStreamableHandle>> HandlesToStart;
-
-	HandlesToStart.Add(AsShared());
-
-	for (int32 i = 0; i < HandlesToStart.Num(); i++)
+	if (bStartStalledHandles)
 	{
-		TSharedRef<FStreamableHandle> Handle = HandlesToStart[i];
+		TArray<TSharedRef<FStreamableHandle>> HandlesToStart;
 
-		if (bStartStalledHandles && Handle->IsStalled())
-		{
-			// If we were stalled, start us now to avoid deadlocks
-			UE_LOG(LogStreamableManager, Warning, TEXT("FStreamableHandle::WaitUntilComplete called on stalled handle %s, forcing load even though resources may not have been acquired yet"), *Handle->GetDebugName());
-			Handle->StartStalledHandle();
-		}
+		HandlesToStart.Add(AsShared());
 
-		for (const TSharedPtr<FStreamableHandle>& ChildHandle : Handle->ChildHandles)
+		for (int32 i = 0; i < HandlesToStart.Num(); i++)
 		{
-			if (ChildHandle.IsValid())
+			TSharedRef<FStreamableHandle> Handle = HandlesToStart[i];
+
+			if (Handle->IsStalled())
 			{
-				HandlesToStart.Add(ChildHandle.ToSharedRef());
+				// If we were stalled, start us now to avoid deadlocks
+				UE_LOG(LogStreamableManager, Warning, TEXT("FStreamableHandle::WaitUntilComplete called on stalled handle %s, forcing load even though resources may not have been acquired yet"), *Handle->GetDebugName());
+				Handle->StartStalledHandle();
+			}
+
+			for (const TSharedPtr<FStreamableHandle>& ChildHandle : Handle->ChildHandles)
+			{
+				if (ChildHandle.IsValid())
+				{
+					HandlesToStart.Add(ChildHandle.ToSharedRef());
+				}
 			}
 		}
 	}
