Explicitly perform PerFrameCleanupIfSkipRenderer and FlushRHIThreadFlushResources when rendering is suspended

* Fixes per-frame memory leak when rendering is suspended

diff --git a/Engine/Source/Runtime/Engine/Private/GameEngine.cpp b/Engine/Source/Runtime/Engine/Private/GameEngine.cpp
--- a/Engine/Source/Runtime/Engine/Private/GameEngine.cpp
+++ b/Engine/Source/Runtime/Engine/Private/GameEngine.cpp
@@ -1878,7 +1878,9 @@ void UGameEngine::Tick( float DeltaSeconds, bool bIdleMode )
 		}
 	}
 
-	if (!bIdleMode && !IsRenderingSuspended() && !IsRunningDedicatedServer() && !IsRunningCommandlet() && FEmbeddedCommunication::IsAwakeForRendering())
+	const bool bRenderingSuspended = IsRenderingSuspended();
+
+	if (!bIdleMode && !bRenderingSuspended && !IsRunningDedicatedServer() && !IsRunningCommandlet() && FEmbeddedCommunication::IsAwakeForRendering())
 	{
 		// Render everything.
 		RedrawViewports();
@@ -1918,6 +1920,16 @@ void UGameEngine::Tick( float DeltaSeconds, bool bIdleMode )
 			FRDGBuilder::TickPoolElements();
 			ICustomResourcePool::TickPoolElements(RHICmdList);
 		});
+
+		if (bRenderingSuspended)
+		{
+			GetRendererModule().PerFrameCleanupIfSkipRenderer();
+			ENQUEUE_RENDER_COMMAND(UGameEngine_Tick_FlushRHIResources)(
+				[](FRHICommandListImmediate& RHICmdList)
+				{
+					RHICmdList.ImmediateFlush(EImmediateFlushType::FlushRHIThreadFlushResources);
+				});
+		}
 	}
 
 #if WITH_EDITOR
