Added logging and additional information to ensure message to help find what assets are causing an ensure in GetAnimationPose
#tests Verified information I expect to see is output from the new log message and when the ensure is hit

diff --git a/Engine/Source/Runtime/Engine/Private/Animation/PoseAsset.cpp b/Engine/Source/Runtime/Engine/Private/Animation/PoseAsset.cpp
--- a/Engine/Source/Runtime/Engine/Private/Animation/PoseAsset.cpp
+++ b/Engine/Source/Runtime/Engine/Private/Animation/PoseAsset.cpp
@@ -610,6 +610,7 @@ bool UPoseAsset::GetAnimationPose(struct FAnimationPoseData& OutAnimationPoseDat
 		FBlendedCurve& OutCurve = OutAnimationPoseData.GetCurve();
 
 		const FBoneContainer& RequiredBones = OutPose.GetBoneContainer();
+		const USkeletalMesh* RequiredBonesMesh = RequiredBones.GetSkeletalMeshAsset();
 		USkeleton* MySkeleton = GetSkeleton();
 		
 		check(PoseContainer.IsValid());
@@ -648,7 +649,7 @@ bool UPoseAsset::GetAnimationPose(struct FAnimationPoseData& OutAnimationPoseDat
 					{
 						const FSkeletonPoseBoneIndex SkeletonBoneIndex = FSkeletonPoseBoneIndex(SkeletonRemapping.IsValid() ? SkeletonRemapping.GetTargetSkeletonBoneIndex(PoseContainer.TrackBoneIndices[TrackIndex]) : PoseContainer.TrackBoneIndices[TrackIndex]);
 
-						ensureMsgf(RequiredBones.IsSkeletonPoseIndexValid(SkeletonBoneIndex), TEXT("PoseAsset [%s] with skeleton [%s] has bones not present in the evaluation container. Bone(%s) index(%d) not found."), *GetPathName(), MySkeleton ? *MySkeleton->GetPathName() : TEXT("<Skeleton Not Found>"), *PoseContainer.Tracks[TrackIndex].ToString(), SkeletonBoneIndex.GetInt());
+						ensureMsgf(RequiredBones.IsSkeletonPoseIndexValid(SkeletonBoneIndex), TEXT("PoseAsset [%s] with skeleton [%s] has bones not present in the evaluation container. Bone(%s) index(%d) not found in skeletal mesh %s."), *GetPathName(), MySkeleton ? *MySkeleton->GetPathName() : TEXT("<Skeleton Not Found>"), *PoseContainer.Tracks[TrackIndex].ToString(), SkeletonBoneIndex.GetInt(), RequiredBonesMesh ? *RequiredBonesMesh->GetName() : TEXT("<Skeletal Mesh Not Found>"));
 
 						const FCompactPoseBoneIndex CompactIndex = RequiredBones.GetCompactPoseIndexFromSkeletonPoseIndex(SkeletonBoneIndex);
 						
@@ -712,7 +713,7 @@ bool UPoseAsset::GetAnimationPose(struct FAnimationPoseData& OutAnimationPoseDat
 			{
 				const FSkeletonPoseBoneIndex SkeletonBoneIndex = FSkeletonPoseBoneIndex(SkeletonRemapping.IsValid() ? SkeletonRemapping.GetTargetSkeletonBoneIndex(PoseContainer.TrackBoneIndices[TrackIndex]) : PoseContainer.TrackBoneIndices[TrackIndex]);
 
-				ensureMsgf(RequiredBones.IsSkeletonPoseIndexValid(SkeletonBoneIndex), TEXT("PoseAsset [%s] with skeleton [%s] has bones not present in the evaluation container. Bone(%s) index(%d) not found."), *GetPathName(), MySkeleton ? *MySkeleton->GetPathName() : TEXT("<Skeleton Not Found>"), *PoseContainer.Tracks[TrackIndex].ToString(), SkeletonBoneIndex.GetInt());
+				ensureMsgf(RequiredBones.IsSkeletonPoseIndexValid(SkeletonBoneIndex), TEXT("PoseAsset [%s] with skeleton [%s] has bones not present in the evaluation container. Bone(%s) index(%d) not found in skeletal mesh %s."), *GetPathName(), MySkeleton ? *MySkeleton->GetPathName() : TEXT("<Skeleton Not Found>"), *PoseContainer.Tracks[TrackIndex].ToString(), SkeletonBoneIndex.GetInt(), RequiredBonesMesh ? *RequiredBonesMesh->GetName() : TEXT("<Skeletal Mesh Not Found>"));
 
 				const FCompactPoseBoneIndex CompactIndex = RequiredBones.GetCompactPoseIndexFromSkeletonPoseIndex(SkeletonBoneIndex);
 
