World Partition Class Descriptors: better loggin when failing to gather redirected assets from the asset registry.

diff --git a/Engine/Source/Runtime/Engine/Private/WorldPartition/WorldPartitionClassDescRegistry.cpp b/Engine/Source/Runtime/Engine/Private/WorldPartition/WorldPartitionClassDescRegistry.cpp
--- a/Engine/Source/Runtime/Engine/Private/WorldPartition/WorldPartitionClassDescRegistry.cpp
+++ b/Engine/Source/Runtime/Engine/Private/WorldPartition/WorldPartitionClassDescRegistry.cpp
@@ -170,7 +170,7 @@ void FWorldPartitionClassDescRegistry::PrefetchClassDescs(const TArray<FTopLevel
 
 			if (RedirectAssets.IsEmpty())
 			{
-				UE_LOG(LogWorldPartition, Warning, TEXT("Failed to find redirected assets for '%s'"), *AssetData.ToSoftObjectPath().ToString());
+				UE_LOG(LogWorldPartition, Warning, TEXT("Failed to find redirected assets for '%s' from '%s'"), *AssetData.ToSoftObjectPath().ToString(), *AssetClassPath.ToString());
 				AssetData = FAssetData();
 				break;
 			}
