Fixed issue where actorchannels being closed due to level being unloaded/streamed out destroyed NetStartupActors unless channel had been explicitly broken.
This behavior can be disabled by cvar net.SkipDestroyNetStartupActorsOnChannelCloseDueToLevelUnloaded 0

diff --git a/Engine/Source/Runtime/Engine/Private/DataChannel.cpp b/Engine/Source/Runtime/Engine/Private/DataChannel.cpp
--- a/Engine/Source/Runtime/Engine/Private/DataChannel.cpp
+++ b/Engine/Source/Runtime/Engine/Private/DataChannel.cpp
@@ -161,6 +161,12 @@ namespace UE::Net
 	static bool GMetAsyncDemoNameDebugChance = false;
 #endif
 
+	static bool bSkipDestroyNetStartupActorsOnChannelCloseDueToLevelUnloaded = true;
+	static FAutoConsoleVariableRef CVarSkipDestroyNetStartupActorsOnChannelCloseDueToLevelUnloaded(
+		TEXT("net.SkipDestroyNetStartupActorsOnChannelCloseDueToLevelUnloaded"),
+		bSkipDestroyNetStartupActorsOnChannelCloseDueToLevelUnloaded,
+		TEXT("Controls if Actor that is a NetStartUpActor assosciated with the channel is destroyed or not when we receive a channel close with ECloseReason::LevelUnloaded."), ECVF_Default);
+
 	// bTearOff is private but we still might need to adjust the flag on clients based on the channel close reason
 	class FTearOffSetter final
 	{
@@ -2474,6 +2480,10 @@ bool UActorChannel::CleanUp(const bool bForDestroy, EChannelCloseReason CloseRea
 				Connection->Driver->NotifyActorFullyDormantForConnection(Actor, Connection);
 				bWasDormant = true;
 			}
+			else if (UE::Net::bSkipDestroyNetStartupActorsOnChannelCloseDueToLevelUnloaded && (CloseReason == EChannelCloseReason::LevelUnloaded) && Actor->IsNetStartupActor())
+			{
+				UE_LOG(LogNetTraffic, Verbose, TEXT("UActorChannel::CleanUp: Skipped Destroying NetStartupActor %s due to EChannelCloseReason::LevelUnloaded"), *Describe());
+			}
 			else if (!Actor->bNetTemporary && Actor->GetWorld() != NULL && !IsEngineExitRequested() && Connection->Driver->ShouldClientDestroyActor(Actor))
 			{
 				UE_LOG(LogNetDormancy, Verbose, TEXT("UActorChannel::CleanUp: Destroying Actor. %s"), *Describe() );
