Remove 5.1 deprecated code + added params struct to get container interface for actor descriptors.

diff --git a/Engine/Source/Runtime/Engine/Private/Tests/WorldPartition/WorldPartitionHandleTests.cpp b/Engine/Source/Runtime/Engine/Private/Tests/WorldPartition/WorldPartitionHandleTests.cpp
--- a/Engine/Source/Runtime/Engine/Private/Tests/WorldPartition/WorldPartitionHandleTests.cpp
+++ b/Engine/Source/Runtime/Engine/Private/Tests/WorldPartition/WorldPartitionHandleTests.cpp
@@ -323,10 +323,10 @@ namespace WorldPartitionTests
 			{
 				TUniquePtr<FWorldPartitionActorDesc> NewActorDesc(AActor::StaticCreateClassActorDesc(ActorNativeClass));
 
-				FWorldPartitionActorDescInitData ActorDescInitData;
-				ActorDescInitData.NativeClass = ActorNativeClass;
-				ActorDescInitData.PackageName = ActorDescIterator->GetActorPackage();
-				ActorDescInitData.ActorPath = ActorDescIterator->GetActorSoftPath();
+				FWorldPartitionActorDescInitData ActorDescInitData = FWorldPartitionActorDescInitData()
+					.SetNativeClass(ActorNativeClass)
+					.SetPackageName(ActorDescIterator->GetActorPackage())
+					.SetActorPath(ActorDescIterator->GetActorSoftPath());
 			
 				ActorDescIterator->SerializeTo(ActorDescInitData.SerializedData);
 				NewActorDesc->Init(ActorDescInitData);
diff --git a/Engine/Source/Runtime/Engine/Private/WorldPartition/LevelInstance/LevelInstanceActorDesc.cpp b/Engine/Source/Runtime/Engine/Private/WorldPartition/LevelInstance/LevelInstanceActorDesc.cpp
--- a/Engine/Source/Runtime/Engine/Private/WorldPartition/LevelInstance/LevelInstanceActorDesc.cpp
+++ b/Engine/Source/Runtime/Engine/Private/WorldPartition/LevelInstance/LevelInstanceActorDesc.cpp
@@ -163,7 +163,7 @@ bool FLevelInstanceActorDesc::IsContainerInstance() const
 	return ULevelInstanceSubsystem::CanUsePackage(LevelPackage);
 }
 
-bool FLevelInstanceActorDesc::GetContainerInstance(FContainerInstance& OutContainerInstance, bool bInBuildFilter) const
+bool FLevelInstanceActorDesc::GetContainerInstance(const FGetContainerInstanceParams& InParams, FContainerInstance& OutContainerInstance) const
 {
 	if (LevelInstanceContainer.IsValid())
 	{
@@ -174,7 +174,7 @@ bool FLevelInstanceActorDesc::GetContainerInstance(FContainerInstance& OutContai
 		FTransform LevelInstancePivotOffsetTransform = FTransform(ULevel::GetLevelInstancePivotOffsetFromPackage(LevelInstanceContainer->GetContainerPackage()));
 		OutContainerInstance.Transform = LevelInstancePivotOffsetTransform * LevelInstanceTransform;
 
-		if (bInBuildFilter)
+		if (InParams.bBuildFilter)
 		{
 			// Fill Container Instance Filter
 			ULevelInstanceSubsystem* LevelInstanceSubsystem = UWorld::GetSubsystem<ULevelInstanceSubsystem>(LevelInstanceContainerWorldContext.Get());
@@ -204,9 +204,9 @@ bool FLevelInstanceActorDesc::GetContainerInstance(FContainerInstance& OutContai
 				}
 			};
 
-			ProcessFilter(OutContainerInstance.GetID(), ContainerFilter);
+			ProcessFilter(InParams.ContainerID, ContainerFilter);
 
-			TFunction<void(const FActorContainerID&, const UActorDescContainer*)> ProcessContainers = [&FilteredOutDataLayersPerContainer, &OutContainerInstance, &ProcessContainers](const FActorContainerID& InContainerID, const UActorDescContainer* InContainer)
+			TFunction<void(const FActorContainerID&, const UActorDescContainer*)> ProcessContainers = [&FilteredOutDataLayersPerContainer, &OutContainerInstance, &ProcessContainers, &InParams](const FActorContainerID& InContainerID, const UActorDescContainer* InContainer)
 			{
 				const TSet<FSoftObjectPath>& Filtered = FilteredOutDataLayersPerContainer.FindChecked(InContainerID);
 				for (FActorDescList::TConstIterator<> ActorDescIt(InContainer); ActorDescIt; ++ActorDescIt)
@@ -227,8 +227,7 @@ bool FLevelInstanceActorDesc::GetContainerInstance(FContainerInstance& OutContai
 					if (ActorDescIt->IsContainerInstance())
 					{
 						FWorldPartitionActorDesc::FContainerInstance ChildContainerInstance;
-						const bool bBuildFilter = false;
-						if (ActorDescIt->GetContainerInstance(ChildContainerInstance, bBuildFilter))
+						if (ActorDescIt->GetContainerInstance(FWorldPartitionActorDesc::FGetContainerInstanceParams(), ChildContainerInstance))
 						{
 							ProcessContainers(FActorContainerID(InContainerID, ActorDescIt->GetGuid()), ChildContainerInstance.Container);
 						}
@@ -236,7 +235,7 @@ bool FLevelInstanceActorDesc::GetContainerInstance(FContainerInstance& OutContai
 				}
 			};
 
-			ProcessContainers(OutContainerInstance.GetID(), OutContainerInstance.Container);
+			ProcessContainers(InParams.ContainerID, OutContainerInstance.Container);
 		}
 		return true;
 	}
diff --git a/Engine/Source/Runtime/Engine/Private/WorldPartition/WorldPartitionActorDescUtils.cpp b/Engine/Source/Runtime/Engine/Private/WorldPartition/WorldPartitionActorDescUtils.cpp
--- a/Engine/Source/Runtime/Engine/Private/WorldPartition/WorldPartitionActorDescUtils.cpp
+++ b/Engine/Source/Runtime/Engine/Private/WorldPartition/WorldPartitionActorDescUtils.cpp
@@ -73,10 +73,10 @@ TUniquePtr<FWorldPartitionActorDesc> FWorldPartitionActorDescUtils::GetActorDesc
 {
 	if (IsValidActorDescriptorFromAssetData(InAssetData))
 	{
-		FWorldPartitionActorDescInitData ActorDescInitData;
-		ActorDescInitData.NativeClass = GetActorNativeClassFromAssetData(InAssetData);
-		ActorDescInitData.PackageName = InAssetData.PackageName;
-		ActorDescInitData.ActorPath = InAssetData.GetSoftObjectPath();
+		FWorldPartitionActorDescInitData ActorDescInitData = FWorldPartitionActorDescInitData()
+			.SetNativeClass(GetActorNativeClassFromAssetData(InAssetData))
+			.SetPackageName(InAssetData.PackageName)
+			.SetActorPath(InAssetData.GetSoftObjectPath());
 
 		FString ActorMetaDataStr;
 		verify(InAssetData.GetTagValue(NAME_ActorMetaData, ActorMetaDataStr));
diff --git a/Engine/Source/Runtime/Engine/Private/WorldPartition/WorldPartitionActorDescView.cpp b/Engine/Source/Runtime/Engine/Private/WorldPartition/WorldPartitionActorDescView.cpp
--- a/Engine/Source/Runtime/Engine/Private/WorldPartition/WorldPartitionActorDescView.cpp
+++ b/Engine/Source/Runtime/Engine/Private/WorldPartition/WorldPartitionActorDescView.cpp
@@ -184,9 +184,9 @@ FName FWorldPartitionActorDescView::GetLevelPackage() const
 	return ActorDesc->GetLevelPackage();
 }
 
-bool FWorldPartitionActorDescView::GetContainerInstance(FWorldPartitionActorDesc::FContainerInstance& OutContainerInstance, bool bInBuildFilter) const
+bool FWorldPartitionActorDescView::GetContainerInstance(const FWorldPartitionActorDesc::FGetContainerInstanceParams& InParams, FWorldPartitionActorDesc::FContainerInstance& OutContainerInstance) const
 {
-	return ActorDesc->GetContainerInstance(OutContainerInstance, bInBuildFilter);
+	return ActorDesc->GetContainerInstance(InParams, OutContainerInstance);
 }
 
 void FWorldPartitionActorDescView::CheckForErrors(IStreamingGenerationErrorHandler* ErrorHandler) const
diff --git a/Engine/Source/Runtime/Engine/Private/WorldPartition/WorldPartitionStreamingGeneration.cpp b/Engine/Source/Runtime/Engine/Private/WorldPartition/WorldPartitionStreamingGeneration.cpp
--- a/Engine/Source/Runtime/Engine/Private/WorldPartition/WorldPartitionStreamingGeneration.cpp
+++ b/Engine/Source/Runtime/Engine/Private/WorldPartition/WorldPartitionStreamingGeneration.cpp
@@ -475,9 +475,14 @@ class FWorldPartitionStreamingGenerator
 		for (const FWorldPartitionActorDescView& ContainerInstanceView : ContainerInstanceViews)
 		{
 			const FGuid ActorGuid = ContainerInstanceView.GetGuid();
-			FWorldPartitionActorDesc::FContainerInstance SubContainerInstance(FActorContainerID(InContainerInstanceDescriptor.ID, ActorGuid));
-			bool bBuildFilter = InContainerInstanceDescriptor.ID.IsMainContainer(); // Build filter on top level child containers
-			if (!ContainerInstanceView.GetContainerInstance(SubContainerInstance, bBuildFilter))
+			const FActorContainerID ContainerID(InContainerInstanceDescriptor.ID, ActorGuid);
+
+			FWorldPartitionActorDesc::FGetContainerInstanceParams GetContainerInstanceParams = FWorldPartitionActorDesc::FGetContainerInstanceParams()
+				.SetContainerID(ContainerID)
+				.SetBuildFilter(InContainerInstanceDescriptor.ID.IsMainContainer()); // Build filter on top level child containers
+
+			FWorldPartitionActorDesc::FContainerInstance SubContainerInstance;
+			if (!ContainerInstanceView.GetContainerInstance(GetContainerInstanceParams, SubContainerInstance))
 			{
 				ErrorHandler->OnLevelInstanceInvalidWorldAsset(ContainerInstanceView, ContainerInstanceView.GetLevelPackage(), IStreamingGenerationErrorHandler::ELevelInstanceInvalidReason::WorldAssetHasInvalidContainer);
 				continue;
@@ -495,7 +500,7 @@ class FWorldPartitionStreamingGenerator
 			FContainerInstanceDescriptor SubContainerInstanceDescriptor;
 
 			// Inherit fields from FContainerInstance
-			SubContainerInstanceDescriptor.ID = SubContainerInstance.GetID();
+			SubContainerInstanceDescriptor.ID = ContainerID;
 			SubContainerInstanceDescriptor.Container = SubContainerInstance.Container;
 			SubContainerInstanceDescriptor.Transform = SubContainerInstance.Transform * InContainerInstanceDescriptor.Transform;
 			SubContainerInstanceDescriptor.FilteredActors = SubContainerInstance.FilteredActors;
diff --git a/Engine/Source/Runtime/Engine/Private/WorldPartitionBlueprintLibrary.cpp b/Engine/Source/Runtime/Engine/Private/WorldPartitionBlueprintLibrary.cpp
--- a/Engine/Source/Runtime/Engine/Private/WorldPartitionBlueprintLibrary.cpp
+++ b/Engine/Source/Runtime/Engine/Private/WorldPartitionBlueprintLibrary.cpp
@@ -52,9 +52,7 @@ FActorDesc::FActorDesc(const FWorldPartitionActorDesc& InActorDesc, const FTrans
 		}
 	}
 	
-	PRAGMA_DISABLE_DEPRECATION_WARNINGS
-	ActorPath = InActorDesc.GetActorPath();
-	PRAGMA_ENABLE_DEPRECATION_WARNINGS
+	ActorPath = *InActorDesc.GetActorSoftPath().ToString();
 }
 
 TMap<UWorldPartition*, TUniquePtr<FLoaderAdapterActorList>> UWorldPartitionBlueprintLibrary::LoaderAdapterActorListMap;
@@ -104,7 +102,7 @@ bool UWorldPartitionBlueprintLibrary::GetActorDescs(const UActorDescContainer* I
 		if (ActorDescIt->IsContainerInstance())
 		{
 			FWorldPartitionActorDesc::FContainerInstance ContainerInstance;
-			if (ActorDescIt->GetContainerInstance(ContainerInstance))
+			if (ActorDescIt->GetContainerInstance(FWorldPartitionActorDesc::FGetContainerInstanceParams(), ContainerInstance))
 			{
 				bResult &= GetActorDescs(ContainerInstance.Container, ContainerInstance.Transform * InTransform, OutActorDescs);
 			}
@@ -129,7 +127,7 @@ bool UWorldPartitionBlueprintLibrary::HandleIntersectingActorDesc(const FWorldPa
 	if (ActorDesc->IsContainerInstance())
 	{
 		FWorldPartitionActorDesc::FContainerInstance ContainerInstance;
-		if (ActorDesc->GetContainerInstance(ContainerInstance))
+		if (ActorDesc->GetContainerInstance(FWorldPartitionActorDesc::FGetContainerInstanceParams(), ContainerInstance))
 		{
 			bResult &= GetIntersectingActorDescs(ContainerInstance.Container, InBox, ContainerInstance.Transform * InTransform, OutActorDescs);
 		}
diff --git a/Engine/Source/Runtime/Engine/Public/WorldPartition/LevelInstance/LevelInstanceActorDesc.h b/Engine/Source/Runtime/Engine/Public/WorldPartition/LevelInstance/LevelInstanceActorDesc.h
--- a/Engine/Source/Runtime/Engine/Public/WorldPartition/LevelInstance/LevelInstanceActorDesc.h
+++ b/Engine/Source/Runtime/Engine/Public/WorldPartition/LevelInstance/LevelInstanceActorDesc.h
@@ -27,7 +27,7 @@ public:
 
 	virtual bool IsContainerInstance() const override;
 	virtual FName GetLevelPackage() const override { return LevelPackage; }
-	virtual bool GetContainerInstance(FContainerInstance& OutContainerInstance, bool bInBuildFilter) const override;
+	virtual bool GetContainerInstance(const FGetContainerInstanceParams& InParams, FContainerInstance& OutContainerInstance) const override;
 	virtual const FWorldPartitionActorFilter* GetContainerFilter() const override { return &Filter; }
 	virtual void CheckForErrors(IStreamingGenerationErrorHandler* ErrorHandler) const override;
 
diff --git a/Engine/Source/Runtime/Engine/Public/WorldPartition/WorldPartitionActorContainerID.h b/Engine/Source/Runtime/Engine/Public/WorldPartition/WorldPartitionActorContainerID.h
--- a/Engine/Source/Runtime/Engine/Public/WorldPartition/WorldPartitionActorContainerID.h
+++ b/Engine/Source/Runtime/Engine/Public/WorldPartition/WorldPartitionActorContainerID.h
@@ -17,7 +17,7 @@ struct FActorContainerID
 	: ID(InOther.ID)
 	{}
 
-	FActorContainerID(const FActorContainerID& InParent, FGuid InActorGuid)
+	FActorContainerID(const FActorContainerID& InParent, const FGuid& InActorGuid)
 	: ID(CityHash64WithSeed((const char*)&InActorGuid, sizeof(InActorGuid), InParent.ID))
 	{}
 
diff --git a/Engine/Source/Runtime/Engine/Public/WorldPartition/WorldPartitionActorDesc.h b/Engine/Source/Runtime/Engine/Public/WorldPartition/WorldPartitionActorDesc.h
--- a/Engine/Source/Runtime/Engine/Public/WorldPartition/WorldPartitionActorDesc.h
+++ b/Engine/Source/Runtime/Engine/Public/WorldPartition/WorldPartitionActorDesc.h
@@ -21,6 +21,10 @@ struct FWorldPartitionActorDescInitData
 	FName PackageName;
 	FSoftObjectPath ActorPath;
 	TArray<uint8> SerializedData;
+
+	FWorldPartitionActorDescInitData& SetNativeClass(UClass* InNativeClass) { NativeClass = InNativeClass; return *this; }
+	FWorldPartitionActorDescInitData& SetPackageName(FName InPackageName) { PackageName = InPackageName; return *this; }
+	FWorldPartitionActorDescInitData& SetActorPath(const FSoftObjectPath& InActorPath) { ActorPath = InActorPath; return *this; }
 };
 
 class AActor;
@@ -77,31 +81,25 @@ class ENGINE_API FWorldPartitionActorDesc
 public:
 	struct FContainerInstance
 	{
-		FContainerInstance() {}
-		FContainerInstance(const FActorContainerID& InContainerID)
-			: ContainerID(InContainerID) {}
-	public:
-		const FActorContainerID& GetID() const { return ContainerID; }
-
 		const UActorDescContainer* Container = nullptr;
 		FTransform Transform = FTransform::Identity;
 		EContainerClusterMode ClusterMode;
 		TMap<FActorContainerID, TSet<FGuid>> FilteredActors;
-	private:
+	};
+
+	struct FGetContainerInstanceParams
+	{
 		FActorContainerID ContainerID;
+		bool bBuildFilter = false;
+
+		FGetContainerInstanceParams& SetContainerID(FActorContainerID InContainerID) { ContainerID = InContainerID; return *this; }
+		FGetContainerInstanceParams& SetBuildFilter(bool bInBuildFilter) { bBuildFilter = bInBuildFilter; return *this; }
 	};
 
 	virtual ~FWorldPartitionActorDesc() {}
 
 	inline const FGuid& GetGuid() const { return Guid; }
 	
-PRAGMA_DISABLE_DEPRECATION_WARNINGS
-	UE_DEPRECATED(5.1, "GetClass is deprecated, GetBaseClass or GetNativeClass should be used instead.")
-	inline FName GetClass() const { return GetNativeClass().ToFName(); }
-	UE_DEPRECATED(5.1, "GetActorPath is deprecated, GetActorSoftPath should be used instead.")
-	inline FName GetActorPath() const { return ActorPath.ToFName(); }
-PRAGMA_ENABLE_DEPRECATION_WARNINGS
-
 	inline FTopLevelAssetPath GetBaseClass() const { return BaseClass; }
 	inline FTopLevelAssetPath GetNativeClass() const { return NativeClass; }
 	inline UClass* GetActorNativeClass() const { return ActorNativeClass; }
@@ -120,9 +118,6 @@ PRAGMA_ENABLE_DEPRECATION_WARNINGS
 	UE_DEPRECATED(5.1, "SetIsSpatiallyLoadedRaw is deprecated and should not be used.")
 	inline void SetIsSpatiallyLoadedRaw(bool bNewIsSpatiallyLoaded) { bIsSpatiallyLoaded = bNewIsSpatiallyLoaded; }
 
-	UE_DEPRECATED(5.1, "GetLevelBoundsRelevant is deprecated.")
-	inline bool GetLevelBoundsRelevant() const { return false; }
-
 	inline bool GetActorIsHLODRelevant() const { return bActorIsHLODRelevant; }
 	inline FName GetHLODLayer() const { return HLODLayer; }
 	inline const TArray<FName>& GetDataLayers() const { return DataLayers; }
@@ -156,7 +151,7 @@ PRAGMA_ENABLE_DEPRECATION_WARNINGS
 	virtual bool IsContainerInstance() const { return false; }
 	virtual FName GetLevelPackage() const { return NAME_None; }
 	virtual const FWorldPartitionActorFilter* GetContainerFilter() const { return nullptr; }
-	virtual bool GetContainerInstance(FContainerInstance& OutContainerInstance, bool bInBuildFilter = false) const { return false; }
+	virtual bool GetContainerInstance(const FGetContainerInstanceParams& InParams, FContainerInstance& OutContainerInstance) const { return false; }
 
 	FGuid GetContentBundleGuid() const;
 
@@ -250,9 +245,6 @@ public:
 	AActor* Load() const;
 	virtual void Unload();
 
-	UE_DEPRECATED(5.1, "ShouldBeLoadedByEditorCells is deprecated, IsEditorRelevant should be used instead.")
-	bool ShouldBeLoadedByEditorCells() const { return IsEditorRelevant(); }
-
 	virtual void Init(const AActor* InActor);
 	virtual void Init(const FWorldPartitionActorDescInitData& DescData);
 
diff --git a/Engine/Source/Runtime/Engine/Public/WorldPartition/WorldPartitionActorDescView.h b/Engine/Source/Runtime/Engine/Public/WorldPartition/WorldPartitionActorDescView.h
--- a/Engine/Source/Runtime/Engine/Public/WorldPartition/WorldPartitionActorDescView.h
+++ b/Engine/Source/Runtime/Engine/Public/WorldPartition/WorldPartitionActorDescView.h
@@ -65,7 +65,7 @@ PRAGMA_ENABLE_DEPRECATION_WARNINGS
 
 	bool IsContainerInstance() const;
 	FName GetLevelPackage() const;
-	bool GetContainerInstance(FWorldPartitionActorDesc::FContainerInstance& OutContainerInstance, bool bInBuildFilter) const;
+	bool GetContainerInstance(const FWorldPartitionActorDesc::FGetContainerInstanceParams& InParams, FWorldPartitionActorDesc::FContainerInstance& OutContainerInstance) const;
 
 	void CheckForErrors(IStreamingGenerationErrorHandler* ErrorHandler) const;
 
