Fix UStaticMesh registry tags showing Nanite has disabled even though Nanite is force enabled.

diff --git a/Engine/Source/Runtime/Engine/Private/StaticMesh.cpp b/Engine/Source/Runtime/Engine/Private/StaticMesh.cpp
--- a/Engine/Source/Runtime/Engine/Private/StaticMesh.cpp
+++ b/Engine/Source/Runtime/Engine/Private/StaticMesh.cpp
@@ -4589,7 +4589,7 @@ int32 UStaticMesh::GetNumSectionsWithCollision() const
 void UStaticMesh::GetAssetRegistryTags(TArray<FAssetRegistryTag>& OutTags) const
 {
 #if WITH_EDITORONLY_DATA
-	OutTags.Add(FAssetRegistryTag("NaniteEnabled", NaniteSettings.bEnabled ? TEXT("True") : TEXT("False"), FAssetRegistryTag::TT_Alphabetical));
+	OutTags.Add(FAssetRegistryTag("NaniteEnabled", IsNaniteEnabled() ? TEXT("True") : TEXT("False"), FAssetRegistryTag::TT_Alphabetical));
 	OutTags.Add(FAssetRegistryTag("NaniteFallbackPercent", FString::Printf(TEXT("%.1f"), NaniteSettings.FallbackPercentTriangles * 100.0f), FAssetRegistryTag::TT_Numerical));
 
 	if (AssetImportData)
