Cameras: improve post-process settings blending

Previously, the handling of post-process settings in the engine only had two layers: the "base" layer established by the world (and the PP volumes therein), and the main active camera, whose PP settings were blended over the world's.

This meant that if a new camera (pending view target) was blending over the current main camera (view target), none of its PP settings would be taken into account. This would result in the new camera's PP settings suddenly "popping" once the blending reached 100% and it took over as the main camera.

This change adds a "blend order" to PP settings, with the pending view target adding its PP settings with the "override" order, rather than the "base" order that preserves the previous behaviour.

diff --git a/Engine/Source/Runtime/Engine/Classes/Camera/PlayerCameraManager.h b/Engine/Source/Runtime/Engine/Classes/Camera/PlayerCameraManager.h
--- a/Engine/Source/Runtime/Engine/Classes/Camera/PlayerCameraManager.h
+++ b/Engine/Source/Runtime/Engine/Classes/Camera/PlayerCameraManager.h
@@ -42,6 +42,13 @@ enum EViewTargetBlendFunction : int
 	VTBlend_MAX,
 };
 
+UENUM()
+enum EViewTargetBlendOrder : int
+{
+	VTBlendOrder_Base,
+	VTBlendOrder_Override
+};
+
 /** 
  * Cached camera POV info, stored as optimization so we only
  * need to do a full camera update once per tick.
@@ -368,17 +375,23 @@ protected:
 	/** Internal list of weights for active post process effects. Parallel array to PostProcessBlendCache. */
 	TArray<float> PostProcessBlendCacheWeights;
 
+	/** Internal list of how to blend post process effects. */
+	TArray<EViewTargetBlendOrder> PostProcessBlendCacheOrders;
+
 public:
 	/** Adds a postprocess effect at the given weight. */
-	void AddCachedPPBlend(struct FPostProcessSettings& PPSettings, float BlendWeight);
+	void AddCachedPPBlend(struct FPostProcessSettings& PPSettings, float BlendWeight, EViewTargetBlendOrder BlendOrder = VTBlendOrder_Base);
+
+	/** Returns active post process info. */
+	void GetCachedPostProcessBlends(TArray<struct FPostProcessSettings> const*& OutPPSettings, TArray<float> const*& OutBlendWeights) const;
+
+	/** Returns active post process info. */
+	void GetCachedPostProcessBlends(TArray<struct FPostProcessSettings> const*& OutPPSettings, TArray<float> const*& OutBlendWeights, TArray<EViewTargetBlendOrder> const*& OutBlendOrders) const;
+
 protected:
 	/** Removes all postprocess effects. */
 	void ClearCachedPPBlends();
 
-public:
-	/** Returns active post process info. */
-	void GetCachedPostProcessBlends(TArray<struct FPostProcessSettings> const*& OutPPSettings, TArray<float> const*& OutBlendWeigthts) const;
-	
 protected:
 	/** Internal. Receives the output of individual camera animations. */
 	UPROPERTY(transient)
diff --git a/Engine/Source/Runtime/Engine/Private/LocalPlayer.cpp b/Engine/Source/Runtime/Engine/Private/LocalPlayer.cpp
--- a/Engine/Source/Runtime/Engine/Private/LocalPlayer.cpp
+++ b/Engine/Source/Runtime/Engine/Private/LocalPlayer.cpp
@@ -878,24 +878,38 @@ FSceneView* ULocalPlayer::CalcSceneView( class FSceneViewFamily* ViewFamily,
 	{
 		View->StartFinalPostprocessSettings(ViewInfo.Location);
 
-		// CameraAnim override
+		TArray<FPostProcessSettings> const* CameraAnimPPSettings = nullptr;
+		TArray<float> const* CameraAnimPPBlendWeights = nullptr;
+		TArray<EViewTargetBlendOrder> const* CameraAnimPPBlendOrders = nullptr;
+
+		// Base overrides (post process volumes, etc)
 		if (PlayerController->PlayerCameraManager)
 		{
-			TArray<FPostProcessSettings> const* CameraAnimPPSettings;
-			TArray<float> const* CameraAnimPPBlendWeights;
-			PlayerController->PlayerCameraManager->GetCachedPostProcessBlends(CameraAnimPPSettings, CameraAnimPPBlendWeights);
+			PlayerController->PlayerCameraManager->GetCachedPostProcessBlends(CameraAnimPPSettings, CameraAnimPPBlendWeights, CameraAnimPPBlendOrders);
 
 			for (int32 PPIdx = 0; PPIdx < CameraAnimPPBlendWeights->Num(); ++PPIdx)
 			{
-				View->OverridePostProcessSettings( (*CameraAnimPPSettings)[PPIdx], (*CameraAnimPPBlendWeights)[PPIdx]);
+				if ((*CameraAnimPPBlendOrders)[PPIdx] == VTBlendOrder_Base)
+				{
+					View->OverridePostProcessSettings( (*CameraAnimPPSettings)[PPIdx], (*CameraAnimPPBlendWeights)[PPIdx]);
+				}
 			}
 		}
 
-		//	CAMERA OVERRIDE
+		// Main camera
 		View->OverridePostProcessSettings(ViewInfo.PostProcessSettings, ViewInfo.PostProcessBlendWeight);
 
+		// Camera overrides (cameras blending in, camera modifiers, etc)
 		if (PlayerController->PlayerCameraManager)
 		{
+			for (int32 PPIdx = 0; PPIdx < CameraAnimPPBlendWeights->Num(); ++PPIdx)
+			{
+				if ((*CameraAnimPPBlendOrders)[PPIdx] == VTBlendOrder_Override)
+				{
+					View->OverridePostProcessSettings( (*CameraAnimPPSettings)[PPIdx], (*CameraAnimPPBlendWeights)[PPIdx]);
+				}
+			}
+
 			PlayerController->PlayerCameraManager->UpdatePhotographyPostProcessing(View->FinalPostProcessSettings);
 		}
 
diff --git a/Engine/Source/Runtime/Engine/Private/PlayerCameraManager.cpp b/Engine/Source/Runtime/Engine/Private/PlayerCameraManager.cpp
--- a/Engine/Source/Runtime/Engine/Private/PlayerCameraManager.cpp
+++ b/Engine/Source/Runtime/Engine/Private/PlayerCameraManager.cpp
@@ -292,23 +292,33 @@ void APlayerCameraManager::ApplyCameraModifiers(float DeltaTime, FMinimalViewInf
 	}
 }
 
-void APlayerCameraManager::AddCachedPPBlend(struct FPostProcessSettings& PPSettings, float BlendWeight)
+void APlayerCameraManager::AddCachedPPBlend(struct FPostProcessSettings& PPSettings, float BlendWeight, EViewTargetBlendOrder BlendOrder)
 {
 	check(PostProcessBlendCache.Num() == PostProcessBlendCacheWeights.Num());
+	check(PostProcessBlendCache.Num() == PostProcessBlendCacheOrders.Num());
 	PostProcessBlendCache.Add(PPSettings);
 	PostProcessBlendCacheWeights.Add(BlendWeight);
+	PostProcessBlendCacheOrders.Add(BlendOrder);
 }
 
 void APlayerCameraManager::ClearCachedPPBlends()
 {
 	PostProcessBlendCache.Empty();
 	PostProcessBlendCacheWeights.Empty();
+	PostProcessBlendCacheOrders.Empty();
 }
 
-void APlayerCameraManager::GetCachedPostProcessBlends(TArray<FPostProcessSettings> const*& OutPPSettings, TArray<float> const*& OutBlendWeigthts) const
+void APlayerCameraManager::GetCachedPostProcessBlends(TArray<FPostProcessSettings> const*& OutPPSettings, TArray<float> const*& OutBlendWeights) const
 {
 	OutPPSettings = &PostProcessBlendCache;
-	OutBlendWeigthts = &PostProcessBlendCacheWeights;
+	OutBlendWeights = &PostProcessBlendCacheWeights;
+}
+
+void APlayerCameraManager::GetCachedPostProcessBlends(TArray<FPostProcessSettings> const*& OutPPSettings, TArray<float> const*& OutBlendWeights, TArray<EViewTargetBlendOrder> const*& OutBlendOrders) const
+{
+	OutPPSettings = &PostProcessBlendCache;
+	OutBlendWeights = &PostProcessBlendCacheWeights;
+	OutBlendOrders = &PostProcessBlendCacheOrders;
 }
 
 void APlayerCameraManager::UpdateViewTargetInternal(FTViewTarget& OutVT, float DeltaTime)
@@ -881,6 +891,14 @@ void APlayerCameraManager::DoUpdateCamera(float DeltaTime)
 			// Update pending view target blend
 			NewPOV = ViewTarget.POV;
 			NewPOV.BlendViewInfo(PendingViewTarget.POV, BlendPct);//@TODO: CAMERA: Make sure the sense is correct!  BlendViewTargets(ViewTarget, PendingViewTarget, BlendPct);
+
+			// Add this pending view target's post-process settings as an override of the main view target's one,
+			// since it is blending on top of it.
+			const float PendingViewTargetPPWeight = PendingViewTarget.POV.PostProcessBlendWeight * BlendPct;
+			if (PendingViewTargetPPWeight > 0.f)
+			{
+				AddCachedPPBlend(PendingViewTarget.POV.PostProcessSettings, PendingViewTargetPPWeight, VTBlendOrder_Override);
+			}
 		}
 		else
 		{
