Make sure not to discard depth only and reversed index buffers when streaming in a LOD if those buffers have already been initialized during initial loading. Discrepancy can occur if r.SupportDepthOnlyIndexBuffers or r.SupportReversedIndexBuffers changes value between the initial load and a later stream-in. However, streaming cannot change the initialized states because the index buffers have already been cached in various rendering data structures.

diff --git a/Engine/Source/Runtime/Engine/Private/StaticMesh.cpp b/Engine/Source/Runtime/Engine/Private/StaticMesh.cpp
--- a/Engine/Source/Runtime/Engine/Private/StaticMesh.cpp
+++ b/Engine/Source/Runtime/Engine/Private/StaticMesh.cpp
@@ -477,8 +477,9 @@ void FStaticMeshLODResources::SerializeBuffers(FArchive& Ar, UStaticMesh* OwnerS
 {
 	Ar.UsingCustomVersion(FUE5ReleaseStreamObjectVersion::GUID);
 
-	bool bEnableDepthOnlyIndexBuffer = (CVarSupportDepthOnlyIndexBuffers.GetValueOnAnyThread() == 1);
-	bool bEnableReversedIndexBuffer = (CVarSupportReversedIndexBuffers.GetValueOnAnyThread() == 1);
+	// If the index buffers have already been initialized, do not change the behavior since the RHI resource pointer may have been cached somewhere already.
+	const bool bEnableDepthOnlyIndexBuffer = bHasDepthOnlyIndices || (CVarSupportDepthOnlyIndexBuffers.GetValueOnAnyThread() == 1);
+	const bool bEnableReversedIndexBuffer = bHasReversedIndices || bHasReversedDepthOnlyIndices || (CVarSupportReversedIndexBuffers.GetValueOnAnyThread() == 1);
 
 	// See if the mesh wants to keep resources CPU accessible
 	bool bMeshCPUAcces = OwnerStaticMesh ? OwnerStaticMesh->bAllowCPUAccess : false;
