Fixed strata re-entrant function by detecting it in a simple fashion for now using stack max depth.
Tested the case it fixes and Opal level.

diff --git a/Engine/Source/Runtime/Engine/Classes/Engine/EngineTypes.h b/Engine/Source/Runtime/Engine/Classes/Engine/EngineTypes.h
--- a/Engine/Source/Runtime/Engine/Classes/Engine/EngineTypes.h
+++ b/Engine/Source/Runtime/Engine/Classes/Engine/EngineTypes.h
@@ -692,6 +692,9 @@ enum EStrataShadingModel : int
 };
 static_assert(SSM_NUM <= 16, "Do not exceed 16 shading models without expanding FStrataMaterialShadingModelField to support uint32 instead of uint16!");
 
+// This used to track cyclic graph which we do not support. We only support acyclic graph and a depth of 128 is already too high for a realistic use case.
+#define STRATA_TREE_MAX_DEPTH 48
+
 /** Gather information from the Substrate material graph to setup material for runtime. */
 USTRUCT()
 struct ENGINE_API FStrataMaterialInfo
@@ -728,6 +731,23 @@ public:
 	bool operator==(const FStrataMaterialInfo& Other) const { return ShadingModelField == Other.GetShadingModelField(); }
 	bool operator!=(const FStrataMaterialInfo& Other) const { return ShadingModelField != Other.GetShadingModelField(); }
 
+#if WITH_EDITOR
+	bool PushStrataTreeStack()
+	{
+		bOutOfStackDepthWhenParsing = bOutOfStackDepthWhenParsing || (++ParsingStackDepth > STRATA_TREE_MAX_DEPTH);
+		return bOutOfStackDepthWhenParsing;
+	}
+	void PopStrataTreeStack()
+	{
+		ParsingStackDepth--;
+		check(ParsingStackDepth >= 0);
+	}
+	bool GetStrataTreeOutOfStackDepthOccurred() 
+	{
+		return bOutOfStackDepthWhenParsing;
+	}
+#endif
+
 private:
 	UPROPERTY()
 	uint16 ShadingModelField = 0;
@@ -742,6 +762,12 @@ private:
 	
 	UPROPERTY()
 	TArray<TObjectPtr<USubsurfaceProfile>> SubsurfaceProfiles;
+
+#if WITH_EDITOR
+	// A simple way to detect and prevent node re-entry due to cycling graph; stop the compilation and avoid crashing.
+	bool bOutOfStackDepthWhenParsing = false;
+	int32 ParsingStackDepth = 0;
+#endif
 };
 
 /** Describes how textures are sampled for materials */
diff --git a/Engine/Source/Runtime/Engine/Private/Materials/HLSLMaterialTranslator.cpp b/Engine/Source/Runtime/Engine/Private/Materials/HLSLMaterialTranslator.cpp
--- a/Engine/Source/Runtime/Engine/Private/Materials/HLSLMaterialTranslator.cpp
+++ b/Engine/Source/Runtime/Engine/Private/Materials/HLSLMaterialTranslator.cpp
@@ -390,6 +390,7 @@ FHLSLMaterialTranslator::FHLSLMaterialTranslator(FMaterial* InMaterial,
 
 	// Default value used as the root of the tree for the first path (when a node parent==nullptr).
 	StrataNodeIdentifierStack.Push(FGuid(0x7AEE, 0xBAD, 0xDEAD, 0xBEEF));
+	bStrataTreeOutOfStackDepthOccurred = false;
 
 	if (InStrataCompilationConfig)
 	{
@@ -814,6 +815,10 @@ bool FHLSLMaterialTranslator::Translate()
 				check(StrataThicknessStack.Num() == 0);
 			}
 
+			if (bStrataTreeOutOfStackDepthOccurred)
+			{
+				Errorf(TEXT(" %s [%s]: Substrate - Cyclic graph detected when we only support acyclic graph."), *Material->GetDebugName(), *Material->GetAssetPath().ToString());
+			}
 			if (!StrataGenerateDerivedMaterialOperatorData())
 			{
 				Errorf(TEXT("Substrate material errors encountered."));
@@ -11747,6 +11752,8 @@ FGuid FHLSLMaterialTranslator::StrataTreeStackPush(UMaterialExpression* Expressi
 	UE_LOG(LogMaterial, Display, TEXT(" StrataTreeStack: %s."), *GuidStack);
 #endif
 
+	bStrataTreeOutOfStackDepthOccurred = bStrataTreeOutOfStackDepthOccurred || (StrataNodeIdentifierStack.Num() > STRATA_TREE_MAX_DEPTH);
+
 	return StrataNodeIdentifierStack.Top();
 }
 
@@ -11782,6 +11789,11 @@ void FHLSLMaterialTranslator::StrataTreeStackPop()
 #endif
 }
 
+bool FHLSLMaterialTranslator::GetStrataTreeOutOfStackDepthOccurred()
+{
+	return bStrataTreeOutOfStackDepthOccurred;
+}
+
 int32 FHLSLMaterialTranslator::StrataThicknessStackGetThicknessIndex()
 {
 	return StrataThicknessStack.Top();
diff --git a/Engine/Source/Runtime/Engine/Private/Materials/HLSLMaterialTranslator.h b/Engine/Source/Runtime/Engine/Private/Materials/HLSLMaterialTranslator.h
--- a/Engine/Source/Runtime/Engine/Private/Materials/HLSLMaterialTranslator.h
+++ b/Engine/Source/Runtime/Engine/Private/Materials/HLSLMaterialTranslator.h
@@ -452,6 +452,12 @@ protected:
 	*/
 	TArray<FGuid> StrataNodeIdentifierStack;
 
+	/**
+	 * This can be used to know if the strata tree we are trying to build is too deep and we should stop the compilation.
+	 * True means that we have likely encountered node re-entry leading to cyclic graph we cannot handle and compile internally: we must fail the compilation.
+	 */
+	bool bStrataTreeOutOfStackDepthOccurred;
+
 	/** Stack of thickness input used for propagating thickness information from root node and vertical operation
 	* This is transient and updated when calling StrataGenerateMaterialTopologyTree. The information is then stored into the FStratOperator
 	*/
@@ -1138,6 +1144,7 @@ protected:
 	virtual FGuid StrataTreeStackGetPathUniqueId() override;
 	virtual FGuid StrataTreeStackGetParentPathUniqueId() override;
 	virtual void StrataTreeStackPop() override;
+	virtual bool GetStrataTreeOutOfStackDepthOccurred() override;
 	
 	virtual int32 StrataThicknessStackGetThicknessIndex() override;
 	virtual int32 StrataThicknessStackGetThicknessCode(int32 Index) override;
diff --git a/Engine/Source/Runtime/Engine/Private/Materials/Material.cpp b/Engine/Source/Runtime/Engine/Private/Materials/Material.cpp
--- a/Engine/Source/Runtime/Engine/Private/Materials/Material.cpp
+++ b/Engine/Source/Runtime/Engine/Private/Materials/Material.cpp
@@ -4805,6 +4805,13 @@ void UMaterial::RebuildShadingModelField()
 		{
 			EditorOnly->FrontMaterial.Expression->GatherStrataMaterialInfo(StrataMaterialInfo, EditorOnly->FrontMaterial.OutputIndex);
 			this->CachedConnectedInputs = StrataMaterialInfo.GetPropertyConnected();
+
+			if (StrataMaterialInfo.GetStrataTreeOutOfStackDepthOccurred())
+			{
+				CancelOutstandingCompilation();
+				UE_LOG(LogMaterial, Error, TEXT("%s: Substrate - Cyclic graph detected when we only support acyclic graph."), *GetName());
+				return;
+			}
 		}
 
 		bool bSanitizeMaterial = false;
diff --git a/Engine/Source/Runtime/Engine/Private/Materials/MaterialExpressions.cpp b/Engine/Source/Runtime/Engine/Private/Materials/MaterialExpressions.cpp
--- a/Engine/Source/Runtime/Engine/Private/Materials/MaterialExpressions.cpp
+++ b/Engine/Source/Runtime/Engine/Private/Materials/MaterialExpressions.cpp
@@ -24040,17 +24040,28 @@ bool UMaterialExpressionStrataConvertToDecal::IsResultStrataMaterial(int32 Outpu
 
 void UMaterialExpressionStrataConvertToDecal::GatherStrataMaterialInfo(FStrataMaterialInfo& StrataMaterialInfo, int32 OutputIndex)
 {
+	if (!StrataMaterialInfo.PushStrataTreeStack())
+	{
+		return;
+	}
+
 	if (DecalMaterial.GetTracedInput().Expression)
 	{
 		DecalMaterial.GetTracedInput().Expression->GatherStrataMaterialInfo(StrataMaterialInfo, DecalMaterial.OutputIndex);
 	}
 	StrataMaterialInfo.AddShadingModel(SSM_Decal);
+
+	StrataMaterialInfo.PopStrataTreeStack();
 }
 
 FStrataOperator* UMaterialExpressionStrataConvertToDecal::StrataGenerateMaterialTopologyTree(class FMaterialCompiler* Compiler, class UMaterialExpression* Parent, int32 OutputIndex)
 {
 	const bool bUseParameterBlending = true;
 	FStrataOperator& StrataOperator = Compiler->StrataCompilationRegisterOperator(STRATA_OPERATOR_WEIGHT, Compiler->StrataTreeStackGetPathUniqueId(), Parent, Compiler->StrataTreeStackGetParentPathUniqueId(), bUseParameterBlending);
+	if (Compiler->GetStrataTreeOutOfStackDepthOccurred())
+	{
+		return &StrataOperator; // Out ot stack space, return now to fail the compilation
+	}
 
 	UMaterialExpression* ChildDecalMaterialExpression = DecalMaterial.GetTracedInput().Expression;
 	FStrataOperator* OpA = nullptr;
@@ -24626,6 +24637,11 @@ bool UMaterialExpressionStrataHorizontalMixing::IsResultStrataMaterial(int32 Out
 
 void UMaterialExpressionStrataHorizontalMixing::GatherStrataMaterialInfo(FStrataMaterialInfo& StrataMaterialInfo, int32 OutputIndex)
 {
+	if (!StrataMaterialInfo.PushStrataTreeStack())
+	{
+		return;
+	}
+
 	if (Foreground.GetTracedInput().Expression)
 	{
 		Foreground.GetTracedInput().Expression->GatherStrataMaterialInfo(StrataMaterialInfo, Foreground.OutputIndex);
@@ -24634,11 +24650,17 @@ void UMaterialExpressionStrataHorizontalMixing::GatherStrataMaterialInfo(FStrata
 	{
 		Background.GetTracedInput().Expression->GatherStrataMaterialInfo(StrataMaterialInfo, Background.OutputIndex);
 	}
+
+	StrataMaterialInfo.PopStrataTreeStack();
 }
 
 FStrataOperator* UMaterialExpressionStrataHorizontalMixing::StrataGenerateMaterialTopologyTree(class FMaterialCompiler* Compiler, class UMaterialExpression* Parent, int32 OutputIndex)
 {
 	FStrataOperator& StrataOperator = Compiler->StrataCompilationRegisterOperator(STRATA_OPERATOR_HORIZONTAL, Compiler->StrataTreeStackGetPathUniqueId(), Parent, Compiler->StrataTreeStackGetParentPathUniqueId(), bUseParameterBlending);
+	if (Compiler->GetStrataTreeOutOfStackDepthOccurred())
+	{
+		return &StrataOperator; // Out ot stack space, return now to fail the compilation
+	}
 
 	UMaterialExpression* ChildAExpression = Background.GetTracedInput().Expression;
 	UMaterialExpression* ChildBExpression = Foreground.GetTracedInput().Expression;
@@ -24809,6 +24831,11 @@ bool UMaterialExpressionStrataVerticalLayering::IsResultStrataMaterial(int32 Out
 
 void UMaterialExpressionStrataVerticalLayering::GatherStrataMaterialInfo(FStrataMaterialInfo& StrataMaterialInfo, int32 OutputIndex)
 {
+	if (!StrataMaterialInfo.PushStrataTreeStack())
+	{
+		return;
+	}
+
 	if (Top.GetTracedInput().Expression)
 	{
 		Top.GetTracedInput().Expression->GatherStrataMaterialInfo(StrataMaterialInfo, Top.OutputIndex);
@@ -24817,11 +24844,17 @@ void UMaterialExpressionStrataVerticalLayering::GatherStrataMaterialInfo(FStrata
 	{
 		Base.GetTracedInput().Expression->GatherStrataMaterialInfo(StrataMaterialInfo, Base.OutputIndex);
 	}
+
+	StrataMaterialInfo.PopStrataTreeStack();
 }
 
 FStrataOperator* UMaterialExpressionStrataVerticalLayering::StrataGenerateMaterialTopologyTree(class FMaterialCompiler* Compiler, class UMaterialExpression* Parent, int32 OutputIndex)
 {
 	FStrataOperator& StrataOperator = Compiler->StrataCompilationRegisterOperator(STRATA_OPERATOR_VERTICAL, Compiler->StrataTreeStackGetPathUniqueId(), Parent, Compiler->StrataTreeStackGetParentPathUniqueId(), bUseParameterBlending);
+	if (Compiler->GetStrataTreeOutOfStackDepthOccurred())
+	{
+		return &StrataOperator; // Out ot stack space, return now to fail the compilation
+	}
 
 	UMaterialExpression* ChildAExpression = Top.GetTracedInput().Expression;
 	UMaterialExpression* ChildBExpression = Base.GetTracedInput().Expression;
@@ -24971,6 +25004,11 @@ bool UMaterialExpressionStrataAdd::IsResultStrataMaterial(int32 OutputIndex)
 
 void UMaterialExpressionStrataAdd::GatherStrataMaterialInfo(FStrataMaterialInfo& StrataMaterialInfo, int32 OutputIndex)
 {
+	if (!StrataMaterialInfo.PushStrataTreeStack())
+	{
+		return;
+	}
+
 	if (A.GetTracedInput().Expression)
 	{
 		A.GetTracedInput().Expression->GatherStrataMaterialInfo(StrataMaterialInfo, A.OutputIndex);	
@@ -24979,11 +25017,17 @@ void UMaterialExpressionStrataAdd::GatherStrataMaterialInfo(FStrataMaterialInfo&
 	{
 		B.GetTracedInput().Expression->GatherStrataMaterialInfo(StrataMaterialInfo, B.OutputIndex);
 	}
+
+	StrataMaterialInfo.PopStrataTreeStack();
 }
 
 FStrataOperator* UMaterialExpressionStrataAdd::StrataGenerateMaterialTopologyTree(class FMaterialCompiler* Compiler, class UMaterialExpression* Parent, int32 OutputIndex)
 {
 	FStrataOperator& StrataOperator = Compiler->StrataCompilationRegisterOperator(STRATA_OPERATOR_ADD, Compiler->StrataTreeStackGetPathUniqueId(), Parent, Compiler->StrataTreeStackGetParentPathUniqueId(), bUseParameterBlending);
+	if (Compiler->GetStrataTreeOutOfStackDepthOccurred())
+	{
+		return &StrataOperator; // Out ot stack space, return now to fail the compilation
+	}
 
 	UMaterialExpression* ChildAExpression = A.GetTracedInput().Expression;
 	UMaterialExpression* ChildBExpression = B.GetTracedInput().Expression;
@@ -25097,15 +25141,26 @@ bool UMaterialExpressionStrataWeight::IsResultStrataMaterial(int32 OutputIndex)
 
 void UMaterialExpressionStrataWeight::GatherStrataMaterialInfo(FStrataMaterialInfo& StrataMaterialInfo, int32 OutputIndex)
 {
+	if (!StrataMaterialInfo.PushStrataTreeStack())
+	{
+		return;
+	}
+
 	if (A.GetTracedInput().Expression)
 	{
 		A.GetTracedInput().Expression->GatherStrataMaterialInfo(StrataMaterialInfo, A.OutputIndex);
 	}
+
+	StrataMaterialInfo.PopStrataTreeStack();
 }
 
 FStrataOperator* UMaterialExpressionStrataWeight::StrataGenerateMaterialTopologyTree(class FMaterialCompiler* Compiler, class UMaterialExpression* Parent, int32 OutputIndex)
 {
 	FStrataOperator& StrataOperator = Compiler->StrataCompilationRegisterOperator(STRATA_OPERATOR_WEIGHT, Compiler->StrataTreeStackGetPathUniqueId(), Parent, Compiler->StrataTreeStackGetParentPathUniqueId());
+	if (Compiler->GetStrataTreeOutOfStackDepthOccurred())
+	{
+		return &StrataOperator; // Out ot stack space, return now to fail the compilation
+	}
 
 	UMaterialExpression* ChildAExpression = A.GetTracedInput().Expression;
 	FStrataOperator* OpA = nullptr;
diff --git a/Engine/Source/Runtime/Engine/Public/MaterialCompiler.h b/Engine/Source/Runtime/Engine/Public/MaterialCompiler.h
--- a/Engine/Source/Runtime/Engine/Public/MaterialCompiler.h
+++ b/Engine/Source/Runtime/Engine/Public/MaterialCompiler.h
@@ -714,6 +714,11 @@ public:
 	 * Pops a node node of the Strata tree being walked. Used when walking up the tree back to its root level.
 	 */
 	virtual void StrataTreeStackPop() = 0;
+	/**
+	 * This can be used to know if the strata tree we are trying to build is too deep and we should stop the compilation.
+	 * This can be used by all nodes taking as input StrataData in order to detect node re-entry leading to cyclic graph we cannot handle and compile internally: we must fail the compilation.
+	 */
+	virtual bool GetStrataTreeOutOfStackDepthOccurred() = 0;
 
 	virtual int32 StrataThicknessStackGetThicknessIndex() = 0;
 	virtual int32 StrataThicknessStackGetThicknessCode(int32 Index) = 0;
@@ -1422,6 +1427,10 @@ public:
 	{
 		Compiler->StrataTreeStackPop();
 	}
+	virtual bool GetStrataTreeOutOfStackDepthOccurred() override
+	{
+		return Compiler->GetStrataTreeOutOfStackDepthOccurred();
+	}
 
 	virtual int32 StrataThicknessStackGetThicknessIndex() override
 	{
