Fix start position being ignored when using sync groups

Change re-sync code for zero-length blends to only operate if the sync group was actually active

diff --git a/Engine/Source/Runtime/Engine/Private/Animation/AnimSync.cpp b/Engine/Source/Runtime/Engine/Private/Animation/AnimSync.cpp
--- a/Engine/Source/Runtime/Engine/Private/Animation/AnimSync.cpp
+++ b/Engine/Source/Runtime/Engine/Private/Animation/AnimSync.cpp
@@ -158,7 +158,8 @@ void FAnimSync::TickAssetPlayerInstances(FAnimInstanceProxy& InProxy, float InDe
 				// (maintains sync during inertialization or zero-length blends where the previous sync leader is no longer active)
 				const bool bGroupLeaderWasActive = (PreviousGroup && PreviousGroup->ActivePlayers.ContainsByPredicate(TickRecordMatchesGroupLeader)) ||
 					(PreviousUngroupedActivePlayers.ContainsByPredicate(TickRecordMatchesGroupLeader));
-				TickContext.SetResyncToSyncGroup(!bGroupLeaderWasActive);
+				const bool bGroupWasActive = PreviousGroup && !PreviousGroup->ActivePlayers.IsEmpty();
+				TickContext.SetResyncToSyncGroup(!bGroupLeaderWasActive && bGroupWasActive);
 
 				TickContext.MarkerTickContext.MarkersPassedThisTick.Reset();
 				TickContext.RootMotionMovementParams.Clear();
