Fixed leader-pose being used in construction scripts potentially leading to distorted meshes

This forces LOD 0 on registration in editor at all times, rather than at re-instancing, as was the previous fix in CL 22051410.

diff --git a/Engine/Source/Runtime/Engine/Private/Components/SkinnedMeshComponent.cpp b/Engine/Source/Runtime/Engine/Private/Components/SkinnedMeshComponent.cpp
--- a/Engine/Source/Runtime/Engine/Private/Components/SkinnedMeshComponent.cpp
+++ b/Engine/Source/Runtime/Engine/Private/Components/SkinnedMeshComponent.cpp
@@ -632,10 +632,10 @@ void USkinnedMeshComponent::OnRegister()
 	}
 
 #if WITH_EDITOR
-	// When we are reinstancing, ensure that the initial setup is done at LOD 0 so that non-ticking components dont
-	// get unexpected behavior when transitioning LODs post-compile.
+	// When we are in editor, ensure that the initial setup is done at LOD 0 so that non-ticking components dont
+	// get unexpected behavior when transitioning LODs post-construction script
 	const int32 OldForcedLOD = GetForcedLOD();
-	if(GIsReinstancing)
+	if(GIsEditor)
 	{
 		SetForcedLOD(1);
 	}
@@ -644,7 +644,7 @@ void USkinnedMeshComponent::OnRegister()
 	UpdateLODStatus();
 
 #if WITH_EDITOR
-	if(GIsReinstancing)
+	if(GIsEditor)
 	{
 		SetForcedLOD(OldForcedLOD);
 	}
