Allow Engine subclasses to modify time between GC passes.

diff --git a/Engine/Source/Runtime/Engine/Classes/Engine/Engine.h b/Engine/Source/Runtime/Engine/Classes/Engine/Engine.h
--- a/Engine/Source/Runtime/Engine/Classes/Engine/Engine.h
+++ b/Engine/Source/Runtime/Engine/Classes/Engine/Engine.h
@@ -2649,11 +2649,11 @@ private:
 	*/
 	void OnExternalUIChange(bool bInIsOpening);
 
-	/** Returns GetTimeBetweenGarbageCollectionPasses but tweaked if its an idle server or not */
-	float GetTimeBetweenGarbageCollectionPasses(bool bHasPlayersConnected) const;
-
 protected:
 
+	/** Returns GetTimeBetweenGarbageCollectionPasses but tweaked if its an idle server or not */
+	virtual float GetTimeBetweenGarbageCollectionPasses(bool bHasPlayersConnected) const;
+
 	/**
 	 * Handles freezing/unfreezing of rendering 
 	 * 
