Remove 2GB texture size check in DDC1 , large sizes are okay now
check(CompressedDataSize < MAX_int32);

diff --git a/Engine/Source/Runtime/Engine/Private/TextureDerivedDataTask.cpp b/Engine/Source/Runtime/Engine/Private/TextureDerivedDataTask.cpp
--- a/Engine/Source/Runtime/Engine/Private/TextureDerivedDataTask.cpp
+++ b/Engine/Source/Runtime/Engine/Private/TextureDerivedDataTask.cpp
@@ -550,6 +550,9 @@ static void DDC1_StoreClassicTextureInDerivedData(
 		check(CompressedImage.RawData.GetTypeSize() == 1);
 		int64 CompressedDataSize = CompressedImage.RawData.Num();
 
+		// CompressedDataSize can exceed int32 ; eg. 16k x 16k x RGBA16F == 2 GB
+		// DDC1 should be 64-bit safe now
+
 		// CompressedImage sizes were padded up to multiple of 4 for d3d, no longer
 		UE_LOG(LogTextureUpload, Verbose, TEXT("Compressed Mip %d PF=%d : %dx%dx%d : %d ; up4 %dx%d=%d"),
 			MipIndex, (int)CompressedImage.PixelFormat,
