Fix for static code analysis

#rnx

diff --git a/Engine/Source/Runtime/Engine/Private/LocalPlayer.cpp b/Engine/Source/Runtime/Engine/Private/LocalPlayer.cpp
--- a/Engine/Source/Runtime/Engine/Private/LocalPlayer.cpp
+++ b/Engine/Source/Runtime/Engine/Private/LocalPlayer.cpp
@@ -900,7 +900,8 @@ FSceneView* ULocalPlayer::CalcSceneView( class FSceneViewFamily* ViewFamily,
 		View->OverridePostProcessSettings(ViewInfo.PostProcessSettings, ViewInfo.PostProcessBlendWeight);
 
 		// Camera overrides (cameras blending in, camera modifiers, etc)
-		if (PlayerController->PlayerCameraManager)
+		if (PlayerController->PlayerCameraManager &&
+				CameraAnimPPSettings && CameraAnimPPBlendWeights && CameraAnimPPBlendOrders)
 		{
 			for (int32 PPIdx = 0; PPIdx < CameraAnimPPBlendWeights->Num(); ++PPIdx)
 			{
