Substrate - Fixed inverted logic causing some issues with material interpretation.

diff --git a/Engine/Source/Runtime/Engine/Classes/Engine/EngineTypes.h b/Engine/Source/Runtime/Engine/Classes/Engine/EngineTypes.h
--- a/Engine/Source/Runtime/Engine/Classes/Engine/EngineTypes.h
+++ b/Engine/Source/Runtime/Engine/Classes/Engine/EngineTypes.h
@@ -732,10 +732,11 @@ public:
 	bool operator!=(const FStrataMaterialInfo& Other) const { return ShadingModelField != Other.GetShadingModelField(); }
 
 #if WITH_EDITOR
+	// Returns true if everything went fine (not out of Substrate tree stack)
 	bool PushStrataTreeStack()
 	{
 		bOutOfStackDepthWhenParsing = bOutOfStackDepthWhenParsing || (++ParsingStackDepth > STRATA_TREE_MAX_DEPTH);
-		return bOutOfStackDepthWhenParsing;
+		return !bOutOfStackDepthWhenParsing;
 	}
 	void PopStrataTreeStack()
 	{
diff --git a/Engine/Source/Runtime/Engine/Private/Materials/Material.cpp b/Engine/Source/Runtime/Engine/Private/Materials/Material.cpp
--- a/Engine/Source/Runtime/Engine/Private/Materials/Material.cpp
+++ b/Engine/Source/Runtime/Engine/Private/Materials/Material.cpp
@@ -4808,6 +4808,10 @@ void UMaterial::RebuildShadingModelField()
 
 			if (StrataMaterialInfo.GetStrataTreeOutOfStackDepthOccurred())
 			{
+				StrataMaterialInfo.AddShadingModel(SSM_Unlit);
+				MaterialDomain = EMaterialDomain::MD_Surface;
+				ShadingModel = MSM_Unlit;
+				ShadingModels.AddShadingModel(MSM_Unlit);
 				CancelOutstandingCompilation();
 				UE_LOG(LogMaterial, Error, TEXT("%s: Substrate - Cyclic graph detected when we only support acyclic graph."), *GetName());
 				return;
