Print an error when the requested DeviceProfile can't be found.

diff --git a/Engine/Source/Runtime/Engine/Private/DeviceProfiles/DeviceProfileManager.cpp b/Engine/Source/Runtime/Engine/Private/DeviceProfiles/DeviceProfileManager.cpp
--- a/Engine/Source/Runtime/Engine/Private/DeviceProfiles/DeviceProfileManager.cpp
+++ b/Engine/Source/Runtime/Engine/Private/DeviceProfiles/DeviceProfileManager.cpp
@@ -696,6 +696,13 @@ UDeviceProfile* UDeviceProfileManager::CreateProfile(const FString& ProfileName,
 		{
 			PlatformConfigFile = FConfigCacheIni::FindOrLoadPlatformConfig(LocalConfigFile, TEXT("DeviceProfiles"), ConfigPlatform);
 		}
+#if !UE_BUILD_SHIPPING
+		if (!PlatformConfigFile->Contains(ProfileName) && !PlatformConfigFile->Contains(ProfileName + " DeviceProfile"))
+		{
+			// Display and not Error to allow tests to create profiles without failing
+			UE_LOG(LogDeviceProfileManager, Display, TEXT("Deviceprofile %s not found."), *ProfileName);
+		}
+#endif
 
 		// Build Parent objects first. Important for setup
 		FString ParentName = InSpecifyParentName;
