Add InstanceLODDistanceScale as property on the UHierarchicalInstancedStaticMeshComponent

#rnx

diff --git a/Engine/Source/Runtime/Engine/Classes/Components/HierarchicalInstancedStaticMeshComponent.h b/Engine/Source/Runtime/Engine/Classes/Components/HierarchicalInstancedStaticMeshComponent.h
--- a/Engine/Source/Runtime/Engine/Classes/Components/HierarchicalInstancedStaticMeshComponent.h
+++ b/Engine/Source/Runtime/Engine/Classes/Components/HierarchicalInstancedStaticMeshComponent.h
@@ -183,6 +183,10 @@ class ENGINE_API UHierarchicalInstancedStaticMeshComponent : public UInstancedSt
 	// Current value of density scaling applied to this component
 	float CurrentDensityScaling;
 
+	/** Scale applied to change the computation of LOD distances when using the StaticMesh screen sizes. */
+	UPROPERTY(EditAnywhere, BlueprintReadOnly, Category = Culling)
+	float InstanceLODDistanceScale;
+
 	// The number of nodes in the occlusion layer
 	UPROPERTY()
 	int32 OcclusionLayerNumNodes;
@@ -245,6 +249,10 @@ public:
 	virtual void PreAllocateInstancesMemory(int32 AddedInstanceCount) override;
 	virtual bool SupportsRemoveSwap() const override { return true; }
 
+	/** Sets the fading start and culling end distances for this component. */
+	UFUNCTION(BlueprintCallable, Category = "Components|HierarchicalInstancedStaticMesh")
+	void SetLODDistanceScale(float InLODDistanceScale);
+
 	/** Get the number of instances that overlap a given sphere */
 	int32 GetOverlappingSphereCount(const FSphere& Sphere) const;
 	/** Get the number of instances that overlap a given box */
diff --git a/Engine/Source/Runtime/Engine/Classes/Engine/InstancedStaticMesh.h b/Engine/Source/Runtime/Engine/Classes/Engine/InstancedStaticMesh.h
--- a/Engine/Source/Runtime/Engine/Classes/Engine/InstancedStaticMesh.h
+++ b/Engine/Source/Runtime/Engine/Classes/Engine/InstancedStaticMesh.h
@@ -35,7 +35,6 @@
 class ULightComponent;
 
 extern TAutoConsoleVariable<float> CVarFoliageMinimumScreenSize;
-extern TAutoConsoleVariable<float> CVarFoliageLODDistanceScale;
 extern TAutoConsoleVariable<float> CVarRandomLODRange;
 extern TAutoConsoleVariable<int32> CVarMinLOD;
 
@@ -227,6 +226,8 @@ struct FInstancingUserData
 	int32 StartCullDistance;
 	int32 EndCullDistance;
 
+	float LODDistanceScale;
+
 	int32 MinLOD;
 
 	bool bRenderSelected;
diff --git a/Engine/Source/Runtime/Engine/Private/HierarchicalInstancedStaticMesh.cpp b/Engine/Source/Runtime/Engine/Private/HierarchicalInstancedStaticMesh.cpp
--- a/Engine/Source/Runtime/Engine/Private/HierarchicalInstancedStaticMesh.cpp
+++ b/Engine/Source/Runtime/Engine/Private/HierarchicalInstancedStaticMesh.cpp
@@ -93,7 +93,12 @@ TAutoConsoleVariable<int32> CVarFoliageMaxEndCullDistance(
 TAutoConsoleVariable<float> CVarFoliageLODDistanceScale(
 	TEXT("foliage.LODDistanceScale"),
 	1.0f,
-	TEXT("Scale factor for the distance used in computing LOD for foliage."));
+	TEXT("Scale factor for the distance used in computing LOD for foliage."),
+	FConsoleVariableDelegate::CreateLambda([](IConsoleVariable* InVariable)
+		{
+			FGlobalComponentRecreateRenderStateContext Context;
+		})
+	);
 
 TAutoConsoleVariable<float> CVarRandomLODRange(
 	TEXT("foliage.RandomLODRange"),
@@ -856,6 +861,12 @@ FHierarchicalStaticMeshSceneProxy::FHierarchicalStaticMeshSceneProxy(UHierarchic
 
 	bIsHierarchicalInstancedStaticMesh = true;
 	bIsLandscapeGrass = (ViewRelevance == EHISMViewRelevanceType::Grass);
+
+	// Store LODDistanceScale so it can be used in FInstancedStaticMeshVertexFactoryShaderParameters::GetElementShaderBindings when dither LOD transitions are enabled
+	float LODDistanceScale = InComponent->InstanceLODDistanceScale * CVarFoliageLODDistanceScale.GetValueOnGameThread();
+	UserData_AllInstances.LODDistanceScale = LODDistanceScale;
+	UserData_SelectedInstances.LODDistanceScale = LODDistanceScale;
+	UserData_DeselectedInstances.LODDistanceScale = LODDistanceScale;
 }
 
 void FHierarchicalStaticMeshSceneProxy::SetupOcclusion(UHierarchicalInstancedStaticMeshComponent* InComponent)
@@ -1648,11 +1659,11 @@ void FHierarchicalStaticMeshSceneProxy::GetDynamicMeshElements(const TArray<cons
 				InstanceParams.ViewOriginInLocalOne = WorldToLocal.TransformPosition(View->GetTemporalLODOrigin(1, bMultipleSections));
 
 				float MinSize = bIsOrtho ? 0.0f : CVarFoliageMinimumScreenSize.GetValueOnRenderThread();
-				float LODScale = CVarFoliageLODDistanceScale.GetValueOnRenderThread();
+				float LODScale = UserData_AllInstances.LODDistanceScale;
 				int MaxEndCullDistance = CVarFoliageMaxEndCullDistance.GetValueOnRenderThread();
 				float LODRandom = CVarRandomLODRange.GetValueOnRenderThread();
 				float MaxDrawDistanceScale = GetCachedScalabilityCVars().ViewDistanceScale;
-				
+								
 				FVector AverageScale(InstanceParams.Tree[0].MinInstanceScale + (InstanceParams.Tree[0].MaxInstanceScale - InstanceParams.Tree[0].MinInstanceScale) / 2.0f);
 				FBoxSphereBounds ScaledBounds = RenderData->Bounds.TransformBy(FTransform(FRotator::ZeroRotator, FVector::ZeroVector, AverageScale));
 				float SphereRadius = ScaledBounds.SphereRadius;
@@ -1846,7 +1857,7 @@ void FHierarchicalStaticMeshSceneProxy::GetDynamicMeshElements(const TArray<cons
 
 						const bool bIsOrtho = !View->ViewMatrices.IsPerspectiveProjection();
 						const float MinSize = bIsOrtho ? 0.0f : CVarFoliageMinimumScreenSize.GetValueOnRenderThread();
-						const float LODScale = CVarFoliageLODDistanceScale.GetValueOnRenderThread();
+						const float LODScale = UserData_AllInstances.LODDistanceScale;
 						int MaxEndCullDistance = CVarFoliageMaxEndCullDistance.GetValueOnRenderThread();
 						const float LODRandom = CVarRandomLODRange.GetValueOnRenderThread();
 						const float MaxDrawDistanceScale = GetCachedScalabilityCVars().ViewDistanceScale;
@@ -2025,6 +2036,7 @@ UHierarchicalInstancedStaticMeshComponent::UHierarchicalInstancedStaticMeshCompo
 	, UnbuiltInstanceBounds(ForceInit)
 	, bEnableDensityScaling(false)
 	, CurrentDensityScaling(1.0f)
+	, InstanceLODDistanceScale(1.0f)
 	, OcclusionLayerNumNodes(0)
 	, InstanceCountToRender(0)
 	, bIsAsyncBuilding(false)
@@ -2640,6 +2652,15 @@ void UHierarchicalInstancedStaticMeshComponent::ClearInstances()
 	FNavigationSystem::UpdateComponentData(*this);
 }
 
+void UHierarchicalInstancedStaticMeshComponent::SetLODDistanceScale(float InLODDistanceScale)
+{
+	if (InstanceLODDistanceScale != InLODDistanceScale)
+	{
+		InstanceLODDistanceScale = InLODDistanceScale;
+		MarkRenderStateDirty();
+	}
+}
+
 bool UHierarchicalInstancedStaticMeshComponent::ShouldCreatePhysicsState() const
 {
 	if (bDisableCollision)
diff --git a/Engine/Source/Runtime/Engine/Private/ISMPartition/ISMComponentDescriptor.cpp b/Engine/Source/Runtime/Engine/Private/ISMPartition/ISMComponentDescriptor.cpp
--- a/Engine/Source/Runtime/Engine/Private/ISMPartition/ISMComponentDescriptor.cpp
+++ b/Engine/Source/Runtime/Engine/Private/ISMPartition/ISMComponentDescriptor.cpp
@@ -117,6 +117,7 @@ void FISMComponentDescriptorBase::InitFrom(const UStaticMeshComponent* Template,
 		if (const UHierarchicalInstancedStaticMeshComponent* HISMTemplate = Cast<UHierarchicalInstancedStaticMeshComponent>(Template))
 		{
 			bEnableDensityScaling = HISMTemplate->bEnableDensityScaling;
+			InstanceLODDistanceScale = HISMTemplate->InstanceLODDistanceScale;
 		}
 	}
 
@@ -314,6 +315,7 @@ void FISMComponentDescriptorBase::InitComponent(UInstancedStaticMeshComponent* I
 	if (UHierarchicalInstancedStaticMeshComponent* HISMComponent = Cast<UHierarchicalInstancedStaticMeshComponent>(ISMComponent))
 	{
 		HISMComponent->bEnableDensityScaling = bEnableDensityScaling;
+		HISMComponent->InstanceLODDistanceScale = InstanceLODDistanceScale;
 	}
 }
 
diff --git a/Engine/Source/Runtime/Engine/Private/InstancedStaticMesh.cpp b/Engine/Source/Runtime/Engine/Private/InstancedStaticMesh.cpp
--- a/Engine/Source/Runtime/Engine/Private/InstancedStaticMesh.cpp
+++ b/Engine/Source/Runtime/Engine/Private/InstancedStaticMesh.cpp
@@ -1544,6 +1544,7 @@ void FInstancedStaticMeshSceneProxy::SetupProxy(UInstancedStaticMeshComponent* I
 	UserData_AllInstances.MeshRenderData = InComponent->GetStaticMesh()->GetRenderData();
 	UserData_AllInstances.StartCullDistance = InComponent->InstanceStartCullDistance;
 	UserData_AllInstances.EndCullDistance = InComponent->InstanceEndCullDistance;
+	UserData_AllInstances.LODDistanceScale = 1.0f;
 	UserData_AllInstances.InstancingOffset = InComponent->GetStaticMesh()->GetBoundingBox().GetCenter();
 	UserData_AllInstances.MinLOD = ClampedMinLOD;
 	UserData_AllInstances.bRenderSelected = true;
@@ -5269,6 +5270,8 @@ void FInstancedStaticMeshVertexFactoryShaderParameters::GetElementShaderBindings
 		FVector4f InstancingTranslatedWorldViewOriginZero(ForceInit);
 		FVector4f InstancingTranslatedWorldViewOriginOne(ForceInit);
 		InstancingTranslatedWorldViewOriginOne.W = 1.0f;
+
+		// InstancedLODRange is only set for HierarchicalInstancedStaticMeshes
 		if (InstancingUserData && BatchElement.InstancedLODRange)
 		{
 			int32 FirstLOD = InstancingUserData->MinLOD;
@@ -5282,7 +5285,7 @@ void FInstancedStaticMeshVertexFactoryShaderParameters::GetElementShaderBindings
 			FBoxSphereBounds ScaledBounds = InstancingUserData->MeshRenderData->Bounds.TransformBy(FTransform(FRotator::ZeroRotator, FVector::ZeroVector, InstancingUserData->AverageInstancesScale));
 			float SphereRadius = ScaledBounds.SphereRadius;
 			float MinSize = View->ViewMatrices.IsPerspectiveProjection() ? CVarFoliageMinimumScreenSize.GetValueOnRenderThread() : 0.0f;
-			float LODScale = CVarFoliageLODDistanceScale.GetValueOnRenderThread();
+			float LODScale = InstancingUserData->LODDistanceScale;
 			float LODRandom = CVarRandomLODRange.GetValueOnRenderThread();
 			float MaxDrawDistanceScale = GetCachedScalabilityCVars().ViewDistanceScale;
 
diff --git a/Engine/Source/Runtime/Engine/Public/ISMPartition/ISMComponentDescriptor.h b/Engine/Source/Runtime/Engine/Public/ISMPartition/ISMComponentDescriptor.h
--- a/Engine/Source/Runtime/Engine/Public/ISMPartition/ISMComponentDescriptor.h
+++ b/Engine/Source/Runtime/Engine/Public/ISMPartition/ISMComponentDescriptor.h
@@ -69,6 +69,9 @@ public:
 	UPROPERTY(EditAnywhere, Category = "Component Settings")
 	int32 InstanceEndCullDistance;
 
+	UPROPERTY(EditAnywhere, Category = "Component Settings")
+	float InstanceLODDistanceScale;
+
 	UPROPERTY(EditAnywhere, Category = "Component Settings")
 	int32 VirtualTextureCullMips;
 
