Add forward decl of UTexture to make non-PCH builds compile.

diff --git a/Engine/Source/Runtime/Engine/Public/SceneManagement.h b/Engine/Source/Runtime/Engine/Public/SceneManagement.h
--- a/Engine/Source/Runtime/Engine/Public/SceneManagement.h
+++ b/Engine/Source/Runtime/Engine/Public/SceneManagement.h
@@ -61,6 +61,7 @@ class USkyAtmosphereComponent;
 class FSkyAtmosphereRenderSceneInfo;
 class USkyLightComponent;
 class UTexture2D;
+class UTexture;
 class UTextureLightProfile;
 struct FDynamicMeshVertex;
 class ULightMapVirtualTexture2D;
