Substrate - fix for crash with function and OutputIndex (because TracedInput can resolve to a difference output index than ExpresionInput)

diff --git a/Engine/Source/Runtime/Engine/Private/Materials/MaterialExpressions.cpp b/Engine/Source/Runtime/Engine/Private/Materials/MaterialExpressions.cpp
--- a/Engine/Source/Runtime/Engine/Private/Materials/MaterialExpressions.cpp
+++ b/Engine/Source/Runtime/Engine/Private/Materials/MaterialExpressions.cpp
@@ -24058,9 +24058,10 @@ void UMaterialExpressionStrataConvertToDecal::GatherStrataMaterialInfo(FStrataMa
 		return;
 	}
 
-	if (DecalMaterial.GetTracedInput().Expression)
+	FExpressionInput TracedInput = DecalMaterial.GetTracedInput();
+	if (TracedInput.Expression)
 	{
-		DecalMaterial.GetTracedInput().Expression->GatherStrataMaterialInfo(StrataMaterialInfo, DecalMaterial.OutputIndex);
+		TracedInput.Expression->GatherStrataMaterialInfo(StrataMaterialInfo, TracedInput.OutputIndex);
 	}
 	StrataMaterialInfo.AddShadingModel(SSM_Decal);
 
@@ -24076,12 +24077,13 @@ FStrataOperator* UMaterialExpressionStrataConvertToDecal::StrataGenerateMaterial
 		return &StrataOperator; // Out ot stack space, return now to fail the compilation
 	}
 
-	UMaterialExpression* ChildDecalMaterialExpression = DecalMaterial.GetTracedInput().Expression;
+	FExpressionInput TracedInput = DecalMaterial.GetTracedInput();
+	UMaterialExpression* ChildDecalMaterialExpression = TracedInput.Expression;
 	FStrataOperator* OpA = nullptr;
 	if (ChildDecalMaterialExpression)
 	{
 		Compiler->StrataTreeStackPush(this, 0);
-		OpA = ChildDecalMaterialExpression->StrataGenerateMaterialTopologyTree(Compiler, this, DecalMaterial.OutputIndex);
+		OpA = ChildDecalMaterialExpression->StrataGenerateMaterialTopologyTree(Compiler, this, TracedInput.OutputIndex);
 		Compiler->StrataTreeStackPop();
 		AssignOperatorIndexIfNotNull(StrataOperator.LeftIndex, OpA);
 	}
@@ -24655,13 +24657,15 @@ void UMaterialExpressionStrataHorizontalMixing::GatherStrataMaterialInfo(FStrata
 		return;
 	}
 
-	if (Foreground.GetTracedInput().Expression)
+	FExpressionInput TracedInputA = Foreground.GetTracedInput();
+	FExpressionInput TracedInputB = Background.GetTracedInput();
+	if (TracedInputA.Expression)
 	{
-		Foreground.GetTracedInput().Expression->GatherStrataMaterialInfo(StrataMaterialInfo, Foreground.OutputIndex);
+		TracedInputA.Expression->GatherStrataMaterialInfo(StrataMaterialInfo, TracedInputA.OutputIndex);
 	}
-	if (Background.GetTracedInput().Expression)
+	if (TracedInputB.Expression)
 	{
-		Background.GetTracedInput().Expression->GatherStrataMaterialInfo(StrataMaterialInfo, Background.OutputIndex);
+		TracedInputB.Expression->GatherStrataMaterialInfo(StrataMaterialInfo, TracedInputB.OutputIndex);
 	}
 
 	StrataMaterialInfo.PopStrataTreeStack();
@@ -24675,21 +24679,23 @@ FStrataOperator* UMaterialExpressionStrataHorizontalMixing::StrataGenerateMateri
 		return &StrataOperator; // Out ot stack space, return now to fail the compilation
 	}
 
-	UMaterialExpression* ChildAExpression = Background.GetTracedInput().Expression;
-	UMaterialExpression* ChildBExpression = Foreground.GetTracedInput().Expression;
+	FExpressionInput TracedInputA = Background.GetTracedInput();
+	FExpressionInput TracedInputB = Foreground.GetTracedInput();
+	UMaterialExpression* ChildAExpression = TracedInputA.Expression;
+	UMaterialExpression* ChildBExpression = TracedInputB.Expression;
 	FStrataOperator* OpA = nullptr;
 	FStrataOperator* OpB = nullptr;
 	if (ChildAExpression)
 	{
 		Compiler->StrataTreeStackPush(this, 0);
-		OpA = ChildAExpression->StrataGenerateMaterialTopologyTree(Compiler, this, Background.OutputIndex);
+		OpA = ChildAExpression->StrataGenerateMaterialTopologyTree(Compiler, this, TracedInputA.OutputIndex);
 		Compiler->StrataTreeStackPop();
 		AssignOperatorIndexIfNotNull(StrataOperator.LeftIndex, OpA);
 	}
 	if (ChildBExpression)
 	{
 		Compiler->StrataTreeStackPush(this, 1);
-		OpB = ChildBExpression->StrataGenerateMaterialTopologyTree(Compiler, this, Foreground.OutputIndex);
+		OpB = ChildBExpression->StrataGenerateMaterialTopologyTree(Compiler, this, TracedInputB.OutputIndex);
 		Compiler->StrataTreeStackPop();
 		AssignOperatorIndexIfNotNull(StrataOperator.RightIndex, OpB);
 	}
@@ -24849,13 +24855,15 @@ void UMaterialExpressionStrataVerticalLayering::GatherStrataMaterialInfo(FStrata
 		return;
 	}
 
-	if (Top.GetTracedInput().Expression)
+	FExpressionInput TracedInputTop  = Top.GetTracedInput();
+	FExpressionInput TracedInputBase = Base.GetTracedInput();
+	if (TracedInputTop.Expression)
 	{
-		Top.GetTracedInput().Expression->GatherStrataMaterialInfo(StrataMaterialInfo, Top.OutputIndex);
+		TracedInputTop.Expression->GatherStrataMaterialInfo(StrataMaterialInfo, TracedInputTop.OutputIndex);
 	}
-	if (Base.GetTracedInput().Expression)
+	if (TracedInputBase.Expression)
 	{
-		Base.GetTracedInput().Expression->GatherStrataMaterialInfo(StrataMaterialInfo, Base.OutputIndex);
+		TracedInputBase.Expression->GatherStrataMaterialInfo(StrataMaterialInfo, TracedInputBase.OutputIndex);
 	}
 
 	StrataMaterialInfo.PopStrataTreeStack();
@@ -24869,8 +24877,10 @@ FStrataOperator* UMaterialExpressionStrataVerticalLayering::StrataGenerateMateri
 		return &StrataOperator; // Out ot stack space, return now to fail the compilation
 	}
 
-	UMaterialExpression* ChildAExpression = Top.GetTracedInput().Expression;
-	UMaterialExpression* ChildBExpression = Base.GetTracedInput().Expression;
+	FExpressionInput TracedInputTop  = Top.GetTracedInput();
+	FExpressionInput TracedInputBase = Base.GetTracedInput();
+	UMaterialExpression* ChildAExpression = TracedInputTop.Expression;
+	UMaterialExpression* ChildBExpression = TracedInputBase.Expression;
 	FStrataOperator* OpA = nullptr;
 	FStrataOperator* OpB = nullptr;
 
@@ -24879,7 +24889,7 @@ FStrataOperator* UMaterialExpressionStrataVerticalLayering::StrataGenerateMateri
 	{
 		Compiler->StrataThicknessStackPush(this, &Thickness);
 		Compiler->StrataTreeStackPush(this, 0);
-		OpA = ChildAExpression->StrataGenerateMaterialTopologyTree(Compiler, this, Top.OutputIndex);
+		OpA = ChildAExpression->StrataGenerateMaterialTopologyTree(Compiler, this, TracedInputTop.OutputIndex);
 		Compiler->StrataTreeStackPop();
 		Compiler->StrataThicknessStackPop();
 		AssignOperatorIndexIfNotNull(StrataOperator.LeftIndex, OpA);
@@ -24888,7 +24898,7 @@ FStrataOperator* UMaterialExpressionStrataVerticalLayering::StrataGenerateMateri
 	if (ChildBExpression)
 	{
 		Compiler->StrataTreeStackPush(this, 1);
-		OpB = ChildBExpression->StrataGenerateMaterialTopologyTree(Compiler, this, Base.OutputIndex);
+		OpB = ChildBExpression->StrataGenerateMaterialTopologyTree(Compiler, this, TracedInputBase.OutputIndex);
 		Compiler->StrataTreeStackPop();
 		AssignOperatorIndexIfNotNull(StrataOperator.RightIndex, OpB);
 	}
@@ -25022,13 +25032,15 @@ void UMaterialExpressionStrataAdd::GatherStrataMaterialInfo(FStrataMaterialInfo&
 		return;
 	}
 
-	if (A.GetTracedInput().Expression)
+	FExpressionInput TracedInputA = A.GetTracedInput();
+	FExpressionInput TracedInputB = B.GetTracedInput();
+	if (TracedInputA.Expression)
 	{
-		A.GetTracedInput().Expression->GatherStrataMaterialInfo(StrataMaterialInfo, A.OutputIndex);	
+		TracedInputA.Expression->GatherStrataMaterialInfo(StrataMaterialInfo, TracedInputA.OutputIndex);
 	}
-	if (B.GetTracedInput().Expression)
+	if (TracedInputB.Expression)
 	{
-		B.GetTracedInput().Expression->GatherStrataMaterialInfo(StrataMaterialInfo, B.OutputIndex);
+		TracedInputB.Expression->GatherStrataMaterialInfo(StrataMaterialInfo, TracedInputB.OutputIndex);
 	}
 
 	StrataMaterialInfo.PopStrataTreeStack();
@@ -25042,21 +25054,23 @@ FStrataOperator* UMaterialExpressionStrataAdd::StrataGenerateMaterialTopologyTre
 		return &StrataOperator; // Out ot stack space, return now to fail the compilation
 	}
 
-	UMaterialExpression* ChildAExpression = A.GetTracedInput().Expression;
-	UMaterialExpression* ChildBExpression = B.GetTracedInput().Expression;
+	FExpressionInput TracedInputA = A.GetTracedInput();
+	FExpressionInput TracedInputB = B.GetTracedInput();
+	UMaterialExpression* ChildAExpression = TracedInputA.Expression;
+	UMaterialExpression* ChildBExpression = TracedInputB.Expression;
 	FStrataOperator* OpA = nullptr;
 	FStrataOperator* OpB = nullptr;
 	if (ChildAExpression)
 	{
 		Compiler->StrataTreeStackPush(this, 0);
-		OpA = ChildAExpression->StrataGenerateMaterialTopologyTree(Compiler, this, A.OutputIndex);
+		OpA = ChildAExpression->StrataGenerateMaterialTopologyTree(Compiler, this, TracedInputA.OutputIndex);
 		Compiler->StrataTreeStackPop();
 		AssignOperatorIndexIfNotNull(StrataOperator.LeftIndex, OpA);
 	}
 	if (ChildBExpression)
 	{
 		Compiler->StrataTreeStackPush(this, 1);
-		OpB = ChildBExpression->StrataGenerateMaterialTopologyTree(Compiler, this, B.OutputIndex);
+		OpB = ChildBExpression->StrataGenerateMaterialTopologyTree(Compiler, this, TracedInputB.OutputIndex);
 		Compiler->StrataTreeStackPop();
 		AssignOperatorIndexIfNotNull(StrataOperator.RightIndex, OpB);
 	}
@@ -25159,9 +25173,10 @@ void UMaterialExpressionStrataWeight::GatherStrataMaterialInfo(FStrataMaterialIn
 		return;
 	}
 
-	if (A.GetTracedInput().Expression)
+	FExpressionInput TracedInputA = A.GetTracedInput();
+	if (TracedInputA.Expression)
 	{
-		A.GetTracedInput().Expression->GatherStrataMaterialInfo(StrataMaterialInfo, A.OutputIndex);
+		TracedInputA.Expression->GatherStrataMaterialInfo(StrataMaterialInfo, TracedInputA.OutputIndex);
 	}
 
 	StrataMaterialInfo.PopStrataTreeStack();
@@ -25175,12 +25190,13 @@ FStrataOperator* UMaterialExpressionStrataWeight::StrataGenerateMaterialTopology
 		return &StrataOperator; // Out ot stack space, return now to fail the compilation
 	}
 
-	UMaterialExpression* ChildAExpression = A.GetTracedInput().Expression;
+	FExpressionInput TracedInputA = A.GetTracedInput();
+	UMaterialExpression* ChildAExpression = TracedInputA.Expression;
 	FStrataOperator* OpA = nullptr;
 	if (ChildAExpression)
 	{
 		Compiler->StrataTreeStackPush(this, 0);
-		OpA = ChildAExpression->StrataGenerateMaterialTopologyTree(Compiler, this, A.OutputIndex);
+		OpA = ChildAExpression->StrataGenerateMaterialTopologyTree(Compiler, this, TracedInputA.OutputIndex);
 		Compiler->StrataTreeStackPop();
 		AssignOperatorIndexIfNotNull(StrataOperator.LeftIndex, OpA);
 	}
