Fix crash when a Substrate material as more shared tangent basis than allowed.

The error is not reported as a shader compilation error that the user can act on.

diff --git a/Engine/Source/Runtime/Engine/Private/Materials/HLSLMaterialTranslator.cpp b/Engine/Source/Runtime/Engine/Private/Materials/HLSLMaterialTranslator.cpp
--- a/Engine/Source/Runtime/Engine/Private/Materials/HLSLMaterialTranslator.cpp
+++ b/Engine/Source/Runtime/Engine/Private/Materials/HLSLMaterialTranslator.cpp
@@ -1627,6 +1627,17 @@ bool FHLSLMaterialTranslator::Translate()
 					ResourcesString += "}\n";
 				}
 			}
+
+			// Check if normal/tangent basis are valid
+			{
+				uint8 UsedSharedLocalBasesCount = 0;
+				uint8 RequestedSharedLocalBasesCount = 0;
+				StrataEvaluateSharedLocalBases(UsedSharedLocalBasesCount, RequestedSharedLocalBasesCount, nullptr, nullptr);
+				if (RequestedSharedLocalBasesCount > STRATA_MAX_SHAREDLOCALBASES_REGISTERS)
+				{
+					Errorf(TEXT(" %s [%s]: Substrate - Material has more unique normal/tangent basis than the allowed limit %d/%d."), *Material->GetDebugName(), *Material->GetAssetPath().ToString(), RequestedSharedLocalBasesCount, STRATA_MAX_SHAREDLOCALBASES_REGISTERS);
+				}
+			}
 		}
 		else
 		{
@@ -2051,7 +2062,8 @@ void FHLSLMaterialTranslator::GetMaterialEnvironment(EShaderPlatform InPlatform,
 		// Compute the shared local basis count and generate the hlsl shader code for it.
 		StrataPixelNormalInitializerValues = FString::Printf(TEXT("\n\n\n\t// Strata normal and tangent\n"));
 		FinalUsedSharedLocalBasesCount = 0;
-		StrataEvaluateSharedLocalBases(FinalUsedSharedLocalBasesCount, &StrataPixelNormalInitializerValues, &OutEnvironment);
+		uint8 RequestedSharedLocalBasesCount = 0;
+		StrataEvaluateSharedLocalBases(FinalUsedSharedLocalBasesCount, RequestedSharedLocalBasesCount, &StrataPixelNormalInitializerValues, &OutEnvironment);
 
 		// Now write some feedback to the user
 		{
@@ -2069,7 +2081,14 @@ void FHLSLMaterialTranslator::GetMaterialEnvironment(EShaderPlatform InPlatform,
 			StrataMaterialDescription += FString::Printf(TEXT(" - Requested Byte Size after simplification   %u (%d UINT32)\r\n"), StrataMaterialRequestedSizeByte, StrataMaterialRequestedSizeByte / 4);
 			StrataMaterialDescription += FString::Printf(TEXT(" - Material complexity                        %s\r\n"), bStrataMaterialIsSingle ? TEXT("SINGLE") : (bStrataMaterialIsSimple ? TEXT("SIMPLE") : TEXT("COMPLEX")));
 			StrataMaterialDescription += FString::Printf(TEXT(" - TotalBSDFCount                             %i\r\n"), StrataMaterialBSDFCount);
-			StrataMaterialDescription += FString::Printf(TEXT(" - SharedLocalBasesCount                      %i\r\n"), FinalUsedSharedLocalBasesCount);
+			if (RequestedSharedLocalBasesCount > STRATA_MAX_SHAREDLOCALBASES_REGISTERS)
+			{
+				StrataMaterialDescription += FString::Printf(TEXT(" - SharedLocalBasesCount                      %i (Requested:%i)\r\n"), FinalUsedSharedLocalBasesCount, RequestedSharedLocalBasesCount);
+			}
+			else
+			{
+				StrataMaterialDescription += FString::Printf(TEXT(" - SharedLocalBasesCount                      %i\r\n"), FinalUsedSharedLocalBasesCount);
+			}
 
 			for (int32 OpIt = 0; OpIt < StrataMaterialExpressionRegisteredOperators.Num(); ++OpIt)
 			{
@@ -10314,9 +10333,10 @@ FStrataOperator* FHLSLMaterialTranslator::StrataCompilationGetOperatorFromIndex(
 }
 
 void FHLSLMaterialTranslator::StrataEvaluateSharedLocalBases(
-	uint8& UsedSharedLocalBasesCount,
+	uint8& OutUsedSharedLocalBasesCount,
+	uint8& OutRequestedSharedLocalBasesCount,
 	FString* OutStrataPixelNormalInitializerValues, 
-	FShaderCompilerEnvironment* OutEnvironment)
+	FShaderCompilerEnvironment* OutEnvironment) const
 {
 	/*
 	* The final output code/workflow for shared tangent basis should look like
@@ -10338,11 +10358,12 @@ void FHLSLMaterialTranslator::StrataEvaluateSharedLocalBases(
 	*/
 
 	FStrataRegisteredSharedLocalBasis UsedSharedLocalBasesInfo[STRATA_MAX_SHAREDLOCALBASES_REGISTERS];
-	UsedSharedLocalBasesCount = 0;
+	OutUsedSharedLocalBasesCount = 0;
+	OutRequestedSharedLocalBasesCount = 0;
 
 	for (int32 OpIt = 0; OpIt < StrataMaterialExpressionRegisteredOperators.Num(); ++OpIt)
 	{
-		FStrataOperator& BSDFOperator = StrataMaterialExpressionRegisteredOperators[OpIt];
+		const FStrataOperator& BSDFOperator = StrataMaterialExpressionRegisteredOperators[OpIt];
 		if (BSDFOperator.BSDFIndex == INDEX_NONE || BSDFOperator.IsDiscarded())
 		{
 			continue;	// not a BSDF or if discarded (i.e. not the root of a parameter blending subtree), then there is no local basis to register
@@ -10356,7 +10377,7 @@ void FHLSLMaterialTranslator::StrataEvaluateSharedLocalBases(
 
 		// First, we check that the normal/tangent has not already written out (avoid 2 BSDFs sharing the same normal to note generate the same code twice)
 		bool bAlreadyProcessed = false;
-		for (uint8 i = 0; i < UsedSharedLocalBasesCount; ++i)
+		for (uint8 i = 0; i < OutUsedSharedLocalBasesCount; ++i)
 		{
 			if (UsedSharedLocalBasesInfo[i].NormalCodeChunkHash == StrataSharedLocalBasesInfo.SharedData.NormalCodeChunkHash &&
 				(UsedSharedLocalBasesInfo[i].TangentCodeChunkHash == StrataSharedLocalBasesInfo.SharedData.TangentCodeChunkHash || BSDFOperator.BSDFRegisteredSharedLocalBasis.TangentCodeChunk == INDEX_NONE))
@@ -10370,8 +10391,13 @@ void FHLSLMaterialTranslator::StrataEvaluateSharedLocalBases(
 			continue;
 		}
 
-		check(UsedSharedLocalBasesCount < STRATA_MAX_SHAREDLOCALBASES_REGISTERS);
-		const uint8 FinalSharedLocalBasisIndex = UsedSharedLocalBasesCount++;
+		++OutRequestedSharedLocalBasesCount;
+		if (OutUsedSharedLocalBasesCount >= STRATA_MAX_SHAREDLOCALBASES_REGISTERS)
+		{
+			continue;
+		}
+
+		const uint8 FinalSharedLocalBasisIndex = OutUsedSharedLocalBasesCount++;
 		UsedSharedLocalBasesInfo[FinalSharedLocalBasisIndex] = StrataSharedLocalBasesInfo.SharedData;
 
 		if (OutStrataPixelNormalInitializerValues)
@@ -10404,7 +10430,7 @@ void FHLSLMaterialTranslator::StrataEvaluateSharedLocalBases(
 	if (OutEnvironment)
 	{
 		// Now write out all the macros, them mapping from the BSDF to the effective position/index in the shared local basis array they should write to.
-		for (TMultiMap<uint64, FStrataSharedLocalBasesInfo>::TIterator It(CodeChunkToStrataSharedLocalBasis); It; ++It)
+		for (TMultiMap<uint64, FStrataSharedLocalBasesInfo>::TConstIterator It(CodeChunkToStrataSharedLocalBasis); It; ++It)
 		{
 			// The default linear output index will be 0 by default, and different if in fact the shared local basis points to one that is effectively in used in the array of shared local bases.
 			uint8 LinearIndex = 0;
@@ -10759,7 +10785,7 @@ bool FHLSLMaterialTranslator::StrataGenerateDerivedMaterialOperatorData()
 		//
 		{
 			// Compute the shared local basis count only
-			// But we cannot use StrataEvaluateSharedLocalBases here, becaise the material has not been compiled yet so al lthe bases would just default to the same.
+			// But we cannot use StrataEvaluateSharedLocalBases here, becaise the material has not been compiled yet so all the bases would just default to the same.
 			// STRATA_TODO: can we do that material generation in two passes? 
 			//		1- A first one to evaluate the normal/tangent code
 			//		2- Operators are processed and simplification computed based on memory budget
@@ -11054,16 +11080,16 @@ FStrataRegisteredSharedLocalBasis FHLSLMaterialTranslator::StrataCompilationInfo
 	return StrataRegisteredSharedLocalBasis;
 }
 
-FHLSLMaterialTranslator::FStrataSharedLocalBasesInfo FHLSLMaterialTranslator::StrataCompilationInfoGetMatchingSharedLocalBasisInfo(const FStrataRegisteredSharedLocalBasis& SearchedSharedLocalBasis)
+FHLSLMaterialTranslator::FStrataSharedLocalBasesInfo FHLSLMaterialTranslator::StrataCompilationInfoGetMatchingSharedLocalBasisInfo(const FStrataRegisteredSharedLocalBasis& SearchedSharedLocalBasis) const
 {
 	check(NextFreeStrataShaderNormalIndex < 255);	// Out of shared local basis slots
 
 	// Find a basis which matches both the Normal & the Tangent code chunks
-	TArray<FStrataSharedLocalBasesInfo*> NormalInfos;
+	TArray<const FStrataSharedLocalBasesInfo*> NormalInfos;
 	CodeChunkToStrataSharedLocalBasis.MultiFindPointer(SearchedSharedLocalBasis.NormalCodeChunkHash, NormalInfos);
 
 	// We first try to find a perfect match for normal and tangent from all the registered element.
-	for (FStrataSharedLocalBasesInfo* NormalInfo : NormalInfos)
+	for (const FStrataSharedLocalBasesInfo* NormalInfo : NormalInfos)
 	{
 		if (SearchedSharedLocalBasis.TangentCodeChunk == INDEX_NONE ||											// We selected the first available normal if there is no tangent specified on the material.
 			SearchedSharedLocalBasis.TangentCodeChunkHash == NormalInfo->SharedData.TangentCodeChunkHash)		// Otherwise we select the normal+tangent that exactly matches the request.
diff --git a/Engine/Source/Runtime/Engine/Private/Materials/HLSLMaterialTranslator.h b/Engine/Source/Runtime/Engine/Private/Materials/HLSLMaterialTranslator.h
--- a/Engine/Source/Runtime/Engine/Private/Materials/HLSLMaterialTranslator.h
+++ b/Engine/Source/Runtime/Engine/Private/Materials/HLSLMaterialTranslator.h
@@ -692,7 +692,7 @@ protected:
 	int32 GenericSwitch(const TCHAR* Function, int32 IfTrue, int32 IfFalse);
 
 	bool StrataGenerateDerivedMaterialOperatorData();
-	void StrataEvaluateSharedLocalBases(uint8& UsedSharedLocalBasesCount, FString* OutStrataPixelNormalInitializerValues, FShaderCompilerEnvironment* OutEnvironment);
+	void StrataEvaluateSharedLocalBases(uint8& UsedSharedLocalBasesCount, uint8& RequestedSharedLocalBasesCount, FString* OutStrataPixelNormalInitializerValues, FShaderCompilerEnvironment* OutEnvironment) const;
 	FString StrataGetCastParameterCode(int32 Index, EMaterialValueType DestType);
 
 	bool ValidateMaterialExpressionPermission(const UMaterialExpression* Expression);
@@ -1161,7 +1161,7 @@ protected:
 	virtual int32 StrataVerticalLayeringParameterBlendingBSDFCoverageToNormalMixCodeChunk(int32 TopCodeChunk) override;
 	virtual int32 StrataHorizontalMixingParameterBlendingBSDFCoverageToNormalMixCodeChunk(int32 BackgroundCodeChunk, int32 ForegroundCodeChunk, int32 HorizontalMixCodeChunk) override;
 
-	FStrataSharedLocalBasesInfo StrataCompilationInfoGetMatchingSharedLocalBasisInfo(const FStrataRegisteredSharedLocalBasis& SearchedSharedLocalBasis);
+	FStrataSharedLocalBasesInfo StrataCompilationInfoGetMatchingSharedLocalBasisInfo(const FStrataRegisteredSharedLocalBasis& SearchedSharedLocalBasis) const;
 
 #if HANDLE_CUSTOM_OUTPUTS_AS_MATERIAL_ATTRIBUTES
 	/** Used to translate code for custom output attributes such as ClearCoatBottomNormal */
diff --git a/Engine/Source/Runtime/Engine/Public/MaterialCompiler.h b/Engine/Source/Runtime/Engine/Public/MaterialCompiler.h
--- a/Engine/Source/Runtime/Engine/Public/MaterialCompiler.h
+++ b/Engine/Source/Runtime/Engine/Public/MaterialCompiler.h
@@ -172,7 +172,7 @@ struct FStrataOperator
 		bBSDFHasAnisotropy = A.bBSDFHasAnisotropy;
 	}
 
-	bool IsDiscarded()
+	bool IsDiscarded() const
 	{
 		return bUseParameterBlending && !bRootOfParameterBlendingSubTree;
 	}
@@ -740,7 +740,7 @@ public:
 
 	virtual FStrataRegisteredSharedLocalBasis StrataCompilationInfoRegisterSharedLocalBasis(int32 NormalCodeChunk) = 0;
 	virtual FStrataRegisteredSharedLocalBasis StrataCompilationInfoRegisterSharedLocalBasis(int32 NormalCodeChunk, int32 TangentCodeChunk) = 0;
-	FString GetStrataSharedLocalBasisIndexMacro(const FStrataRegisteredSharedLocalBasis& SharedLocalBasis)
+	FString GetStrataSharedLocalBasisIndexMacro(const FStrataRegisteredSharedLocalBasis& SharedLocalBasis) const
 	{
 		return FString::Printf(TEXT("SHAREDLOCALBASIS_INDEX_%u"), SharedLocalBasis.GraphSharedLocalBasisIndex);
 	}
