World Partition: remove 5.1 deprecated code.

diff --git a/Engine/Plugins/Editor/WorldPartitionHLODUtilities/Source/Private/WorldPartition/HLOD/Utilities/WorldPartitionHLODUtilities.cpp b/Engine/Plugins/Editor/WorldPartitionHLODUtilities/Source/Private/WorldPartition/HLOD/Utilities/WorldPartitionHLODUtilities.cpp
--- a/Engine/Plugins/Editor/WorldPartitionHLODUtilities/Source/Private/WorldPartition/HLOD/Utilities/WorldPartitionHLODUtilities.cpp
+++ b/Engine/Plugins/Editor/WorldPartitionHLODUtilities/Source/Private/WorldPartition/HLOD/Utilities/WorldPartitionHLODUtilities.cpp
@@ -144,10 +144,7 @@ static uint32 ComputeHLODHash(AWorldPartitionHLOD* InHLODActor, const TArray<AAc
 
 void AddSubActor(const FWorldPartitionActorDescView& ActorDescView, const IStreamingGenerationContext::FActorInstance& ActorInstance, TSet<FHLODSubActor>& SubActors)
 {
-PRAGMA_DISABLE_DEPRECATION_WARNINGS
-	// Leaving this as deprecated for now until we fix up the serialization format for SubActorsInfo and do proper upgrade/deprecation 
-	FName ActorPath = ActorDescView.GetActorPath();
-PRAGMA_ENABLE_DEPRECATION_WARNINGS
+	const FName ActorPath = *ActorDescView.GetActorSoftPath().ToString();
 
 	// Add the actor
 	bool bIsAlreadyInSet = false;
diff --git a/Engine/Source/Runtime/Engine/Private/WorldPartition/WorldPartitionRuntimeLevelStreamingCell.cpp b/Engine/Source/Runtime/Engine/Private/WorldPartition/WorldPartitionRuntimeLevelStreamingCell.cpp
--- a/Engine/Source/Runtime/Engine/Private/WorldPartition/WorldPartitionRuntimeLevelStreamingCell.cpp
+++ b/Engine/Source/Runtime/Engine/Private/WorldPartition/WorldPartitionRuntimeLevelStreamingCell.cpp
@@ -205,10 +205,8 @@ void UWorldPartitionRuntimeLevelStreamingCell::SetIsAlwaysLoaded(bool bInIsAlway
 void UWorldPartitionRuntimeLevelStreamingCell::AddActorToCell(const FWorldPartitionActorDescView& ActorDescView, const FActorContainerID& InContainerID, const FTransform& InContainerTransform, const UActorDescContainer* InContainer)
 {
 	check(!ActorDescView.GetActorIsEditorOnly());
-	// Leaving this using deprecated functions until the serialization format of Packages is updated.
-PRAGMA_DISABLE_DEPRECATION_WARNINGS
-	Packages.Emplace(ActorDescView.GetActorPackage(), ActorDescView.GetActorPath(), InContainerID, InContainerTransform, InContainer->GetContainerPackage(), GetWorld()->GetPackage()->GetFName(), ActorDescView.GetContentBundleGuid());
-PRAGMA_ENABLE_DEPRECATION_WARNINGS
+
+	Packages.Emplace(ActorDescView.GetActorPackage(), *ActorDescView.GetActorSoftPath().ToString(), InContainerID, InContainerTransform, InContainer->GetContainerPackage(), GetWorld()->GetPackage()->GetFName(), ActorDescView.GetContentBundleGuid());
 }
 
 bool UWorldPartitionRuntimeLevelStreamingCell::PopulateGeneratorPackageForCook(TArray<UPackage*>& OutModifiedPackages)
diff --git a/Engine/Source/Runtime/Engine/Public/WorldPartition/WorldPartitionActorDescView.h b/Engine/Source/Runtime/Engine/Public/WorldPartition/WorldPartitionActorDescView.h
--- a/Engine/Source/Runtime/Engine/Public/WorldPartition/WorldPartitionActorDescView.h
+++ b/Engine/Source/Runtime/Engine/Public/WorldPartition/WorldPartitionActorDescView.h
@@ -42,10 +42,6 @@ public:
 	const TArray<FName>& GetTags() const;
 	FName GetActorPackage() const;
 	
-PRAGMA_DISABLE_DEPRECATION_WARNINGS
-	UE_DEPRECATED(5.1, "GetActorPath is deprecated, GetActorSoftPath should be used instead.")
-	inline FName GetActorPath() const { return GetActorSoftPath().ToFName(); }
-PRAGMA_ENABLE_DEPRECATION_WARNINGS
 	FSoftObjectPath GetActorSoftPath() const;
 	FName GetActorLabel() const;
 
