Added a mode (disabled by default) that prevents MaterialInstances from having as parent Materials or MaterialInstances that don't satisfy a set of requirements.

When this mode is on, a Material or MaterialInstance is allowed to be a parent to another MaterialInstance if and only if:
- It is uncooked.
- It is a special engine material.
- It belongs to the same module as the child MaterialInstance.
- It is belongs to a set of "base modules".

diff --git a/Engine/Source/Editor/MaterialEditor/Private/MaterialEditorInstanceDetailCustomization.cpp b/Engine/Source/Editor/MaterialEditor/Private/MaterialEditorInstanceDetailCustomization.cpp
--- a/Engine/Source/Editor/MaterialEditor/Private/MaterialEditorInstanceDetailCustomization.cpp
+++ b/Engine/Source/Editor/MaterialEditor/Private/MaterialEditorInstanceDetailCustomization.cpp
@@ -956,7 +956,17 @@ bool FMaterialInstanceParameterDetails::OnShouldSetAsset(const FAssetData& Asset
 				EAppMsgType::Ok,
 				FText::Format(LOCTEXT("CannotSetExistingChildAsParent", "Cannot set {0} as a parent as it is already a child of this material instance."), FText::FromName(AssetData.AssetName)));
 		}
-		return !bIsChild;
+
+		if (bIsChild)
+		{
+			return false;
+		}
+	}
+
+	if (!MaterialEditorInstance->SourceInstance->IsSpecificMaterialValidParent(MaterialInstance))
+	{
+		FMessageDialog::Open(EAppMsgType::Ok, NSLOCTEXT("UnrealEd", "Error_MaterialEditor_InvalidMaterialInstanceParent", "Only uncooked, user-defined, engine or game Materials/MaterialInstances can be selected as a MaterialInstance parent."));
+		return false;
 	}
 
 	return true;
diff --git a/Engine/Source/Editor/UnrealEd/Private/Commandlets/SavePackageUtilitiesCommandlet.cpp b/Engine/Source/Editor/UnrealEd/Private/Commandlets/SavePackageUtilitiesCommandlet.cpp
--- a/Engine/Source/Editor/UnrealEd/Private/Commandlets/SavePackageUtilitiesCommandlet.cpp
+++ b/Engine/Source/Editor/UnrealEd/Private/Commandlets/SavePackageUtilitiesCommandlet.cpp
@@ -37,7 +37,7 @@ int32 USavePackageUtilitiesCommandlet::Main(const FString& Params)
 
 		// CookData should only be nonzero if we are cooking.
 		TOptional<FArchiveCookData> CookData;
-		FArchiveCookContext CookContext(Package, FArchiveCookContext::ECookTypeUnknown);
+		FArchiveCookContext CookContext(Package, FArchiveCookContext::ECookTypeUnknown, FArchiveCookContext::ECookingDLCUnknown);
 		if (TargetPlatform != nullptr)
 		{
 			CookData.Emplace(*TargetPlatform, CookContext);
diff --git a/Engine/Source/Editor/UnrealEd/Private/CookOnTheFlyServer.cpp b/Engine/Source/Editor/UnrealEd/Private/CookOnTheFlyServer.cpp
--- a/Engine/Source/Editor/UnrealEd/Private/CookOnTheFlyServer.cpp
+++ b/Engine/Source/Editor/UnrealEd/Private/CookOnTheFlyServer.cpp
@@ -5669,7 +5669,9 @@ FSaveCookedPackageContext::FSaveCookedPackageContext(UCookOnTheFlyServer& InCOTF
 	, Package(PackageData.GetPackage())
 	, PackageName(Package ? Package->GetName() : FString())
 	, Filename(PackageData.GetFileName().ToString())
-	, ArchiveCookContext(Package, InCOTFS.IsDirectorCookByTheBook() ? FArchiveCookContext::ECookByTheBook : FArchiveCookContext::ECookOnTheFly)
+	, ArchiveCookContext(Package,
+		InCOTFS.IsDirectorCookByTheBook() ? FArchiveCookContext::ECookByTheBook : FArchiveCookContext::ECookOnTheFly,
+		InCOTFS.IsCookingDLC() ? FArchiveCookContext::ECookingDLCYes : FArchiveCookContext::ECookingDLCNo)
 {
 }
 
diff --git a/Engine/Source/Runtime/CoreUObject/Private/UObject/SavePackage.cpp b/Engine/Source/Runtime/CoreUObject/Private/UObject/SavePackage.cpp
--- a/Engine/Source/Runtime/CoreUObject/Private/UObject/SavePackage.cpp
+++ b/Engine/Source/Runtime/CoreUObject/Private/UObject/SavePackage.cpp
@@ -20,7 +20,7 @@ FSavePackageResultStruct UPackage::Save(UPackage* InOuter, UObject* Base, EObjec
 {
 	// CookData should only be nonzero if we are cooking.
 	TOptional<FArchiveCookData> CookData;
-	FArchiveCookContext CookContext(InOuter, FArchiveCookContext::ECookTypeUnknown);
+	FArchiveCookContext CookContext(InOuter, FArchiveCookContext::ECookTypeUnknown, FArchiveCookContext::ECookingDLCUnknown);
 	if (TargetPlatform != nullptr)
 	{		
 		CookData.Emplace(*TargetPlatform, CookContext);
@@ -45,7 +45,7 @@ bool UPackage::SavePackage(UPackage* InOuter, UObject* Base, EObjectFlags TopLev
 {
 	// CookData should only be nonzero if we are cooking.
 	TOptional<FArchiveCookData> CookData;
-	FArchiveCookContext CookContext(InOuter, FArchiveCookContext::ECookTypeUnknown);
+	FArchiveCookContext CookContext(InOuter, FArchiveCookContext::ECookTypeUnknown, FArchiveCookContext::ECookingDLCUnknown);
 	if (TargetPlatform != nullptr)
 	{
 		CookData.Emplace(*TargetPlatform, CookContext);
diff --git a/Engine/Source/Runtime/CoreUObject/Public/UObject/ArchiveCookContext.h b/Engine/Source/Runtime/CoreUObject/Public/UObject/ArchiveCookContext.h
--- a/Engine/Source/Runtime/CoreUObject/Public/UObject/ArchiveCookContext.h
+++ b/Engine/Source/Runtime/CoreUObject/Public/UObject/ArchiveCookContext.h
@@ -17,20 +17,29 @@ public:
 		ECookByTheBook
 	};
 
+	enum ECookingDLC
+	{
+		ECookingDLCUnknown,
+		ECookingDLCYes,
+		ECookingDLCNo,
+	};
+
 private:
 
 	// Only valid for cook by the book.
 	FCookTagList CookTagList;
 	bool bCookTagListEnabled = false;
 	ECookType CookType;
+	ECookingDLC CookingDLC = ECookingDLCUnknown;
 
 public:
 
 
-	FArchiveCookContext(UPackage* InPackage, ECookType InCookType) :
+	FArchiveCookContext(UPackage* InPackage, ECookType InCookType, ECookingDLC CookingDLC) :
 		CookTagList(InPackage),
 		bCookTagListEnabled(InPackage && InCookType == ECookByTheBook),
-		CookType(InCookType)
+		CookType(InCookType),
+		CookingDLC(CookingDLC)
 	{
 		
 	}
@@ -46,4 +55,5 @@ public:
 	bool IsCookOnTheFly() const { return CookType == ECookType::ECookOnTheFly; }
 	bool IsCookTypeUnknown() const { return CookType == ECookType::ECookTypeUnknown; }
 	ECookType GetCookType() const { return CookType; }
+	ECookingDLC GetCookingDLC() const { return CookingDLC; }
 };
\ No newline at end of file
diff --git a/Engine/Source/Runtime/Engine/Classes/Materials/MaterialInstance.h b/Engine/Source/Runtime/Engine/Classes/Materials/MaterialInstance.h
--- a/Engine/Source/Runtime/Engine/Classes/Materials/MaterialInstance.h
+++ b/Engine/Source/Runtime/Engine/Classes/Materials/MaterialInstance.h
@@ -850,6 +850,16 @@ public:
 
 	ENGINE_API void SwapLayerParameterIndices(int32 OriginalIndex, int32 NewIndex);
 	ENGINE_API void RemoveLayerParameterIndex(int32 Index);
+	
+	/**
+	 * Returns whether specified MaterialInterface is a valid parent for this MaterialInstance.
+	 */
+	ENGINE_API bool IsSpecificMaterialValidParent(UMaterialInterface* CandidateParent) const;
+
+	/**
+	 * Ensures that current parent is a valid for this material instance and if not, it resets the parent to null.
+	 */
+	ENGINE_API void ValidateParent();
 
 #endif // WITH_EDITOR
 
diff --git a/Engine/Source/Runtime/Engine/Classes/Materials/MaterialInterface.h b/Engine/Source/Runtime/Engine/Classes/Materials/MaterialInterface.h
--- a/Engine/Source/Runtime/Engine/Classes/Materials/MaterialInterface.h
+++ b/Engine/Source/Runtime/Engine/Classes/Materials/MaterialInterface.h
@@ -275,6 +275,10 @@ public:
 	UPROPERTY(EditAnywhere, BlueprintReadOnly, Category = Material, meta = (DisplayName = "Subsurface Profile"))
 	TObjectPtr<class USubsurfaceProfile> SubsurfaceProfile;
 
+	/** Whether this material interface is included in the base game (and not in a DLC) */
+	UPROPERTY()
+	uint8 bIncludedInBaseGame : 1;
+
 	/* -------------------------- */
 
 	/** A fence to track when the primitive is no longer used as a parent */
@@ -340,7 +344,7 @@ public:
 	/** Importing data and options used for this material */
 	UPROPERTY(EditAnywhere, Instanced, Category = ImportSettings)
 	TObjectPtr<class UAssetImportData> AssetImportData;
-
+	
 private:
 	/** Unique ID for this material, used for caching during distributed lighting */
 	UPROPERTY()
diff --git a/Engine/Source/Runtime/Engine/Private/Materials/MaterialInstance.cpp b/Engine/Source/Runtime/Engine/Private/Materials/MaterialInstance.cpp
--- a/Engine/Source/Runtime/Engine/Private/Materials/MaterialInstance.cpp
+++ b/Engine/Source/Runtime/Engine/Private/Materials/MaterialInstance.cpp
@@ -61,6 +61,10 @@ DECLARE_CYCLE_STAT(TEXT("MaterialInstance CopyMatInstParams"), STAT_MaterialInst
 DECLARE_CYCLE_STAT(TEXT("MaterialInstance Serialize"), STAT_MaterialInstance_Serialize, STATGROUP_Shaders);
 DECLARE_CYCLE_STAT(TEXT("MaterialInstance CopyUniformParamsInternal"), STAT_MaterialInstance_CopyUniformParamsInternal, STATGROUP_Shaders);
 
+// This flag controls whether MaterialInstances parents should be restricted to be either uncooked, to be
+// user defined, part of the engine or part of the base game.
+bool bEnableRestrictiveMaterialInstanceParents = false;
+
 const FMaterialInstanceCachedData FMaterialInstanceCachedData::EmptyData{};
 
 void UMaterialInstance::StartCacheUniformExpressions() const
@@ -2271,7 +2275,7 @@ void UMaterialInstance::CacheResourceShadersForRendering(EMaterialShaderPrecompi
 			{
 				// Load the shader map for this resource, if needed
 				FMaterialResource Tmp;
-				FName PackageFileName = GetOutermost()->FileName;
+				FName PackageFileName = GetPackage()->FileName;
 				UE_CLOG(PackageFileName.IsNone(), LogMaterial, Warning,
 					TEXT("UMaterialInstance::CacheResourceShadersForRendering - Can't reload material resource '%s'. File system based reload is unsupported in this build."),
 					*GetFullName());
@@ -2848,7 +2852,7 @@ void UMaterialInstance::Serialize(FArchive& Ar)
 #if WITH_EDITOR
 		else
 		{
-			const bool bLoadedByCookedMaterial = FPlatformProperties::RequiresCookedData() || GetOutermost()->bIsCookedForEditor;
+			const bool bLoadedByCookedMaterial = FPlatformProperties::RequiresCookedData() || GetPackage()->bIsCookedForEditor;
 
 			FMaterialResource LegacyResource;
 			LegacyResource.LegacySerialize(Ar);
@@ -2997,6 +3001,10 @@ void UMaterialInstance::PostLoad()
 		Parent->ConditionalPostLoad();
 	}
 
+#if WITH_EDITOR
+	ValidateParent();
+#endif
+
 	// Add references to the expression object if we do not have one already, and fix up any names that were changed.
 	UpdateParameters();
 
@@ -3133,6 +3141,56 @@ void UMaterialInstance::PostLoad()
 	//DumpDebugInfo(*GLog);
 }
 
+#if WITH_EDITOR
+
+bool UMaterialInstance::IsSpecificMaterialValidParent(UMaterialInterface* CandidateParent) const
+{
+	// Nothing to do if specified parent is null or if restrictive mode is disabled
+	if (!CandidateParent || !bEnableRestrictiveMaterialInstanceParents)
+	{
+		return true;
+	}
+
+	// Cache this material package
+	UPackage* Package = GetPackage();
+
+	// Controls whether material is allowed.
+	bool bIsMaterialParentValid = false;
+	
+	// Or if this material instance package is cooked
+	bIsMaterialParentValid |= Package->HasAnyPackageFlags(PKG_Cooked);
+	
+	// Or if this material instance is transient
+	bIsMaterialParentValid |= Package == GetTransientPackage();
+	
+	// Or if the candidate material is uncooked
+	bIsMaterialParentValid |= !CandidateParent->GetPackage()->HasAnyPackageFlags(PKG_Cooked);
+	
+	// Or if the candidate parent is a Mateiral and it is flagged to be used as a special engine material
+	if (UMaterial* Material = Cast<UMaterial>(CandidateParent))
+	{
+		bIsMaterialParentValid |= Material->bUsedAsSpecialEngineMaterial;
+	}
+	
+	// Allow candidate material if it is included in base game.
+	bIsMaterialParentValid |= CandidateParent->bIncludedInBaseGame;
+
+	return bIsMaterialParentValid;
+}
+
+void UMaterialInstance::ValidateParent()
+{
+	const UMaterialInterface* PrevParent = Parent;
+	if (!IsSpecificMaterialValidParent(Parent))
+	{
+		// We don't allow Material Instances to parent to cooked materials.
+		UE_LOG(LogMaterial, Warning, TEXT("Can't set material instance '%s' to cooked non-user non-base parent material '%s'. Setting parent to null."), *GetName(), *PrevParent->GetName());
+		SetParentInternal(nullptr, true);
+	}
+}
+
+#endif // WITH_EDITOR
+
 #if WITH_EDITORONLY_DATA
 void UMaterialInstance::DeclareConstructClasses(TArray<FTopLevelAssetPath>& OutConstructClasses, const UClass* SpecificSubclass)
 {
@@ -3916,6 +3974,11 @@ void UMaterialInstance::PostEditChangeProperty(FPropertyChangedEvent& PropertyCh
 		FGlobalComponentRecreateRenderStateContext RecreateComponentsRenderState;
 	}
 
+	if (PropertyChangedEvent.MemberProperty != nullptr && PropertyChangedEvent.MemberProperty->GetFName() == GET_MEMBER_NAME_CHECKED(UMaterialInstance, Parent))
+	{
+		ValidateParent();
+	}
+
 	// If BLEND_TranslucentColoredTransmittance is selected while Strata is not enabled, force BLEND_Translucent blend mode
 	if (!Strata::IsStrataEnabled())
 	{
diff --git a/Engine/Source/Runtime/Engine/Private/Materials/MaterialInstanceConstant.cpp b/Engine/Source/Runtime/Engine/Private/Materials/MaterialInstanceConstant.cpp
--- a/Engine/Source/Runtime/Engine/Private/Materials/MaterialInstanceConstant.cpp
+++ b/Engine/Source/Runtime/Engine/Private/Materials/MaterialInstanceConstant.cpp
@@ -91,6 +91,7 @@ void UMaterialInstanceConstant::SetParentEditorOnly(UMaterialInterface* NewParen
 
 	if (SetParentInternal(NewParent, RecacheShader))
 	{
+		ValidateParent();
 		UpdateCachedData();
 	}
 }
diff --git a/Engine/Source/Runtime/Engine/Private/Materials/MaterialInterface.cpp b/Engine/Source/Runtime/Engine/Private/Materials/MaterialInterface.cpp
--- a/Engine/Source/Runtime/Engine/Private/Materials/MaterialInterface.cpp
+++ b/Engine/Source/Runtime/Engine/Private/Materials/MaterialInterface.cpp
@@ -36,6 +36,7 @@
 #include "ObjectCacheContext.h"
 #include "MaterialCachedData.h"
 #include "Components/DecalComponent.h"
+#include "UObject/ArchiveCookContext.h"
 
 #if WITH_EDITOR
 #include "ObjectCacheEventSink.h"
@@ -166,6 +167,13 @@ void UMaterialInterface::Serialize(FArchive& Ar)
 {
 	Ar.UsingCustomVersion(FUE5ReleaseStreamObjectVersion::GUID);
 
+	if (Ar.IsCooking())
+	{
+		// Mark whether this material is part of the base game. This provides information on whether it is safe to be
+		// used as a parent safely in child modules.
+		bIncludedInBaseGame = Ar.GetCookContext()->GetCookingDLC() == FArchiveCookContext::ECookingDLCNo;
+	}
+
 	Super::Serialize(Ar);
 
 	bool bSavedCachedExpressionData = false;
diff --git a/Engine/Source/Runtime/Engine/Private/ShaderCompiler/ShaderCompiler.cpp b/Engine/Source/Runtime/Engine/Private/ShaderCompiler/ShaderCompiler.cpp
--- a/Engine/Source/Runtime/Engine/Private/ShaderCompiler/ShaderCompiler.cpp
+++ b/Engine/Source/Runtime/Engine/Private/ShaderCompiler/ShaderCompiler.cpp
@@ -7603,7 +7603,7 @@ FString SaveGlobalShaderFile(EShaderPlatform Platform, FString SavePath, class I
 
 #if WITH_EDITOR
 		TOptional<FArchiveCookData> CookData;
-		FArchiveCookContext CookContext(nullptr /*InPackage*/, FArchiveCookContext::ECookTypeUnknown);
+		FArchiveCookContext CookContext(nullptr /*InPackage*/, FArchiveCookContext::ECookTypeUnknown, FArchiveCookContext::ECookingDLCUnknown);
 		if (TargetPlatform != nullptr)
 		{
 			CookData.Emplace(*TargetPlatform, CookContext);
@@ -8218,7 +8218,7 @@ void CompileGlobalShaderMapForRemote(
 	FNameAsStringProxyArchive Ar(MemWriter);
 
 	TOptional<FArchiveCookData> CookData;
-	FArchiveCookContext CookContext(nullptr /*InPackage*/, FArchiveCookContext::ECookTypeUnknown);
+	FArchiveCookContext CookContext(nullptr /*InPackage*/, FArchiveCookContext::ECookTypeUnknown, FArchiveCookContext::ECookingDLCUnknown);
 	if (TargetPlatform != nullptr)
 	{
 		CookData.Emplace(*TargetPlatform, CookContext);
@@ -8236,7 +8236,7 @@ void SaveShaderMapsForRemote(ITargetPlatform* TargetPlatform, const TMap<FString
 	FNameAsStringProxyArchive Ar(MemWriter);
 
 	TOptional<FArchiveCookData> CookData;
-	FArchiveCookContext CookContext(nullptr /*InPackage*/, FArchiveCookContext::ECookTypeUnknown);
+	FArchiveCookContext CookContext(nullptr /*InPackage*/, FArchiveCookContext::ECookTypeUnknown, FArchiveCookContext::ECookingDLCUnknown);
 	if (TargetPlatform != nullptr)
 	{
 		CookData.Emplace(*TargetPlatform, CookContext);
