World Partition: assert when destroying a world partition object that wasn't properly uninitialized.

diff --git a/Engine/Source/Runtime/Engine/Private/WorldPartition/WorldPartition.cpp b/Engine/Source/Runtime/Engine/Private/WorldPartition/WorldPartition.cpp
--- a/Engine/Source/Runtime/Engine/Private/WorldPartition/WorldPartition.cpp
+++ b/Engine/Source/Runtime/Engine/Private/WorldPartition/WorldPartition.cpp
@@ -1324,6 +1324,12 @@ bool UWorldPartition::ResolveSubobject(const TCHAR* SubObjectPath, UObject*& Out
 	return false;
 }
 
+void UWorldPartition::BeginDestroy()
+{
+	check(InitState == EWorldPartitionInitState::Uninitialized);
+	Super::BeginDestroy();
+}
+
 void UWorldPartition::AddReferencedObjects(UObject* InThis, FReferenceCollector& Collector)
 {
 #if WITH_EDITOR
diff --git a/Engine/Source/Runtime/Engine/Public/WorldPartition/WorldPartition.h b/Engine/Source/Runtime/Engine/Public/WorldPartition/WorldPartition.h
--- a/Engine/Source/Runtime/Engine/Public/WorldPartition/WorldPartition.h
+++ b/Engine/Source/Runtime/Engine/Public/WorldPartition/WorldPartition.h
@@ -149,6 +149,7 @@ public:
 	virtual void Serialize(FArchive& Ar) override;
 	virtual UWorld* GetWorld() const override;
 	virtual bool ResolveSubobject(const TCHAR* SubObjectPath, UObject*& OutObject, bool bLoadIfExists) override;
+	virtual void BeginDestroy() override;
 	static void AddReferencedObjects(UObject* InThis, FReferenceCollector& Collector);
 	//~ End UObject Interface
 
