work around to avoid crash when importing textures on ensure(!IsAssetStreamingSuspended())
a proper fix is still needed
streaming should not be disabled in viewport resize

diff --git a/Engine/Source/Runtime/Engine/Private/StreamableRenderAsset.cpp b/Engine/Source/Runtime/Engine/Private/StreamableRenderAsset.cpp
--- a/Engine/Source/Runtime/Engine/Private/StreamableRenderAsset.cpp
+++ b/Engine/Source/Runtime/Engine/Private/StreamableRenderAsset.cpp
@@ -368,10 +368,14 @@ void UStreamableRenderAsset::WaitForPendingInitOrStreaming(bool bWaitForLODTrans
 {
 	TRACE_CPUPROFILER_EVENT_SCOPE(UStreamableRenderAsset::WaitForPendingInitOrStreaming);
 
+	// this can be called with IsAssetStreamingSuspended == true
+	//	because Interchange Tasks due to PostEditChange do Texture UpdateResource
+	//	those tasks can be retracted in the D3D RHI Wait which runs in the Viewport resize
+	//	which turns off streaming
+	// @todo : Viewport resize should not turn off streaming
+
 	while (HasPendingInitOrStreaming(bWaitForLODTransition))
 	{
-		ensure(!IsAssetStreamingSuspended());
-
 		// Advance the streaming state.
 		TickStreaming(bSendCompletionEvents);
 		// Make sure any render commands are executed, in particular things like InitRHI, or asset updates on the render thread.
@@ -380,6 +384,9 @@ void UStreamableRenderAsset::WaitForPendingInitOrStreaming(bool bWaitForLODTrans
 		// Most of the time, sleeping is not required, so avoid loosing a whole quantum (10ms on W10Pro) unless stricly necessary.
 		if (HasPendingInitOrStreaming(bWaitForLODTransition))
 		{
+			// try to make sure streaming is enabled before doing horrible busy wait
+			ensure(!IsAssetStreamingSuspended());
+
 			// Give some time increment so that LOD transition can complete, and also for the gamethread to give room for streaming async tasks.
 			FPlatformProcess::Sleep(RENDER_ASSET_STREAMING_SLEEP_DT);
 		}
diff --git a/Engine/Source/Runtime/SlateRHIRenderer/Private/SlateRHIRenderer.cpp b/Engine/Source/Runtime/SlateRHIRenderer/Private/SlateRHIRenderer.cpp
--- a/Engine/Source/Runtime/SlateRHIRenderer/Private/SlateRHIRenderer.cpp
+++ b/Engine/Source/Runtime/SlateRHIRenderer/Private/SlateRHIRenderer.cpp
@@ -453,6 +453,9 @@ void FSlateRHIRenderer::ConditionalResizeViewport(FViewportInfo* ViewInfo, uint3
 
 		// Prevent the texture update logic to use the RHI while the viewport is resized. 
 		// This could happen if a streaming IO request completes and throws a callback.
+		// @todo : this does not in fact stop texture tasks from using the RHI while the viewport is resized
+		//		because they can be running in other threads, or even in retraction on this thread inside the D3D Wait
+		//		this should be removed and whatever streaming thread safety is needed during a viewport resize should be done correctly
 		SuspendTextureStreamingRenderTasks();
 
 		// cannot resize the viewport while potentially using it.
