Fix reinstancing crash when using a custom anim instance proxy as a linked anim graph

Also address a crash found when reinstancing a linked layer/graph with a null skeletal mesh applied (and thus an invalid bone container)

diff --git a/Engine/Source/Runtime/Engine/Private/Animation/AnimNode_LinkedAnimGraph.cpp b/Engine/Source/Runtime/Engine/Private/Animation/AnimNode_LinkedAnimGraph.cpp
--- a/Engine/Source/Runtime/Engine/Private/Animation/AnimNode_LinkedAnimGraph.cpp
+++ b/Engine/Source/Runtime/Engine/Private/Animation/AnimNode_LinkedAnimGraph.cpp
@@ -475,10 +475,26 @@ void FAnimNode_LinkedAnimGraph::HandleObjectsReinstanced_Impl(UObject* InSourceO
 	if(UseLegacyAnimInstanceReinstancingBehavior == nullptr || !UseLegacyAnimInstanceReinstancingBehavior->GetBool())
 	{
 		UAnimInstance* SourceAnimInstance = CastChecked<UAnimInstance>(InSourceObject);
-	
+		FAnimInstanceProxy& SourceProxy = SourceAnimInstance->GetProxyOnAnyThread<FAnimInstanceProxy>();
+
+		// Call Initialize here to ensure any custom proxies are initialized (as they may have been re-created during
+		// re-instancing, and they dont call the constructor that takes a UAnimInstance*)
+		SourceProxy.Initialize(SourceAnimInstance);
+
 		InitializeProperties(SourceAnimInstance, GetTargetClass());
 		DynamicUnlink(SourceAnimInstance);
 		DynamicLink(SourceAnimInstance);
+
+		SourceProxy.InitializeCachedClassData();
+
+		// Ensure we have a valid mesh at this point, as calling into the graph without one can result in crashes
+		// as we assume a valid bone container/reference skeleton is present
+		USkeletalMeshComponent* MeshComponent = SourceAnimInstance->GetSkelMeshComponent();
+		if(MeshComponent && MeshComponent->GetSkeletalMeshAsset())
+		{
+			FAnimationInitializeContext Context(&SourceProxy);
+			InitializeSubGraph_AnyThread(Context);
+		}
 	}
 }
 #endif
diff --git a/Engine/Source/Runtime/Engine/Private/Animation/AnimNode_LinkedAnimLayer.cpp b/Engine/Source/Runtime/Engine/Private/Animation/AnimNode_LinkedAnimLayer.cpp
--- a/Engine/Source/Runtime/Engine/Private/Animation/AnimNode_LinkedAnimLayer.cpp
+++ b/Engine/Source/Runtime/Engine/Private/Animation/AnimNode_LinkedAnimLayer.cpp
@@ -229,9 +229,15 @@ void FAnimNode_LinkedAnimLayer::HandleObjectsReinstanced_Impl(UObject* InSourceO
 		DynamicLink(SourceAnimInstance);
 
 		SourceProxy.InitializeCachedClassData();
-		
-		FAnimationInitializeContext Context(&SourceProxy);
-		InitializeSubGraph_AnyThread(Context);
+
+		// Ensure we have a valid mesh at this point, as calling into the graph without one can result in crashes
+		// as we assume a valid bone container/reference skeleton is present
+		USkeletalMeshComponent* MeshComponent = SourceAnimInstance->GetSkelMeshComponent();
+		if(MeshComponent && MeshComponent->GetSkeletalMeshAsset())
+		{
+			FAnimationInitializeContext Context(&SourceProxy);
+			InitializeSubGraph_AnyThread(Context);
+		}
 	}
 }
 
