Interchange: Fix morph target import, The iteration was done on the wrong number of vertices, so we were missing some vertex in the blend shape.

#rnx

diff --git a/Engine/Plugins/Interchange/Runtime/Source/Import/Private/Mesh/InterchangeSkeletalMeshFactory.cpp b/Engine/Plugins/Interchange/Runtime/Source/Import/Private/Mesh/InterchangeSkeletalMeshFactory.cpp
--- a/Engine/Plugins/Interchange/Runtime/Source/Import/Private/Mesh/InterchangeSkeletalMeshFactory.cpp
+++ b/Engine/Plugins/Interchange/Runtime/Source/Import/Private/Mesh/InterchangeSkeletalMeshFactory.cpp
@@ -239,9 +239,14 @@ namespace UE
 							MorphTargetData.Points[VertexID.GetValue()] = (FVector3f)GlobalTransform.TransformPosition((FVector)VertexPositions[VertexID]);
 						}
 
-						for (int32 PointIdx = VertexOffset; PointIdx < DestinationVertexIndexMax; ++PointIdx)
+						for (int32 PointIdx = VertexOffset; PointIdx < DestinationSkeletalMeshImportData.Points.Num(); ++PointIdx)
 						{
 							int32 OriginalPointIdx = DestinationSkeletalMeshImportData.PointToRawMap[PointIdx] - VertexOffset;
+							if (!MorphTargetData.Points.IsValidIndex(OriginalPointIdx))
+							{
+								//We break if we get over a valid point it mean we are inspecting a different part of the mesh
+								break;
+							}
 							//Rebuild the data with only the modified point
 							if ((MorphTargetData.Points[OriginalPointIdx] - DestinationSkeletalMeshImportData.Points[PointIdx]).SizeSquared() > FMath::Square(THRESH_POINTS_ARE_SAME))
 							{
diff --git a/Engine/Source/Runtime/Engine/Private/Rendering/SkeletalMeshLODImporterData.cpp b/Engine/Source/Runtime/Engine/Private/Rendering/SkeletalMeshLODImporterData.cpp
--- a/Engine/Source/Runtime/Engine/Private/Rendering/SkeletalMeshLODImporterData.cpp
+++ b/Engine/Source/Runtime/Engine/Private/Rendering/SkeletalMeshLODImporterData.cpp
@@ -1944,18 +1944,20 @@ FSkeletalMeshImportData FSkeletalMeshImportData::CreateFromMeshDescription(const
 	{
 		//We can use GetValue because the Meshdescription was compacted before the copy
 		SkelMeshImportData.Points[VertexID.GetValue()] = VertexPositions[VertexID];
-
-		FVertexBoneWeightsConst BoneWeights = VertexSkinWeights.Get(VertexID); 
-		const int32 InfluenceCount = BoneWeights.Num();
-		
-		const int32 InfluenceOffsetIndex = SkelMeshImportData.Influences.Num();
-		SkelMeshImportData.Influences.AddDefaulted(InfluenceCount);
-		for (int32 InfluenceIndex = 0; InfluenceIndex < InfluenceCount; ++InfluenceIndex)
+		if (VertexSkinWeights.IsValid())
 		{
-			SkeletalMeshImportData::FRawBoneInfluence& BoneInfluence = SkelMeshImportData.Influences[InfluenceOffsetIndex + InfluenceIndex];
-			BoneInfluence.VertexIndex = VertexID.GetValue();
-			BoneInfluence.BoneIndex = BoneWeights[InfluenceIndex].GetBoneIndex();
-			BoneInfluence.Weight = BoneWeights[InfluenceIndex].GetWeight();
+			FVertexBoneWeightsConst BoneWeights = VertexSkinWeights.Get(VertexID);
+			const int32 InfluenceCount = BoneWeights.Num();
+
+			const int32 InfluenceOffsetIndex = SkelMeshImportData.Influences.Num();
+			SkelMeshImportData.Influences.AddDefaulted(InfluenceCount);
+			for (int32 InfluenceIndex = 0; InfluenceIndex < InfluenceCount; ++InfluenceIndex)
+			{
+				SkeletalMeshImportData::FRawBoneInfluence& BoneInfluence = SkelMeshImportData.Influences[InfluenceOffsetIndex + InfluenceIndex];
+				BoneInfluence.VertexIndex = VertexID.GetValue();
+				BoneInfluence.BoneIndex = BoneWeights[InfluenceIndex].GetBoneIndex();
+				BoneInfluence.Weight = BoneWeights[InfluenceIndex].GetWeight();
+			}
 		}
 	}
 
