Set VR-related properties on capture views properly.

diff --git a/Engine/Source/Runtime/Engine/Private/SceneView.cpp b/Engine/Source/Runtime/Engine/Private/SceneView.cpp
--- a/Engine/Source/Runtime/Engine/Private/SceneView.cpp
+++ b/Engine/Source/Runtime/Engine/Private/SceneView.cpp
@@ -878,9 +878,10 @@ PRAGMA_ENABLE_DEPRECATION_WARNINGS
 			UE_LOG(LogMultiView, Fatal, TEXT("Family requires Mobile Multi-View, but it is not supported by the RHI and no fallback is available."));
 		}
 
-		bIsInstancedStereoEnabled = Aspects.IsInstancedStereoEnabled();
-		bIsMultiViewportEnabled = Aspects.IsInstancedMultiViewportEnabled();
-		bIsMobileMultiViewEnabled = (!bIsSceneCapture && !bIsReflectionCapture && !bIsPlanarReflection && Family && Family->bRequireMultiView && Aspects.IsMobileMultiViewEnabled());
+		const bool bCaptureOfSorts = bIsSceneCapture || bIsReflectionCapture || bIsPlanarReflection;
+		bIsInstancedStereoEnabled = !bCaptureOfSorts && Aspects.IsInstancedStereoEnabled();
+		bIsMultiViewportEnabled = !bCaptureOfSorts && Aspects.IsInstancedMultiViewportEnabled();
+		bIsMobileMultiViewEnabled = !bCaptureOfSorts && Family && Family->bRequireMultiView && Aspects.IsMobileMultiViewEnabled();
 
 		PRAGMA_DISABLE_DEPRECATION_WARNINGS
 		bIsMultiViewEnabled = bIsMultiViewportEnabled;	// temporary, as a graceful way to support plugins/licensee mods
