Add ULocalPlayerSaveGame base class to the engine that provides utility functions for saving/loading a save game object attached to a Local Player.
This wraps the existing save/load logic on GameplayStatics and provides functionality needed to support save/load on non-PC platforms.

diff --git a/Engine/Source/Runtime/Engine/Classes/GameFramework/SaveGame.h b/Engine/Source/Runtime/Engine/Classes/GameFramework/SaveGame.h
--- a/Engine/Source/Runtime/Engine/Classes/GameFramework/SaveGame.h
+++ b/Engine/Source/Runtime/Engine/Classes/GameFramework/SaveGame.h
@@ -5,8 +5,12 @@
 #include "CoreMinimal.h"
 #include "UObject/ObjectMacros.h"
 #include "UObject/Object.h"
+#include "Misc/CoreMiscDefines.h"
 #include "SaveGame.generated.h"
 
+class ULocalPlayer;
+class APlayerController;
+
 /** 
  *	This class acts as a base class for a save game object that can be used to save state about the game. 
  *	When you create your own save game subclass, you would add member variables for the information that you want to save.
@@ -26,5 +30,197 @@ class ENGINE_API USaveGame : public UObject
 	 *	@see UGameplayStatics::DeleteGameInSlot
 	 */
 
-	GENERATED_UCLASS_BODY()
+	GENERATED_BODY()
 };
+
+
+DECLARE_DYNAMIC_DELEGATE_OneParam(FOnLocalPlayerSaveGameLoaded, class ULocalPlayerSaveGame*, SaveGame);
+DECLARE_DELEGATE_OneParam(FOnLocalPlayerSaveGameLoadedNative, class ULocalPlayerSaveGame*);
+
+/**
+ * Abstract subclass of USaveGame that provides utility functions that let you associate a Save Game object with a specific local player.
+ * These objects can also be loaded using the functions on GameplayStatics, but you would need to call functions like InitializeSaveGame manually.
+ * For simple games, it is fine to blueprint this class directly and add parameters and override functions in blueprint,
+ * but for complicated games you will want to subclass this in native code and set up proper versioning.
+ */
+UCLASS(abstract, Blueprintable, BlueprintType)
+class ENGINE_API ULocalPlayerSaveGame : public USaveGame
+{
+	GENERATED_BODY()
+public:
+
+	/**
+	 * Synchronously loads a save game object in the specified slot for the local player, stalling the main thread until it completes.
+	 * This will return null for invalid parameters, but will create a new instance if the parameters are valid but loading fails.
+	 */
+	UFUNCTION(BlueprintCallable, Category = "SaveGame|LocalPlayer")
+	static ULocalPlayerSaveGame* LoadOrCreateSaveGameForLocalPlayer(TSubclassOf<ULocalPlayerSaveGame> SaveGameClass, APlayerController* LocalPlayerController, const FString& SlotName);
+
+	/**
+	 * Native version of above function, this takes a ULocalPlayer because you can have a local player before a player controller.
+	 * This will return null for invalid parameters, but will create a new instance if the parameters are valid but loading fails.
+	 */
+	static ULocalPlayerSaveGame* LoadOrCreateSaveGameForLocalPlayer(TSubclassOf<ULocalPlayerSaveGame> SaveGameClass, const ULocalPlayer* LocalPlayer, const FString& SlotName);
+
+	/**
+	 * Asynchronously loads a save game object in the specified slot for the local player, if this returns true the delegate will get called later.
+	 * False means the load was never scheduled, otherwise it will create and initialize a new instance before calling the delegate if loading failed.
+	 */
+	UFUNCTION(BlueprintCallable, Category = "SaveGame|LocalPlayer")
+	static bool AsyncLoadOrCreateSaveGameForLocalPlayer(TSubclassOf<ULocalPlayerSaveGame> SaveGameClass, APlayerController* LocalPlayerController, const FString& SlotName, FOnLocalPlayerSaveGameLoaded Delegate);
+
+	/**
+	 * Native version of above function, this takes a ULocalPlayer and calls a native delegate.
+	 * False means the load was never scheduled, otherwise it will create and initialize a new instance before calling the delegate if loading failed.
+	 */
+	static bool AsyncLoadOrCreateSaveGameForLocalPlayer(TSubclassOf<ULocalPlayerSaveGame> SaveGameClass, const ULocalPlayer* LocalPlayer, const FString& SlotName, FOnLocalPlayerSaveGameLoadedNative Delegate);
+
+	/**
+	 * Create a brand new save game, possibly ready for saving but without trying to load from disk. This will always succeed.
+	 */
+	static ULocalPlayerSaveGame* CreateNewSaveGameForLocalPlayer(TSubclassOf<ULocalPlayerSaveGame> SaveGameClass, const ULocalPlayer* LocalPlayer, const FString& SlotName);
+
+	/**
+	 * Synchronously save using the slot and user index, stalling the main thread until it completes. 
+	 * This will return true if the save was requested, and errors should be handled by the HandlePostSave function that will be called immediately.
+	 */
+	UFUNCTION(BlueprintCallable, Category = "SaveGame|LocalPlayer")
+	virtual bool SaveGameToSlotForLocalPlayer();
+
+	/**
+	 * Asynchronously save to the slot and user index.
+	 * This will return true if the save was requested, and errors should be handled by the HandlePostSave function after the save succeeds or fails
+	 */
+	UFUNCTION(BlueprintCallable, Category = "SaveGame|LocalPlayer")
+	virtual bool AsyncSaveGameToSlotForLocalPlayer();
+
+
+	/** Returns the local player controller this is associated with, this will be valid if it is ready to save */
+	UFUNCTION(BlueprintCallable, Category = "SaveGame|LocalPlayer")
+	virtual APlayerController* GetLocalPlayerController() const;
+
+	/** Returns the local player this is associated with, this will be valid if it is ready to save */
+	virtual const ULocalPlayer* GetLocalPlayer() const;
+
+	/** Associates this save game with a local player, this is called automatically during load/create */
+	virtual void SetLocalPlayer(const ULocalPlayer* LocalPlayer);
+
+	/** Returns the platform user to save to, based on Local Player by default */
+	UFUNCTION(BlueprintCallable, Category = "SaveGame|LocalPlayer")
+	virtual FPlatformUserId GetPlatformUserId() const;
+
+	/** Returns the user index to save to, based on Local Player by default */
+	UFUNCTION(BlueprintCallable, Category = "SaveGame|LocalPlayer")
+	virtual int32 GetPlatformUserIndex() const;
+
+	/** Returns the save slot name to use */
+	UFUNCTION(BlueprintCallable, Category = "SaveGame|LocalPlayer")
+	virtual FString GetSaveSlotName() const;
+
+	/** Sets the save slot name for any future saves */
+	virtual void SetSaveSlotName(const FString& SlotName);
+
+	/** Returns the game-specific version number this was last saved/loaded as */
+	UFUNCTION(BlueprintCallable, Category = "SaveGame|LocalPlayer")
+	virtual int32 GetSavedDataVersion() const;
+
+	/** Returns the invalid save data version, which means it has never been saved/loaded */
+	UFUNCTION(BlueprintCallable, Category = "SaveGame|LocalPlayer")
+	virtual int32 GetInvalidDataVersion() const;
+
+	/** Returns the latest save data version, this is used when the new data is saved */
+	UFUNCTION(BlueprintCallable, Category = "SaveGame|LocalPlayer")
+	virtual int32 GetLatestDataVersion() const;
+
+	/** Returns true if this was loaded from an existing save */
+	UFUNCTION(BlueprintCallable, Category = "SaveGame|LocalPlayer")
+	virtual bool WasLoaded() const;
+
+	/** Returns true if a save is in progress */
+	UFUNCTION(BlueprintCallable, Category = "SaveGame|LocalPlayer")
+	virtual bool IsSaveInProgress() const;
+
+	/** Returns true if it has been saved at least once and the last save was successful */
+	UFUNCTION(BlueprintCallable, Category = "SaveGame|LocalPlayer")
+	virtual bool WasLastSaveSuccessful() const;
+
+	/** Returns true if a save was ever requested, may still be in progress */
+	UFUNCTION(BlueprintCallable, Category = "SaveGame|LocalPlayer")
+	virtual bool WasSaveRequested() const;
+
+
+	/** Initializes this save after either loading or initial creation, automatically called by load/create functions above */
+	virtual void InitializeSaveGame(const ULocalPlayer* LocalPlayer, FString InSlotName, bool bWasLoaded);
+
+	/** Resets all saved data to default values, called when the load fails or manually */
+	UFUNCTION(BlueprintCallable, Category = "SaveGame|LocalPlayer")
+	virtual void ResetToDefault();
+
+	/** Blueprint event called to reset all saved data to default, called when the load fails or manually */
+	UFUNCTION(BlueprintImplementableEvent, Category = "SaveGame|LocalPlayer")
+	void OnResetToDefault();
+
+	/** Called after loading, this is not called for newly created saves */
+	virtual void HandlePostLoad();
+
+	/** Blueprint event called after loading, is not called for newly created saves */
+	UFUNCTION(BlueprintImplementableEvent, Category = "SaveGame|LocalPlayer")
+	void OnPostLoad();
+
+	/** Called before saving, do any game-specific fixup here */
+	virtual void HandlePreSave();
+
+	/** Blueprint event called before saving, do any game-specific fixup here  */
+	UFUNCTION(BlueprintImplementableEvent, Category = "SaveGame|LocalPlayer")
+	void OnPreSave();
+
+	/** Called after saving finishes with success/failure result */
+	virtual void HandlePostSave(bool bSuccess);
+
+	/** Blueprint event called after saving finishes with success/failure result */
+	UFUNCTION(BlueprintImplementableEvent, Category = "SaveGame|LocalPlayer")
+	void OnPostSave(bool bSuccess);
+
+protected:
+	/** Internal callback that will call HandlePostSave */
+	virtual void ProcessSaveComplete(const FString& SlotName, const int32 UserIndex, bool bSuccess, int32 SaveRequest);
+
+	/** Internal helper function used by both the sync and async save */
+	static ULocalPlayerSaveGame* ProcessLoadedSave(USaveGame* BaseSave, const FString& SlotName, const int32 UserIndex, TSubclassOf<ULocalPlayerSaveGame> SaveGameClass, const ULocalPlayer* LocalPlayer);
+
+	/** The local player this is connected to, can be null if subclasses override Get/Set Local Player or it hasn't been initialized */
+	UPROPERTY(Transient)
+	TObjectPtr<const ULocalPlayer> OwningPlayer;
+
+	/** The slot name this was loaded from and that will be used to save to in the future */
+	UPROPERTY(Transient)
+	FString SaveSlotName;
+
+	/** 
+	 * The value of GetLatestDataVersion when this was last saved.
+	 * Subclasses can override GetLatestDataVersion and then handle fixups in HandlePostLoad.
+	 * This defaults to 0 so old save games that didn't previously subclass ULocalPlayerSaveGame will have 0 instead of the invalid version.
+	 */
+	UPROPERTY()
+	int32 SavedDataVersion = 0;
+
+	/**
+	 * The value of SavedDataVersion when a save was last loaded, this will be -1 for newly created saves
+	 */
+	UPROPERTY(Transient)
+	int32 LoadedDataVersion = -1;
+
+
+	/** Integer that is incremented every time a save has been requested in the current session, can be used to know if one is in progress */
+	UPROPERTY(Transient)
+	int32 CurrentSaveRequest = 0;
+
+	/** Integer that is set when a save completes successfully, if this equals RequestedSaveCount then the last save was successful */
+	UPROPERTY(Transient)
+	int32 LastSuccessfulSaveRequest = 0;
+
+	/** Integer that is set when a save fails */
+	UPROPERTY(Transient)
+	int32 LastErrorSaveRequest = 0;
+
+};
\ No newline at end of file
diff --git a/Engine/Source/Runtime/Engine/Private/GameFramework/SaveGame.cpp b/Engine/Source/Runtime/Engine/Private/GameFramework/SaveGame.cpp
--- a/Engine/Source/Runtime/Engine/Private/GameFramework/SaveGame.cpp
+++ b/Engine/Source/Runtime/Engine/Private/GameFramework/SaveGame.cpp
@@ -1,12 +1,355 @@
 // Copyright Epic Games, Inc. All Rights Reserved.
 
 #include "GameFramework/SaveGame.h"
+#include "Engine/LocalPlayer.h"
+#include "Kismet/GameplayStatics.h"
+#include "EngineLogs.h"
 
 #include UE_INLINE_GENERATED_CPP_BY_NAME(SaveGame)
 
-USaveGame::USaveGame(const FObjectInitializer& ObjectInitializer)
-	: Super(ObjectInitializer)
+ULocalPlayerSaveGame* ULocalPlayerSaveGame::LoadOrCreateSaveGameForLocalPlayer(TSubclassOf<ULocalPlayerSaveGame> SaveGameClass, APlayerController* LocalPlayerController, const FString& SlotName)
 {
+	const ULocalPlayer* LocalPlayer = LocalPlayerController ? LocalPlayerController->GetLocalPlayer() : nullptr;
 
+	if (!ensure(LocalPlayer))
+	{
+		UE_LOG(LogPlayerManagement, Error, TEXT("LoadOrCreateSaveGameForLocalPlayer called with null or remote player controller!"));
+		return nullptr;
+	}
+
+	return LoadOrCreateSaveGameForLocalPlayer(SaveGameClass, LocalPlayer, SlotName);
+}
+
+ULocalPlayerSaveGame* ULocalPlayerSaveGame::LoadOrCreateSaveGameForLocalPlayer(TSubclassOf<ULocalPlayerSaveGame> SaveGameClass, const ULocalPlayer* LocalPlayer, const FString& SlotName)
+{
+	if (!ensure(LocalPlayer))
+	{
+		UE_LOG(LogPlayerManagement, Error, TEXT("LoadOrCreateSaveGame called with null local player!"));
+		return nullptr;
+	}
+
+	if (!ensure(SaveGameClass))
+	{
+		UE_LOG(LogPlayerManagement, Error, TEXT("LoadOrCreateSaveGame called with null SaveGameClass!"));
+		return nullptr;
+	}
+
+	if (!ensure(SlotName.Len() > 0))
+	{
+		UE_LOG(LogPlayerManagement, Error, TEXT("LoadOrCreateSaveGame called with empty slot name!"));
+		return nullptr;
+	}
+
+	int32 UserIndex = LocalPlayer->GetPlatformUserIndex();
+
+	// If the save game exists, load it.
+	ULocalPlayerSaveGame* LoadedSave = nullptr;
+	USaveGame* BaseSave = nullptr;
+	if (UGameplayStatics::DoesSaveGameExist(SlotName, UserIndex))
+	{
+		BaseSave = UGameplayStatics::LoadGameFromSlot(SlotName, UserIndex);
+	}
+
+	LoadedSave = ProcessLoadedSave(BaseSave, SlotName, UserIndex, SaveGameClass, LocalPlayer);
+
+	return LoadedSave;
+}
+
+bool ULocalPlayerSaveGame::AsyncLoadOrCreateSaveGameForLocalPlayer(TSubclassOf<ULocalPlayerSaveGame> SaveGameClass, APlayerController* LocalPlayerController, const FString& SlotName, FOnLocalPlayerSaveGameLoaded Delegate)
+{
+	const ULocalPlayer* LocalPlayer = LocalPlayerController ? LocalPlayerController->GetLocalPlayer() : nullptr;
+
+	if (!ensure(LocalPlayer))
+	{
+		UE_LOG(LogPlayerManagement, Error, TEXT("AsyncLoadOrCreateSaveGameForLocalPlayer called with null or remote player controller!"));
+		return false;
+	}
+
+	// Simple lambda that calls the dynamic delegate
+	FOnLocalPlayerSaveGameLoadedNative Lambda = FOnLocalPlayerSaveGameLoadedNative::CreateLambda([Delegate]
+		(ULocalPlayerSaveGame* LoadedSave)
+		{
+			Delegate.ExecuteIfBound(LoadedSave);
+		});
+
+	return AsyncLoadOrCreateSaveGameForLocalPlayer(SaveGameClass, LocalPlayer, SlotName, Lambda);
+}
+
+bool ULocalPlayerSaveGame::AsyncLoadOrCreateSaveGameForLocalPlayer(TSubclassOf<ULocalPlayerSaveGame> SaveGameClass, const ULocalPlayer* LocalPlayer, const FString& SlotName, FOnLocalPlayerSaveGameLoadedNative Delegate)
+{
+	if (!ensure(LocalPlayer))
+	{
+		UE_LOG(LogPlayerManagement, Error, TEXT("AsyncLoadOrCreateSaveGame called with null local player!"));
+		return false;
+	}
+
+	if (!ensure(SaveGameClass))
+	{
+		UE_LOG(LogPlayerManagement, Error, TEXT("AsyncLoadOrCreateSaveGame called with null SaveGameClass!"));
+		return false;
+	}
+
+	if (!ensure(SlotName.Len() > 0))
+	{
+		UE_LOG(LogPlayerManagement, Error, TEXT("AsyncLoadOrCreateSaveGame called with empty slot name!"));
+		return false;
+	}
+	
+	FAsyncLoadGameFromSlotDelegate Lambda = FAsyncLoadGameFromSlotDelegate::CreateWeakLambda(LocalPlayer, [LocalPlayer, Delegate, SaveGameClass]
+		(const FString& SlotName, const int32 UserIndex, USaveGame* BaseSave) 
+		{
+			ULocalPlayerSaveGame* LoadedSave = ProcessLoadedSave(BaseSave, SlotName, UserIndex, SaveGameClass, LocalPlayer);
+			
+			Delegate.ExecuteIfBound(LoadedSave);
+		});
+
+	UGameplayStatics::AsyncLoadGameFromSlot(SlotName, LocalPlayer->GetPlatformUserIndex(), Lambda);
+
+	return true;
+}
+
+ULocalPlayerSaveGame* ULocalPlayerSaveGame::ProcessLoadedSave(USaveGame* BaseSave, const FString& SlotName, const int32 UserIndex, TSubclassOf<ULocalPlayerSaveGame> SaveGameClass, const ULocalPlayer* LocalPlayer)
+{
+	ULocalPlayerSaveGame* LoadedSave = Cast<ULocalPlayerSaveGame>(BaseSave);
+	if (BaseSave && (!LoadedSave || !LoadedSave->IsA(SaveGameClass.Get())))
+	{
+		UE_LOG(LogPlayerManagement, Warning, TEXT("ProcessLoadedSave found invalid save game object %s in slot %s for player %d!"), *GetNameSafe(LoadedSave), *SlotName, UserIndex);
+		LoadedSave = nullptr;
+	}
+
+	if (LoadedSave == nullptr)
+	{
+		LoadedSave = CreateNewSaveGameForLocalPlayer(SaveGameClass, LocalPlayer, SlotName);
+	}
+	else
+	{
+		LoadedSave->InitializeSaveGame(LocalPlayer, SlotName, true);
+	}
+
+	return LoadedSave;
+}
+
+ULocalPlayerSaveGame* ULocalPlayerSaveGame::CreateNewSaveGameForLocalPlayer(TSubclassOf<ULocalPlayerSaveGame> SaveGameClass, const ULocalPlayer* LocalPlayer, const FString& SlotName)
+{
+	ULocalPlayerSaveGame* LoadedSave = Cast<ULocalPlayerSaveGame>(UGameplayStatics::CreateSaveGameObject(SaveGameClass));
+	if (ensure(LoadedSave))
+	{
+		LoadedSave->InitializeSaveGame(LocalPlayer, SlotName, false);
+	}
+	return LoadedSave;
 }
 
+bool ULocalPlayerSaveGame::SaveGameToSlotForLocalPlayer()
+{
+	// Subclasses can override this to add additional checks like if one is in progress
+	if (!ensure(GetLocalPlayer()))
+	{
+		UE_LOG(LogPlayerManagement, Error, TEXT("SaveGameToSlotForLocalPlayer called with null local player!"));
+		return false;
+	}
+
+	int32 RequestUserIndex = GetPlatformUserIndex();
+	FString RequestSlotName = GetSaveSlotName();
+	if (!ensure(RequestSlotName.Len() > 0))
+	{
+		UE_LOG(LogPlayerManagement, Error, TEXT("SaveGameToSlotForLocalPlayer called with empty slot name!"));
+		return false;
+	}
+
+	HandlePreSave();
+
+	bool bSuccess = UGameplayStatics::SaveGameToSlot(this, RequestSlotName, RequestUserIndex);
+
+	ProcessSaveComplete(RequestSlotName, RequestUserIndex, bSuccess, CurrentSaveRequest);
+
+	return true;
+}
+
+bool ULocalPlayerSaveGame::AsyncSaveGameToSlotForLocalPlayer()
+{
+	// Subclasses can override this to add additional checks like if one is in progress
+	if (!ensure(GetLocalPlayer()))
+	{
+		UE_LOG(LogPlayerManagement, Error, TEXT("AsyncSaveGameToSlotForLocalPlayer called with null local player!"));
+		return false;
+	}
+	
+	int32 RequestUserIndex = GetPlatformUserIndex();
+	FString RequestSlotName = GetSaveSlotName();
+	if (!ensure(RequestSlotName.Len() > 0))
+	{
+		UE_LOG(LogPlayerManagement, Error, TEXT("AsyncSaveGameToSlotForLocalPlayer called with empty slot name!"));
+		return false;
+	}
+
+	HandlePreSave();
+
+	// Start an async save and pass through the current save request count
+	FAsyncSaveGameToSlotDelegate SavedDelegate;
+	SavedDelegate.BindUObject(this, &ULocalPlayerSaveGame::ProcessSaveComplete, CurrentSaveRequest);
+	UGameplayStatics::AsyncSaveGameToSlot(this, RequestSlotName, RequestUserIndex, SavedDelegate);
+
+	return true;
+}
+
+void ULocalPlayerSaveGame::ProcessSaveComplete(const FString& SlotName, const int32 UserIndex, bool bSuccess, int32 SaveRequest)
+{
+	// Now that a save completed, update the success/failure
+	// Every time CurrentSaveRequest is incremented it will result in setting either success or error, but there could be multiple in flight that are processed in order
+	if (bSuccess)
+	{
+		ensure(CurrentSaveRequest > LastSuccessfulSaveRequest);
+		LastSuccessfulSaveRequest = CurrentSaveRequest;
+	}
+	else
+	{
+		ensure(CurrentSaveRequest > LastErrorSaveRequest);
+		LastErrorSaveRequest = CurrentSaveRequest;
+	}
+
+	HandlePostSave(bSuccess);
+}
+
+APlayerController* ULocalPlayerSaveGame::GetLocalPlayerController() const
+{
+	if (OwningPlayer)
+	{
+		return OwningPlayer->GetPlayerController(nullptr);
+	}
+
+	return nullptr;
+}
+
+const ULocalPlayer* ULocalPlayerSaveGame::GetLocalPlayer() const
+{
+	return OwningPlayer;
+}
+
+void ULocalPlayerSaveGame::SetLocalPlayer(const ULocalPlayer* LocalPlayer)
+{
+	OwningPlayer = LocalPlayer;
+}
+
+FPlatformUserId ULocalPlayerSaveGame::GetPlatformUserId() const
+{
+	if (OwningPlayer)
+	{
+		return OwningPlayer->GetPlatformUserId();
+	}
+
+	return FPlatformUserId();
+}
+
+int32 ULocalPlayerSaveGame::GetPlatformUserIndex() const
+{
+	if (OwningPlayer)
+	{
+		return OwningPlayer->GetPlatformUserIndex();
+	}
+
+	return INDEX_NONE;
+}
+
+FString ULocalPlayerSaveGame::GetSaveSlotName() const
+{
+	return SaveSlotName;
+}
+
+void ULocalPlayerSaveGame::SetSaveSlotName(const FString& SlotName)
+{
+	SaveSlotName = SlotName;
+}
+
+int32 ULocalPlayerSaveGame::GetSavedDataVersion() const
+{
+	return SavedDataVersion;
+}
+
+int32 ULocalPlayerSaveGame::GetInvalidDataVersion() const
+{
+	return -1;
+}
+
+int32 ULocalPlayerSaveGame::GetLatestDataVersion() const
+{
+	// Override with game-specific logic
+	return 0;
+}
+
+bool ULocalPlayerSaveGame::WasLoaded() const
+{
+	return LoadedDataVersion != GetInvalidDataVersion();
+}
+
+bool ULocalPlayerSaveGame::IsSaveInProgress() const
+{
+	return (CurrentSaveRequest > LastErrorSaveRequest) || (CurrentSaveRequest > LastSuccessfulSaveRequest);
+}
+
+bool ULocalPlayerSaveGame::WasLastSaveSuccessful() const
+{
+	return (WasSaveRequested() && LastSuccessfulSaveRequest > LastErrorSaveRequest);
+}
+
+bool ULocalPlayerSaveGame::WasSaveRequested() const
+{
+	return CurrentSaveRequest > 0;
+}
+
+void ULocalPlayerSaveGame::InitializeSaveGame(const ULocalPlayer* LocalPlayer, FString InSlotName, bool bWasLoaded)
+{
+	SetLocalPlayer(LocalPlayer);
+	SetSaveSlotName(InSlotName);
+
+	if (bWasLoaded)
+	{
+		LoadedDataVersion = SavedDataVersion;
+		HandlePostLoad();
+	}
+	else
+	{
+		ResetToDefault();
+	}
+}
+
+void ULocalPlayerSaveGame::ResetToDefault()
+{
+	SavedDataVersion = GetInvalidDataVersion();
+	LoadedDataVersion = SavedDataVersion;
+
+	OnResetToDefault();
+}
+
+void ULocalPlayerSaveGame::HandlePostLoad()
+{
+	// Override this to handle SavedDataVersion fixups
+
+	OnPostLoad();
+}
+
+void ULocalPlayerSaveGame::HandlePreSave()
+{
+	OnPreSave();
+
+	// Set the save data version and increment the requested count
+
+	SavedDataVersion = GetLatestDataVersion();
+	CurrentSaveRequest++;
+
+	UE_LOG(LogPlayerManagement, Verbose, TEXT("Starting to save game %s request %d to slot %s for user %d"), *GetName(), CurrentSaveRequest, *GetSaveSlotName(), GetPlatformUserIndex());
+}
+
+void ULocalPlayerSaveGame::HandlePostSave(bool bSuccess)
+{
+	// Override for error handling
+	if (bSuccess)
+	{
+		UE_LOG(LogPlayerManagement, Verbose, TEXT("Successfully saved game %s request %d to slot %s for user %d"), *GetName(), LastSuccessfulSaveRequest, *GetSaveSlotName(), GetPlatformUserIndex());
+	}
+	else
+	{
+		UE_LOG(LogPlayerManagement, Error, TEXT("Failed to save game %s request %d to slot %s for user %d!"), *GetName(), LastErrorSaveRequest, *GetSaveSlotName(), GetPlatformUserIndex());
+	}
+
+	OnPostSave(bSuccess);
+}
