expose a ComponentOverlapComponent function on the primitive component that also returns the overlap result

diff --git a/Engine/Source/Runtime/Engine/Classes/Components/PrimitiveComponent.h b/Engine/Source/Runtime/Engine/Classes/Components/PrimitiveComponent.h
--- a/Engine/Source/Runtime/Engine/Classes/Components/PrimitiveComponent.h
+++ b/Engine/Source/Runtime/Engine/Classes/Components/PrimitiveComponent.h
@@ -2776,9 +2776,23 @@ public:
 	bool ComponentOverlapComponent(class UPrimitiveComponent* PrimComp, const FVector Pos, const FQuat& Rot, const FCollisionQueryParams& Params);
 	bool ComponentOverlapComponent(class UPrimitiveComponent* PrimComp, const FVector Pos, const FRotator Rot, const FCollisionQueryParams& Params);
 
+	/**
+	 *  Test the collision of the supplied component at the supplied location/rotation, and determine if it overlaps this component.
+	 *  @note This overload taking rotation as a FQuat is slightly faster than the version using FRotator.
+	 *  @note This simply calls the virtual ComponentOverlapComponentImpl() which can be overridden to implement custom behavior.
+	 *  @param  PrimComp        Component to use geometry from to test against this component. Transform of this component is ignored.
+	 *  @param  Pos             Location to place PrimComp geometry at
+	 *  @param  Rot             Rotation to place PrimComp geometry at
+	 *  @param  Params          Parameter for trace. TraceTag is only used.
+	 *  @param  OutOverlap      Also returns all the sub-overlaps within the component.
+	 *  @return true if PrimComp overlaps this component at the specified location/rotation
+	 */
+	bool ComponentOverlapComponentWithResult(const class UPrimitiveComponent* const PrimComp, const FVector& Pos, const FQuat& Rot, const FCollisionQueryParams& Params, TArray<FOverlapResult>& OutOverlap) const;
+	bool ComponentOverlapComponentWithResult(const class UPrimitiveComponent* const PrimComp, const FVector& Pos, const FRotator& Rot, const FCollisionQueryParams& Params, TArray<FOverlapResult>& OutOverlap) const;
 protected:
 	/** Override this method for custom behavior for ComponentOverlapComponent() */
 	virtual bool ComponentOverlapComponentImpl(class UPrimitiveComponent* PrimComp, const FVector Pos, const FQuat& Rot, const FCollisionQueryParams& Params);
+	virtual bool ComponentOverlapComponentWithResultImpl(const class UPrimitiveComponent* const PrimComp, const FVector& Pos, const FQuat& Rot, const FCollisionQueryParams& Params, TArray<FOverlapResult>& OutOverlap) const;
 
 public:	
 	/** 
@@ -2919,6 +2933,16 @@ FORCEINLINE_DEBUGGABLE bool UPrimitiveComponent::ComponentOverlapComponent(class
 	return ComponentOverlapComponentImpl(PrimComp, Pos, Rot.Quaternion(), Params);
 }
 
+FORCEINLINE_DEBUGGABLE bool UPrimitiveComponent::ComponentOverlapComponentWithResult(const class UPrimitiveComponent* const PrimComp, const FVector& Pos, const FQuat& Rot, const FCollisionQueryParams& Params, TArray<FOverlapResult>& OutOverlap) const
+{
+	return ComponentOverlapComponentWithResultImpl(PrimComp, Pos, Rot, Params, OutOverlap);
+}
+
+FORCEINLINE_DEBUGGABLE bool UPrimitiveComponent::ComponentOverlapComponentWithResult(const class UPrimitiveComponent* const PrimComp, const FVector& Pos, const FRotator& Rot, const FCollisionQueryParams& Params, TArray<FOverlapResult>& OutOverlap) const
+{
+	return ComponentOverlapComponentWithResultImpl(PrimComp, Pos, Rot.Quaternion(), Params, OutOverlap);
+}
+
 FORCEINLINE_DEBUGGABLE const TArray<FOverlapInfo>& UPrimitiveComponent::GetOverlapInfos() const
 {
 	return OverlappingComponents;
diff --git a/Engine/Source/Runtime/Engine/Classes/PhysicsEngine/ClusterUnionComponent.h b/Engine/Source/Runtime/Engine/Classes/PhysicsEngine/ClusterUnionComponent.h
--- a/Engine/Source/Runtime/Engine/Classes/PhysicsEngine/ClusterUnionComponent.h
+++ b/Engine/Source/Runtime/Engine/Classes/PhysicsEngine/ClusterUnionComponent.h
@@ -209,6 +209,7 @@ public:
 	virtual bool LineTraceComponent(FHitResult& OutHit, const FVector Start, const FVector End, const FCollisionQueryParams& Params) override;
 	virtual bool SweepComponent(FHitResult& OutHit, const FVector Start, const FVector End, const FQuat& ShapeWorldRotation, const FCollisionShape& CollisionShape, bool bTraceComplex) override;
 	virtual bool OverlapComponentWithResult(const FVector& Pos, const FQuat& Rot, const FCollisionShape& CollisionShape, TArray<FOverlapResult>& OutOverlap) const override;
+	virtual bool ComponentOverlapComponentWithResultImpl(const class UPrimitiveComponent* const PrimComp, const FVector& Pos, const FQuat& Rot, const FCollisionQueryParams& Params, TArray<FOverlapResult>& OutOverlap) const override;
 	//~ End UPrimitiveComponent Interface
 
 	//~ Begin USceneComponent Interface
diff --git a/Engine/Source/Runtime/Engine/Private/Components/PrimitiveComponent.cpp b/Engine/Source/Runtime/Engine/Private/Components/PrimitiveComponent.cpp
--- a/Engine/Source/Runtime/Engine/Private/Components/PrimitiveComponent.cpp
+++ b/Engine/Source/Runtime/Engine/Private/Components/PrimitiveComponent.cpp
@@ -3064,6 +3064,50 @@ bool UPrimitiveComponent::ComponentOverlapComponentImpl(class UPrimitiveComponen
 	return false;
 }
 
+bool UPrimitiveComponent::ComponentOverlapComponentWithResultImpl(const class UPrimitiveComponent* const PrimComp, const FVector& Pos, const FQuat& Rot, const FCollisionQueryParams& Params, TArray<FOverlapResult>& OutOverlap) const
+{
+	const FTransform InTransform{ Rot, Pos };
+	TArray<Chaos::FPhysicsObjectHandle> InObjects = PrimComp->GetAllPhysicsObjects();
+	TArray<Chaos::FPhysicsObjectHandle> ThisObjects = GetAllPhysicsObjects();
+
+	FLockedReadPhysicsObjectExternalInterface Interface = FPhysicsObjectExternalInterface::LockRead(InObjects);
+	InObjects = InObjects.FilterByPredicate(
+		[&Interface](Chaos::FPhysicsObjectHandle Handle)
+		{
+			return !Interface->AreAllDisabled({ &Handle, 1 });
+		}
+	);
+
+	ThisObjects = ThisObjects.FilterByPredicate(
+		[&Interface](Chaos::FPhysicsObjectHandle Handle)
+		{
+			return !Interface->AreAllDisabled({ &Handle, 1 });
+		}
+	);
+
+	for (Chaos::FPhysicsObjectHandle InObject : InObjects)
+	{
+		for (Chaos::FPhysicsObjectHandle ThisObject : ThisObjects)
+		{
+			TArray<ChaosInterface::FOverlapHit> OverlapHits;
+			if (Interface->PhysicsObjectOverlap(ThisObject, FTransform::Identity, InObject, InTransform, Params.bTraceComplex, OverlapHits))
+			{
+				TArray<FOverlapResult> Overlaps;
+
+				FCollisionFilterData QueryFilter;
+				QueryFilter.Word1 = 0xFFFFF;
+				ConvertOverlapResults(OverlapHits.Num(), OverlapHits.GetData(), QueryFilter, Overlaps);
+
+				if (!Overlaps.IsEmpty())
+				{
+					OutOverlap = Overlaps;
+					return true;
+				}
+			}
+		}
+	}
+	return false;
+}
 
 bool UPrimitiveComponent::OverlapComponent(const FVector& Pos, const FQuat& Rot, const struct FCollisionShape& CollisionShape) const
 {
diff --git a/Engine/Source/Runtime/Engine/Private/PhysicsEngine/ClusterUnionComponent.cpp b/Engine/Source/Runtime/Engine/Private/PhysicsEngine/ClusterUnionComponent.cpp
--- a/Engine/Source/Runtime/Engine/Private/PhysicsEngine/ClusterUnionComponent.cpp
+++ b/Engine/Source/Runtime/Engine/Private/PhysicsEngine/ClusterUnionComponent.cpp
@@ -610,6 +610,24 @@ bool UClusterUnionComponent::OverlapComponentWithResult(const FVector& Pos, cons
 	return bHasOverlap;
 }
 
+bool UClusterUnionComponent::ComponentOverlapComponentWithResultImpl(const class UPrimitiveComponent* const PrimComp, const FVector& Pos, const FQuat& Rot, const FCollisionQueryParams& Params, TArray<FOverlapResult>& OutOverlap) const
+{
+	bool bHasOverlap = false;
+	for (const TPair<TObjectKey<UPrimitiveComponent>, FClusteredComponentData>& Kvp : ComponentToPhysicsObjects)
+	{
+		if (UPrimitiveComponent* Component = Kvp.Key.ResolveObjectPtr())
+		{
+			TArray<FOverlapResult> SubOverlaps;
+			if (Component->ComponentOverlapComponentWithResult(PrimComp, Pos, Rot, Params, SubOverlaps))
+			{
+				bHasOverlap = true;
+				OutOverlap.Append(SubOverlaps);
+			}
+		}
+	}
+	return bHasOverlap;
+}
+
 void UClusterUnionComponent::AddReferencedObjects(UObject* InThis, FReferenceCollector& Collector)
 {
 	UPrimitiveComponent::AddReferencedObjects(InThis, Collector);
diff --git a/Engine/Source/Runtime/Experimental/Chaos/Public/Chaos/PhysicsObjectInterface.h b/Engine/Source/Runtime/Experimental/Chaos/Public/Chaos/PhysicsObjectInterface.h
--- a/Engine/Source/Runtime/Experimental/Chaos/Public/Chaos/PhysicsObjectInterface.h
+++ b/Engine/Source/Runtime/Experimental/Chaos/Public/Chaos/PhysicsObjectInterface.h
@@ -88,6 +88,30 @@ namespace Chaos
 		// This function will not compute any overlap heuristic.
 		bool PhysicsObjectOverlap(const FConstPhysicsObjectHandle ObjectA, const FTransform& InTransformA, const FConstPhysicsObjectHandle ObjectB, const FTransform& InTransformB, bool bTraceComplex);
 
+		// Returns all the overlaps within A given a shape B.
+		template<typename TOverlapHit>
+		bool PhysicsObjectOverlap(const FConstPhysicsObjectHandle ObjectA, const FTransform& InTransformA, const FConstPhysicsObjectHandle ObjectB, const FTransform& InTransformB, bool bTraceComplex, TArray<TOverlapHit>& OutOverlaps)
+		{
+			static_assert(std::is_same_v<TOverlapHit, ChaosInterface::TThreadOverlapHit<Id>>);
+			return PairwiseShapeOverlapHelper(
+				ObjectA,
+				InTransformA,
+				ObjectB,
+				InTransformB,
+				bTraceComplex,
+				false,
+				FVector::Zero(),
+				[this, ObjectA, &OutOverlaps](const FShapeOverlapData& A , const FShapeOverlapData& B, const FMTDInfo&)
+				{
+					ChaosInterface::FOverlapHit Overlap;
+					Overlap.Shape = A.Shape;
+					Overlap.Actor = GetParticle(ObjectA);
+					OutOverlaps.Add(Overlap);
+					return false;
+				}
+			);
+		}
+
 		// This function does the same as GetPhysicsObjectOverlap but also computes the MTD metric.
 		bool PhysicsObjectOverlapWithMTD(const FConstPhysicsObjectHandle ObjectA, const FTransform& InTransformA, const FConstPhysicsObjectHandle ObjectB, const FTransform& InTransformB, bool bTraceComplex, FMTDInfo& OutMTD);
 
