WorldPartitionLevelHelper: Change log category from LogEngine to LogWorldPartition

#rnx

diff --git a/Engine/Source/Runtime/Engine/Private/WorldPartition/WorldPartitionLevelHelper.cpp b/Engine/Source/Runtime/Engine/Private/WorldPartition/WorldPartitionLevelHelper.cpp
--- a/Engine/Source/Runtime/Engine/Private/WorldPartition/WorldPartitionLevelHelper.cpp
+++ b/Engine/Source/Runtime/Engine/Private/WorldPartition/WorldPartitionLevelHelper.cpp
@@ -16,6 +16,7 @@
 #include "Model.h"
 #include "UnrealEngine.h"
 #include "WorldPartition/WorldPartition.h"
+#include "WorldPartition/WorldPartitionLog.h"
 #include "WorldPartition/WorldPartitionPackageHelper.h"
 #include "WorldPartition/ContentBundle/ContentBundleEditor.h"
 #include "LevelUtils.h"
@@ -160,7 +161,7 @@ void FWorldPartitionLevelHelper::MoveExternalActorsToLevel(const TArray<FWorldPa
 		}
 		else
 		{
-			UE_LOG(LogEngine, Warning, TEXT("Can't find actor %s."), *PackageObjectMapping.Path.ToString());
+			UE_LOG(LogWorldPartition, Warning, TEXT("Can't find actor %s."), *PackageObjectMapping.Path.ToString());
 		}
 	}
 }
@@ -488,22 +489,22 @@ bool FWorldPartitionLevelHelper::LoadActors(UWorld* InOwningWorld, ULevel* InDes
 					});
 				}
 
-				UE_LOG(LogEngine, Verbose, TEXT(" ==> Loaded %s (remaining: %d)"), *Actor->GetFullName(), LoadProgress->NumPendingLoadRequests);
+				UE_LOG(LogWorldPartition, Verbose, TEXT(" ==> Loaded %s (remaining: %d)"), *Actor->GetFullName(), LoadProgress->NumPendingLoadRequests);
 			}
 			else
 			{
 				if (LoadedPackage)
 				{
-					UE_LOG(LogEngine, Warning, TEXT("Failed to find actor in package %s. Package Content:"), *LoadedPackageName.ToString());
+					UE_LOG(LogWorldPartition, Warning, TEXT("Failed to find actor in package %s. Package Content:"), *LoadedPackageName.ToString());
 					ForEachObjectWithPackage(LoadedPackage, [](UObject* Object)
 					{
-						UE_LOG(LogEngine, Warning, TEXT("\t Object %s, Flags 0x%llx"), *Object->GetName(), static_cast<uint64>(Object->GetFlags()));
+						UE_LOG(LogWorldPartition, Warning, TEXT("\t Object %s, Flags 0x%llx"), *Object->GetName(), static_cast<uint64>(Object->GetFlags()));
 						return true;
 					}, false);
 				}
 				else
 				{
-					UE_LOG(LogEngine, Warning, TEXT("Failed to load actor package %s"), *LoadedPackageName.ToString());
+					UE_LOG(LogWorldPartition, Warning, TEXT("Failed to load actor package %s"), *LoadedPackageName.ToString());
 				}
 
 				//@todo_ow: cumulate and process when NumPendingActorRequests == 0
