fix Texture Cube Array creation not correctly forbidden like 2d Arrays were
Also fix Texture Cube Array crash when given textures without sources

diff --git a/Engine/Plugins/Editor/EngineAssetDefinitions/Source/Private/AssetDefinition_TextureCube.cpp b/Engine/Plugins/Editor/EngineAssetDefinitions/Source/Private/AssetDefinition_TextureCube.cpp
--- a/Engine/Plugins/Editor/EngineAssetDefinitions/Source/Private/AssetDefinition_TextureCube.cpp
+++ b/Engine/Plugins/Editor/EngineAssetDefinitions/Source/Private/AssetDefinition_TextureCube.cpp
@@ -56,13 +56,22 @@ namespace MenuExtension_TextureCubeArray
 			FToolMenuSection& Section = Menu->FindOrAddSection("GetAssetActions");
 			Section.AddDynamicEntry(NAME_None, FNewToolMenuSectionDelegate::CreateLambda([](FToolMenuSection& InSection)
 			{
+				TSharedRef<FPathPermissionList> ClassPathPermissionList = IAssetTools::Get().GetAssetClassPathPermissionList(EAssetClassAction::CreateAsset);
+				
+				// check that UTextureCubeArray classes are allowed in this project :
+				if (ClassPathPermissionList->PassesFilter(UTextureCubeArray::StaticClass()->GetClassPathName().ToString()))
 				{
-					const TAttribute<FText> Label = LOCTEXT("Texture_TextureCubeArray", "Create Texture Cube Array");
-					const TAttribute<FText> ToolTip = LOCTEXT("Texture_CreateTextureCubeArrayTooltip", "Creates a new texture cube array.");
-					const FSlateIcon Icon = FSlateIcon(FAppStyle::GetAppStyleSetName(), "ClassIcon.TextureCube");
-					const FToolMenuExecuteAction UIAction = FToolMenuExecuteAction::CreateStatic(&ExecuteCreateTextureArray);
+					// share the AllowTexture2DArrayCreation CVar to toggle this rather than make a separate one for cubes
+					static const auto AllowTextureArrayAssetCreationVar = IConsoleManager::Get().FindTConsoleVariableDataInt(TEXT("r.AllowTexture2DArrayCreation"));
+					if (AllowTextureArrayAssetCreationVar->GetValueOnGameThread() != 0)
+					{
+						const TAttribute<FText> Label = LOCTEXT("Texture_TextureCubeArray", "Create Texture Cube Array");
+						const TAttribute<FText> ToolTip = LOCTEXT("Texture_CreateTextureCubeArrayTooltip", "Creates a new texture cube array.");
+						const FSlateIcon Icon = FSlateIcon(FAppStyle::GetAppStyleSetName(), "ClassIcon.TextureCube");
+						const FToolMenuExecuteAction UIAction = FToolMenuExecuteAction::CreateStatic(&ExecuteCreateTextureArray);
 
-					InSection.AddMenuEntry("Texture_TextureCubeArray", Label, ToolTip, Icon, UIAction);
+						InSection.AddMenuEntry("Texture_TextureCubeArray", Label, ToolTip, Icon, UIAction);
+					}
 				}
 			}));
 		}));
diff --git a/Engine/Source/Runtime/Engine/Private/TextureCubeArray.cpp b/Engine/Source/Runtime/Engine/Private/TextureCubeArray.cpp
--- a/Engine/Source/Runtime/Engine/Private/TextureCubeArray.cpp
+++ b/Engine/Source/Runtime/Engine/Private/TextureCubeArray.cpp
@@ -602,7 +602,7 @@ ENGINE_API bool UTextureCubeArray::CheckArrayTexturesCompatibility()
 	for (int32 TextureIndex = 0; TextureIndex < SourceTextures.Num(); ++TextureIndex)
 	{
 		// Do not create array till all texture slots are filled.
-		if (!SourceTextures[TextureIndex])
+		if (!SourceTextures[TextureIndex] || ! SourceTextures[TextureIndex]->Source.IsValid())
 		{
 			return false;
 		}
@@ -628,6 +628,8 @@ ENGINE_API bool UTextureCubeArray::CheckArrayTexturesCompatibility()
 		const FString TextureNameCmp = SourceTextures[TextureCmpIndex]->GetName();
 		const ETextureSourceFormat SourceFormatCmp = TextureSourceCmp.GetFormat();
 
+		check( TextureSourceCmp.IsValid() );
+
 		if (TextureSourceCmp.GetSizeX() != SizeX || TextureSourceCmp.GetSizeY() != SizeY)
 		{
 			UE_LOG(LogTexture, Warning, TEXT("TextureCubeArray creation failed. Textures %s and %s have different sizes."), *TextureName, *TextureNameCmp);
@@ -661,6 +663,8 @@ ENGINE_API bool UTextureCubeArray::UpdateSourceFromSourceTextures(bool bCreating
 	Modify();
 
 	const FTextureSource& InitialSource = SourceTextures[0]->Source;
+	check( InitialSource.IsValid() ); // checked by CheckArrayTexturesCompatibility
+
 	// Format and format size.
 	const EPixelFormat PixelFormat = SourceTextures[0]->GetPixelFormat();
 	const ETextureSourceFormat Format = InitialSource.GetFormat();
@@ -710,6 +714,7 @@ ENGINE_API bool UTextureCubeArray::UpdateSourceFromSourceTextures(bool bCreating
 		{
 			const uint64 Size = MipSizeBytes[MipIndex];
 			const uint64 CheckSize = SourceTextures[SourceTexIndex]->Source.CalcMipSize(MipIndex);
+			check( SourceTextures[SourceTexIndex]->Source.IsValid() ); // checked by CheckArrayTexturesCompatibility
 			check(Size == CheckSize);
 
 			RefCubeData.Reset();
@@ -839,6 +844,16 @@ ENGINE_API void UTextureCubeArray::PostEditChangeProperty(FPropertyChangedEvent&
 
 	if (PropertyName == GET_MEMBER_NAME_CHECKED(UTextureCubeArray, SourceTextures))
 	{
+		for (int i=0;i<SourceTextures.Num();i++)
+		{
+			UTextureCube * SourceTexture = SourceTextures[i];
+			if (SourceTexture && !SourceTexture->Source.IsValid())
+			{
+				UE_LOG(LogTexture, Warning, TEXT("Texture has no Source, cannot be used [%s]"), *SourceTexture->GetFullName());
+				SourceTextures[i] = nullptr;
+			}
+		}
+
 		// Empty SourceTextures, remove any resources if present.
 		if (SourceTextures.Num() == 0)
 		{
