Chaos - speculative fix for cooking crash in UStaticMesh::GetPhysicsTriMeshDataCheckComplex
- also fixes any calls from non-editor builds when when LOD0 has no CPU data

diff --git a/Engine/Source/Runtime/Engine/Private/StaticMesh.cpp b/Engine/Source/Runtime/Engine/Private/StaticMesh.cpp
--- a/Engine/Source/Runtime/Engine/Private/StaticMesh.cpp
+++ b/Engine/Source/Runtime/Engine/Private/StaticMesh.cpp
@@ -7126,7 +7126,8 @@ bool UStaticMesh::GetPhysicsTriMeshDataCheckComplex(struct FTriMeshCollisionData
 		return ComplexCollisionMesh->GetPhysicsTriMeshDataCheckComplex(CollisionData, bInUseAllTriData, false); // Only one level of recursion
 	}
 #else // #if WITH_EDITORONLY_DATA
-	// the static mesh needs to be tagged for CPUAccess in order to access TriMeshData in runtime mode : 
+	// the static mesh needs to be tagged for CPUAccess in order to access TriMeshData in runtime mode
+	// we must also check that the selected LOD has CPU data (see below)
 	if (!bAllowCPUAccess)
 	{
 		UE_LOG(LogStaticMesh, Warning, TEXT("UStaticMesh::GetPhysicsTriMeshData: Triangle data from '%s' cannot be accessed at runtime on a mesh that isn't flagged as Allow CPU Access. This asset needs to be flagged as such (in the Advanced section)."), *GetFullName());
@@ -7145,6 +7146,14 @@ bool UStaticMesh::GetPhysicsTriMeshDataCheckComplex(struct FTriMeshCollisionData
 
 	FStaticMeshLODResources& LOD = GetRenderData()->LODResources[UseLODIndex];
 
+	// Make sure the LOD we selected actually has CPU data
+	// NOTE: for non-editor builds we forced LOD0
+	if (!LOD.IndexBuffer.GetAllowCPUAccess())
+	{
+		UE_LOG(LogStaticMesh, Warning, TEXT("UStaticMesh::GetPhysicsTriMeshData: CPU data not available on selected LOD (UseLODIndex=%d, LODForCollision=%d) on '%s'."), UseLODIndex, LODForCollision, *GetFullName());
+		return false;
+	}
+
 	FIndexArrayView Indices = LOD.IndexBuffer.GetArrayView();
 
 	TMap<int32, int32> MeshToCollisionVertMap; // map of static mesh verts to collision verts
