the cluster union replicated proxy component should be marked as pending removal when its requested to be removed on the GT rather than needing to wait for the PT

diff --git a/Engine/Source/Runtime/Engine/Classes/PhysicsEngine/ClusterUnionReplicatedProxyComponent.h b/Engine/Source/Runtime/Engine/Classes/PhysicsEngine/ClusterUnionReplicatedProxyComponent.h
--- a/Engine/Source/Runtime/Engine/Classes/PhysicsEngine/ClusterUnionReplicatedProxyComponent.h
+++ b/Engine/Source/Runtime/Engine/Classes/PhysicsEngine/ClusterUnionReplicatedProxyComponent.h
@@ -43,6 +43,12 @@ public:
 	UFUNCTION()
 	void SetParticleChildToParent(int32 BoneId, const FTransform& ChildToParent);
 
+	UFUNCTION()
+	void MarkPendingDeletion() { bIsPendingDeletion = true; }
+
+	UFUNCTION()
+	bool IsPendingDeletion() { return bIsPendingDeletion; }
+
 protected:
 
 	UFUNCTION()
@@ -82,6 +88,9 @@ private:
 	UPROPERTY()
 	bool bNetUpdateParticleChildToParents;
 
+	UPROPERTY()
+	bool bIsPendingDeletion = false;
+
 	void DeferUntilChildClusteredComponentInParentUnion(TFunction<void()> Func);
 
 	//~ Begin UActorComponent Interface
diff --git a/Engine/Source/Runtime/Engine/Private/PhysicsEngine/ClusterUnionComponent.cpp b/Engine/Source/Runtime/Engine/Private/PhysicsEngine/ClusterUnionComponent.cpp
--- a/Engine/Source/Runtime/Engine/Private/PhysicsEngine/ClusterUnionComponent.cpp
+++ b/Engine/Source/Runtime/Engine/Private/PhysicsEngine/ClusterUnionComponent.cpp
@@ -129,6 +129,17 @@ void UClusterUnionComponent::RemoveComponentFromCluster(UPrimitiveComponent* InC
 
 	if (FClusteredComponentData* ComponentData = ComponentToPhysicsObjects.Find(InComponent))
 	{
+		// We need to mark the replicated proxy as pending deletion.
+		// This way anyone who tries to use the replicated proxy component knows that it
+		// doesn't actually denote a meaningful cluster union relationship.
+		if (IsAuthority())
+		{
+			if (UClusterUnionReplicatedProxyComponent* Component = ComponentData->ReplicatedProxyComponent.Get())
+			{
+				Component->MarkPendingDeletion();
+			}
+		}
+
 		PhysicsObjectsToRemove = ComponentData->PhysicsObjects;
 	}
 
