Fix ensure caused by WorldPartition->UpdateStreamingState triggering an async load flush and modifying the array of world partitions in the WorldPartitionSubsystem while they are being iterated over.

diff --git a/Engine/Source/Runtime/Engine/Private/WorldPartition/WorldPartitionSubsystem.cpp b/Engine/Source/Runtime/Engine/Private/WorldPartition/WorldPartitionSubsystem.cpp
--- a/Engine/Source/Runtime/Engine/Private/WorldPartition/WorldPartitionSubsystem.cpp
+++ b/Engine/Source/Runtime/Engine/Private/WorldPartition/WorldPartitionSubsystem.cpp
@@ -349,9 +349,11 @@ void UWorldPartitionSubsystem::DumpStreamingSources(FOutputDevice& OutputDevice)
 
 void UWorldPartitionSubsystem::UpdateStreamingState()
 {
-	for (UWorldPartition* RegisteredWorldPartition : RegisteredWorldPartitions)
+	//make temp copy of array as UpdateStreamingState may FlushAsyncLoading, which may add a new world partition to RegisteredWorldPartitions while iterating
+	const TArray<UWorldPartition*> RegisteredWorldPartitionsCopy = RegisteredWorldPartitions;
+	for (UWorldPartition* RegisteredWorldParition : RegisteredWorldPartitionsCopy)
 	{
-		RegisteredWorldPartition->UpdateStreamingState();
+		RegisteredWorldParition->UpdateStreamingState();
 	}
 }
 
