Level Instance: assign deterministic guids to actors in instanced level instances.

diff --git a/Engine/Plugins/Editor/WorldPartitionHLODUtilities/Source/Private/WorldPartition/HLOD/Utilities/WorldPartitionHLODUtilities.cpp b/Engine/Plugins/Editor/WorldPartitionHLODUtilities/Source/Private/WorldPartition/HLOD/Utilities/WorldPartitionHLODUtilities.cpp
--- a/Engine/Plugins/Editor/WorldPartitionHLODUtilities/Source/Private/WorldPartition/HLOD/Utilities/WorldPartitionHLODUtilities.cpp
+++ b/Engine/Plugins/Editor/WorldPartitionHLODUtilities/Source/Private/WorldPartition/HLOD/Utilities/WorldPartitionHLODUtilities.cpp
@@ -50,7 +50,19 @@ static UWorldPartitionLevelStreamingDynamic* CreateLevelStreamingFromHLODActor(A
 	const FName LevelStreamingName = FName(*FString::Printf(TEXT("HLODLevelStreaming_%s"), *InHLODActor->GetName()));
 	TArray<FWorldPartitionRuntimeCellObjectMapping> Mappings;
 	Mappings.Reserve(InHLODActor->GetSubActors().Num());
-	Algo::Transform(InHLODActor->GetSubActors(), Mappings, [World, InHLODActor](const FHLODSubActor& SubActor) { return FWorldPartitionRuntimeCellObjectMapping(SubActor.ActorPackage, SubActor.ActorPath, SubActor.ContainerID, SubActor.ContainerTransform, SubActor.ContainerPackage, World->GetPackage()->GetFName(), InHLODActor->GetContentBundleGuid()); });
+	Algo::Transform(InHLODActor->GetSubActors(), Mappings, [World, InHLODActor](const FHLODSubActor& SubActor)
+	{
+		return FWorldPartitionRuntimeCellObjectMapping(
+			SubActor.ActorPackage, 
+			SubActor.ActorPath, 
+			SubActor.ContainerID, 
+			SubActor.ContainerTransform, 
+			SubActor.ContainerPackage, 
+			World->GetPackage()->GetFName(), 
+			InHLODActor->GetContentBundleGuid(),
+			SubActor.ActorGuid
+		);
+	});
 
 	UWorldPartitionLevelStreamingDynamic* LevelStreaming = UWorldPartitionLevelStreamingDynamic::LoadInEditor(World, LevelStreamingName, Mappings);
 	check(LevelStreaming);
diff --git a/Engine/Source/Runtime/Core/Private/Misc/Guid.cpp b/Engine/Source/Runtime/Core/Private/Misc/Guid.cpp
--- a/Engine/Source/Runtime/Core/Private/Misc/Guid.cpp
+++ b/Engine/Source/Runtime/Core/Private/Misc/Guid.cpp
@@ -235,6 +235,15 @@ FGuid FGuid::NewGuid()
 	return Result;
 }
 
+FGuid FGuid::Combine(const FGuid& GuidA, const FGuid& GuidB)
+{
+	return FGuid(
+		HashCombine(GuidA.A, GuidB.A),
+		HashCombine(GuidA.B, GuidB.B),
+		HashCombine(GuidA.C, GuidB.C),
+		HashCombine(GuidA.D, GuidB.D)
+	);
+}
 
 bool FGuid::Parse(const FString& GuidString, FGuid& OutGuid)
 {
diff --git a/Engine/Source/Runtime/Core/Public/Misc/Guid.h b/Engine/Source/Runtime/Core/Public/Misc/Guid.h
--- a/Engine/Source/Runtime/Core/Public/Misc/Guid.h
+++ b/Engine/Source/Runtime/Core/Public/Misc/Guid.h
@@ -361,6 +361,13 @@ public:
 	 */
 	static CORE_API FGuid NewGuid();
 
+	/**
+	 * Returns a GUID which is a combinationof the two provided ones.
+	 *
+	 * @return The combined GUID.
+	 */
+	static CORE_API FGuid Combine(const FGuid& GuidA, const FGuid& GuidB);
+
 	/**
 	 * Converts a string to a GUID.
 	 *
diff --git a/Engine/Source/Runtime/Engine/Classes/GameFramework/Actor.h b/Engine/Source/Runtime/Engine/Classes/GameFramework/Actor.h
--- a/Engine/Source/Runtime/Engine/Classes/GameFramework/Actor.h
+++ b/Engine/Source/Runtime/Engine/Classes/GameFramework/Actor.h
@@ -4209,6 +4209,9 @@ private:
 	friend class UEngine;
 	friend class UExternalActorsCommandlet;
 	friend class UWorldPartitionConvertCommandlet;
+	friend class ULevelStreamingLevelInstance;
+	friend class UWorldPartitionLevelStreamingDynamic;
+	friend class FWorldPartitionLevelHelper;
 };
 
 struct FSetActorContentBundleGuid
diff --git a/Engine/Source/Runtime/Engine/Private/LevelInstance/LevelInstanceLevelStreaming.cpp b/Engine/Source/Runtime/Engine/Private/LevelInstance/LevelInstanceLevelStreaming.cpp
--- a/Engine/Source/Runtime/Engine/Private/LevelInstance/LevelInstanceLevelStreaming.cpp
+++ b/Engine/Source/Runtime/Engine/Private/LevelInstance/LevelInstanceLevelStreaming.cpp
@@ -221,6 +221,8 @@ ULevelStreamingLevelInstance* ULevelStreamingLevelInstance::LoadInstance(ILevelI
 		{
 			GEngine->BlockTillLevelStreamingCompleted(LevelInstanceActor->GetWorld());
 
+			const FGuid& LevelInstanceGuid = LevelInstance->GetLevelInstanceID().GetGuid();
+
 			// Most of the code here is meant to allow partial support for undo/redo of LevelInstance Instance Loading:
 			// by setting the objects RF_Transient and !RF_Transactional we can check when unloading if those flags
 			// have been changed and figure out if we need to clear the transaction buffer or not.
@@ -250,6 +252,8 @@ ULevelStreamingLevelInstance* ULevelStreamingLevelInstance::LoadInstance(ILevelI
 					if (LevelActor)
 					{
 						LevelStreaming->PrepareLevelInstanceLoadedActor(*LevelActor, LevelInstance, false);
+
+						FSetActorGuid SetActorGuid(LevelActor, FGuid::Combine(LevelInstanceGuid, LevelActor->GetActorGuid()));
 					}
 				}
 
diff --git a/Engine/Source/Runtime/Engine/Private/LevelInstance/LevelInstanceSubsystem.cpp b/Engine/Source/Runtime/Engine/Private/LevelInstance/LevelInstanceSubsystem.cpp
--- a/Engine/Source/Runtime/Engine/Private/LevelInstance/LevelInstanceSubsystem.cpp
+++ b/Engine/Source/Runtime/Engine/Private/LevelInstance/LevelInstanceSubsystem.cpp
@@ -135,6 +135,13 @@ FLevelInstanceID::FLevelInstanceID(ULevelInstanceSubsystem* LevelInstanceSubsyst
 	}
 	
 	Hash = CityHash64WithSeed((const char*)Guids.GetData(), Guids.Num() * sizeof(FGuid), NameHash);
+
+	// Make sure to start with a zero guid as the top-level container ID is zero.
+	Guid.Invalidate();
+	for (int32 GuidIndex = Guids.Num() - 1; GuidIndex >= 0; GuidIndex--)
+	{
+		Guid = FGuid::Combine(Guid, Guids[GuidIndex]);
+	}
 }
 
 FLevelInstanceID ULevelInstanceSubsystem::RegisterLevelInstance(ILevelInstanceInterface* LevelInstance)
diff --git a/Engine/Source/Runtime/Engine/Private/WorldPartition/WorldPartitionLevelHelper.cpp b/Engine/Source/Runtime/Engine/Private/WorldPartition/WorldPartitionLevelHelper.cpp
--- a/Engine/Source/Runtime/Engine/Private/WorldPartition/WorldPartitionLevelHelper.cpp
+++ b/Engine/Source/Runtime/Engine/Private/WorldPartition/WorldPartitionLevelHelper.cpp
@@ -445,6 +445,9 @@ bool FWorldPartitionLevelHelper::LoadActors(UWorld* InOwningWorld, ULevel* InDes
 					TransformParams.Actor = Actor;
 					TransformParams.bDoPostEditMove = false;
 					FLevelUtils::ApplyLevelTransform(TransformParams);
+
+					// Set the actor's guid
+					FSetActorGuid SetActorGuid(Actor, PackageObjectMapping->ActorGuid);
 						
 					// Path to use when searching for this actor in MoveExternalActorsToLevel
 					PackageObjectMapping->LoadedPath = *Actor->GetPathName();
diff --git a/Engine/Source/Runtime/Engine/Private/WorldPartition/WorldPartitionRuntimeHash.cpp b/Engine/Source/Runtime/Engine/Private/WorldPartition/WorldPartitionRuntimeHash.cpp
--- a/Engine/Source/Runtime/Engine/Private/WorldPartition/WorldPartitionRuntimeHash.cpp
+++ b/Engine/Source/Runtime/Engine/Private/WorldPartition/WorldPartitionRuntimeHash.cpp
@@ -222,8 +222,10 @@ void UWorldPartitionRuntimeHash::PopulateRuntimeCell(UWorldPartitionRuntimeCell*
 	FBox CellContentBounds(ForceInit);
 	for (const IStreamingGenerationContext::FActorInstance& ActorInstance : ActorInstances)
 	{
+		check(ActorInstance.GetActorGuid().IsValid());
+
 		const FWorldPartitionActorDescView& ActorDescView = ActorInstance.GetActorDescView();
-		RuntimeCell->AddActorToCell(ActorDescView, ActorInstance.GetContainerID(), ActorInstance.GetTransform(), ActorInstance.GetActorDescContainer());
+		RuntimeCell->AddActorToCell(ActorDescView, ActorInstance.GetContainerID(), ActorInstance.GetTransform(), ActorInstance.GetActorDescContainer(), ActorInstance.GetActorGuid());
 		CellContentBounds += ActorDescView.GetRuntimeBounds().TransformBy(ActorInstance.GetTransform());
 					
 		if (ActorInstance.GetContainerID().IsMainContainer() && RuntimeCell->UnsavedActorsContainer)
diff --git a/Engine/Source/Runtime/Engine/Private/WorldPartition/WorldPartitionRuntimeLevelStreamingCell.cpp b/Engine/Source/Runtime/Engine/Private/WorldPartition/WorldPartitionRuntimeLevelStreamingCell.cpp
--- a/Engine/Source/Runtime/Engine/Private/WorldPartition/WorldPartitionRuntimeLevelStreamingCell.cpp
+++ b/Engine/Source/Runtime/Engine/Private/WorldPartition/WorldPartitionRuntimeLevelStreamingCell.cpp
@@ -202,11 +202,19 @@ void UWorldPartitionRuntimeLevelStreamingCell::SetIsAlwaysLoaded(bool bInIsAlway
 }
 
 #if WITH_EDITOR
-void UWorldPartitionRuntimeLevelStreamingCell::AddActorToCell(const FWorldPartitionActorDescView& ActorDescView, const FActorContainerID& InContainerID, const FTransform& InContainerTransform, const UActorDescContainer* InContainer)
+void UWorldPartitionRuntimeLevelStreamingCell::AddActorToCell(const FWorldPartitionActorDescView& ActorDescView, const FActorContainerID& InContainerID, const FTransform& InContainerTransform, const UActorDescContainer* InContainer, const FGuid& InActorGuid)
 {
 	check(!ActorDescView.GetActorIsEditorOnly());
 
-	Packages.Emplace(ActorDescView.GetActorPackage(), *ActorDescView.GetActorSoftPath().ToString(), InContainerID, InContainerTransform, InContainer->GetContainerPackage(), GetWorld()->GetPackage()->GetFName(), ActorDescView.GetContentBundleGuid());
+	Packages.Emplace(
+		ActorDescView.GetActorPackage(), 
+		*ActorDescView.GetActorSoftPath().ToString(), 
+		InContainerID, InContainerTransform, 
+		InContainer->GetContainerPackage(), 
+		GetWorld()->GetPackage()->GetFName(), 
+		ActorDescView.GetContentBundleGuid(),
+		InActorGuid
+	);
 }
 
 bool UWorldPartitionRuntimeLevelStreamingCell::PopulateGeneratorPackageForCook(TArray<UPackage*>& OutModifiedPackages)
@@ -310,8 +318,9 @@ void UWorldPartitionRuntimeLevelStreamingCell::DumpStateLog(FHierarchicalLogArch
 
 	for (const FWorldPartitionRuntimeCellObjectMapping& Mapping : Packages)
 	{
-		Ar.Printf(TEXT("Actor Path: %s"), *Mapping.Path.ToString());
+		Ar.Printf(TEXT("   Actor Path: %s"), *Mapping.Path.ToString());
 		Ar.Printf(TEXT("Actor Package: %s"), *Mapping.Package.ToString());
+		Ar.Printf(TEXT("   Actor Guid: %s"), *Mapping.ActorGuid.ToString());
 	}
 }
 #endif
diff --git a/Engine/Source/Runtime/Engine/Private/WorldPartition/WorldPartitionStreamingGeneration.cpp b/Engine/Source/Runtime/Engine/Private/WorldPartition/WorldPartitionStreamingGeneration.cpp
--- a/Engine/Source/Runtime/Engine/Private/WorldPartition/WorldPartitionStreamingGeneration.cpp
+++ b/Engine/Source/Runtime/Engine/Private/WorldPartition/WorldPartitionStreamingGeneration.cpp
@@ -161,6 +161,16 @@ class FWorldPartitionStreamingGenerator
 					ActorSetInstance.ContainerID = ContainerInstanceDescriptor.ID;
 					ActorSetInstance.Transform = ContainerInstanceDescriptor.Transform;
 
+					// Apply the per-container actor guids
+					if (!ContainerID.IsMainContainer())
+					{
+						ActorSetInstance.ActorsSetGuids.Reserve(ActorSet.Actors.Num());
+						for (const FGuid& ActorGuid : ActorSet.Actors)
+						{
+							ActorSetInstance.ActorsSetGuids.Add(ActorGuid, FGuid::Combine(ContainerInstanceDescriptor.Guid, ActorGuid));
+						}
+					}
+
 					// Apply AND logic on spatially loaded flag
 					ActorSetInstance.bIsSpatiallyLoaded = ReferenceActorDescView.GetIsSpatiallyLoaded() && ContainerInstanceDescriptor.bIsSpatiallyLoaded;
 
@@ -263,6 +273,7 @@ class FWorldPartitionStreamingGenerator
 		FString OwnerName;
 		FActorContainerID ID;
 		FActorContainerID ParentID;
+		FGuid Guid;
 		TMap<FActorContainerID, TSet<FGuid>> FilteredActors;
 	};
 
@@ -507,6 +518,11 @@ class FWorldPartitionStreamingGenerator
 			SubContainerInstanceDescriptor.ParentID = InContainerInstanceDescriptor.ID;
 			SubContainerInstanceDescriptor.OwnerName = *ContainerInstanceView.GetActorLabelOrName().ToString();
 
+			FContainerInstanceDescriptor& ParentContainer = ContainerInstanceDescriptorsMap.FindChecked(InContainerInstanceDescriptor.ParentID);
+			check(ParentContainer.ID.IsMainContainer() || ParentContainer.Guid.IsValid());
+
+			SubContainerInstanceDescriptor.Guid = FGuid::Combine(ParentContainer.Guid, ActorGuid);
+
 			// Inherit parent container properties
 			FName InheritedRuntimeGrid = InContainerInstanceDescriptor.RuntimeGrid;
 			bool bInheritedbIsSpatiallyLoaded = InContainerInstanceDescriptor.bIsSpatiallyLoaded;
@@ -985,7 +1001,9 @@ public:
 				{
 					UE_SCOPED_INDENT_LOG_ARCHIVE(Ar.PrintfIndent(TEXT("%s:"), *ContainerInstanceDescriptor.OwnerName));
 
-					Ar.Printf(TEXT("       ID: 0x%016llx"), ContainerID.ID);
+					Ar.Printf(TEXT("       ID: %s"), *ContainerInstanceDescriptor.ID.ToString());
+					Ar.Printf(TEXT(" ParentID: %s"), *ContainerInstanceDescriptor.ParentID.ToString());
+					Ar.Printf(TEXT("     GUID: %s"), *ContainerInstanceDescriptor.Guid.ToString());
 					Ar.Printf(TEXT("   Bounds: %s"), *ContainerInstanceDescriptor.Bounds.ToString());
 					Ar.Printf(TEXT("Transform: %s"), *ContainerInstanceDescriptor.Transform.ToString());
 					Ar.Printf(TEXT("Container: %s"), *ContainerInstanceDescriptor.Container->GetContainerPackage().ToString());
diff --git a/Engine/Source/Runtime/Engine/Private/WorldPartition/WorldPartitionStreamingGenerationContext.cpp b/Engine/Source/Runtime/Engine/Private/WorldPartition/WorldPartitionStreamingGenerationContext.cpp
--- a/Engine/Source/Runtime/Engine/Private/WorldPartition/WorldPartitionStreamingGenerationContext.cpp
+++ b/Engine/Source/Runtime/Engine/Private/WorldPartition/WorldPartitionStreamingGenerationContext.cpp
@@ -14,6 +14,15 @@ const FActorContainerID& IStreamingGenerationContext::FActorInstance::GetContain
 	return ActorSetInstance->ContainerID;
 }
 
+const FGuid& IStreamingGenerationContext::FActorInstance::GetActorGuid() const
+{
+	if (const FGuid* RemappedGuid = ActorSetInstance->ActorsSetGuids.Find(ActorGuid))
+	{
+		return *RemappedGuid;
+	}
+	return ActorGuid;
+}
+
 const FTransform& IStreamingGenerationContext::FActorInstance::GetTransform() const
 {
 	return ActorSetInstance->Transform;
diff --git a/Engine/Source/Runtime/Engine/Public/LevelInstance/LevelInstanceTypes.h b/Engine/Source/Runtime/Engine/Public/LevelInstance/LevelInstanceTypes.h
--- a/Engine/Source/Runtime/Engine/Public/LevelInstance/LevelInstanceTypes.h
+++ b/Engine/Source/Runtime/Engine/Public/LevelInstance/LevelInstanceTypes.h
@@ -33,8 +33,11 @@ struct FLevelInstanceID
 
 	inline uint64 GetHash() const { return Hash; }
 
+	inline const FGuid& GetGuid() const { return Guid; }
+
 private:
 	uint64 Hash = 0;
+	FGuid Guid;
 	TArray<FGuid> Guids;
 	FName ActorName;
 };
diff --git a/Engine/Source/Runtime/Engine/Public/WorldPartition/WorldPartitionRuntimeCell.h b/Engine/Source/Runtime/Engine/Public/WorldPartition/WorldPartitionRuntimeCell.h
--- a/Engine/Source/Runtime/Engine/Public/WorldPartition/WorldPartitionRuntimeCell.h
+++ b/Engine/Source/Runtime/Engine/Public/WorldPartition/WorldPartitionRuntimeCell.h
@@ -42,7 +42,7 @@ struct FWorldPartitionRuntimeCellObjectMapping
 #endif
 	{}
 
-	FWorldPartitionRuntimeCellObjectMapping(FName InPackage, FName InPath, const FActorContainerID& InContainerID, const FTransform& InContainerTransform, FName InContainerPackage, FName InWorldPackage, const FGuid& InContentBundleGuid)
+	FWorldPartitionRuntimeCellObjectMapping(FName InPackage, FName InPath, const FActorContainerID& InContainerID, const FTransform& InContainerTransform, FName InContainerPackage, FName InWorldPackage, const FGuid& InContentBundleGuid, const FGuid& InActorGuid)
 #if WITH_EDITORONLY_DATA
 		: Package(InPackage)
 		, Path(InPath)
@@ -51,6 +51,7 @@ struct FWorldPartitionRuntimeCellObjectMapping
 		, ContainerPackage(InContainerPackage)
 		, WorldPackage(InWorldPackage)
 		, ContentBundleGuid(InContentBundleGuid)
+		, ActorGuid(InActorGuid)
 		, LoadedPath(InPath)
 #endif
 	{}
@@ -98,6 +99,12 @@ struct FWorldPartitionRuntimeCellObjectMapping
 	UPROPERTY()
 	FGuid ContentBundleGuid;
 
+	/**
+	 * GUID of the actor instance
+	 */
+	UPROPERTY()
+	FGuid ActorGuid;
+
 	/**
 	* Loaded actor path (when cooking or pie)
 	* 
@@ -220,7 +227,7 @@ class ENGINE_API UWorldPartitionRuntimeCell : public UObject, public IWorldParti
 	virtual TSet<FName> GetActorPackageNames() const override { return TSet<FName>(); }
 	//~End IWorldPartitionCell Interface
 
-	virtual void AddActorToCell(const FWorldPartitionActorDescView& ActorDescView, const FActorContainerID& InContainerID, const FTransform& InContainerTransform, const UActorDescContainer* InContainer) PURE_VIRTUAL(UWorldPartitionRuntimeCell::AddActorToCell,);
+	virtual void AddActorToCell(const FWorldPartitionActorDescView& ActorDescView, const FActorContainerID& InContainerID, const FTransform& InContainerTransform, const UActorDescContainer* InContainer, const FGuid& InActorGuid) PURE_VIRTUAL(UWorldPartitionRuntimeCell::AddActorToCell,);
 	virtual int32 GetActorCount() const PURE_VIRTUAL(UWorldPartitionRuntimeCell::GetActorCount, return 0;);
 
 	// Cook methods
diff --git a/Engine/Source/Runtime/Engine/Public/WorldPartition/WorldPartitionRuntimeLevelStreamingCell.h b/Engine/Source/Runtime/Engine/Public/WorldPartition/WorldPartitionRuntimeLevelStreamingCell.h
--- a/Engine/Source/Runtime/Engine/Public/WorldPartition/WorldPartitionRuntimeLevelStreamingCell.h
+++ b/Engine/Source/Runtime/Engine/Public/WorldPartition/WorldPartitionRuntimeLevelStreamingCell.h
@@ -48,7 +48,7 @@ class ENGINE_API UWorldPartitionRuntimeLevelStreamingCell : public UWorldPartiti
 
 #if WITH_EDITOR
 	//~Begin UWorldPartitionRuntimeCell Interface
-	virtual void AddActorToCell(const FWorldPartitionActorDescView& ActorDescView, const FActorContainerID& InContainerID, const FTransform& InContainerTransform, const UActorDescContainer* InContainer) override;
+	virtual void AddActorToCell(const FWorldPartitionActorDescView& ActorDescView, const FActorContainerID& InContainerID, const FTransform& InContainerTransform, const UActorDescContainer* InContainer, const FGuid& InActorGuid) override;
 	virtual int32 GetActorCount() const override;
 	virtual void DumpStateLog(FHierarchicalLogArchive& Ar) override;
 	// Cook methods
diff --git a/Engine/Source/Runtime/Engine/Public/WorldPartition/WorldPartitionStreamingGenerationContext.h b/Engine/Source/Runtime/Engine/Public/WorldPartition/WorldPartitionStreamingGenerationContext.h
--- a/Engine/Source/Runtime/Engine/Public/WorldPartition/WorldPartitionStreamingGenerationContext.h
+++ b/Engine/Source/Runtime/Engine/Public/WorldPartition/WorldPartitionStreamingGenerationContext.h
@@ -43,6 +43,7 @@ public:
 		FActorContainerID ContainerID;
 		FTransform Transform;
 		const FActorSet* ActorSet;
+		TMap<FGuid, FGuid> ActorsSetGuids;
 	};
 
 	struct ENGINE_API FActorInstance
@@ -57,6 +58,7 @@ public:
 
 		const FWorldPartitionActorDescView& GetActorDescView() const;
 		const FActorContainerID& GetContainerID() const;
+		const FGuid& GetActorGuid() const;
 		const FTransform& GetTransform() const;
 		const UActorDescContainer* GetActorDescContainer() const;
 	};
