World Partition Class Descriptors: proper fix for wrongly reported missing classes for cooked editor (revert 24626972).
In cooked builds, we need to react to loaded blueprint generated class assets instead of getting the class from the blueprint object.

diff --git a/Engine/Source/Runtime/Engine/Private/BlueprintGeneratedClass.cpp b/Engine/Source/Runtime/Engine/Private/BlueprintGeneratedClass.cpp
--- a/Engine/Source/Runtime/Engine/Private/BlueprintGeneratedClass.cpp
+++ b/Engine/Source/Runtime/Engine/Private/BlueprintGeneratedClass.cpp
@@ -31,6 +31,7 @@
 #include "Kismet2/BlueprintEditorUtils.h"
 #include "BlueprintCompilationManager.h"
 #include "Engine/LevelScriptBlueprint.h"
+#include "WorldPartition/WorldPartitionActorDescUtils.h"
 extern UNREALED_API class UEditorEngine* GEditor;
 #else
 #include "UObject/LinkerLoad.h"
@@ -250,6 +251,13 @@ void UBlueprintGeneratedClass::GetAssetRegistryTags(TArray<FAssetRegistryTag>& O
 		OutTags.Add(FAssetRegistryTag(EditorTag.Key, EditorTag.Value, FAssetRegistryTag::TT_Hidden));
 	}
 #endif //#if WITH_EDITORONLY_DATA
+
+#if WITH_EDITOR
+	if (AActor* BlueprintCDO = Cast<AActor>(ClassDefaultObject))
+	{
+		FWorldPartitionActorDescUtils::AppendAssetDataTagsFromActor(BlueprintCDO, OutTags);
+	}
+#endif
 }
 
 #if WITH_EDITOR
diff --git a/Engine/Source/Runtime/Engine/Private/WorldPartition/ActorDescContainer.cpp b/Engine/Source/Runtime/Engine/Private/WorldPartition/ActorDescContainer.cpp
--- a/Engine/Source/Runtime/Engine/Private/WorldPartition/ActorDescContainer.cpp
+++ b/Engine/Source/Runtime/Engine/Private/WorldPartition/ActorDescContainer.cpp
@@ -86,7 +86,7 @@ void UActorDescContainer::Initialize(const FInitializeParams& InitParams)
 
 			if (!ActorDesc.IsValid() || 
 				!ActorDesc->GetNativeClass().IsValid() || 
-				//(ActorDesc->GetBaseClass().IsValid() && !ClassDescRegistry.IsRegisteredClass(ActorDesc->GetBaseClass())) || 
+				(ActorDesc->GetBaseClass().IsValid() && !ClassDescRegistry.IsRegisteredClass(ActorDesc->GetBaseClass())) || 
 				FActorDescList::GetActorDesc(ActorDesc->GetGuid()) ||
 				(InitParams.FilterActorDesc && !InitParams.FilterActorDesc(ActorDesc.Get())))
 			{
diff --git a/Engine/Source/Runtime/Engine/Private/WorldPartition/WorldPartitionClassDescRegistry.cpp b/Engine/Source/Runtime/Engine/Private/WorldPartition/WorldPartitionClassDescRegistry.cpp
--- a/Engine/Source/Runtime/Engine/Private/WorldPartition/WorldPartitionClassDescRegistry.cpp
+++ b/Engine/Source/Runtime/Engine/Private/WorldPartition/WorldPartitionClassDescRegistry.cpp
@@ -1,5 +1,4 @@
 // Copyright Epic Games, Inc. All Rights Reserved.
-// Copyright Epic Games, Inc. All Rights Reserved.
 
 #include "WorldPartition/WorldPartitionClassDescRegistry.h"
 
@@ -21,6 +20,34 @@
 #include "WorldPartition/WorldPartitionActorDesc.h"
 #include "WorldPartition/WorldPartitionActorDescUtils.h"
 
+/*
+ * We have to deal with some different realities between non-cooked and cooked editor builds: for non-cooked editor builds, UBlueprint are considered assets
+ * and will be the opnly thing visible from the asset registry standpoint. On the other hand, for cooked editor builds, UBlueprint are not considered assets
+ * while their corresponding UBlueprintGeneratedClass will be. Also, we expect redirectors only in non-cooked editor builds.
+ */	
+static FTopLevelAssetPath GetAssetDataClassNameForBlueprint(const FString& InAssetDataClassName)
+{
+	return FTopLevelAssetPath(InAssetDataClassName + TEXT("_C"));
+}
+
+static FTopLevelAssetPath GetAssetDataClassNameForBlueprintGeneratedClass(const FString& InAssetDataClassName)
+{
+	return FTopLevelAssetPath(InAssetDataClassName);
+}
+
+static FTopLevelAssetPath GetAssetDataClassName(const FAssetData& InAssetData)
+{
+	const FString AssetClassName = InAssetData.ToSoftObjectPath().ToString();
+
+	if (InAssetData.GetClass()->IsChildOf<UBlueprint>())
+	{
+		return GetAssetDataClassNameForBlueprint(AssetClassName);
+	}
+	
+	check(InAssetData.GetClass()->IsChildOf<UBlueprintGeneratedClass>());
+	return GetAssetDataClassNameForBlueprintGeneratedClass(AssetClassName);
+}
+
 FWorldPartitionClassDescRegistry& FWorldPartitionClassDescRegistry::Get()
 {
 	return TLazySingleton<FWorldPartitionClassDescRegistry>::Get();
@@ -121,6 +148,7 @@ void FWorldPartitionClassDescRegistry::PrefetchClassDescs(const TArray<FTopLevel
 	{
 		FARFilter Filter;
 		Filter.ClassPaths.Add(UBlueprint::StaticClass()->GetClassPathName());
+		Filter.ClassPaths.Add(UBlueprintGeneratedClass::StaticClass()->GetClassPathName());
 		Filter.ClassPaths.Add(UObjectRedirector::StaticClass()->GetClassPathName());
 		Filter.bRecursiveClasses = true;
 		Filter.PackageNames.Reserve(FilePaths.Num());
@@ -158,8 +186,8 @@ void FWorldPartitionClassDescRegistry::PrefetchClassDescs(const TArray<FTopLevel
 			FString DestinationObjectPath;
 			if (!AssetData.GetTagValue(TEXT("DestinationObject"), DestinationObjectPath))
 			{
-				AssetData = FAssetData();
 				UE_LOG(LogWorldPartition, Warning, TEXT("Failed to follow class redirector for '%s'"), *AssetData.ToSoftObjectPath().ToString());
+				AssetData = FAssetData();
 				break;
 			}
 
@@ -185,8 +213,8 @@ void FWorldPartitionClassDescRegistry::PrefetchClassDescs(const TArray<FTopLevel
 		// Register redirects
 		for (const FTopLevelAssetPath& ClassRedirect : ClassRedirects)
 		{
-			FTopLevelAssetPath SourceClassPath(ClassRedirect.ToString() + TEXT("_C"));
-			FTopLevelAssetPath RedirectedClassPath(AssetData.ToSoftObjectPath().ToString() + TEXT("_C"));
+			const FTopLevelAssetPath SourceClassPath(GetAssetDataClassNameForBlueprint(ClassRedirect.ToString()));
+			const FTopLevelAssetPath RedirectedClassPath(GetAssetDataClassName(AssetData));
 			RedirectClassMap.Add(SourceClassPath, RedirectedClassPath);
 		}
 
@@ -195,7 +223,8 @@ void FWorldPartitionClassDescRegistry::PrefetchClassDescs(const TArray<FTopLevel
 			FString ParentClassName;
 			if (AssetData.GetTagValue(FBlueprintTags::ParentClassPath, ParentClassName))
 			{
-				ParentClassPathsMap.Add(FTopLevelAssetPath(AssetData.ToSoftObjectPath().ToString() + TEXT("_C")), FTopLevelAssetPath(ParentClassName));
+				const FTopLevelAssetPath AssetClassName = GetAssetDataClassName(AssetData);
+				ParentClassPathsMap.Add(AssetClassName, FTopLevelAssetPath(ParentClassName));
 			}
 		}
 	});
@@ -216,8 +245,8 @@ void FWorldPartitionClassDescRegistry::PrefetchClassDescs(const TArray<FTopLevel
 	// Register current class descriptors
 	Algo::ForEach(Assets, [this](FAssetData& AssetData)
 	{
-		// Create the class descriptor
-		AssetData.AssetName = *FString(AssetData.AssetName.ToString() + TEXT("_C"));
+		const FTopLevelAssetPath AssetDataClassName(GetAssetDataClassName(AssetData));
+		AssetData.AssetName = *AssetDataClassName.GetAssetName().ToString();
 
 		// Lookup for an already registered class
 		if (!ClassByPath.Contains(FTopLevelAssetPath(AssetData.ToSoftObjectPath().ToString())))
@@ -300,7 +329,7 @@ void FWorldPartitionClassDescRegistry::UnregisterClassDescriptor(FWorldPartition
 
 void FWorldPartitionClassDescRegistry::RegisterClassDescriptorFromAssetData(const FAssetData& InAssetData)
 {
-	const FTopLevelAssetPath ClassPath = FTopLevelAssetPath(InAssetData.ToSoftObjectPath().ToString() + TEXT("_C"));
+	const FTopLevelAssetPath ClassPath(GetAssetDataClassName(InAssetData));
 	check(!ClassByPath.Contains(ClassPath));
 
 	UBlueprint* Blueprint = CastChecked<UBlueprint>(InAssetData.GetAsset());
@@ -333,17 +362,24 @@ void FWorldPartitionClassDescRegistry::OnAssetLoaded(UObject* InAssetLoaded)
 {
 	check(IsInitialized());
 
-	if (UBlueprint* Blueprint = Cast<UBlueprint>(InAssetLoaded))
+	UBlueprintGeneratedClass* BlueprintGeneratedClass = Cast<UBlueprintGeneratedClass>(InAssetLoaded);
+
+	if (!BlueprintGeneratedClass)
 	{
-		if (Blueprint->GeneratedClass && Blueprint->GeneratedClass->IsChildOf<AActor>())
+		if (UBlueprint* Blueprint = Cast<UBlueprint>(InAssetLoaded))
 		{
-			if (!ClassByPath.Contains(FTopLevelAssetPath(Blueprint->GeneratedClass->GetPathName())))
-			{
-				ValidateInternalState();
-				RegisterClassDescriptorFromActorClass(Blueprint->GeneratedClass);
-				ValidateInternalState();
-			}
-		}		
+			BlueprintGeneratedClass = Cast<UBlueprintGeneratedClass>(Blueprint->GeneratedClass);
+		}
+	}
+
+	if (BlueprintGeneratedClass && BlueprintGeneratedClass->IsChildOf<AActor>())
+	{
+		if (!ClassByPath.Contains(FTopLevelAssetPath(BlueprintGeneratedClass->GetPathName())))
+		{
+			ValidateInternalState();
+			RegisterClass(BlueprintGeneratedClass);
+			ValidateInternalState();
+		}
 	}
 }
 
@@ -412,7 +448,7 @@ void FWorldPartitionClassDescRegistry::OnAssetRemoved(const FAssetData& InAssetD
 
 	if (InAssetData.GetClass() && InAssetData.GetClass()->IsChildOf(UBlueprint::StaticClass()))
 	{
-		const FTopLevelAssetPath ClassPath = FTopLevelAssetPath(InAssetData.ToSoftObjectPath().ToString() + TEXT("_C"));
+		const FTopLevelAssetPath ClassPath(GetAssetDataClassName(InAssetData));
 
 		if (TUniquePtr<FWorldPartitionActorDesc>** ExistingClassDesc = ClassByPath.Find(ClassPath))
 		{
@@ -438,8 +474,8 @@ void FWorldPartitionClassDescRegistry::OnAssetRenamed(const FAssetData& InAssetD
 	{
 		ValidateInternalState();
 
-		const FTopLevelAssetPath OldClassPath = FTopLevelAssetPath(InOldObjectPath + TEXT("_C"));
-		const FTopLevelAssetPath NewClassPath = FTopLevelAssetPath(InAssetData.ToSoftObjectPath().ToString() + TEXT("_C"));
+		const FTopLevelAssetPath OldClassPath(GetAssetDataClassNameForBlueprint(InOldObjectPath));
+		const FTopLevelAssetPath NewClassPath(GetAssetDataClassName(InAssetData));
 
 		if (TUniquePtr<FWorldPartitionActorDesc>** ExistingClassDesc = ClassByPath.Find(OldClassPath))
 		{
