Clarify wording on the PrimaryDataAsset class comment to explain that Data Only Blueprints is the preferred way to handle blueprint-defined data assets

diff --git a/Engine/Source/Runtime/Engine/Classes/Engine/DataAsset.h b/Engine/Source/Runtime/Engine/Classes/Engine/DataAsset.h
--- a/Engine/Source/Runtime/Engine/Classes/Engine/DataAsset.h
+++ b/Engine/Source/Runtime/Engine/Classes/Engine/DataAsset.h
@@ -31,8 +31,11 @@ private:
 
 /**
  * A DataAsset that implements GetPrimaryAssetId and has asset bundle support, which allows it to be manually loaded/unloaded from the AssetManager.
- * Instances of native subclasses can be created directly as Data Assets and will use the name of the native class as the PrimaryAssetType.
- * Or, blueprint subclasses can be created with a PrimaryAssetType equal to the name of the first native class going up the hierarchy, or the highest level blueprint class.
+ * Instances of native subclasses can be created directly as Data Assets in the editor and will use the name of the native class as the PrimaryAssetType.
+ * Or, blueprint subclasses can be created to add variables and then subclassed again by Data Only Blueprints that set those variables.
+ * With blueprint subclasses, use Data Only Blueprints (and not Data Asset instances) to properly handle data inheritance and updating the parent class.
+ *
+ * The PrimaryAssetType will be equal to the name of the first native class going up the hierarchy, or the highest level blueprint class.
  * IE, if you have UPrimaryDataAsset -> UParentNativeClass -> UChildNativeClass -> DataOnlyBlueprintClass the type will be ChildNativeClass.
  * Whereas if you have UPrimaryDataAsset -> ParentBlueprintClass -> DataOnlyBlueprintClass the type will be ParentBlueprintClass.
  * To change this behavior, override GetPrimaryAssetId in your native class or copy those functions into a different native base class.
