Path Tracer: Add a blueprint call to invalidate path tracing output

The intent here is to catch cases that are not currently handled by automatic invalidations.

diff --git a/Engine/Source/Runtime/Engine/Classes/Kismet/KismetRenderingLibrary.h b/Engine/Source/Runtime/Engine/Classes/Kismet/KismetRenderingLibrary.h
--- a/Engine/Source/Runtime/Engine/Classes/Kismet/KismetRenderingLibrary.h
+++ b/Engine/Source/Runtime/Engine/Classes/Kismet/KismetRenderingLibrary.h
@@ -206,9 +206,15 @@ class UKismetRenderingLibrary : public UBlueprintFunctionLibrary
 	UFUNCTION(BlueprintPure, Category="Rendering", meta = (DisplayName = "Calculate Projection Matrix (Minimal View Info)", ScriptMethod = "CalculateProjectionMatrix"))
 	static ENGINE_API FMatrix CalculateProjectionMatrix(const FMinimalViewInfo& MinimalViewInfo);
 
-	/** Enables or disables the path tracer for all viewports simulatenously.
+	/** Enables or disables the path tracer for the current Game Viewport.
 	 * This command is equivalent to setting ShowFlag.PathTracing, but is accessible even from shipping builds.
 	 */
 	UFUNCTION(BlueprintCallable, Category = "Rendering", meta = (Keywords="EnablePathTracing"))
 	static ENGINE_API void EnablePathTracing(bool bEnablePathTracer);
+
+	/** Forces the path tracer to restart sample accumulation.
+	 * This can be used to force the path tracer to compute a new frame in situations where it can not detect a change in the scene automatically.
+	 */
+	UFUNCTION(BlueprintCallable, Category = "Rendering", meta = (Keywords = "PathTracing"))
+	static ENGINE_API void RefreshPathTracingOutput();
 };
diff --git a/Engine/Source/Runtime/Engine/Private/KismetRenderingLibrary.cpp b/Engine/Source/Runtime/Engine/Private/KismetRenderingLibrary.cpp
--- a/Engine/Source/Runtime/Engine/Private/KismetRenderingLibrary.cpp
+++ b/Engine/Source/Runtime/Engine/Private/KismetRenderingLibrary.cpp
@@ -24,6 +24,7 @@
 #include "Engine/Texture2D.h"
 #include "UObject/Package.h"
 #include "EngineModule.h"
+#include "SceneInterface.h"
 
 #include UE_INLINE_GENERATED_CPP_BY_NAME(KismetRenderingLibrary)
 
@@ -906,5 +907,17 @@ ENGINE_API void UKismetRenderingLibrary::EnablePathTracing(bool bEnablePathTrace
 	}
 }
 
+ENGINE_API void UKismetRenderingLibrary::RefreshPathTracingOutput()
+{
+	if (GEngine != nullptr && GEngine->GameViewport != nullptr)
+	{
+		UWorld* World = GEngine->GameViewport->GetWorld();
+		if (World != nullptr && World->Scene != nullptr)
+		{
+			World->Scene->InvalidatePathTracedOutput();
+		}
+	}
+}
+
 #undef LOCTEXT_NAMESPACE
 
