Path Tracer: Invalidate path traced image when a skeletal mesh morph target changes

diff --git a/Engine/Source/Runtime/Engine/Private/SkeletalRenderGPUSkin.cpp b/Engine/Source/Runtime/Engine/Private/SkeletalRenderGPUSkin.cpp
--- a/Engine/Source/Runtime/Engine/Private/SkeletalRenderGPUSkin.cpp
+++ b/Engine/Source/Runtime/Engine/Private/SkeletalRenderGPUSkin.cpp
@@ -410,21 +410,23 @@ void FSkeletalMeshObjectGPUSkin::Update(
 	uint32 RevisionNumber = 0;
 
 	FGPUSkinCache* GPUSkinCache = nullptr;
+	FSceneInterface* Scene = nullptr;
 	if (InMeshComponent && InMeshComponent->SceneProxy)
 	{
 		// We allow caching of per-frame, per-scene data
-		FrameNumberToPrepare = InMeshComponent->SceneProxy->GetScene().GetFrameNumber() + 1;
-		GPUSkinCache = InMeshComponent->SceneProxy->GetScene().GetGPUSkinCache();
+		Scene = &(InMeshComponent->SceneProxy->GetScene());
+		FrameNumberToPrepare = Scene->GetFrameNumber() + 1;
+		GPUSkinCache = Scene->GetGPUSkinCache();
 		RevisionNumber = InMeshComponent->GetBoneTransformRevisionNumber();
 	}
 
 	// queue a call to update this data
 	FSkeletalMeshObjectGPUSkin* MeshObject = this;
 	ENQUEUE_RENDER_COMMAND(SkelMeshObjectUpdateDataCommand)(
-		[MeshObject, FrameNumberToPrepare, RevisionNumber, NewDynamicData, GPUSkinCache](FRHICommandListImmediate& RHICmdList)
+		[MeshObject, FrameNumberToPrepare, RevisionNumber, NewDynamicData, GPUSkinCache, Scene](FRHICommandListImmediate& RHICmdList)
 		{
 			FScopeCycleCounter Context(MeshObject->GetStatId());
-			MeshObject->UpdateDynamicData_RenderThread(GPUSkinCache, RHICmdList, NewDynamicData, nullptr, FrameNumberToPrepare, RevisionNumber);
+			MeshObject->UpdateDynamicData_RenderThread(GPUSkinCache, RHICmdList, NewDynamicData, Scene, FrameNumberToPrepare, RevisionNumber);
 		}
 	);
 }
@@ -499,6 +501,13 @@ void FSkeletalMeshObjectGPUSkin::UpdateDynamicData_RenderThread(FGPUSkinCache* G
 		!DynamicData->ActiveMorphTargetsEqual(InDynamicData->ActiveMorphTargets, InDynamicData->MorphTargetWeights))
 		: true);
 
+#if RHI_RAYTRACING
+	if (bMorphNeedsUpdate && Scene)
+	{
+		Scene->InvalidatePathTracedOutput();
+	}
+#endif
+
 	WaitForRHIThreadFenceForDynamicData();
 	if (DynamicData)
 	{
