Remove custom icon declaration from UHeterogeneousVolumesComponent

diff --git a/Engine/Source/Runtime/Engine/Private/Components/HeterogeneousVolumeComponent.cpp b/Engine/Source/Runtime/Engine/Private/Components/HeterogeneousVolumeComponent.cpp
--- a/Engine/Source/Runtime/Engine/Private/Components/HeterogeneousVolumeComponent.cpp
+++ b/Engine/Source/Runtime/Engine/Private/Components/HeterogeneousVolumeComponent.cpp
@@ -293,12 +293,10 @@ AHeterogeneousVolume::AHeterogeneousVolume(const FObjectInitializer& ObjectIniti
 		// Structure to hold one-time initialization
 		struct FConstructorStatics
 		{
-			ConstructorHelpers::FObjectFinderOptional<UTexture2D> HeterogeneousVolumeTextureObject;
 			FName ID_HeterogeneousVolume;
 			FText NAME_HeterogeneousVolume;
 			FConstructorStatics()
-				: HeterogeneousVolumeTextureObject(TEXT("/Engine/EditorResources/S_HeterogeneousVolume"))
-				, ID_HeterogeneousVolume(TEXT("Fog"))
+				: ID_HeterogeneousVolume(TEXT("Fog"))
 				, NAME_HeterogeneousVolume(NSLOCTEXT("SpriteCategory", "Fog", "Fog"))
 			{
 			}
@@ -307,7 +305,6 @@ AHeterogeneousVolume::AHeterogeneousVolume(const FObjectInitializer& ObjectIniti
 
 		if (GetSpriteComponent())
 		{
-			GetSpriteComponent()->Sprite = ConstructorStatics.HeterogeneousVolumeTextureObject.Get();
 			GetSpriteComponent()->SetRelativeScale3D(FVector(0.5f, 0.5f, 0.5f));
 			GetSpriteComponent()->SpriteInfo.Category = ConstructorStatics.ID_HeterogeneousVolume;
 			GetSpriteComponent()->SpriteInfo.DisplayName = ConstructorStatics.NAME_HeterogeneousVolume;
