(Experimental) functionality to purge scene interface (and all its allocations) from the editor world prior to bringing up a play-in-editor world, and restore editor world back to normal once the play session has finished. This avoids 2x allocations with both world scene interfaces fully initialized in memory.

diff --git a/Engine/Source/Editor/UnrealEd/Private/PlayLevel.cpp b/Engine/Source/Editor/UnrealEd/Private/PlayLevel.cpp
--- a/Engine/Source/Editor/UnrealEd/Private/PlayLevel.cpp
+++ b/Engine/Source/Editor/UnrealEd/Private/PlayLevel.cpp
@@ -888,6 +888,9 @@ void UEditorEngine::TeardownPlaySession(FWorldContext& PieWorldContext)
 	GWorld = EditorWorld;
 	GIsPlayInEditorWorld = false;
 
+	// Restore the previously purged scene interface for the editor world to its original glory.
+	GWorld->RestoreScene();
+
 	FWorldContext& EditorWorldContext = GEditor->GetEditorWorldContext();
 
 	// Let the viewport know about leaving PIE/Simulate session. Do it after everything's been cleaned up
@@ -2700,6 +2703,9 @@ void UEditorEngine::StartPlayInEditorSession(FRequestPlaySessionParams& InReques
 	// Let navigation know PIE is starting so it can avoid any blueprint creation/deletion/instantiation affect editor map's navmesh changes
 	FNavigationSystem::OnPIEStart(*InWorld);
 
+	// Purge the existing scene interface from the editor world to avoid 2x GPU allocations with the additional play-in-editor world that is about to init
+	InWorld->PurgeScene();
+
 	ULevelEditorPlaySettings* EditorPlaySettings = InRequestParams.EditorPlaySettings;
 	check(EditorPlaySettings);
 
diff --git a/Engine/Source/Runtime/Engine/Classes/Engine/World.h b/Engine/Source/Runtime/Engine/Classes/Engine/World.h
--- a/Engine/Source/Runtime/Engine/Classes/Engine/World.h
+++ b/Engine/Source/Runtime/Engine/Classes/Engine/World.h
@@ -1292,6 +1292,9 @@ public:
 
 	/** Indicates that toggling between Play-in-Editor and Simulate-in-Editor happened this frame. */
 	uint8 bToggledBetweenPIEandSIEThisFrame : 1;
+
+	/** Indicates that the renderer scene for this editor world was purged while Play-in-Editor. */
+	uint8 bPurgedScene : 1;
 #endif
 
 	/** Keeps track whether actors moved via PostEditMove and therefore constraint syncup should be performed. */
@@ -1455,7 +1458,13 @@ public:
 	/** Change the feature level that this world is current rendering with */
 	void ChangeFeatureLevel(ERHIFeatureLevel::Type InFeatureLevel, bool bShowSlowProgressDialog = true);
 
-	void RecreateScene(ERHIFeatureLevel::Type InFeatureLevel);
+	void RecreateScene(ERHIFeatureLevel::Type InFeatureLevel, bool bBroadcastChange = true);
+
+	/** Recreate the editor world's FScene with a null scene interface to drop extra GPU memory during PIE */
+	void PurgeScene();
+
+	/** Restore the purged editor world FScene back to the proper GPU representation */
+	void RestoreScene();
 
 #endif // WITH_EDITOR
 
diff --git a/Engine/Source/Runtime/Engine/Private/World.cpp b/Engine/Source/Runtime/Engine/Private/World.cpp
--- a/Engine/Source/Runtime/Engine/Private/World.cpp
+++ b/Engine/Source/Runtime/Engine/Private/World.cpp
@@ -139,6 +139,12 @@ CSV_DECLARE_CATEGORY_MODULE_EXTERN(CORE_API, Basic);
 static int32 bDisableRemapScriptActors = 0;
 FAutoConsoleVariableRef CVarDisableRemapScriptActors(TEXT("net.DisableRemapScriptActors"), bDisableRemapScriptActors, TEXT("When set, disables name remapping of compiled script actors (for networking)"));
 
+static TAutoConsoleVariable<int32> CVarPurgeEditorSceneDuringPIE(
+	TEXT("r.PurgeEditorSceneDuringPIE"),
+	0,
+	TEXT("0 to keep editor scene fully initialized during PIE (default)\n")
+	TEXT("1 to purge editor scene from memory during PIE and restore when the session finishes."));
+
 template<class Function>
 static void ForEachNetDriver(UEngine* Engine, const UWorld* const World, const Function InFunction)
 {
@@ -378,6 +384,7 @@ UWorld::UWorld( const FObjectInitializer& ObjectInitializer )
 #if WITH_EDITOR
 , bDebugFrameStepExecutedThisFrame(false)
 , bToggledBetweenPIEandSIEThisFrame(false)
+, bPurgedScene(false)
 #endif
 , bShouldTick(true)
 , bHasEverBeenInitialized(false)
@@ -8559,7 +8566,7 @@ void UWorld::ChangeFeatureLevel(ERHIFeatureLevel::Type InFeatureLevel, bool bSho
 	}
 }
 
-void UWorld::RecreateScene(ERHIFeatureLevel::Type InFeatureLevel)
+void UWorld::RecreateScene(ERHIFeatureLevel::Type InFeatureLevel, bool bBroadcastChange)
 {
 	if (Scene)
 	{
@@ -8573,7 +8580,10 @@ void UWorld::RecreateScene(ERHIFeatureLevel::Type InFeatureLevel)
 		IRendererModule& RendererModule = GetRendererModule();
 		RendererModule.RemoveScene(Scene);
 
-		FRenderResource::ChangeFeatureLevel(InFeatureLevel);
+		if (bBroadcastChange)
+		{
+			FRenderResource::ChangeFeatureLevel(InFeatureLevel);
+		}
 
 		RendererModule.AllocateScene(this, bRequiresHitProxies, FXSystem != nullptr, InFeatureLevel);
 
@@ -8586,6 +8596,56 @@ void UWorld::RecreateScene(ERHIFeatureLevel::Type InFeatureLevel)
 	}
 }
 
+// Recreate the editor world's FScene with a null scene interface to drop extra GPU memory during PIE
+void UWorld::PurgeScene()
+{
+	if (CVarPurgeEditorSceneDuringPIE.GetValueOnGameThread() == 0)
+	{
+		return;
+	}
+
+	if (!bPurgedScene && this->IsEditorWorld())
+	{
+		FGlobalComponentReregisterContext RecreateComponents;
+
+		// Finish any deferred / async render cleanup work.
+		GetRendererModule().PerFrameCleanupIfSkipRenderer();
+		FlushRenderingCommands();
+
+		const bool bOldVal = GUsingNullRHI;
+		GUsingNullRHI = true;
+		{
+			RecreateScene(FeatureLevel, false /* bBroadcastChange */);
+		}
+		GUsingNullRHI = bOldVal;
+
+		InvalidateAllSkyCaptures();
+		TriggerStreamingDataRebuild();
+
+		bPurgedScene = true;
+	}
+}
+
+// Restore the purged editor world FScene back to the proper GPU representation
+void UWorld::RestoreScene()
+{
+	if (bPurgedScene)
+	{
+		FGlobalComponentReregisterContext RecreateComponents;
+
+		// Finish any deferred / async render cleanup work.
+		GetRendererModule().PerFrameCleanupIfSkipRenderer();
+		FlushRenderingCommands();
+
+		RecreateScene(FeatureLevel, false /* bBroadcastChange */);
+
+		InvalidateAllSkyCaptures();
+		TriggerStreamingDataRebuild();
+
+		bPurgedScene = false;
+	}
+}
+
 void UWorld::GetAssetRegistryTags(TArray<FAssetRegistryTag>& OutTags) const
 {
 	Super::GetAssetRegistryTags(OutTags);
