World Partition Class Descriptors: add debug code to try to recover from missing redirectors in the asset registry to gather more information when it happens on build machines.

diff --git a/Engine/Source/Runtime/Engine/Private/WorldPartition/WorldPartitionClassDescRegistry.cpp b/Engine/Source/Runtime/Engine/Private/WorldPartition/WorldPartitionClassDescRegistry.cpp
--- a/Engine/Source/Runtime/Engine/Private/WorldPartition/WorldPartitionClassDescRegistry.cpp
+++ b/Engine/Source/Runtime/Engine/Private/WorldPartition/WorldPartitionClassDescRegistry.cpp
@@ -144,13 +144,14 @@ void FWorldPartitionClassDescRegistry::PrefetchClassDescs(const TArray<FTopLevel
 	// the shortest name.
 	auto SortBlueprintAssetsByNameLength = [](const FAssetData& Lhs, const FAssetData& Rhs) { return Lhs.AssetName.GetStringLength() < Rhs.AssetName.GetStringLength(); };
 
-	auto GetBlueprintAssets = [&AssetRegistry](const TArray<FString>& FilePaths, TArray<FAssetData>& Assets)
+	auto GetBlueprintAssets = [&AssetRegistry](const TArray<FString>& FilePaths, TArray<FAssetData>& Assets, bool bIncludeOnlyOnDiskAssets = false)
 	{
 		FARFilter Filter;
 		Filter.ClassPaths.Add(UBlueprint::StaticClass()->GetClassPathName());
 		Filter.ClassPaths.Add(UBlueprintGeneratedClass::StaticClass()->GetClassPathName());
 		Filter.ClassPaths.Add(UObjectRedirector::StaticClass()->GetClassPathName());
 		Filter.bRecursiveClasses = true;
+		Filter.bIncludeOnlyOnDiskAssets = bIncludeOnlyOnDiskAssets;
 		Filter.PackageNames.Reserve(FilePaths.Num());
 		Algo::Transform(FilePaths, Filter.PackageNames, [](const FString& ClassPath) { return *ClassPath; });
 	
@@ -199,8 +200,16 @@ void FWorldPartitionClassDescRegistry::PrefetchClassDescs(const TArray<FTopLevel
 			if (RedirectAssets.IsEmpty())
 			{
 				UE_LOG(LogWorldPartition, Warning, TEXT("Failed to find redirected assets for '%s' from '%s'"), *AssetData.ToSoftObjectPath().ToString(), *AssetClassPath.ToString());
-				AssetData = FAssetData();
-				break;
+
+				GetBlueprintAssets({ AssetClassPath.GetPackageName().ToString() }, RedirectAssets, true);
+
+				if (RedirectAssets.IsEmpty())
+				{
+					AssetData = FAssetData();
+					break;
+				}
+
+				UE_LOG(LogWorldPartition, Warning, TEXT("\tRecovered redirected assets using bIncludeOnlyOnDiskAssets=true"));
 			}
 
 			ClassRedirects.Add(FTopLevelAssetPath(AssetData.ToSoftObjectPath().ToString()));
