Gather subclasses when calling FExternalPackageHelper::LoadObjectsFromExternalPackages.

diff --git a/Engine/Source/Runtime/Engine/Public/ExternalPackageHelper.h b/Engine/Source/Runtime/Engine/Public/ExternalPackageHelper.h
--- a/Engine/Source/Runtime/Engine/Public/ExternalPackageHelper.h
+++ b/Engine/Source/Runtime/Engine/Public/ExternalPackageHelper.h
@@ -115,6 +115,7 @@ void FExternalPackageHelper::LoadObjectsFromExternalPackages(UObject* InOuter, T
 		Filter.bRecursivePaths = true;
 		Filter.bIncludeOnlyOnDiskAssets = true;
 		Filter.ClassPaths.Add(T::StaticClass()->GetClassPathName());
+		Filter.bRecursiveClasses = true;
 		Filter.PackagePaths.Add(*ExternalObjectsPath);
 
 		TArray<FAssetData> Assets;
