GC duration API + debug facilities

Added console commands to show/hide a live GC debug graph + generate or unlink a dense UObject tree as a stress test for GC marking.

diff --git a/Engine/Source/Runtime/CoreUObject/Private/UObject/GarbageCollection.cpp b/Engine/Source/Runtime/CoreUObject/Private/UObject/GarbageCollection.cpp
--- a/Engine/Source/Runtime/CoreUObject/Private/UObject/GarbageCollection.cpp
+++ b/Engine/Source/Runtime/CoreUObject/Private/UObject/GarbageCollection.cpp
@@ -37,6 +37,7 @@
 #include "HAL/RunnableThread.h"
 #include "UObject/FieldPathProperty.h"
 #include "UObject/GarbageCollectionHistory.h"
+#include "UObject/GarbageCollectionTesting.h"
 
 #include <atomic>
 
@@ -84,7 +85,13 @@ static bool GObjCurrentPurgeObjectIndexNeedsReset = true;
 static TArray<FUObjectItem*> GUnreachableObjects;
 static FCriticalSection GUnreachableObjectsCritical;
 static int32 GUnrechableObjectIndex = 0;
-static double GLastGCTime = 0;
+
+struct FGCTimingInfo
+{
+	double LastGCTime{};
+	double LastGCDuration = -1;
+};
+static FGCTimingInfo GTimingInfo;
 
 bool GIsGarbageCollecting = false;
 
@@ -242,6 +249,36 @@ static FAutoConsoleVariableRef CMultithreadedDestructionEnabled(
 	ECVF_Default
 );
 
+static UObjectReachabilityStressData* GReachabilityStressData;
+static void AllocateReachabilityStressData(FOutputDevice&)
+{
+	if (GReachabilityStressData)
+	{
+		return;
+	}
+	GReachabilityStressData = GenerateReachabilityStressData();
+}
+
+static void UnlinkReachabilityStressData(FOutputDevice&)
+{
+	if (!GReachabilityStressData)
+	{
+		return;
+	}
+	UnlinkReachabilityStressData(GReachabilityStressData);
+	GReachabilityStressData = nullptr;
+}
+
+static FAutoConsoleCommandWithOutputDevice GGenerateReachabilityStressDataCmd(TEXT("gc.GenerateReachabilityStressData"),
+		        						      TEXT("Allocate deeply-nested UObject tree to "
+						        			   "stress test reachability analysis."),
+FConsoleCommandWithOutputDeviceDelegate::CreateStatic(&AllocateReachabilityStressData));
+
+static FAutoConsoleCommandWithOutputDevice GUnlinkReachabilityStressDataCmd(TEXT("gc.UnlinkReachabilityStressData"),
+                                 					    TEXT("Unlink previously-generated reachability analysis "
+										 "stress test data for collection in the next cycle."),
+FConsoleCommandWithOutputDeviceDelegate::CreateStatic(&UnlinkReachabilityStressData));
+
 #if UE_BUILD_SHIPPING
 static constexpr int32 GGarbageReferenceTrackingEnabled = 0;
 #else
@@ -4500,6 +4537,8 @@ FORCENOINLINE static void CollectGarbageFull(EObjectFlags KeepFlags)
 
 FORCEINLINE void CollectGarbageInternal(EObjectFlags KeepFlags, bool bPerformFullPurge)
 {
+	const double StartTime = FPlatformTime::Seconds();
+
 	if (bPerformFullPurge)
 	{
 		CollectGarbageFull(KeepFlags);
@@ -4508,6 +4547,8 @@ FORCEINLINE void CollectGarbageInternal(EObjectFlags KeepFlags, bool bPerformFul
 	{
 		CollectGarbageIncremental(KeepFlags);
 	}
+
+	GTimingInfo.LastGCDuration = FPlatformTime::Seconds() - StartTime;
 }
 
 /** 
@@ -4722,7 +4763,7 @@ void CollectGarbageImpl(EObjectFlags KeepFlags)
 		FCoreUObjectDelegates::GetPostGarbageCollect().Broadcast();
 	}
 
-	GLastGCTime = FPlatformTime::Seconds();
+	GTimingInfo.LastGCTime = FPlatformTime::Seconds();
 	STAT_ADD_CUSTOMMESSAGE_NAME( STAT_NamedMarker, TEXT( "GarbageCollection - End" ) );
 }
 #endif // !UE_WITH_GC
@@ -4748,7 +4789,12 @@ FString FGarbageReferenceInfo::GetReferencingObjectInfo() const
 
 double GetLastGCTime()
 {
-	return GLastGCTime;
+	return GTimingInfo.LastGCTime;
+}
+
+double GetLastGCDuration()
+{
+	return GTimingInfo.LastGCDuration;
 }
 
 bool IsIncrementalUnhashPending()
diff --git a/Engine/Source/Runtime/CoreUObject/Private/UObject/GarbageCollectionTesting.cpp b/Engine/Source/Runtime/CoreUObject/Private/UObject/GarbageCollectionTesting.cpp
new file mode 100644
--- /dev/null
+++ b/Engine/Source/Runtime/CoreUObject/Private/UObject/GarbageCollectionTesting.cpp
@@ -0,0 +1,56 @@
+// Copyright Epic Games, Inc. All Rights Reserved.
+
+#include "GarbageCollectionTesting.h"
+
+/*
+	TODO additional features for a future changeset:
+	- gc.GenerateReachabilityStressData should have extra arg
+	  to specify # of levels in UObject tree
+	- gc.GenerateReachabilityStressData should have extra arg
+	  to specify whether or not the tree is gc-rooted (to see
+		how long it takes to collect)
+	- UObjectReachabilityStressData can expose more members to
+	  GC to make the token stream bigger
+ */
+
+static void GenerateReachabilityStressData(int Levels, UObjectReachabilityStressData* Data);
+
+UObjectReachabilityStressData* GenerateReachabilityStressData()
+{
+	UObjectReachabilityStressData* Data = NewObject<UObjectReachabilityStressData>();
+	GenerateReachabilityStressData(20, Data);
+	return Data;
+}
+
+static void GenerateReachabilityStressData(int Levels, UObjectReachabilityStressData* Data)
+{
+	Data->AddToRoot();
+
+	if (Levels == 0)
+	{
+		return;
+	}
+
+	const int N = 2;
+	for (int I = 0; I < N; ++I)
+	{
+		UObjectReachabilityStressData* Child = NewObject<UObjectReachabilityStressData>();
+		GenerateReachabilityStressData(Levels - 1, Child);
+		Data->Children.Add(Child);
+	}
+}
+
+void UnlinkReachabilityStressData(UObjectReachabilityStressData* Data)
+{
+	Data->RemoveFromRoot();
+	for (const auto Child : Data->Children)
+	{
+		UnlinkReachabilityStressData(Child);
+	}
+}
+
+IMPLEMENT_CORE_INTRINSIC_CLASS(UObjectReachabilityStressData, UObject,
+	{
+		UE::GC::FTokenStreamBuilder& Builder = UE::GC::FIntrinsicClassTokens::AllocateBuilder(Class);
+		Builder.EmitObjectArrayReference(STRUCT_OFFSET(UObjectReachabilityStressData, Children), TEXT("Children"));
+	});
diff --git a/Engine/Source/Runtime/CoreUObject/Private/UObject/GarbageCollectionTesting.h b/Engine/Source/Runtime/CoreUObject/Private/UObject/GarbageCollectionTesting.h
new file mode 100644
--- /dev/null
+++ b/Engine/Source/Runtime/CoreUObject/Private/UObject/GarbageCollectionTesting.h
@@ -0,0 +1,22 @@
+// Copyright Epic Games, Inc. All Rights Reserved.
+
+#pragma once
+
+#include "CoreMinimal.h"
+#include "UObject/Object.h"
+#include "UObject/ObjectMacros.h"
+#include "UObject/UObjectGlobals.h"
+
+class UObjectReachabilityStressData : public UObject
+{
+	DECLARE_CLASS_INTRINSIC(UObjectReachabilityStressData,
+		UObject,
+		CLASS_Transient,
+		TEXT("/Script/CoreUObject"));
+
+public:
+	TArray<UObjectReachabilityStressData*> Children;
+};
+
+UObjectReachabilityStressData* GenerateReachabilityStressData();
+void UnlinkReachabilityStressData(UObjectReachabilityStressData* Data);
diff --git a/Engine/Source/Runtime/CoreUObject/Public/UObject/GarbageCollection.h b/Engine/Source/Runtime/CoreUObject/Public/UObject/GarbageCollection.h
--- a/Engine/Source/Runtime/CoreUObject/Public/UObject/GarbageCollection.h
+++ b/Engine/Source/Runtime/CoreUObject/Public/UObject/GarbageCollection.h
@@ -518,6 +518,13 @@ extern COREUOBJECT_API bool GIsGarbageCollecting;
  */
 COREUOBJECT_API double GetLastGCTime();
 
+/**
+ * Gets the duration of the last GC run.
+ *
+ * @return	Returns the last GC duration, -1 if GC has never run.
+ */
+COREUOBJECT_API double GetLastGCDuration();
+
 /**
 * Whether we are inside garbage collection
 */
diff --git a/Engine/Source/Runtime/Engine/Private/DebugGarbageCollectionGraph.cpp b/Engine/Source/Runtime/Engine/Private/DebugGarbageCollectionGraph.cpp
new file mode 100644
--- /dev/null
+++ b/Engine/Source/Runtime/Engine/Private/DebugGarbageCollectionGraph.cpp
@@ -0,0 +1,146 @@
+// Copyright Epic Games, Inc. All Rights Reserved.
+
+#include "DebugGarbageCollectionGraph.h"
+
+#include "Debug/DebugDrawService.h"
+#include "Engine/GameEngine.h"
+
+/*
+	TODO additional features for a future changeset:
+	- compute and display the all-time (not windowed) average
+	  collection time so far
+ */
+
+int32 UDebugGarbageCollectionGraph::SafeDurationThreshold = 4;
+
+void UDebugGarbageCollectionGraph::StartDrawing()
+{
+	if (DrawHandle.IsValid())
+	{
+		return;
+	}
+	DrawHandle = UDebugDrawService::Register(TEXT("Game"),
+		FDebugDrawDelegate::CreateUObject(this,
+			&UDebugGarbageCollectionGraph::Draw));
+}
+
+void UDebugGarbageCollectionGraph::StopDrawing()
+{
+	if (!DrawHandle.IsValid())
+	{
+		return;
+	}
+	UDebugDrawService::Unregister(DrawHandle);
+	DrawHandle.Reset();
+}
+
+void UDebugGarbageCollectionGraph::Draw(UCanvas* Canvas, APlayerController*)
+{
+	static constexpr double FrameSeconds     = 1 / 60.0;
+	static constexpr int    ViewX            = 20;
+	static constexpr int    ViewY            = 20;
+	static constexpr int    ViewWidth        = 400;
+	static constexpr int    ViewHeight       = 400;
+	static constexpr int    OneMinuteSeconds = 60;
+	static constexpr double WidthSeconds     = OneMinuteSeconds * 0.5;
+	static constexpr double HeightSeconds    = FrameSeconds * 10;
+	static constexpr double PixelsPerSecond  = ViewWidth / WidthSeconds;
+
+	const double LastGCTime = GetLastGCTime();
+	const double CurTime    = FPlatformTime::Seconds();
+
+	if (LastGCTime != History.back().Time)
+	{
+		History.push_back({LastGCTime, GetLastGCDuration()});
+	}
+	while (!History.empty() && CurTime - History.front().Time > WidthSeconds)
+	{
+		History.erase(History.begin());
+	}
+
+	auto&& DrawLine =
+		[Canvas](double X0, double Y0, double X1, double Y1, const FLinearColor& Color, double Thickness)
+		{
+			FCanvasLineItem Line{FVector2D{X0, Y0}, FVector2D{X1, Y1}};
+			Line.SetColor(Color);
+			Line.LineThickness = Thickness;
+			Canvas->DrawItem(Line);
+		};
+
+	auto&& DrawString =
+		[Canvas](const FString& Str, int X, int Y)
+		{
+			FCanvasTextItem Text{FVector2D(X, Y), FText::FromString(Str), GEngine->GetSmallFont(), FLinearColor::Yellow};
+			Text.EnableShadow(FLinearColor::Black);
+			Text.bCentreX = true;
+			Text.bCentreY = true;
+			Canvas->DrawItem(Text);
+		};
+
+	{ // draw graph edges + label
+		const FLinearColor Color = FLinearColor::White;
+		DrawLine(ViewX, ViewY + ViewHeight, ViewX + ViewWidth, ViewY + ViewHeight, Color, 1);
+		DrawLine(ViewX, ViewY, ViewX, ViewY + ViewHeight, Color, 1);
+		DrawLine(ViewX + ViewWidth, ViewY, ViewX + ViewWidth, ViewY + ViewHeight, Color, 1);
+		DrawString(TEXT("GC Mark Durations"), ViewX + ViewWidth * 0.5, ViewY + ViewHeight + 16);
+	}
+
+	for (int I = History.size() - 1; I >= 0; --I)
+	{
+		const auto&        Item = History[I];
+		const double       X    = ViewX + ViewWidth - PixelsPerSecond * (CurTime - Item.Time);
+		const double       H    = ViewHeight * (Item.Duration / HeightSeconds);
+		const FLinearColor Color =
+			Item.Duration <= (SafeDurationThreshold * 0.001)
+				? FLinearColor::Green
+				: FLinearColor::Red;
+		const int LineThickness = 3;
+
+		DrawLine(X, ViewY + ViewHeight - 1, X, ViewY + ViewHeight - H, Color, LineThickness);
+		DrawString(
+			FString::Printf(TEXT("%dms"), int(Item.Duration * 1000)), X, ViewY + ViewHeight - H - 11);
+	}
+}
+
+static UDebugGarbageCollectionGraph* GDebugGraph;
+
+static void ShowDebugGraph(FOutputDevice&)
+{
+	if (GDebugGraph)
+	{
+		return;
+	}
+	GDebugGraph = NewObject<UDebugGarbageCollectionGraph>();
+	GDebugGraph->AddToRoot();
+	GDebugGraph->StartDrawing();
+}
+
+static void HideDebugGraph(FOutputDevice&)
+{
+	if (!GDebugGraph)
+	{
+		return;
+	}
+	GDebugGraph->StopDrawing();
+	GDebugGraph->RemoveFromRoot();
+	GDebugGraph = nullptr;
+}
+
+static FAutoConsoleCommandWithOutputDevice
+	GShowDebugGraphCmd(
+		TEXT("gc.DebugGraphShow"),
+		TEXT("Show GC debug graph. "
+			 "(See also: DebugGraphSafeDurationThresholdMs)"),
+		FConsoleCommandWithOutputDeviceDelegate::CreateStatic(&ShowDebugGraph));
+
+static FAutoConsoleCommandWithOutputDevice
+	GHideDebugGraphCmd(
+		TEXT("gc.DebugGraphHide"),
+		TEXT("Hide GC debug graph."),
+		FConsoleCommandWithOutputDeviceDelegate::CreateStatic(&HideDebugGraph));
+
+static FAutoConsoleVariableRef CVarDebugGraphSafeDurationThreshold(
+	TEXT("gc.DebugGraphSafeDurationThresholdMs"),
+	UDebugGarbageCollectionGraph::SafeDurationThreshold,
+	TEXT("GC Debug Graph: Safe GC duration threshold (in milliseconds)."),
+	ECVF_Default);
diff --git a/Engine/Source/Runtime/Engine/Private/DebugGarbageCollectionGraph.h b/Engine/Source/Runtime/Engine/Private/DebugGarbageCollectionGraph.h
new file mode 100644
--- /dev/null
+++ b/Engine/Source/Runtime/Engine/Private/DebugGarbageCollectionGraph.h
@@ -0,0 +1,34 @@
+// Copyright Epic Games, Inc. All Rights Reserved.
+
+#pragma once
+
+#include "GameFramework/HUD.h"
+#include <vector>
+
+#include "DebugGarbageCollectionGraph.generated.h"
+
+UCLASS()
+class UDebugGarbageCollectionGraph : public UObject
+{
+	GENERATED_BODY()
+
+public:
+	void StartDrawing();
+	void StopDrawing();
+
+	static int32 SafeDurationThreshold;
+
+private:
+	void Draw(UCanvas* Canvas, class APlayerController* PC);
+
+	struct HistoryItem
+	{
+		double Time;
+		double Duration;
+	};
+	std::vector<HistoryItem> History = {
+		{0, 0}
+    };
+
+	FDelegateHandle DrawHandle;
+};
