Fix deprecated FeatureLevel calls

diff --git a/Engine/Source/Runtime/Engine/Private/World.cpp b/Engine/Source/Runtime/Engine/Private/World.cpp
--- a/Engine/Source/Runtime/Engine/Private/World.cpp
+++ b/Engine/Source/Runtime/Engine/Private/World.cpp
@@ -8615,7 +8615,7 @@ void UWorld::PurgeScene()
 		const bool bOldVal = GUsingNullRHI;
 		GUsingNullRHI = true;
 		{
-			RecreateScene(FeatureLevel, false /* bBroadcastChange */);
+			RecreateScene(GetFeatureLevel(), false /* bBroadcastChange */);
 		}
 		GUsingNullRHI = bOldVal;
 
@@ -8637,7 +8637,7 @@ void UWorld::RestoreScene()
 		GetRendererModule().PerFrameCleanupIfSkipRenderer();
 		FlushRenderingCommands();
 
-		RecreateScene(FeatureLevel, false /* bBroadcastChange */);
+		RecreateScene(GetFeatureLevel(), false /* bBroadcastChange */);
 
 		InvalidateAllSkyCaptures();
 		TriggerStreamingDataRebuild();
