Add ability to retrieve and set BasePose(Index) for PoseAssets

diff --git a/Engine/Source/Runtime/Engine/Classes/Animation/PoseAsset.h b/Engine/Source/Runtime/Engine/Classes/Animation/PoseAsset.h
--- a/Engine/Source/Runtime/Engine/Classes/Animation/PoseAsset.h
+++ b/Engine/Source/Runtime/Engine/Classes/Animation/PoseAsset.h
@@ -298,6 +298,14 @@ public:
 	UFUNCTION(BlueprintPure, Category=PoseAsset)
 	void GetPoseNames(TArray<FName>& PoseNames) const;
 
+	/** Returns base pose name, only valid when additive, NAME_None indicates reference pose */
+	UFUNCTION(BlueprintPure, Category=PoseAsset)
+	FName GetBasePoseName() const;
+
+	/** Set base pose index by name, NAME_None indicates reference pose - returns true if set successfully */
+	UFUNCTION(BlueprintCallable, Category=PoseAsset)
+    bool SetBasePoseName(const FName& NewBasePoseName);
+
 	UE_DEPRECATED(5.3, "Please use AddPoseWithUniqueName.")
 	ENGINE_API bool AddOrUpdatePoseWithUniqueName(const USkeletalMeshComponent* MeshComponent, FSmartName* OutPoseName = nullptr) { return false; }
 	
diff --git a/Engine/Source/Runtime/Engine/Private/Animation/PoseAsset.cpp b/Engine/Source/Runtime/Engine/Private/Animation/PoseAsset.cpp
--- a/Engine/Source/Runtime/Engine/Private/Animation/PoseAsset.cpp
+++ b/Engine/Source/Runtime/Engine/Private/Animation/PoseAsset.cpp
@@ -1176,6 +1176,33 @@ void UPoseAsset::RenamePose(const FName& OriginalPoseName, const FName& NewPoseN
 	ModifyPoseName(OriginalPoseName, NewPoseName);
 }
 
+FName UPoseAsset::GetBasePoseName() const
+{
+	return PoseContainer.PoseNames.IsValidIndex(BasePoseIndex) ? PoseContainer.PoseNames[BasePoseIndex].DisplayName : NAME_None;
+}
+
+bool UPoseAsset::SetBasePoseName(const FName& NewBasePoseName)
+{
+	if (NewBasePoseName != NAME_None)
+	{
+		const int32 NewIndex = PoseContainer.PoseNames.IndexOfByPredicate([NewBasePoseName](const FSmartName& SmartName)
+		{
+			return SmartName.DisplayName == NewBasePoseName;			
+		});
+		if (NewIndex != INDEX_NONE)
+		{
+			BasePoseIndex = NewIndex;
+		}
+		else
+		{
+			return false;
+		}
+	}
+	
+	BasePoseIndex = INDEX_NONE;
+	return true;
+}
+
 void UPoseAsset::GetPoseNames(TArray<FName>& PoseNames) const
 {	
 	const int32 NumPoses = GetNumPoses();
