Early init world partition module to avoid a commandlet crash, will make a better fix later.

diff --git a/Engine/Source/Runtime/Engine/Private/UnrealEngine.cpp b/Engine/Source/Runtime/Engine/Private/UnrealEngine.cpp
--- a/Engine/Source/Runtime/Engine/Private/UnrealEngine.cpp
+++ b/Engine/Source/Runtime/Engine/Private/UnrealEngine.cpp
@@ -1985,6 +1985,10 @@ void UEngine::Init(IEngineLoop* InEngineLoop)
 	SelectedMaterialColor = DefaultSelectedMaterialColor;
 	SelectionOutlineColor = DefaultSelectedMaterialColor;
 
+#if WITH_EDITOR
+	FModuleManager::Get().LoadModule("WorldPartitionEditor");
+#endif
+
 	InitializeObjectReferences();
 
 	if (GConfig)
@@ -2087,7 +2091,6 @@ void UEngine::Init(IEngineLoop* InEngineLoop)
 		FModuleManager::Get().LoadModuleChecked("MovieSceneTracks");
 		FModuleManager::Get().LoadModule("LevelSequence");
 #if WITH_EDITOR
-		FModuleManager::Get().LoadModule("WorldPartitionEditor");
 		// The SparseVolumeTexture module containing the importer is only loaded and used in the editor.
 		FModuleManager::Get().LoadModuleChecked("SparseVolumeTexture");
 #endif
