WorldPartition: Add some missing logs for invalid actors

#rnx

diff --git a/Engine/Source/Runtime/Engine/Private/WorldPartition/ActorDescContainer.cpp b/Engine/Source/Runtime/Engine/Private/WorldPartition/ActorDescContainer.cpp
--- a/Engine/Source/Runtime/Engine/Private/WorldPartition/ActorDescContainer.cpp
+++ b/Engine/Source/Runtime/Engine/Private/WorldPartition/ActorDescContainer.cpp
@@ -83,12 +83,32 @@ void UActorDescContainer::Initialize(const FInitializeParams& InitParams)
 		for (const FAssetData& Asset : Assets)
 		{
 			TUniquePtr<FWorldPartitionActorDesc> ActorDesc = FWorldPartitionActorDescUtils::GetActorDescriptorFromAssetData(Asset);
+								
+			bool bValid = true;
+			if (!ActorDesc.IsValid())
+			{
+				UE_LOG(LogWorldPartition, Warning, TEXT("Invalid actor descriptor for actor '%s' from package '%s'"), *Asset.GetObjectPathString(), *Asset.PackageName.ToString());
+				bValid = false;
+			} 
+			else if (FWorldPartitionActorDesc* ExistingDesc = FActorDescList::GetActorDesc(ActorDesc->GetGuid()))
+			{
+				check(ExistingDesc->GetGuid() == ActorDesc->GetGuid());
+				UE_LOG(LogWorldPartition, Warning, TEXT("Duplicate actor descriptor guid `%s`: Actor: '%s' from package '%s' -> Existing actor '%s' from package '%s'"), 
+					*ActorDesc->GetGuid().ToString(), 
+					*ActorDesc->GetActorName().ToString(), 
+					*ActorDesc->GetActorPackage().ToString(),
+					*ExistingDesc->GetActorName().ToString(),
+					*ExistingDesc->GetActorPackage().ToString());
+				bValid = false;
+			}
+			else if(!ActorDesc->GetNativeClass().IsValid() || 
+					(ActorDesc->GetBaseClass().IsValid() && !ClassDescRegistry.IsRegisteredClass(ActorDesc->GetBaseClass())) || 
+					(InitParams.FilterActorDesc && !InitParams.FilterActorDesc(ActorDesc.Get())))
+			{
+				bValid = false;
+			}
 
-			if (!ActorDesc.IsValid() || 
-				!ActorDesc->GetNativeClass().IsValid() || 
-				(ActorDesc->GetBaseClass().IsValid() && !ClassDescRegistry.IsRegisteredClass(ActorDesc->GetBaseClass())) || 
-				FActorDescList::GetActorDesc(ActorDesc->GetGuid()) ||
-				(InitParams.FilterActorDesc && !InitParams.FilterActorDesc(ActorDesc.Get())))
+			if (!bValid)
 			{
 				InvalidActors.Emplace(Asset);
 				continue;
