World Partition: make streaming generation logs deterministic (groundwork for streaming generation unit tests).

diff --git a/Engine/Source/Runtime/Engine/Private/WorldPartition/WorldPartitionRuntimeHash.cpp b/Engine/Source/Runtime/Engine/Private/WorldPartition/WorldPartitionRuntimeHash.cpp
--- a/Engine/Source/Runtime/Engine/Private/WorldPartition/WorldPartitionRuntimeHash.cpp
+++ b/Engine/Source/Runtime/Engine/Private/WorldPartition/WorldPartitionRuntimeHash.cpp
@@ -335,15 +335,24 @@ void UWorldPartitionRuntimeHash::DumpStateLog(FHierarchicalLogArchive& Ar)
 	Ar.Printf(TEXT("----------------------------------------------------------------------------------------------------------------"));
 	{
 		FHierarchicalLogArchive::FIndentScope CellIndentScope = Ar.PrintfIndent(TEXT("Content of %s Persistent Level"), *GetWorld()->GetName());
-		Ar.Printf(TEXT("Always loaded Actor Count: %d "), GetWorld()->PersistentLevel->Actors.Num());
+
+		TArray<const AActor*> Actors;
 		for (const AActor* Actor : GetWorld()->PersistentLevel->Actors)
 		{
 			if (Actor)
 			{
-				Ar.Printf(TEXT("Actor Path: %s"), *Actor->GetPathName());
-				Ar.Printf(TEXT("Actor Package: %s"), *Actor->GetPackage()->GetName());
+				Actors.Add(Actor);
 			}
 		}
+
+		Actors.Sort([this](const AActor& A, const AActor& B) { return A.GetFName().LexicalLess(B.GetFName()); });
+
+		Ar.Printf(TEXT("Always loaded Actor Count: %d "), Actors.Num());
+		for (const AActor* Actor : Actors)
+		{
+			Ar.Printf(TEXT("Actor Path: %s"), *Actor->GetPathName());
+			Ar.Printf(TEXT("Actor Package: %s"), *Actor->GetPackage()->GetName());
+		}
 	}
 	Ar.Printf(TEXT(""));
 }
diff --git a/Engine/Source/Runtime/Engine/Private/WorldPartition/WorldPartitionRuntimeLevelStreamingCell.cpp b/Engine/Source/Runtime/Engine/Private/WorldPartition/WorldPartitionRuntimeLevelStreamingCell.cpp
--- a/Engine/Source/Runtime/Engine/Private/WorldPartition/WorldPartitionRuntimeLevelStreamingCell.cpp
+++ b/Engine/Source/Runtime/Engine/Private/WorldPartition/WorldPartitionRuntimeLevelStreamingCell.cpp
@@ -316,7 +316,11 @@ void UWorldPartitionRuntimeLevelStreamingCell::DumpStateLog(FHierarchicalLogArch
 {
 	Super::DumpStateLog(Ar);
 
-	for (const FWorldPartitionRuntimeCellObjectMapping& Mapping : Packages)
+	TArray<FWorldPartitionRuntimeCellObjectMapping> SortedPackages(Packages);
+
+	SortedPackages.Sort([](const FWorldPartitionRuntimeCellObjectMapping& A, const FWorldPartitionRuntimeCellObjectMapping& B) { return A.Path.LexicalLess(B.Path); });
+
+	for (const FWorldPartitionRuntimeCellObjectMapping& Mapping : SortedPackages)
 	{
 		Ar.Printf(TEXT("   Actor Path: %s"), *Mapping.Path.ToString());
 		Ar.Printf(TEXT("Actor Package: %s"), *Mapping.Package.ToString());
diff --git a/Engine/Source/Runtime/Engine/Private/WorldPartition/WorldPartitionRuntimeSpatialHash.cpp b/Engine/Source/Runtime/Engine/Private/WorldPartition/WorldPartitionRuntimeSpatialHash.cpp
--- a/Engine/Source/Runtime/Engine/Private/WorldPartition/WorldPartitionRuntimeSpatialHash.cpp
+++ b/Engine/Source/Runtime/Engine/Private/WorldPartition/WorldPartitionRuntimeSpatialHash.cpp
@@ -1237,14 +1237,22 @@ void UWorldPartitionRuntimeSpatialHash::DumpStateLog(FHierarchicalLogArchive& Ar
 			{
 				FHierarchicalLogArchive::FIndentScope LevelIndentScope = Ar.PrintfIndent(TEXT("Content of Grid Level %d"), Level++);
 
+				TArray<UWorldPartitionRuntimeCell*> Cells;
 				for (const FSpatialHashStreamingGridLayerCell& LayerCell : GridLevel.LayerCells)
 				{
 					for (const TObjectPtr<UWorldPartitionRuntimeCell>& Cell : LayerCell.GridCells)
 					{
-						FHierarchicalLogArchive::FIndentScope CellIndentScope = Ar.PrintfIndent(TEXT("Content of Cell %s"), *Cell->GetDebugName());
-						Cell->DumpStateLog(Ar);
+						Cells.Add(Cell);
 					}
 				}
+
+				Cells.Sort([this](const UWorldPartitionRuntimeCell& A, const UWorldPartitionRuntimeCell& B) { return A.GetFName().LexicalLess(B.GetFName()); });
+
+				for (UWorldPartitionRuntimeCell* Cell : Cells)
+				{
+					FHierarchicalLogArchive::FIndentScope CellIndentScope = Ar.PrintfIndent(TEXT("Content of Cell %s"), *Cell->GetDebugName());
+					Cell->DumpStateLog(Ar);
+				}
 			}
 		}
 		Ar.Printf(TEXT(""));
