limit the hit impact normal change to be more limited to implict object unions - resolves a regression in character movement behavior

diff --git a/Engine/Source/Runtime/Engine/Private/Collision/CollisionConversions.cpp b/Engine/Source/Runtime/Engine/Private/Collision/CollisionConversions.cpp
--- a/Engine/Source/Runtime/Engine/Private/Collision/CollisionConversions.cpp
+++ b/Engine/Source/Runtime/Engine/Private/Collision/CollisionConversions.cpp
@@ -4,6 +4,7 @@
 #include "BodySetupCore.h"
 #include "Engine/World.h"
 #include "Chaos/Capsule.h"
+#include "Chaos/ImplicitObjectType.h"
 #include "Components/PrimitiveComponent.h"
 
 #include "EngineLogs.h"
@@ -100,17 +101,21 @@ static FVector FindGeomOpposingNormal(ECollisionShapeType QueryGeomType, const T
 	// TODO: can we support other shapes here as well?
 	if (QueryGeomType == ECollisionShapeType::Capsule || QueryGeomType == ECollisionShapeType::Sphere)
 	{
-		if (!Hit.FaceNormal.IsNearlyZero())
+		if (const FPhysicsShape* Shape = GetShape(Hit))
 		{
-			return Hit.FaceNormal;
-		}
-		else if (const FPhysicsShape* Shape = GetShape(Hit))
-		{
-			const FTransform ActorTM(Hit.Actor->R(), Hit.Actor->X());
-			const FVector LocalInNormal = ActorTM.InverseTransformVectorNoScale(InNormal);
-			const FVector LocalTraceDirectionDenorm = ActorTM.InverseTransformVectorNoScale(TraceDirectionDenorm);
-			const FVector LocalNormal = Shape->GetGeometry()->FindGeometryOpposingNormal(LocalTraceDirectionDenorm, Hit.FaceIndex, LocalInNormal);
-			return ActorTM.TransformVectorNoScale(LocalNormal);
+			const Chaos::EImplicitObjectType ShapeType = Shape->GetGeometry()->GetType();
+			if (ShapeType == Chaos::ImplicitObjectType::Union && !Hit.FaceNormal.IsNearlyZero())
+			{
+				return Hit.FaceNormal;
+			}
+			else
+			{
+				const FTransform ActorTM(Hit.Actor->R(), Hit.Actor->X());
+				const FVector LocalInNormal = ActorTM.InverseTransformVectorNoScale(InNormal);
+				const FVector LocalTraceDirectionDenorm = ActorTM.InverseTransformVectorNoScale(TraceDirectionDenorm);
+				const FVector LocalNormal = Shape->GetGeometry()->FindGeometryOpposingNormal(LocalTraceDirectionDenorm, Hit.FaceIndex, LocalInNormal);
+				return ActorTM.TransformVectorNoScale(LocalNormal);
+			}
 		}
 	}
 
