Fix for deadlock: reverting soundcue priming from game thread.
#pf6414b64c3d275ade38db21e3

diff --git a/Engine/Source/Runtime/Engine/Private/SoundNode.cpp b/Engine/Source/Runtime/Engine/Private/SoundNode.cpp
--- a/Engine/Source/Runtime/Engine/Private/SoundNode.cpp
+++ b/Engine/Source/Runtime/Engine/Private/SoundNode.cpp
@@ -18,6 +18,13 @@ FAutoConsoleVariableRef CVarBypassRetainInSoundNodes(
 	TEXT("When set to 1, we ignore the loading behavior of sound classes set on a Sound Cue directly.\n"),
 	ECVF_Default);
 
+static int32 ManuallyPrimeChildNodesCVar = 1;
+FAutoConsoleVariableRef CVarManuallyPrimeChildNodes(			
+	TEXT("au.streamcache.priming.ManuallyPrimeChildNodes"),
+	ManuallyPrimeChildNodesCVar,			
+	TEXT("When set to 1, we ignore the loading behavior of sound classes set on a Sound Cue directly.\n"),
+	ECVF_Default);
+
 /*-----------------------------------------------------------------------------
 	USoundNode implementation.
 -----------------------------------------------------------------------------*/
@@ -65,7 +72,6 @@ bool USoundNode::CanBeInCluster() const
 void USoundNode::AddReferencedObjects(UObject* InThis, FReferenceCollector& Collector)
 {
 	USoundNode* This = CastChecked<USoundNode>(InThis);
-
 	Collector.AddReferencedObject(This->GraphNode, This);
 
 	Super::AddReferencedObjects(InThis, Collector);
@@ -106,30 +112,43 @@ void USoundNode::PrimeChildWavePlayers(bool bRecurse)
 		return;
 	}
 
-	// Search child nodes for wave players, then prime each soundwave
-	IAudioStreamingManager &Mgr = IStreamingManager::Get().GetAudioStreamingManager();
-	for (USoundNode* ChildNode : ChildNodes)
-	{
-		if (ChildNode)
+	// Note: it is not safe to call IAudioStreamingManager::RequestChunk from the game thread
+	// if there is no async loading thread.  This can deadlock:
+
+	// The call will try to obtain a lock held by another thread, that thread may be waiting for
+	// any pending async loads to complete, which will never happen if WE are the (non)"async loading thread"
+	// for now, in the synchronous loading case, we fallback to the old behavior.
+	if(!BypassRetainInSoundNodesCVar && ManuallyPrimeChildNodesCVar && IsAsyncLoadingMultithreaded())
+	{	
+		// Search child nodes for wave players, then prime each soundwave
+		IAudioStreamingManager &Mgr = IStreamingManager::Get().GetAudioStreamingManager();
+		for (USoundNode* ChildNode : ChildNodes)
 		{
-			ChildNode->ConditionalPostLoad();
-			if (bRecurse)
+			if (ChildNode)
 			{
-				ChildNode->PrimeChildWavePlayers(true);
-			}
+				ChildNode->ConditionalPostLoad();
+				if (bRecurse)
+				{
+					ChildNode->PrimeChildWavePlayers(true);
+				}
 
-			if (USoundNodeWavePlayer* WavePlayer = Cast<USoundNodeWavePlayer>(ChildNode))
-			{
-				if (USoundWave* SoundWave = WavePlayer->GetSoundWave())
+				if (USoundNodeWavePlayer* WavePlayer = Cast<USoundNodeWavePlayer>(ChildNode))
 				{
-					if (SoundWave->IsStreaming() && SoundWave->GetNumChunks() > 1)
+					if (USoundWave* SoundWave = WavePlayer->GetSoundWave())
 					{
-						Mgr.RequestChunk(SoundWave->CreateSoundWaveProxy(), 1, [](EAudioChunkLoadResult) {});
+						if (SoundWave->IsStreaming() && SoundWave->GetNumChunks() > 1)
+						{
+							Mgr.RequestChunk(SoundWave->CreateSoundWaveProxy(), 1, [](EAudioChunkLoadResult) {});
+						}
 					}
 				}
 			}
 		}
 	}
+	else
+	{
+		OverrideLoadingBehaviorOnChildWaves(bRecurse, ESoundWaveLoadingBehavior::PrimeOnLoad);
+	}
 }
 
 void USoundNode::RetainChildWavePlayers(bool bRecurse)
