Fix UWorld::UpdateLevelStreaming() not being called when IsRenderingSuspended()

diff --git a/Engine/Source/Runtime/Engine/Private/GameEngine.cpp b/Engine/Source/Runtime/Engine/Private/GameEngine.cpp
--- a/Engine/Source/Runtime/Engine/Private/GameEngine.cpp
+++ b/Engine/Source/Runtime/Engine/Private/GameEngine.cpp
@@ -1880,13 +1880,30 @@ void UGameEngine::Tick( float DeltaSeconds, bool bIdleMode )
 
 	const bool bRenderingSuspended = IsRenderingSuspended();
 
-	if (!bIdleMode && !bRenderingSuspended && !IsRunningDedicatedServer() && !IsRunningCommandlet() && FEmbeddedCommunication::IsAwakeForRendering())
+	if (!bIdleMode && !IsRunningDedicatedServer() && !IsRunningCommandlet() && FEmbeddedCommunication::IsAwakeForRendering())
 	{
-		// Render everything.
-		RedrawViewports();
+		if (!bRenderingSuspended)
+		{
+			// Render everything.
+			RedrawViewports();
 
-		// Some tasks can only be done once we finish all scenes/viewports
-		GetRendererModule().PostRenderAllViewports();
+			// Some tasks can only be done once we finish all scenes/viewports
+			GetRendererModule().PostRenderAllViewports();
+		}
+		else
+		{
+			// Still need to call UpdateLevelStreaming() even when not rendering
+			if (GameViewport && GameViewport->Viewport)
+			{
+				if (FViewportClient* ViewportClient = GameViewport->Viewport->GetClient())
+				{
+					if (UWorld* World = ViewportClient->GetWorld())
+					{
+						World->UpdateLevelStreaming();
+					}
+				}
+			}
+		}
 	}
 
 	if( GIsClient )
