Fix CL# 24700861 in light of CL# 24421503
#rnx

diff --git a/Engine/Source/Runtime/Engine/Private/Animation/PoseAsset.cpp b/Engine/Source/Runtime/Engine/Private/Animation/PoseAsset.cpp
--- a/Engine/Source/Runtime/Engine/Private/Animation/PoseAsset.cpp
+++ b/Engine/Source/Runtime/Engine/Private/Animation/PoseAsset.cpp
@@ -2,29 +2,27 @@
 
 #include "Animation/PoseAsset.h"
 
-#include "Animation/AnimStats.h"
 #include "BonePose.h"
-#include "UObject/FrameworkObjectVersion.h"
 #include "UObject/FortniteMainBranchObjectVersion.h"
+#include "UObject/FrameworkObjectVersion.h"
 #include "UObject/UE5MainStreamObjectVersion.h"
+#include "UObject/UE5ReleaseStreamObjectVersion.h"
 #include "AnimationRuntime.h"
 #include "Components/SkeletalMeshComponent.h"
+#include "Animation/AnimCurveUtils.h"
 #include "Animation/AnimSequence.h"
+#include "Animation/AnimStats.h"
 #include "Animation/AnimationPoseData.h"
 #include "Animation/AnimData/AnimDataModel.h"
 #include "Animation/AnimSequenceHelpers.h"
+#include "Animation/SkeletonRemapping.h"
+#include "Animation/SkeletonRemappingRegistry.h"
 #include "Engine/SkeletalMesh.h"
 #include "UObject/LinkerLoad.h"
 #include "UObject/ObjectSaveContext.h"
 #include "UObject/Package.h"
 #include "UObject/UnrealType.h"
-#include "Animation/SkeletonRemappingRegistry.h"
-#include "Animation/SkeletonRemapping.h"
-#include "UObject/UE5MainStreamObjectVersion.h"
-#include "UObject/UE5ReleaseStreamObjectVersion.h"
 #include "Widgets/Notifications/SNotificationList.h"
-#include "Animation/AnimCurveUtils.h"
-#include "UObject/UE5MainStreamObjectVersion.h"
 
 #include UE_INLINE_GENERATED_CPP_BY_NAME(PoseAsset)
 
@@ -1178,17 +1176,14 @@ void UPoseAsset::RenamePose(const FName& OriginalPoseName, const FName& NewPoseN
 
 FName UPoseAsset::GetBasePoseName() const
 {
-	return PoseContainer.PoseNames.IsValidIndex(BasePoseIndex) ? PoseContainer.PoseNames[BasePoseIndex].DisplayName : NAME_None;
+	return PoseContainer.PoseFNames.IsValidIndex(BasePoseIndex) ? PoseContainer.PoseFNames[BasePoseIndex] : NAME_None;
 }
 
 bool UPoseAsset::SetBasePoseName(const FName& NewBasePoseName)
 {
 	if (NewBasePoseName != NAME_None)
 	{
-		const int32 NewIndex = PoseContainer.PoseNames.IndexOfByPredicate([NewBasePoseName](const FSmartName& SmartName)
-		{
-			return SmartName.DisplayName == NewBasePoseName;			
-		});
+		const int32 NewIndex = PoseContainer.PoseFNames.IndexOfByKey(NewBasePoseName);
 		if (NewIndex != INDEX_NONE)
 		{
 			BasePoseIndex = NewIndex;
