AActor::IsHLODRelevant() - Reordered conditions

diff --git a/Engine/Source/Runtime/Engine/Private/Actor.cpp b/Engine/Source/Runtime/Engine/Private/Actor.cpp
--- a/Engine/Source/Runtime/Engine/Private/Actor.cpp
+++ b/Engine/Source/Runtime/Engine/Private/Actor.cpp
@@ -5901,12 +5901,12 @@ void AActor::PostRename(UObject* OldOuter, const FName OldName)
 
 bool AActor::IsHLODRelevant() const
 {
-	if (HasAnyFlags(RF_Transient))
+	if (!IsValidChecked(this))
 	{
 		return false;
 	}
 
-	if (IsHidden())
+	if (HasAnyFlags(RF_Transient))
 	{
 		return false;
 	}
@@ -5916,7 +5916,7 @@ bool AActor::IsHLODRelevant() const
 		return false;
 	}
 
-	if (!IsValidChecked(this))
+	if (IsHidden())
 	{
 		return false;
 	}
