Unify material related cook stats.

- Adds the following new material cook stats to account for critical path, Game Thread time that we hope to optimize:
    - MaterialTranslateCalls
    - MaterialTranslateTimeSec
    - FinishCacheShadersSec
- Unify all the Material cook stats into a single namespace `MaterialCookStats`
- Unify all the Material cook stats into a single cook stat category `Material`

Sample output looks like the following:

Material
    MeshMaterialShadersCompiled=8109
    FinishCacheShadersSec=29.906448
    MaterialShadersCompiled=2185
    MaterialInstanceUpdateCachedExpressionDataSec=0.014973
    MaterialUpdateCachedExpressionDataSec=0.030023
    MaterialTranslateCalls=592
    MaterialTranslateTimeSec=1.284383

diff --git a/Engine/Source/Runtime/Engine/Private/Materials/HLSLMaterialTranslator.cpp b/Engine/Source/Runtime/Engine/Private/Materials/HLSLMaterialTranslator.cpp
--- a/Engine/Source/Runtime/Engine/Private/Materials/HLSLMaterialTranslator.cpp
+++ b/Engine/Source/Runtime/Engine/Private/Materials/HLSLMaterialTranslator.cpp
@@ -18,6 +18,7 @@
 #include "Materials/HLSLMaterialDerivativeAutogen.h"
 #include "Materials/MaterialFunction.h"
 #include "MaterialExpressionSettings.h"
+#include "ProfilingDebugging/CookStats.h"
 #include "Engine/RendererSettings.h"
 #include "DataDrivenShaderPlatformInfo.h"
 #include "Materials/Material.h"
@@ -50,6 +51,23 @@
 #include "ShaderPlatformCachedIniValue.h"
 #endif
 
+#if ENABLE_COOK_STATS
+#include "ProfilingDebugging/ScopedTimers.h"
+namespace MaterialTranslatorCookStats
+{
+	static int32 MaterialTranslateCalls = 0;
+	static double MaterialTranslateTimeSec = 0.0f;
+
+	static FCookStatsManager::FAutoRegisterCallback RegisterCookStats([](FCookStatsManager::AddStatFuncRef AddStat)
+		{
+			AddStat(TEXT("Material"), FCookStatsManager::CreateKeyValueArray(
+				TEXT("MaterialTranslateCalls"), MaterialTranslateCalls,
+				TEXT("MaterialTranslateTimeSec"), MaterialTranslateTimeSec
+			));
+		});
+}
+#endif
+
 static int32 GetLWCTruncateMode()
 {
 	static const auto CVar = IConsoleManager::Get().FindTConsoleVariableDataInt(TEXT("r.MaterialEditor.LWCTruncateMode"));
@@ -765,6 +783,8 @@ void FHLSLMaterialTranslator::ValidateVtPropertyLimits()
 bool FHLSLMaterialTranslator::Translate()
 {
 	TRACE_CPUPROFILER_EVENT_SCOPE(FHLSLMaterialTranslator::Translate);
+	COOK_STAT(MaterialTranslatorCookStats::MaterialTranslateCalls++);
+	COOK_STAT(FScopedDurationTimer DurationTimer(MaterialTranslatorCookStats::MaterialTranslateTimeSec));
 
 	STAT(double HLSLTranslateTime = 0);
 	{
diff --git a/Engine/Source/Runtime/Engine/Private/Materials/Material.cpp b/Engine/Source/Runtime/Engine/Private/Materials/Material.cpp
--- a/Engine/Source/Runtime/Engine/Private/Materials/Material.cpp
+++ b/Engine/Source/Runtime/Engine/Private/Materials/Material.cpp
@@ -103,12 +103,12 @@
 #if ENABLE_COOK_STATS
 namespace MaterialCookStats
 {
-	static double UpdateCachedExpressionDataSec = 0.0;
+	static double MaterialUpdateCachedExpressionDataSec = 0.0;
 
 	static FCookStatsManager::FAutoRegisterCallback RegisterCookStats([](FCookStatsManager::AddStatFuncRef AddStat)
 		{
 			AddStat(TEXT("Material"), FCookStatsManager::CreateKeyValueArray(
-				TEXT("UpdateCachedExpressionDataSec"), UpdateCachedExpressionDataSec
+				TEXT("MaterialUpdateCachedExpressionDataSec"), MaterialUpdateCachedExpressionDataSec
 			));
 		});
 }
@@ -2027,7 +2027,7 @@ void UMaterial::UpdateTransientExpressionData()
 #if WITH_EDITOR
 void UMaterial::UpdateCachedExpressionData()
 {
-	COOK_STAT(FScopedDurationTimer BlockingTimer(MaterialCookStats::UpdateCachedExpressionDataSec));
+	COOK_STAT(FScopedDurationTimer BlockingTimer(MaterialCookStats::MaterialUpdateCachedExpressionDataSec));
 
 	if (bLoadedCachedExpressionData)
 	{
diff --git a/Engine/Source/Runtime/Engine/Private/Materials/MaterialInstanceConstant.cpp b/Engine/Source/Runtime/Engine/Private/Materials/MaterialInstanceConstant.cpp
--- a/Engine/Source/Runtime/Engine/Private/Materials/MaterialInstanceConstant.cpp
+++ b/Engine/Source/Runtime/Engine/Private/Materials/MaterialInstanceConstant.cpp
@@ -21,12 +21,12 @@
 #include "ProfilingDebugging/ScopedTimers.h"
 namespace MaterialInstanceCookStats
 {
-static double UpdateCachedExpressionDataSec = 0.0;
+static double MaterialInstanceUpdateCachedExpressionDataSec = 0.0;
 
 static FCookStatsManager::FAutoRegisterCallback RegisterCookStats([](FCookStatsManager::AddStatFuncRef AddStat)
 	{
-		AddStat(TEXT("MaterialInstance"), FCookStatsManager::CreateKeyValueArray(
-			TEXT("UpdateCachedExpressionDataSec"), UpdateCachedExpressionDataSec
+		AddStat(TEXT("Material"), FCookStatsManager::CreateKeyValueArray(
+			TEXT("MaterialInstanceUpdateCachedExpressionDataSec"), MaterialInstanceUpdateCachedExpressionDataSec
 		));
 	});
 }
@@ -187,7 +187,7 @@ void FMaterialInstanceCachedData::InitializeForConstant(const FMaterialLayersFun
 
 void UMaterialInstanceConstant::UpdateCachedData()
 {
-	COOK_STAT(FScopedDurationTimer BlockingTimer(MaterialInstanceCookStats::UpdateCachedExpressionDataSec));
+	COOK_STAT(FScopedDurationTimer BlockingTimer(MaterialInstanceCookStats::MaterialInstanceUpdateCachedExpressionDataSec));
 
 	// Don't need to rebuild cached data if it was serialized
 	if (!bLoadedCachedData)
diff --git a/Engine/Source/Runtime/Engine/Private/Materials/MaterialShader.cpp b/Engine/Source/Runtime/Engine/Private/Materials/MaterialShader.cpp
--- a/Engine/Source/Runtime/Engine/Private/Materials/MaterialShader.cpp
+++ b/Engine/Source/Runtime/Engine/Private/Materials/MaterialShader.cpp
@@ -56,12 +56,12 @@ static TAutoConsoleVariable<FString> CVarShaderCompilerDebugDDCKeyAsset(
 namespace MaterialShaderCookStats
 {
 	static FCookStats::FDDCResourceUsageStats UsageStats;
-	static int32 ShadersCompiled = 0;
+	static int32 MaterialShadersCompiled = 0;
 	static FCookStatsManager::FAutoRegisterCallback RegisterCookStats([](FCookStatsManager::AddStatFuncRef AddStat)
 	{
 		UsageStats.LogStats(AddStat, TEXT("MaterialShader.Usage"), TEXT(""));
-		AddStat(TEXT("MaterialShader.Misc"), FCookStatsManager::CreateKeyValueArray(
-			TEXT("ShadersCompiled"), ShadersCompiled
+		AddStat(TEXT("Material"), FCookStatsManager::CreateKeyValueArray(
+			TEXT("MaterialShadersCompiled"), MaterialShadersCompiled
 			));
 	});
 }
@@ -1201,7 +1201,7 @@ static void PrepareMaterialShaderCompileJob(EShaderPlatform Platform,
 	FShaderCompilerEnvironment& ShaderEnvironment = NewJob->Input.Environment;
 
 	UE_LOG(LogShaders, Verbose, TEXT("			%s"), ShaderType->GetName());
-	COOK_STAT(MaterialShaderCookStats::ShadersCompiled++);
+	COOK_STAT(MaterialShaderCookStats::MaterialShadersCompiled++);
 
 	//update material shader stats
 	UpdateMaterialShaderCompilingStats(Material);
diff --git a/Engine/Source/Runtime/Engine/Private/Materials/MaterialShared.cpp b/Engine/Source/Runtime/Engine/Private/Materials/MaterialShared.cpp
--- a/Engine/Source/Runtime/Engine/Private/Materials/MaterialShared.cpp
+++ b/Engine/Source/Runtime/Engine/Private/Materials/MaterialShared.cpp
@@ -53,11 +53,26 @@
 #include "RenderCore.h"
 #include "StrataDefinitions.h"
 #include "DataDrivenShaderPlatformInfo.h"
+#include "ProfilingDebugging/CookStats.h"
 
 #define LOCTEXT_NAMESPACE "MaterialShared"
 
 DEFINE_LOG_CATEGORY(LogMaterial);
 
+#if ENABLE_COOK_STATS
+namespace MaterialSharedCookStats
+{
+	static double FinishCacheShadersSec = 0.0;
+
+	static FCookStatsManager::FAutoRegisterCallback RegisterCookStats([](FCookStatsManager::AddStatFuncRef AddStat)
+		{
+			AddStat(TEXT("Material"), FCookStatsManager::CreateKeyValueArray(
+				TEXT("FinishCacheShadersSec"), FinishCacheShadersSec
+			));
+		});
+}
+#endif
+
 IMPLEMENT_TYPE_LAYOUT(FHashedMaterialParameterInfo);
 IMPLEMENT_TYPE_LAYOUT(FUniformExpressionSet);
 IMPLEMENT_TYPE_LAYOUT(FMaterialCompilationOutput);
@@ -2926,6 +2941,8 @@ bool FMaterial::IsCachingShaders() const
 
 bool FMaterial::FinishCacheShaders() const
 {
+	COOK_STAT(FScopedDurationTimer BlockingTimer(MaterialSharedCookStats::FinishCacheShadersSec));
+
 	if (CacheShadersCompletion)
 	{
 		TRACE_CPUPROFILER_EVENT_SCOPE(FinishCacheShaders);
diff --git a/Engine/Source/Runtime/Engine/Private/Materials/MeshMaterialShader.cpp b/Engine/Source/Runtime/Engine/Private/Materials/MeshMaterialShader.cpp
--- a/Engine/Source/Runtime/Engine/Private/Materials/MeshMaterialShader.cpp
+++ b/Engine/Source/Runtime/Engine/Private/Materials/MeshMaterialShader.cpp
@@ -11,11 +11,11 @@
 #if ENABLE_COOK_STATS
 namespace MaterialMeshCookStats
 {
-	static int32 ShadersCompiled = 0;
+	static int32 MeshMaterialShadersCompiled = 0;
 	static FCookStatsManager::FAutoRegisterCallback RegisterCookStats([](FCookStatsManager::AddStatFuncRef AddStat)
 	{
-		AddStat(TEXT("MeshMaterial.Misc"), FCookStatsManager::CreateKeyValueArray(
-			TEXT("ShadersCompiled"), ShadersCompiled
+		AddStat(TEXT("Material"), FCookStatsManager::CreateKeyValueArray(
+			TEXT("MeshMaterialShadersCompiled"), MeshMaterialShadersCompiled
 			));
 	});
 }
@@ -55,7 +55,7 @@ static void PrepareMeshMaterialShaderCompileJob(EShaderPlatform Platform,
 	UpdateMaterialShaderCompilingStats(Material);
 
 	UE_LOG(LogShaders, Verbose, TEXT("			%s"), ShaderType->GetName());
-	COOK_STAT(MaterialMeshCookStats::ShadersCompiled++);
+	COOK_STAT(MaterialMeshCookStats::MeshMaterialShadersCompiled++);
 
 	// Allow the shader type to modify the compile environment.
 	ShaderType->SetupCompileEnvironment(Platform, MaterialParameters, VertexFactoryType, Key.PermutationId, PermutationFlags, ShaderEnvironment);
