Remove quad tree out-of-bounds log from shipping builds to prevent log spam and performance hit.

diff --git a/Engine/Source/Runtime/Engine/Public/GenericQuadTree.h b/Engine/Source/Runtime/Engine/Public/GenericQuadTree.h
--- a/Engine/Source/Runtime/Engine/Public/GenericQuadTree.h
+++ b/Engine/Source/Runtime/Engine/Public/GenericQuadTree.h
@@ -232,11 +232,13 @@ void TQuadTree<ElementType, NodeCapacity>::Split()
 template <typename ElementType, int32 NodeCapacity>
 void TQuadTree<ElementType, NodeCapacity>::Insert(const ElementType& Element, const FBox2D& Box, const TCHAR* DebugContext)
 {
+#if !UE_BUILD_SHIPPING
 	if (!Box.Intersect(TreeBox))
 	{
 		// Elements shouldn't be added outside the bounds of the top-level quad
 		UE_LOG(LogQuadTree, Warning, TEXT("[%s] Adding element (%s) that is outside the bounds of the quadtree root (%s). Consider resizing."), DebugContext ? DebugContext : TEXT("Unknown Source"), *Box.ToString(), *TreeBox.ToString());
 	}
+#endif // !UE_BUILD_SHIPPING
 
 	InsertElementRecursive(Element, Box, DebugContext);
 }
