Added UBlueprint utility methods (& template variants) to retrieve components from AActor derived class - BP or not
* GetActorClassDefaultComponents()
* GetActorClassDefaultComponent()
* GetActorClassDefaultComponentByName()
* ForEachComponentOfActorClassDefault()

Removed/deprecated these functions (mostly duplicated)
* UHLODTemplatedInstancedStaticMeshComponent::GetActorClassDefaultComponentByName()
* UE::AI::GetActorClassDefaultComponents()
* UBlueprintThumbnailRenderer::CanVisualizeAsset()
* UPCGActorHelpers::GetActorClassDefaultComponents()
* And others in private projects...

diff --git a/Engine/Plugins/Experimental/PCG/Source/PCG/Private/Elements/PCGSpawnActor.cpp b/Engine/Plugins/Experimental/PCG/Source/PCG/Private/Elements/PCGSpawnActor.cpp
--- a/Engine/Plugins/Experimental/PCG/Source/PCG/Private/Elements/PCGSpawnActor.cpp
+++ b/Engine/Plugins/Experimental/PCG/Source/PCG/Private/Elements/PCGSpawnActor.cpp
@@ -144,27 +144,21 @@ UPCGGraphInterface* UPCGSpawnActorSettings::GetSubgraphInterface() const
 		return nullptr;
 	}
 
-	TArray<UActorComponent*> PCGComponents;
-	UPCGActorHelpers::GetActorClassDefaultComponents(TemplateActorClass, PCGComponents, UPCGComponent::StaticClass());
+	UPCGGraphInterface* Result = nullptr;
 
-	if (PCGComponents.IsEmpty())
+	UBlueprint::ForEachComponentOfActorClassDefault<UPCGComponent>(TemplateActorClass, [&](const UPCGComponent* PCGComponent)
 	{
-		return nullptr;
-	}
-
-	for (UActorComponent* Component : PCGComponents)
-	{
-		if (UPCGComponent* PCGComponent = Cast<UPCGComponent>(Component))
+		// If there is no graph, there is no graph instance
+		if (PCGComponent->GetGraph() && PCGComponent->bActivated)
 		{
-			// If there is no graph, there is no graph instance
-			if (PCGComponent->GetGraph() && PCGComponent->bActivated)
-			{
-				return PCGComponent->GetGraphInstance();
-			}
+			Result = PCGComponent->GetGraphInstance();
+			return false;
 		}
-	}
+		
+		return true;
+	});
 
-	return nullptr;
+	return Result;
 }
 
 void UPCGSpawnActorSettings::BeginDestroy()
@@ -438,41 +432,38 @@ bool FPCGSpawnActorElement::ExecuteInternal(FPCGContext* Context) const
 		if (!bFullySkippedDueToReuse && Settings->Option == EPCGSpawnActorOption::CollapseActors)
 		{
 			TRACE_CPUPROFILER_EVENT_SCOPE(FPCGSpawnActorElement::ExecuteInternal::CollapseActors);
-			TArray<UActorComponent*> Components;
-			UPCGActorHelpers::GetActorClassDefaultComponents(Settings->TemplateActorClass, Components, UStaticMeshComponent::StaticClass());
-			
+
 			TMap<FPCGISMCBuilderParameters, TArray<FTransform>> MeshDescriptorTransforms;
 
-			for (UActorComponent* Component : Components)
+			UBlueprint::ForEachComponentOfActorClassDefault<UStaticMeshComponent>(Settings->TemplateActorClass, [&](const UStaticMeshComponent* StaticMeshComponent)
 			{
-				if (UStaticMeshComponent* StaticMeshComponent = Cast<UStaticMeshComponent>(Component))
-				{
-					FPCGISMCBuilderParameters Params;
-					Params.Descriptor.InitFrom(StaticMeshComponent);
-					// TODO: No custom data float support?
+				FPCGISMCBuilderParameters Params;
+				Params.Descriptor.InitFrom(StaticMeshComponent);
+				// TODO: No custom data float support?
 
-					TArray<FTransform>& Transforms = MeshDescriptorTransforms.FindOrAdd(Params);
+				TArray<FTransform>& Transforms = MeshDescriptorTransforms.FindOrAdd(Params)
 
-					if (UInstancedStaticMeshComponent* InstancedStaticMeshComponent = Cast<UInstancedStaticMeshComponent>(StaticMeshComponent))
-					{
-						const int32 NumInstances = InstancedStaticMeshComponent->GetInstanceCount();
-						Transforms.Reserve(Transforms.Num() + NumInstances);
+				if (const UInstancedStaticMeshComponent* InstancedStaticMeshComponent = Cast<UInstancedStaticMeshComponent>(StaticMeshComponent))
+				{
+					const int32 NumInstances = InstancedStaticMeshComponent->GetInstanceCount();
+					Transforms.Reserve(Transforms.Num() + NumInstances);
 						
-						for (int32 InstanceIndex = 0; InstanceIndex < NumInstances; InstanceIndex++)
+					for (int32 InstanceIndex = 0; InstanceIndex < NumInstances; InstanceIndex++)
+					{
+						FTransform InstanceTransform;
+						if (InstancedStaticMeshComponent->GetInstanceTransform(InstanceIndex, InstanceTransform))
 						{
-							FTransform InstanceTransform;
-							if (InstancedStaticMeshComponent->GetInstanceTransform(InstanceIndex, InstanceTransform))
-							{
-								Transforms.Add(InstanceTransform);
-							}
+							Transforms.Add(InstanceTransform);
 						}
 					}
-					else
-					{
-						Transforms.Add(StaticMeshComponent->GetRelativeTransform());
-					}
 				}
-			}
+				else
+				{
+					Transforms.Add(StaticMeshComponent->GetRelativeTransform());
+				}
+
+				return true;
+			});
 				
 			for(const TPair<FPCGISMCBuilderParameters, TArray<FTransform>>& ISMCBuilderTransforms : MeshDescriptorTransforms)
 			{
diff --git a/Engine/Plugins/Experimental/PCG/Source/PCG/Private/Helpers/PCGActorHelpers.cpp b/Engine/Plugins/Experimental/PCG/Source/PCG/Private/Helpers/PCGActorHelpers.cpp
--- a/Engine/Plugins/Experimental/PCG/Source/PCG/Private/Helpers/PCGActorHelpers.cpp
+++ b/Engine/Plugins/Experimental/PCG/Source/PCG/Private/Helpers/PCGActorHelpers.cpp
@@ -7,10 +7,9 @@
 
 #include "PCGComponent.h"
 #include "PCGManagedResource.h"
+#include "Engine/Blueprint.h"
 #include "EngineUtils.h"
 #include "Components/HierarchicalInstancedStaticMeshComponent.h"
-#include "Engine/InheritableComponentHandler.h"
-#include "Engine/SCS_Node.h"
 #include "Engine/StaticMesh.h"
 #include "PCGModule.h"
 
@@ -217,77 +216,14 @@ bool UPCGActorHelpers::DeleteActors(UWorld* World, const TArray<TSoftObjectPtr<A
 	return true;
 }
 
-// Note: this is copied verbatim from AIHelpers.h/.cpp
 void UPCGActorHelpers::GetActorClassDefaultComponents(const TSubclassOf<AActor>& ActorClass, TArray<UActorComponent*>& OutComponents, const TSubclassOf<UActorComponent>& InComponentClass)
 {
-	if (!ensure(ActorClass.Get()))
+	OutComponents.Reset();
+	UBlueprint::ForEachComponentOfActorClassDefault(ActorClass, InComponentClass, [&OutComponents](const UActorComponent* TemplateComponent)
 	{
-		return;
-	}
-
-	UClass* ClassPtr = InComponentClass.Get();
-	TArray<UActorComponent*> ResultComponents;
-
-	// Get the components defined on the native class.
-	AActor* CDO = ActorClass->GetDefaultObject<AActor>();
-	check(CDO);
-	if (ClassPtr)
-	{
-		CDO->GetComponents(InComponentClass, ResultComponents);
-	}
-	else
-	{
-		CDO->GetComponents(ResultComponents);
-	}
-
-	// Try to get the components off the BP class.
-	UBlueprintGeneratedClass* BPClass = Cast<UBlueprintGeneratedClass>(*ActorClass);
-	if (BPClass)
-	{
-		// A BlueprintGeneratedClass has a USimpleConstructionScript member. This member has an array of RootNodes
-		// which contains the SCSNode for the root SceneComponent and non-SceneComponents. For the SceneComponent
-		// hierarchy, each SCSNode knows its children SCSNodes. Each SCSNode stores the component template that will
-		// be created when the Actor is spawned.
-		//
-		// WARNING: This may change in future engine versions!
-
-		TArray<UActorComponent*> Unfiltered;
-		// using this semantic to avoid duplicating following loops or adding a filtering check condition inside the loops
-		TArray<UActorComponent*>& TmpComponents = ClassPtr ? Unfiltered : ResultComponents;
-
-		// Check added components.
-		USimpleConstructionScript* ConstructionScript = BPClass->SimpleConstructionScript;
-		if (ConstructionScript)
-		{
-			for (const USCS_Node* Node : ConstructionScript->GetAllNodes())
-			{
-				TmpComponents.Add(Node->ComponentTemplate);
-			}
-		}
-		// Check modified inherited components.
-		UInheritableComponentHandler* InheritableComponentHandler = BPClass->InheritableComponentHandler;
-		if (InheritableComponentHandler)
-		{
-			for (TArray<FComponentOverrideRecord>::TIterator It = InheritableComponentHandler->CreateRecordIterator(); It; ++It)
-			{
-				TmpComponents.Add(It->ComponentTemplate);
-			}
-		}
-
-		// Filter to the ones matching the requested class.
-		if (ClassPtr)
-		{
-			for (UActorComponent* TemplateComponent : Unfiltered)
-			{
-				if (TemplateComponent->IsA(ClassPtr))
-				{
-					ResultComponents.Add(TemplateComponent);
-				}
-			}
-		}
-	}
-
-	OutComponents = MoveTemp(ResultComponents);
+		OutComponents.Add(const_cast<UActorComponent*>(TemplateComponent));
+		return true;
+	});
 }
 
 void UPCGActorHelpers::ForEachActorInLevel(ULevel* Level, TSubclassOf<AActor> ActorClass, TFunctionRef<bool(AActor*)> Callback)
diff --git a/Engine/Plugins/Experimental/PCG/Source/PCG/Public/Helpers/PCGActorHelpers.h b/Engine/Plugins/Experimental/PCG/Source/PCG/Public/Helpers/PCGActorHelpers.h
--- a/Engine/Plugins/Experimental/PCG/Source/PCG/Public/Helpers/PCGActorHelpers.h
+++ b/Engine/Plugins/Experimental/PCG/Source/PCG/Public/Helpers/PCGActorHelpers.h
@@ -42,14 +42,6 @@ public:
 	static UPCGManagedISMComponent* GetOrCreateManagedISMC(AActor* InTargetActor, UPCGComponent* SourceComponent, uint64 SettingsUID, const FPCGISMCBuilderParameters& Params);
 	static bool DeleteActors(UWorld* World, const TArray<TSoftObjectPtr<AActor>>& ActorsToDelete);
 
-	/**
-	* Fetches all the components of ActorClass's CDO, including the ones added via the BP editor (which AActor.GetComponents fails to do)
-	* @param ActorClass class of AActor for which we will retrieve all components
-	* @param OutComponents this is where the found components will end up. Note that the preexisting contents of OutComponents will get overridden.
-	* @param InComponentClass if supplied will be used to filter the results
-	*/
-	static void GetActorClassDefaultComponents(const TSubclassOf<AActor>& ActorClass, TArray<UActorComponent*>& OutComponents, const TSubclassOf<UActorComponent>& InComponentClass = TSubclassOf<UActorComponent>());
-
 	template <typename T, typename = typename std::enable_if_t<std::is_base_of_v<AActor, T>>>
 	inline static void ForEachActorInLevel(ULevel* Level, TFunctionRef<bool(AActor*)> Callback)
 	{
@@ -105,6 +97,9 @@ public:
 	 * Return the grid cell coordinates on the PCG partition grid given a position and the grid size.
 	 */
 	static FIntVector GetCellCoord(FVector InPosition, int InGridSize, bool bUse2DGrid);
+
+	UE_DEPRECATED(5.3, "Please use UE::BlueprintTools::GetActorClassDefaultComponents() instead")
+	static void GetActorClassDefaultComponents(const TSubclassOf<AActor>& ActorClass, TArray<UActorComponent*>& OutComponents, const TSubclassOf<UActorComponent>& InComponentClass = TSubclassOf<UActorComponent>());
 };
 
 #if UE_ENABLE_INCLUDE_ORDER_DEPRECATED_IN_5_2
diff --git a/Engine/Source/Editor/UnrealEd/Private/ThumbnailHelpers.cpp b/Engine/Source/Editor/UnrealEd/Private/ThumbnailHelpers.cpp
--- a/Engine/Source/Editor/UnrealEd/Private/ThumbnailHelpers.cpp
+++ b/Engine/Source/Editor/UnrealEd/Private/ThumbnailHelpers.cpp
@@ -1118,18 +1118,18 @@ void FClassActorThumbnailScene::ClearStaleActors()
 	}
 }
 
-bool FClassActorThumbnailScene::IsValidComponentForVisualization(UActorComponent* Component)
+bool FClassActorThumbnailScene::IsValidComponentForVisualization(const UActorComponent* Component)
 {
-	UPrimitiveComponent* PrimComp = Cast<UPrimitiveComponent>(Component);
+	const UPrimitiveComponent* PrimComp = Cast<UPrimitiveComponent>(Component);
 	if ( PrimComp && PrimComp->IsVisible() && !PrimComp->bHiddenInGame )
 	{
-		UStaticMeshComponent* StaticMeshComp = Cast<UStaticMeshComponent>(Component);
+		const UStaticMeshComponent* StaticMeshComp = Cast<UStaticMeshComponent>(Component);
 		if ( StaticMeshComp && StaticMeshComp->GetStaticMesh())
 		{
 			return true;
 		}
 
-		USkeletalMeshComponent* SkelMeshComp = Cast<USkeletalMeshComponent>(Component);
+		const USkeletalMeshComponent* SkelMeshComp = Cast<USkeletalMeshComponent>(Component);
 		if ( SkelMeshComp && SkelMeshComp->GetSkeletalMeshAsset())
 		{
 			return true;
diff --git a/Engine/Source/Editor/UnrealEd/Private/ThumbnailRendering/BlueprintThumbnailRenderer.cpp b/Engine/Source/Editor/UnrealEd/Private/ThumbnailRendering/BlueprintThumbnailRenderer.cpp
--- a/Engine/Source/Editor/UnrealEd/Private/ThumbnailRendering/BlueprintThumbnailRenderer.cpp
+++ b/Engine/Source/Editor/UnrealEd/Private/ThumbnailRendering/BlueprintThumbnailRenderer.cpp
@@ -6,8 +6,7 @@
 #include "SceneView.h"
 #include "Misc/App.h"
 
-#include "Engine/InheritableComponentHandler.h"
-#include "Engine/SCS_Node.h"
+#include "Engine/Blueprint.h"
 
 UBlueprintThumbnailRenderer::UBlueprintThumbnailRenderer(const FObjectInitializer& ObjectInitializer)
 	: Super(ObjectInitializer)
@@ -19,68 +18,19 @@ bool UBlueprintThumbnailRenderer::CanVisualizeAsset(UObject* Object)
 {
 	UBlueprint* Blueprint = Cast<UBlueprint>(Object);
 
+	bool bCanVisualizeAsset = false;
+
 	// Only visualize actor based blueprints
 	if (Blueprint && Blueprint->GeneratedClass && Blueprint->GeneratedClass->IsChildOf(AActor::StaticClass()))
 	{
-		// Try to find any visible primitive components in the native class' CDO
-		AActor* CDO = Blueprint->GeneratedClass->GetDefaultObject<AActor>();
-
-		for (UActorComponent* Component : CDO->GetComponents())
-		{
-			if (FBlueprintThumbnailScene::IsValidComponentForVisualization(Component))
-			{
-				return true;
-			}
-		}
-
-		UBlueprint* BlueprintToHarvestComponents = Blueprint;
-		TSet<UBlueprint*> AllVisitedBlueprints;
-		while (BlueprintToHarvestComponents)
+		UBlueprint::ForEachComponentOfActorClassDefault<UActorComponent>(TSubclassOf<AActor>(Blueprint->GeneratedClass), [&](const UActorComponent* Component)
 		{
-			AllVisitedBlueprints.Add(BlueprintToHarvestComponents);
-
-			// Try to find any visible primitive components in the simple construction script
-			if (BlueprintToHarvestComponents->SimpleConstructionScript)
-			{
-				for (USCS_Node* Node : BlueprintToHarvestComponents->SimpleConstructionScript->GetAllNodes())
-				{
-					if (FBlueprintThumbnailScene::IsValidComponentForVisualization(Node->ComponentTemplate))
-					{
-						return true;
-					}
-				}
-			}
-
-			// Check if any inheritable components from parents have valid data
-			if (BlueprintToHarvestComponents->InheritableComponentHandler)
-			{
-				for (TArray<FComponentOverrideRecord>::TIterator InheritedComponentsIter = BlueprintToHarvestComponents->InheritableComponentHandler->CreateRecordIterator(); InheritedComponentsIter; ++InheritedComponentsIter)
-				{
-					if (FBlueprintThumbnailScene::IsValidComponentForVisualization(InheritedComponentsIter->ComponentTemplate))
-					{
-						return true;
-					}
-				}
-			}
-
-			UClass* ParentClass = BlueprintToHarvestComponents->ParentClass;
-			BlueprintToHarvestComponents = nullptr;
-
-			// If the parent class was a blueprint generated class, check it's simple construction script components as well
-			if (ParentClass)
-			{
-				UBlueprint* ParentBlueprint = Cast<UBlueprint>(ParentClass->ClassGeneratedBy);
-
-				// Also make sure we haven't visited the blueprint already. This would only happen if there was a loop of parent classes.
-				if (ParentBlueprint && !AllVisitedBlueprints.Contains(ParentBlueprint))
-				{
-					BlueprintToHarvestComponents = ParentBlueprint;
-				}
-			}
-		}
+			bCanVisualizeAsset = FBlueprintThumbnailScene::IsValidComponentForVisualization(Component);
+			return !bCanVisualizeAsset;
+		});
 	}
 
-	return false;
+	return bCanVisualizeAsset;
 }
 
 void UBlueprintThumbnailRenderer::Draw(UObject* Object, int32 X, int32 Y, uint32 Width, uint32 Height, FRenderTarget* RenderTarget, FCanvas* Canvas, bool bAdditionalViewFamily)
diff --git a/Engine/Source/Editor/UnrealEd/Public/ThumbnailHelpers.h b/Engine/Source/Editor/UnrealEd/Public/ThumbnailHelpers.h
--- a/Engine/Source/Editor/UnrealEd/Public/ThumbnailHelpers.h
+++ b/Engine/Source/Editor/UnrealEd/Public/ThumbnailHelpers.h
@@ -248,7 +248,7 @@ public:
 	FClassActorThumbnailScene();
 
 	/** Returns true if this component can be visualized */
-	static bool IsValidComponentForVisualization(UActorComponent* Component);
+	static bool IsValidComponentForVisualization(const UActorComponent* Component);
 
 protected:
 	// FThumbnailPreviewScene implementation
diff --git a/Engine/Source/Runtime/AIModule/Private/AIHelpers.cpp b/Engine/Source/Runtime/AIModule/Private/AIHelpers.cpp
--- a/Engine/Source/Runtime/AIModule/Private/AIHelpers.cpp
+++ b/Engine/Source/Runtime/AIModule/Private/AIHelpers.cpp
@@ -2,8 +2,7 @@
 
 #include "AIHelpers.h"
 
-#include "Engine/InheritableComponentHandler.h"
-#include "Engine/SCS_Node.h"
+#include "Engine/Blueprint.h"
 
 namespace FAISystem
 {
@@ -118,73 +117,11 @@ TOptional<float> GetYawFromQuaternion(const FQuat& Quaternion)
 
 void GetActorClassDefaultComponents(const TSubclassOf<AActor>& ActorClass, TArray<UActorComponent*>& OutComponents, const TSubclassOf<UActorComponent>& InComponentClass)
 {
-	if (!ensure(ActorClass.Get()))
+	OutComponents.Reset();
+	UBlueprint::ForEachComponentOfActorClassDefault(ActorClass, InComponentClass, [&OutComponents](const UActorComponent* TemplateComponent)
 	{
-		return;
-	}
-
-	UClass* ClassPtr = InComponentClass.Get();
-	TArray<UActorComponent*> ResultComponents;
-
-	// Get the components defined on the native class.
-	AActor* CDO = ActorClass->GetDefaultObject<AActor>();
-	check(CDO);
-	if (ClassPtr)
-	{
-		CDO->GetComponents(InComponentClass, ResultComponents);
-	}
-	else
-	{
-		CDO->GetComponents(ResultComponents);
-	}
-
-	// Try to get the components off the BP class.
-	UBlueprintGeneratedClass* BPClass = Cast<UBlueprintGeneratedClass>(*ActorClass);
-	if (BPClass)
-	{
-		// A BlueprintGeneratedClass has a USimpleConstructionScript member. This member has an array of RootNodes
-		// which contains the SCSNode for the root SceneComponent and non-SceneComponents. For the SceneComponent
-		// hierarchy, each SCSNode knows its children SCSNodes. Each SCSNode stores the component template that will
-		// be created when the Actor is spawned.
-		//
-		// WARNING: This may change in future engine versions!
-
-		TArray<UActorComponent*> Unfiltered;
-		// using this semantic to avoid duplicating following loops or adding a filtering check condition inside the loops
-		TArray<UActorComponent*>& TmpComponents = ClassPtr ? Unfiltered : ResultComponents;
-
-		// Check added components.
-		USimpleConstructionScript* ConstructionScript = BPClass->SimpleConstructionScript;
-		if (ConstructionScript)
-		{
-			for (const USCS_Node* Node : ConstructionScript->GetAllNodes())
-			{
-				TmpComponents.Add(Node->ComponentTemplate);
-			}
-		}
-		// Check modified inherited components.
-		UInheritableComponentHandler* InheritableComponentHandler = BPClass->InheritableComponentHandler;
-		if (InheritableComponentHandler)
-		{
-			for (TArray<FComponentOverrideRecord>::TIterator It = InheritableComponentHandler->CreateRecordIterator(); It; ++It)
-			{
-				TmpComponents.Add(It->ComponentTemplate);
-			}
-		}
-
-		// Filter to the ones matching the requested class.
-		if (ClassPtr)
-		{
-			for (UActorComponent* TemplateComponent : Unfiltered)
-			{
-				if (TemplateComponent->IsA(ClassPtr))
-				{
-					ResultComponents.Add(TemplateComponent);
-				}
-			}
-		}
-	}
-
-	OutComponents = MoveTemp(ResultComponents);
+		OutComponents.Add(const_cast<UActorComponent*>(TemplateComponent));
+		return true;
+	});
 }
 } // UE::AI
\ No newline at end of file
diff --git a/Engine/Source/Runtime/AIModule/Public/AIHelpers.h b/Engine/Source/Runtime/AIModule/Public/AIHelpers.h
--- a/Engine/Source/Runtime/AIModule/Public/AIHelpers.h
+++ b/Engine/Source/Runtime/AIModule/Public/AIHelpers.h
@@ -75,11 +75,6 @@ namespace UE::AI
 	UE_DEPRECATED(5.2, "Deprecated as returning float instead of FVector::FReal. Use FQuat::Rotator()::Yaw etc.")
 	extern AIMODULE_API TOptional<float> GetYawFromQuaternion(const FQuat& Quaternion);
 
-	/**
-	 * Fetches all the components of ActorClass's CDO, including the ones added via the BP editor (which AActor.GetComponents fails to do)
-	 * @param ActorClass class of AActor for which we will retrieve all components
-	 * @param OutComponents this is where the found components will end up. Note that the preexisting contents of OutComponents will get overridden.
-	 * @param InComponentClass if supplied will be used to filter the results
-	 */
+	UE_DEPRECATED(5.3, "Please use UE::BlueprintTools::GetActorClassDefaultComponents() instead")
 	extern AIMODULE_API void GetActorClassDefaultComponents(const TSubclassOf<AActor>& ActorClass, TArray<UActorComponent*>& OutComponents, const TSubclassOf<UActorComponent>& InComponentClass = TSubclassOf<UActorComponent>());
 } // UE::AI
diff --git a/Engine/Source/Runtime/Engine/Classes/Components/ActorComponent.h b/Engine/Source/Runtime/Engine/Classes/Components/ActorComponent.h
--- a/Engine/Source/Runtime/Engine/Classes/Components/ActorComponent.h
+++ b/Engine/Source/Runtime/Engine/Classes/Components/ActorComponent.h
@@ -1137,7 +1137,7 @@ public:
 	/** Override to specify that a component is relevant to the HLOD generation. */
 	virtual bool IsHLODRelevant() const { return false; }
 
-	/** Prefix used to identify template component instances */
+	/** Suffix used to identify template component instances */
 	static const FString ComponentTemplateNameSuffix;
 
 	DECLARE_MULTICAST_DELEGATE_OneParam(FOnMarkRenderStateDirty, UActorComponent&);
diff --git a/Engine/Source/Runtime/Engine/Classes/Engine/Blueprint.h b/Engine/Source/Runtime/Engine/Classes/Engine/Blueprint.h
--- a/Engine/Source/Runtime/Engine/Classes/Engine/Blueprint.h
+++ b/Engine/Source/Runtime/Engine/Classes/Engine/Blueprint.h
@@ -7,7 +7,9 @@
 #include "UObject/Object.h"
 #include "Misc/Guid.h"
 #include "UObject/Class.h"
+#include "Templates/EnableIf.h"
 #include "Templates/SubclassOf.h"
+#include "Templates/UnrealTypeTraits.h"
 #include "Engine/EngineTypes.h"
 #include "EdGraph/EdGraphPin.h"
 #include "Engine/BlueprintCore.h"
@@ -1094,6 +1096,15 @@ public:
 	 */
 	static bool GetBlueprintHierarchyFromClass(const UClass* InClass, TArray<UBlueprintGeneratedClass*>& OutBlueprintParents);
 
+	/**
+	 * Iterate over all BPGCs used to generate this class and its parents, calling InFunc on them. First element is the BPGC used to generate InClass
+	 *
+	 * @param InClass				The class to get the blueprint lineage for
+	 * @param InFunc				Function that will be called for each BPGC. Must return true to continue iteration, or false to stop.
+	 * @return						true if there were no status errors in any of the parent blueprints, otherwise false
+	 */
+	static bool ForEachBlueprintGeneratedClassInHierarchy(const UClass* InClass, TFunctionRef<bool(UBlueprintGeneratedClass*)> InFunc);
+
 private:
 	/**
 	 * Gets an array of all IBlueprintPropertyGuidProviders for this class and its parents.  0th elements is the IBlueprintPropertyGuidProvider for InClass
@@ -1110,6 +1121,85 @@ public:
 	static bool IsBlueprintHierarchyErrorFree(const UClass* InClass);
 #endif
 
+	/**
+	* Fetches all the components of ActorClass's CDO, including the ones added via the BP editor (which AActor.GetComponents fails to do for CDOs).
+	* @param InActorClass		Class of AActor for which we will retrieve all components.
+	* @param InComponentClass	Only retrieve components of this type.
+	* @param OutComponents this is where the found components will end up. Note that the preexisting contents of OutComponents will get overridden.
+	*/
+	static void GetActorClassDefaultComponents(const TSubclassOf<AActor>& InActorClass, const TSubclassOf<UActorComponent>& InComponentClass, TArray<const UActorComponent*>& OutComponents);
+
+	/**
+	* Fetches the first component of ActorClass's CDO which match the requested component class. Will include the components added via the BP editor (which AActor.GetComponents fails to do for CDOs).
+	* @param InActorClass		Class of AActor for which we will retrieve all components.
+	* @param InComponentClass	Only retrieve components of this type.
+	* @param OutComponents this is where the found components will end up. Note that the preexisting contents of OutComponents will get overridden.
+	*/
+	static const UActorComponent* GetActorClassDefaultComponent(const TSubclassOf<AActor>& InActorClass, const TSubclassOf<UActorComponent>& InComponentClass);
+
+	/**
+	* Get the component of ActorClass's CDO that matches the given object name. Will consider all components, including the ones added via the BP editor (which AActor.GetComponents fails to do for CDOs).
+	* @param InActorClass		Class of AActor for which we will search all components.
+	* @param InComponentClass	Only consider components of this type.
+	* @param OutComponents this is where the found components will end up. Note that the preexisting contents of OutComponents will get overridden.
+	*/
+	static const UActorComponent* GetActorClassDefaultComponentByName(const TSubclassOf<AActor>& InActorClass, const TSubclassOf<UActorComponent>& InComponentClass, FName InComponentName);
+
+	/**
+	* Iterate over the components of ActorClass's CDO, including the ones added via the BP editor (which AActor.GetComponents fails to return).
+	* @param InActorClass		Class of AActor for which we will retrieve all components.
+	* @param InComponentClass	Only consider components of this type.
+	* @param InFunc				Code that will be executed for each component. Must return true to continue iteration, or false to stop.
+	*/
+	static void ForEachComponentOfActorClassDefault(const TSubclassOf<AActor>& InActorClass, const TSubclassOf<UActorComponent>& InComponentClass, TFunctionRef<bool(const UActorComponent*)> InFunc);
+
+	/**
+	* Templated version of GetActorClassDefaultComponents()
+	* @see GetActorClassDefaultComponents
+	*/
+	template <typename TComponentClass = UActorComponent, typename = typename TEnableIf<TIsDerivedFrom<TComponentClass, UActorComponent>::IsDerived>::Type>
+	static void GetActorClassDefaultComponents(const TSubclassOf<AActor>& InActorClass, TArray<const TComponentClass*>& OutComponents)
+	{
+		ForEachComponentOfActorClassDefault(InActorClass, TComponentClass::StaticClass(), [&](const UActorComponent* TemplateComponent)
+		{
+			OutComponents.Add(CastChecked<TComponentClass>(TemplateComponent));
+			return true;
+		});
+	}
+
+	/**
+	* Templated version of GetActorClassDefaultComponent()
+	* @see GetActorClassDefaultComponent
+	*/
+	template <typename TComponentClass = UActorComponent, typename = typename TEnableIf<TIsDerivedFrom<TComponentClass, UActorComponent>::IsDerived>::Type>
+	static const TComponentClass* GetActorClassDefaultComponent(const TSubclassOf<AActor>& InActorClass)
+	{
+		return Cast<TComponentClass>(GetActorClassDefaultComponent(InActorClass, TComponentClass::StaticClass()));
+	}
+
+	/**
+	* Templated version of GetActorClassDefaultComponentByName()
+	* @see GetActorClassDefaultComponentByName
+	*/
+	template <typename TComponentClass = UActorComponent, typename = typename TEnableIf<TIsDerivedFrom<TComponentClass, UActorComponent>::IsDerived>::Type>
+	static const TComponentClass* GetActorClassDefaultComponentByName(const TSubclassOf<AActor>& InActorClass, FName InComponentName)
+	{
+		return Cast<TComponentClass>(GetActorClassDefaultComponentByName(InActorClass, TComponentClass::StaticClass(), InComponentName));
+	}
+
+	/**
+	* Templated version of ForEachComponentOfActorClassDefault()
+	* @see ForEachComponentOfActorClassDefault
+	*/
+	template <typename TComponentClass = UActorComponent, typename = typename TEnableIf<TIsDerivedFrom<TComponentClass, UActorComponent>::IsDerived>::Type>
+	static void ForEachComponentOfActorClassDefault(const TSubclassOf<AActor>& InActorClass, TFunctionRef<bool(const TComponentClass*)> InFunc)
+	{
+		ForEachComponentOfActorClassDefault(InActorClass, TComponentClass::StaticClass(), [&](const UActorComponent* TemplateComponent)
+		{
+			return InFunc(CastChecked<TComponentClass>(TemplateComponent));
+		});
+	}
+
 	template<class TFieldType>
 	static FName GetFieldNameFromClassByGuid(const UClass* InClass, const FGuid VarGuid)
 	{
diff --git a/Engine/Source/Runtime/Engine/Private/Blueprint.cpp b/Engine/Source/Runtime/Engine/Private/Blueprint.cpp
--- a/Engine/Source/Runtime/Engine/Private/Blueprint.cpp
+++ b/Engine/Source/Runtime/Engine/Private/Blueprint.cpp
@@ -13,10 +13,12 @@
 #include "Components/TimelineComponent.h"
 #include "Modules/ModuleManager.h"
 #include "UObject/TextProperty.h"
+#include "GameFramework/Actor.h"
+#include "Engine/SimpleConstructionScript.h"
+#include "Engine/SCS_Node.h"
 
 #if WITH_EDITOR
 #include "BlueprintCompilationManager.h"
-#include "Engine/SCS_Node.h"
 #include "Kismet2/BlueprintEditorUtils.h"
 #include "Kismet2/ComponentEditorUtils.h"
 #include "Kismet2/KismetEditorUtilities.h"
@@ -1616,6 +1618,41 @@ bool UBlueprint::GetBlueprintHierarchyFromClass(const UClass* InClass, TArray<UB
 	return bNoErrors;
 }
 
+bool UBlueprint::ForEachBlueprintGeneratedClassInHierarchy(const UClass* InClass, TFunctionRef<bool(UBlueprintGeneratedClass*)> InFunc)
+{
+	bool bNoErrors = true;
+	UBlueprintGeneratedClass* CurrentClass = Cast<UBlueprintGeneratedClass>(const_cast<UClass*>(InClass));
+	while (CurrentClass)
+	{
+		if (!InFunc(CurrentClass))
+		{
+			return bNoErrors;
+		}
+
+#if WITH_EDITORONLY_DATA
+		UBlueprint* BP = UBlueprint::GetBlueprintFromClass(CurrentClass);
+
+		if (BP)
+		{
+			bNoErrors &= (BP->Status != BS_Error);
+		}
+
+		// If valid, use stored ParentClass rather than the actual UClass::GetSuperClass(); handles the case when the class has not been recompiled yet after a reparent operation.
+		if (BP && BP->ParentClass)
+		{
+			CurrentClass = Cast<UBlueprintGeneratedClass>(BP->ParentClass);
+		}
+		else
+#endif // #if WITH_EDITORONLY_DATA
+		{
+			check(CurrentClass);
+			CurrentClass = Cast<UBlueprintGeneratedClass>(CurrentClass->GetSuperClass());
+		}
+	}
+
+	return bNoErrors;
+}
+
 bool UBlueprint::GetBlueprintHierarchyFromClass(const UClass* InClass, TArray<IBlueprintPropertyGuidProvider*>& OutBlueprintParents)
 {
 	OutBlueprintParents.Reset();
@@ -1681,6 +1718,106 @@ bool UBlueprint::IsBlueprintHierarchyErrorFree(const UClass* InClass)
 }
 #endif
 
+void UBlueprint::GetActorClassDefaultComponents(const TSubclassOf<AActor>& InActorClass, const TSubclassOf<UActorComponent>& InComponentClass, TArray<const UActorComponent*>& OutComponents)
+{
+	OutComponents.Reset();
+
+	ForEachComponentOfActorClassDefault(InActorClass, InComponentClass, [&](const UActorComponent* TemplateComponent)
+	{
+		OutComponents.Add(TemplateComponent);
+		return true;
+	});
+}
+
+const UActorComponent* UBlueprint::GetActorClassDefaultComponentByName(const TSubclassOf<AActor>& InActorClass, const TSubclassOf<UActorComponent>& InComponentClass, FName InComponentName)
+{
+	const UActorComponent* Result = nullptr;
+
+	ForEachComponentOfActorClassDefault(InActorClass, InComponentClass, [&Result, InComponentName](const UActorComponent* TemplateComponent)
+	{
+		// Try to strip suffix used to identify template component instances
+		FString StrippedName = TemplateComponent->GetName();
+		if (StrippedName.RemoveFromEnd(UActorComponent::ComponentTemplateNameSuffix))
+		{
+			if (StrippedName == InComponentName.ToString())
+			{
+				Result = TemplateComponent;
+				return false;
+			}
+		}
+		else if (TemplateComponent->GetFName() == InComponentName)
+		{
+			Result = TemplateComponent;
+			return false;
+		}
+
+		return true;
+	});
+
+	return Result;
+}
+
+const UActorComponent* UBlueprint::GetActorClassDefaultComponent(const TSubclassOf<AActor>& InActorClass, const TSubclassOf<UActorComponent>& InComponentClass)
+{
+	const UActorComponent* Result = nullptr;
+
+	ForEachComponentOfActorClassDefault(InActorClass, InComponentClass, [&Result](const UActorComponent* TemplateComponent)
+	{
+		Result = TemplateComponent;
+		return false;
+	});
+
+	return Result;
+}
+
+void UBlueprint::ForEachComponentOfActorClassDefault(const TSubclassOf<AActor>& ActorClass, const TSubclassOf<UActorComponent>& InComponentClass, TFunctionRef<bool(const UActorComponent*)> InFunc)
+{
+	if (!ActorClass.Get())
+	{
+		return;
+	}
+
+	auto FilterFunc = [&](const UActorComponent* TemplateComponent)
+	{
+		if (!InComponentClass.Get() || TemplateComponent->IsA(InComponentClass))
+		{
+			return InFunc(TemplateComponent);
+		}
+
+		return true;
+	};
+
+	// Process native components
+	const AActor* CDO = ActorClass->GetDefaultObject<AActor>();
+	for (const UActorComponent* Component : CDO->GetComponents())
+	{
+		if (!FilterFunc(Component))
+		{
+			return;
+		}
+	}
+
+	// Process blueprint components
+	if (UBlueprintGeneratedClass* ActorBlueprintGeneratedClass = Cast<UBlueprintGeneratedClass>(ActorClass))
+	{
+		ForEachBlueprintGeneratedClassInHierarchy(ActorClass, [&](UBlueprintGeneratedClass* CurrentBPGC)
+		{
+			if (const USimpleConstructionScript* const ConstructionScript = CurrentBPGC->SimpleConstructionScript)
+			{
+				// Gets all BP added components
+				for (const USCS_Node* const Node : ConstructionScript->GetAllNodes())
+				{
+					if (!FilterFunc(Node->GetActualComponentTemplate(ActorBlueprintGeneratedClass)))
+					{
+						return false;
+					}
+				}
+			}
+			return true;
+		});
+	}
+}
+
 FName UBlueprint::FindBlueprintPropertyNameFromGuid(const FGuid& PropertyGuid) const
 {
 #if WITH_EDITORONLY_DATA
diff --git a/Engine/Source/Runtime/Engine/Private/WorldPartition/HLOD/HLODTemplatedInstancedStaticMeshComponent.cpp b/Engine/Source/Runtime/Engine/Private/WorldPartition/HLOD/HLODTemplatedInstancedStaticMeshComponent.cpp
--- a/Engine/Source/Runtime/Engine/Private/WorldPartition/HLOD/HLODTemplatedInstancedStaticMeshComponent.cpp
+++ b/Engine/Source/Runtime/Engine/Private/WorldPartition/HLOD/HLODTemplatedInstancedStaticMeshComponent.cpp
@@ -2,11 +2,9 @@
 
 #include "WorldPartition/HLOD/HLODTemplatedInstancedStaticMeshComponent.h"
 
-#include "GameFramework/Actor.h"
-#include "Engine/BlueprintGeneratedClass.h"
-#include "Engine/InheritableComponentHandler.h"
-#include "Engine/SCS_Node.h"
+#include "Engine/Blueprint.h"
 #include "Engine/StaticMesh.h"
+#include "GameFramework/Actor.h"
 #include "UObject/ObjectSaveContext.h"
 #include "UObject/ObjectSaveOverride.h"
 
@@ -68,93 +66,9 @@ void UHLODTemplatedInstancedStaticMeshComponent::SetTemplateComponentName(const
 	TemplateComponentName = InTemplatePropertyName;
 }
 
-// Note: this is greatly inspired by a similar function in AIHelpers.h/.cpp
-template <class TComp>
-static const TComp* GetActorClassDefaultComponentByName(const TSubclassOf<AActor>& ActorClass, FName InComponentName)
-{
-	auto IsMatchingComponent = [InComponentName](const UActorComponent* InComponent, bool bStripGenSuffix = false)
-	{
-		if (const TComp* TypedComponent = Cast<TComp>(InComponent))
-		{
-			if (bStripGenSuffix)
-			{
-				FString StrippedName = InComponent->GetName();
-				StrippedName.RemoveFromEnd(TEXT("_GEN_VARIABLE"));
-
-				if (StrippedName == InComponentName.ToString())
-				{
-					return true;
-				}
-			}
-			else if (TypedComponent->GetFName() == InComponentName)
-			{
-				return true;
-			}
-		}
-
-		return false;
-	};
-
-	if (!ActorClass.Get())
-	{
-		return nullptr;
-	}
-
-	// Test the components defined on the native class.
-	const AActor* CDO = ActorClass->GetDefaultObject<AActor>();
-	check(CDO);
-
-	for (const UActorComponent* Component : CDO->GetComponents())
-	{
-		if (IsMatchingComponent(Component))
-		{
-			return CastChecked<TComp>(Component);
-		}
-	}
-
-	// Try to get the components off the BP class.
-	UBlueprintGeneratedClass* BPClass = Cast<UBlueprintGeneratedClass>(*ActorClass);
-	if (BPClass)
-	{
-		// A BlueprintGeneratedClass has a USimpleConstructionScript member. This member has an array of RootNodes
-		// which contains the SCSNode for the root SceneComponent and non-SceneComponents. For the SceneComponent
-		// hierarchy, each SCSNode knows its children SCSNodes. Each SCSNode stores the component template that will
-		// be created when the Actor is spawned.
-		//
-		// WARNING: This may change in future engine versions!
-
-		// Check added components.
-		USimpleConstructionScript* ConstructionScript = BPClass->SimpleConstructionScript;
-		if (ConstructionScript)
-		{
-			for (const USCS_Node* Node : ConstructionScript->GetAllNodes())
-			{
-				if (IsMatchingComponent(Node->ComponentTemplate, true))
-				{
-					return CastChecked<TComp>(Node->ComponentTemplate);
-				}
-			}
-		}
-		// Check modified inherited components.
-		UInheritableComponentHandler* InheritableComponentHandler = BPClass->InheritableComponentHandler;
-		if (InheritableComponentHandler)
-		{
-			for (TArray<FComponentOverrideRecord>::TIterator It = InheritableComponentHandler->CreateRecordIterator(); It; ++It)
-			{
-				if (IsMatchingComponent(It->ComponentTemplate, true))
-				{
-					return CastChecked<TComp>(It->ComponentTemplate);
-				}
-			}
-		}
-	}
-
-	return nullptr;
-}
-
 void UHLODTemplatedInstancedStaticMeshComponent::RestoreAssetsFromActorTemplate()
 {
-	const UStaticMeshComponent* TemplateSMC = GetActorClassDefaultComponentByName<UStaticMeshComponent>(TemplateActorClass, TemplateComponentName);
+	const UStaticMeshComponent* TemplateSMC = UBlueprint::GetActorClassDefaultComponentByName<UStaticMeshComponent>(TemplateActorClass, TemplateComponentName);
 	if (TemplateSMC)
 	{
 		// StaticMesh
