World Partition - HLOD: Ensure all scene components are attached to the root

diff --git a/Engine/Plugins/Editor/WorldPartitionHLODUtilities/Source/Private/WorldPartition/HLOD/Utilities/WorldPartitionHLODUtilities.cpp b/Engine/Plugins/Editor/WorldPartitionHLODUtilities/Source/Private/WorldPartition/HLOD/Utilities/WorldPartitionHLODUtilities.cpp
--- a/Engine/Plugins/Editor/WorldPartitionHLODUtilities/Source/Private/WorldPartition/HLOD/Utilities/WorldPartitionHLODUtilities.cpp
+++ b/Engine/Plugins/Editor/WorldPartitionHLODUtilities/Source/Private/WorldPartition/HLOD/Utilities/WorldPartitionHLODUtilities.cpp
@@ -119,7 +119,7 @@ static uint32 ComputeHLODHash(AWorldPartitionHLOD* InHLODActor, const TArray<AAc
 	FArchiveCrc32 Ar;
 
 	// Base key, changing this will force a rebuild of all HLODs
-	FString HLODBaseKey = "A6C1517FC49B478E833A540E2C06289F";
+	FString HLODBaseKey = "BFCB19852F0E4DF3BA4CAE6B85C3078D";
 	Ar << HLODBaseKey;
 
 	// Sub Actors
diff --git a/Engine/Source/Runtime/Engine/Private/WorldPartition/HLOD/HLODActor.cpp b/Engine/Source/Runtime/Engine/Private/WorldPartition/HLOD/HLODActor.cpp
--- a/Engine/Source/Runtime/Engine/Private/WorldPartition/HLOD/HLODActor.cpp
+++ b/Engine/Source/Runtime/Engine/Private/WorldPartition/HLOD/HLODActor.cpp
@@ -242,9 +242,11 @@ void AWorldPartitionHLOD::SetHLODComponents(const TArray<UActorComponent*>& InHL
 		const bool ComponentReplicates = Component->GetIsReplicated();
 		bReplicates |= ComponentReplicates;
 
+		// Avoid using a dummy scene root component (for efficiency), choose one component as the root
 		if (USceneComponent* SceneComponent = Cast<USceneComponent>(Component))
 		{
-			// Prefer a replicated root component
+			// If we have one, prefer a replicated component as our root.
+			// This is required, otherwise the actor won't even be considered for replication
 			if (!RootComponent || (!RootComponent->GetIsReplicated() && ComponentReplicates))
 			{
 				RootComponent = SceneComponent;
@@ -253,6 +255,18 @@ void AWorldPartitionHLOD::SetHLODComponents(const TArray<UActorComponent*>& InHL
 	
 		Component->RegisterComponent();
 	}
+
+	// Attach all scene components to our root.
+	const bool bIncludeFromChildActors = false;
+	ForEachComponent<USceneComponent>(bIncludeFromChildActors, [&](USceneComponent* Component)
+	{
+		// Skip the root component
+		if (Component != GetRootComponent())
+		{
+			// Keep world transform intact while attaching to root component
+			Component->AttachToComponent(GetRootComponent(), FAttachmentTransformRules::KeepWorldTransform);
+		}
+	});
 }
 
 void AWorldPartitionHLOD::SetSubActors(const TArray<FHLODSubActor>& InHLODSubActors)
