Don't include package dirty state resets in the undo buffer during Blueprint Actor instance reconstruction events that occur within a transactional scope.

diff --git a/Engine/Source/Runtime/Engine/Private/ActorConstruction.cpp b/Engine/Source/Runtime/Engine/Private/ActorConstruction.cpp
--- a/Engine/Source/Runtime/Engine/Private/ActorConstruction.cpp
+++ b/Engine/Source/Runtime/Engine/Private/ActorConstruction.cpp
@@ -282,18 +282,11 @@ void AActor::RerunConstructionScripts()
 
 		// Mark package as clean on exit if not transient and not already dirty.
 		bool bMarkPackageClean = false;
-		const UPackage* ActorPackage = GetPackage();
+		UPackage* ActorPackage = GetPackage();
 		if (ActorPackage && !ActorPackage->IsDirty() && ActorPackage != GetTransientPackage())
 		{
 			bMarkPackageClean = true;
 		}
-		ON_SCOPE_EXIT
-		{
-			if (bMarkPackageClean)
-			{
-				GetPackage()->SetDirtyFlag(false);
-			}
-		};
 		
 		// Create cache to store component data across rerunning construction scripts
 		FComponentInstanceDataCache* InstanceDataCache;
@@ -612,6 +605,13 @@ void AActor::RerunConstructionScripts()
 			}
 		}
 
+		// If any of the code above caused the package dirty state to change, we reset it back to a "clean" state now. Note that we have to do this
+		// before we restore the undo buffer below - otherwise, the state change will become part of the current transaction, and we don't want that.
+		if (bMarkPackageClean && ActorPackage->IsDirty())
+		{
+			ActorPackage->SetDirtyFlag(false);
+		}
+
 		// Restore the undo buffer
 		GUndo = CurrentTransaction;
 
