ISMDescriptor - Added missing MaterialOverlay property to the descriptor

diff --git a/Engine/Source/Runtime/Engine/Private/ISMPartition/ISMComponentDescriptor.cpp b/Engine/Source/Runtime/Engine/Private/ISMPartition/ISMComponentDescriptor.cpp
--- a/Engine/Source/Runtime/Engine/Private/ISMPartition/ISMComponentDescriptor.cpp
+++ b/Engine/Source/Runtime/Engine/Private/ISMPartition/ISMComponentDescriptor.cpp
@@ -33,6 +33,7 @@ FISMComponentDescriptor::FISMComponentDescriptor(const FSoftISMComponentDescript
 {
 	StaticMesh = Other.StaticMesh.LoadSynchronous();
 	Algo::Transform(Other.OverrideMaterials, OverrideMaterials, [](TSoftObjectPtr<UMaterialInterface> Material) { return Material.LoadSynchronous(); });
+	OverlayMaterial = Other.OverlayMaterial.LoadSynchronous();
 	Algo::Transform(Other.RuntimeVirtualTextures, RuntimeVirtualTextures, [](TSoftObjectPtr<URuntimeVirtualTexture> RVT) { return RVT.LoadSynchronous(); });
 	Hash = Other.Hash;
 }
@@ -49,6 +50,7 @@ FSoftISMComponentDescriptor::FSoftISMComponentDescriptor(const FISMComponentDesc
 {
 	StaticMesh = Other.StaticMesh;
 	Algo::Transform(Other.OverrideMaterials, OverrideMaterials, [](TObjectPtr<UMaterialInterface> Material) { return Material; });
+	OverlayMaterial = Other.OverlayMaterial;
 	Algo::Transform(Other.RuntimeVirtualTextures, RuntimeVirtualTextures, [](TObjectPtr<URuntimeVirtualTexture> RVT) { return RVT; });
 	Hash = Other.Hash;
 }
@@ -131,6 +133,7 @@ void FISMComponentDescriptor::InitFrom(const UStaticMeshComponent* Template, boo
 {
 	StaticMesh = Template->GetStaticMesh();
 	OverrideMaterials = Template->OverrideMaterials;
+	OverlayMaterial = Template->OverlayMaterial;
 	RuntimeVirtualTextures = Template->RuntimeVirtualTextures;
 
 	Super::InitFrom(Template, bInitBodyInstance);
@@ -140,6 +143,7 @@ void FSoftISMComponentDescriptor::InitFrom(const UStaticMeshComponent* Template,
 {
 	StaticMesh = Template->GetStaticMesh();
 	Algo::Transform(Template->OverrideMaterials, OverrideMaterials, [](TObjectPtr<UMaterialInterface> Material) { return Material; });
+	OverlayMaterial = Template->OverlayMaterial;
 	Algo::Transform(Template->RuntimeVirtualTextures, RuntimeVirtualTextures, [](TObjectPtr<URuntimeVirtualTexture> RVT) { return RVT; });
 
 	Super::InitFrom(Template, bInitBodyInstance);
@@ -214,6 +218,7 @@ bool FISMComponentDescriptor::operator==(const FISMComponentDescriptor& Other) c
 	return (Hash == 0 || Other.Hash == 0 || Hash == Other.Hash) && // Check hash first, other checks are in case of Hash collision
 		StaticMesh == Other.StaticMesh &&
 		OverrideMaterials == Other.OverrideMaterials &&
+		OverlayMaterial == Other.OverlayMaterial &&
 		RuntimeVirtualTextures == Other.RuntimeVirtualTextures &&
 		Super::operator==(Other);
 }
@@ -223,6 +228,7 @@ bool FSoftISMComponentDescriptor::operator==(const FSoftISMComponentDescriptor&
 	return (Hash == 0 || Other.Hash == 0 || Hash == Other.Hash) && // Check hash first, other checks are in case of Hash collision
 		StaticMesh == Other.StaticMesh &&
 		OverrideMaterials == Other.OverrideMaterials &&
+		OverlayMaterial == Other.OverlayMaterial &&
 		RuntimeVirtualTextures == Other.RuntimeVirtualTextures &&
 		Super::operator==(Other);
 }
@@ -265,7 +271,7 @@ UInstancedStaticMeshComponent* FSoftISMComponentDescriptor::CreateComponent(UObj
 	InitComponent(ISMComponent);
 
 	return ISMComponent;
-	}
+}
 
 void FISMComponentDescriptorBase::InitComponent(UInstancedStaticMeshComponent* ISMComponent) const
 {
@@ -323,25 +329,24 @@ void FISMComponentDescriptor::InitComponent(UInstancedStaticMeshComponent* ISMCo
 {
 	ISMComponent->SetStaticMesh(StaticMesh);
 
-	ISMComponent->OverrideMaterials.Empty(OverrideMaterials.Num());
-	for (UMaterialInterface* OverrideMaterial : OverrideMaterials)
+	auto GetMaterial = [ISMComponent](UMaterialInterface* MaterialInterface)
 	{
-		if (OverrideMaterial && !OverrideMaterial->IsAsset())
+		if (MaterialInterface && !MaterialInterface->IsAsset())
 		{
 			// If the material is equivalent to its parent, just take a reference to its parent rather than making another redundant object 
-			if (UMaterialInstance* Instance = Cast<UMaterialInstance>(OverrideMaterial); Instance && Instance->IsRedundant())
+			if (UMaterialInstance* Instance = Cast<UMaterialInstance>(MaterialInterface); Instance && Instance->IsRedundant())
 			{
-				OverrideMaterial = Instance->Parent;
+				MaterialInterface = Instance->Parent;
 			}
 			else
 			{
 				// As override materials are normally outered to their owner component, we need to duplicate them here to make sure we don't create
 				// references to actors in other levels (for packed level instances or HLOD actors).
-				OverrideMaterial = DuplicateObject<UMaterialInterface>(OverrideMaterial, ISMComponent);
+				MaterialInterface = DuplicateObject<UMaterialInterface>(MaterialInterface, ISMComponent);
 
 				// If the MID we just duplicated has a nanite override that's also not an asset, duplicate that too
-				UMaterialInstanceDynamic* OverrideMID = Cast<UMaterialInstanceDynamic>(OverrideMaterial);
-				UMaterialInterface* NaniteOverride = OverrideMID ? OverrideMID->GetNaniteOverride() : nullptr; 
+				UMaterialInstanceDynamic* OverrideMID = Cast<UMaterialInstanceDynamic>(MaterialInterface);
+				UMaterialInterface* NaniteOverride = OverrideMID ? OverrideMID->GetNaniteOverride() : nullptr;
 				if (NaniteOverride && !NaniteOverride->IsAsset())
 				{
 					OverrideMID->SetNaniteOverride(DuplicateObject<UMaterialInterface>(NaniteOverride, ISMComponent));
@@ -349,9 +354,15 @@ void FISMComponentDescriptor::InitComponent(UInstancedStaticMeshComponent* ISMCo
 			}
 		}
 
-		ISMComponent->OverrideMaterials.Add(OverrideMaterial);
-	}
+		return MaterialInterface;
+	};
 
+	ISMComponent->OverrideMaterials.Empty(OverrideMaterials.Num());
+	for (UMaterialInterface* OverrideMaterial : OverrideMaterials)
+	{
+		ISMComponent->OverrideMaterials.Add(GetMaterial(OverrideMaterial));
+	}
+	ISMComponent->OverlayMaterial = GetMaterial(OverlayMaterial);
 	ISMComponent->RuntimeVirtualTextures = RuntimeVirtualTextures;
 
 	Super::InitComponent(ISMComponent);
@@ -361,25 +372,24 @@ void FSoftISMComponentDescriptor::InitComponent(UInstancedStaticMeshComponent* I
 {
 	ISMComponent->SetStaticMesh(StaticMesh.LoadSynchronous());
 
-	ISMComponent->OverrideMaterials.Empty(OverrideMaterials.Num());
-	for (const TSoftObjectPtr<UMaterialInterface>& OverrideMaterialPtr : OverrideMaterials)
+	auto GetMaterial = [ISMComponent](const TSoftObjectPtr<UMaterialInterface>& MaterialInterfacePtr)
 	{
-		UMaterialInterface* OverrideMaterial = OverrideMaterialPtr.LoadSynchronous();
-		if (OverrideMaterial && !OverrideMaterial->IsAsset())
+		UMaterialInterface* MaterialInterface = MaterialInterfacePtr.LoadSynchronous();
+		if (MaterialInterface && !MaterialInterface->IsAsset())
 		{
 			// If the material is equivalent to its parent, just take a reference to its parent rather than making another redundant object 
-			if (UMaterialInstance* Instance = Cast<UMaterialInstance>(OverrideMaterial); Instance && Instance->IsRedundant())
+			if (UMaterialInstance* Instance = Cast<UMaterialInstance>(MaterialInterface); Instance && Instance->IsRedundant())
 			{
-				OverrideMaterial = Instance->Parent;
+				MaterialInterface = Instance->Parent;
 			}
 			else
 			{
 				// As override materials are normally outered to their owner component, we need to duplicate them here to make sure we don't create
 				// references to actors in other levels (for packed level instances or HLOD actors).
-				OverrideMaterial = DuplicateObject<UMaterialInterface>(OverrideMaterial, ISMComponent);
+				MaterialInterface = DuplicateObject<UMaterialInterface>(MaterialInterface, ISMComponent);
 
 				// If the MID we just duplicated has a nanite override that's also not an asset, duplicate that too
-				UMaterialInstanceDynamic* OverrideMID = Cast<UMaterialInstanceDynamic>(OverrideMaterial);
+				UMaterialInstanceDynamic* OverrideMID = Cast<UMaterialInstanceDynamic>(MaterialInterface);
 				UMaterialInterface* NaniteOverride = OverrideMID ? OverrideMID->GetNaniteOverride() : nullptr; 
 				if (NaniteOverride && !NaniteOverride->IsAsset())
 				{
@@ -387,9 +397,16 @@ void FSoftISMComponentDescriptor::InitComponent(UInstancedStaticMeshComponent* I
 				}
 			}
 		}
+		
+		return MaterialInterface;
+	};
 
-		ISMComponent->OverrideMaterials.Add(OverrideMaterial);
+	ISMComponent->OverrideMaterials.Empty(OverrideMaterials.Num());
+	for (const TSoftObjectPtr<UMaterialInterface>& OverrideMaterialPtr : OverrideMaterials)
+	{
+		ISMComponent->OverrideMaterials.Add(GetMaterial(OverrideMaterialPtr));
 	}
+	ISMComponent->OverlayMaterial = GetMaterial(OverlayMaterial);
 
 	ISMComponent->RuntimeVirtualTextures.Empty(RuntimeVirtualTextures.Num());
 	for (const TSoftObjectPtr<URuntimeVirtualTexture>& RuntimeVirtualTexturePtr : RuntimeVirtualTextures)
@@ -401,4 +418,4 @@ void FSoftISMComponentDescriptor::InitComponent(UInstancedStaticMeshComponent* I
 	}
 
 	Super::InitComponent(ISMComponent);
-}
\ No newline at end of file
+}
diff --git a/Engine/Source/Runtime/Engine/Public/ISMPartition/ISMComponentDescriptor.h b/Engine/Source/Runtime/Engine/Public/ISMPartition/ISMComponentDescriptor.h
--- a/Engine/Source/Runtime/Engine/Public/ISMPartition/ISMComponentDescriptor.h
+++ b/Engine/Source/Runtime/Engine/Public/ISMPartition/ISMComponentDescriptor.h
@@ -202,6 +202,9 @@ public:
 	UPROPERTY(EditAnywhere, Category = "Component Settings", meta = (DisplayAfter = "ComponentClass"))
 	TArray<TObjectPtr<UMaterialInterface>> OverrideMaterials;
 
+	UPROPERTY(EditAnywhere, Category = "Component Settings", meta = (DisplayAfter = "ComponentClass"))
+	TObjectPtr<UMaterialInterface> OverlayMaterial;
+
 	UPROPERTY(EditAnywhere, Category = "Component Settings", meta = (DisplayAfter = "ComponentClass"))
 	TArray<TObjectPtr<URuntimeVirtualTexture>> RuntimeVirtualTextures;
 };
@@ -244,6 +247,9 @@ public:
 	UPROPERTY(EditAnywhere, Category = "Component Settings", meta = (DisplayAfter = "ComponentClass"))
 	TArray<TSoftObjectPtr<UMaterialInterface>> OverrideMaterials;
 
+	UPROPERTY(EditAnywhere, Category = "Component Settings", meta = (DisplayAfter = "ComponentClass"))
+	TSoftObjectPtr<UMaterialInterface> OverlayMaterial;
+
 	UPROPERTY(EditAnywhere, Category = "Component Settings", meta = (DisplayAfter = "ComponentClass"))
 	TArray<TSoftObjectPtr<URuntimeVirtualTexture>> RuntimeVirtualTextures;
-};
\ No newline at end of file
+};
