World Partition Streaming Generation Unit Test:
- Added a new unit test that will generate streaming on a test map and compare the generation log with the reference one, and bring the system's diff tool when they don't match.
- When making changes to streaming generation that affects the log, the reference output log will need to be submitted in order to maintan the unit test validation.

diff --git a/Engine/Source/Runtime/Engine/Private/Tests/WorldPartition/WorldPartitionPathConversionTests.cpp b/Engine/Source/Runtime/Engine/Private/Tests/WorldPartition/WorldPartitionPathConversionTests.cpp
new file mode 100644
--- /dev/null
+++ b/Engine/Source/Runtime/Engine/Private/Tests/WorldPartition/WorldPartitionPathConversionTests.cpp
@@ -0,0 +1,169 @@
+// Copyright Epic Games, Inc. All Rights Reserved.
+
+#include "Misc/AutomationTest.h"
+
+#if WITH_EDITOR
+#include "WorldPartition/WorldPartition.h"
+#include "WorldPartition/WorldPartitionHelpers.h"
+#include "EditorWorldUtils.h"
+#endif
+
+#if WITH_DEV_AUTOMATION_TESTS
+
+#define TEST_NAME_ROOT "System.Engine.WorldPartition"
+
+namespace WorldPartitionTests
+{
+	IMPLEMENT_SIMPLE_AUTOMATION_TEST(FWorldPartitionPathConversionTest, TEST_NAME_ROOT ".PathConversion", EAutomationTestFlags::ApplicationContextMask | EAutomationTestFlags::EngineFilter)
+
+	bool FWorldPartitionPathConversionTest::RunTest(const FString& Parameters)
+	{
+#if WITH_EDITOR
+		ON_SCOPE_EXIT
+		{
+			CollectGarbage(GARBAGE_COLLECTION_KEEPFLAGS);
+		};
+
+		FScopedEditorWorld ScopedEditorWorld(
+			TEXTVIEW("/Engine/WorldPartition/WorldPartitionUnitTest"),
+			UWorld::InitializationValues()
+				.RequiresHitProxies(false)
+				.ShouldSimulatePhysics(false)
+				.EnableTraceCollision(false)
+				.CreateNavigation(false)
+				.CreateAISystem(false)
+				.AllowAudioPlayback(false)
+				.CreatePhysicsScene(true)
+		);
+		
+		UWorld* World = ScopedEditorWorld.GetWorld();
+		if (!TestNotNull(TEXT("Missing World Object"), World))
+		{
+			return false;
+		}
+
+		check(World);
+		UWorldPartition* WorldPartition = World->GetWorldPartition();
+		if (!TestNotNull(TEXT("Missing World Partition Object"), WorldPartition))
+		{
+			return false;
+		}
+		
+		check(WorldPartition);
+		UActorDescContainer* ActorDescMainContainer = WorldPartition->GetActorDescContainer();
+		if (!TestNotNull(TEXT("Missing World Partition Container"), ActorDescMainContainer))
+		{
+			return false;
+		}
+
+		UWorldPartition::FGenerateStreamingParams Params = UWorldPartition::FGenerateStreamingParams().SetActorDescContainer(ActorDescMainContainer);
+		UWorldPartition::FGenerateStreamingContext Context;
+
+		if (!TestTrue(TEXT("World Partition Generate Streaming"), WorldPartition->GenerateContainerStreaming(Params, Context)))
+		{
+			return false;
+		}
+
+		FWorldPartitionReference ActorRef(ActorDescMainContainer, FGuid(TEXT("5D9F93BA407A811AFDDDAAB4F1CECC6A")));
+		if (!TestTrue(TEXT("Invalid Actor Reference"), ActorRef.IsValid()))
+		{
+			return false;
+		}
+
+		AActor* Actor = ActorRef->GetActor();
+		if (!TestNotNull(TEXT("Missing Actor"), Actor))
+		{
+			return false;
+		}
+
+		FWorldPartitionHandle ActorHandle(ActorDescMainContainer, FGuid(TEXT("0D2B04D240BE5DE58FE437A8D2DBF5C9")));
+		if (!TestTrue(TEXT("Invalid Actor Handle"), ActorHandle.IsValid()))
+		{
+			return false;
+		}
+
+		if (!TestNull(TEXT("Actor Handle Not Loaded"), ActorHandle->GetActor()))
+		{
+			return false;
+		}
+
+		UObject* ResolvedObject = ActorHandle->GetActorSoftPath().TryLoad();
+		if (!TestNotNull(TEXT("Actor Handle Loaded"), ActorHandle->GetActor()))
+		{
+			return false;
+		}
+
+		if (!TestTrue(TEXT("Resolving Runtime Actor From Editor Path Failed"), ResolvedObject == ActorHandle->GetActor()))
+		{
+			return false;
+		}
+
+		auto TestEditorRuntimePathConversion = [this](const FSoftObjectPath& InEditorPath) -> bool
+		{
+			FSoftObjectPath RuntimePath;
+			if (!TestTrue(TEXT("Path Editor to Runtime Conversion Failed"), FWorldPartitionHelpers::ConvertEditorPathToRuntimePath(InEditorPath, RuntimePath)))
+			{
+				return false;
+			}
+
+			FSoftObjectPath EditorPath;
+			if (!TestTrue(TEXT("Path Runtime to Editor Conversion Failed"), FWorldPartitionHelpers::ConvertRuntimePathToEditorPath(RuntimePath, EditorPath)))
+			{
+				return false;
+			}
+
+			if (!TestTrue(TEXT("Path Editor to Runtime to Editor Conversion Failed"), EditorPath == InEditorPath))
+			{
+				return false;
+			}
+
+			return true;
+		};
+
+		// Test path conversions to an actor
+		if (!TestEditorRuntimePathConversion(FSoftObjectPath(Actor)))
+		{
+			return false;
+		}
+
+		// Test path conversions to an actor component
+		if (!TestEditorRuntimePathConversion(FSoftObjectPath(Actor->GetRootComponent())))
+		{
+			return false;
+		}
+
+		// Test path conversions to an instanced level
+		UPackage* WorldPackage = World->GetPackage();
+		WorldPackage->Rename(*FString(WorldPackage->GetName() + TEXT("_LevelInstance_1")));
+
+		// Test path conversions to an instanced actor
+		if (!TestEditorRuntimePathConversion(FSoftObjectPath(Actor)))
+		{
+			return false;
+		}
+
+		// Test path conversions to an instanced actor component
+		if (!TestEditorRuntimePathConversion(FSoftObjectPath(Actor->GetRootComponent())))
+		{
+			return false;
+		}
+
+		// Test path conversions to a newly spawned actor
+		AActor* NewActor = World->SpawnActor(AActor::StaticClass());
+		if (!TestNotNull(TEXT("Sapwning Actor Failed"), NewActor))
+		{
+			return false;
+		}
+
+		if (!TestEditorRuntimePathConversion(FSoftObjectPath(NewActor)))
+		{
+			return false;
+		}
+#endif
+		return true;
+	}
+}
+
+#undef TEST_NAME_ROOT
+
+#endif 
diff --git a/Engine/Source/Runtime/Engine/Private/Tests/WorldPartition/WorldPartitionStreamingGenerationTests.cpp b/Engine/Source/Runtime/Engine/Private/Tests/WorldPartition/WorldPartitionStreamingGenerationTests.cpp
--- a/Engine/Source/Runtime/Engine/Private/Tests/WorldPartition/WorldPartitionStreamingGenerationTests.cpp
+++ b/Engine/Source/Runtime/Engine/Private/Tests/WorldPartition/WorldPartitionStreamingGenerationTests.cpp
@@ -6,6 +6,10 @@
 #include "WorldPartition/WorldPartition.h"
 #include "WorldPartition/WorldPartitionHelpers.h"
 #include "EditorWorldUtils.h"
+#include "Misc/CommandLine.h"
+#include "Misc/FileHelper.h"
+#include "Settings/EditorLoadingSavingSettings.h"
+#include "AssetToolsModule.h"
 #endif
 
 #if WITH_DEV_AUTOMATION_TESTS
@@ -25,7 +29,7 @@ namespace WorldPartitionTests
 		};
 
 		FScopedEditorWorld ScopedEditorWorld(
-			TEXTVIEW("/Engine/WorldPartition/WorldPartitionUnitTest"),
+			TEXTVIEW("/Engine/WorldPartition/UnitTests/WPUnitTests"),
 			UWorld::InitializationValues()
 				.RequiresHitProxies(false)
 				.ShouldSimulatePhysics(false)
@@ -48,112 +52,42 @@ namespace WorldPartitionTests
 		{
 			return false;
 		}
-		
-		check(WorldPartition);
-		UActorDescContainer* ActorDescMainContainer = WorldPartition->GetActorDescContainer();
-		if (!TestNotNull(TEXT("Missing World Partition Container"), ActorDescMainContainer))
-		{
-			return false;
-		}
 
-		if (!TestTrue(TEXT("World Partition Generate Streaming"), WorldPartition->GenerateContainerStreaming(ActorDescMainContainer)))
-		{
-			return false;
-		}
+		UWorldPartition::FGenerateStreamingParams Params = UWorldPartition::FGenerateStreamingParams()
+			.SetOutputLogPath(TEXT("StreamingGenerationTest"));
 
-		FWorldPartitionReference ActorRef(ActorDescMainContainer, FGuid(TEXT("5D9F93BA407A811AFDDDAAB4F1CECC6A")));
-		if (!TestTrue(TEXT("Invalid Actor Reference"), ActorRef.IsValid()))
-		{
-			return false;
-		}
+		UWorldPartition::FGenerateStreamingContext Context;
+		WorldPartition->GenerateStreaming(Params, Context);
 
-		AActor* Actor = ActorRef->GetActor();
-		if (!TestNotNull(TEXT("Missing Actor"), Actor))
+		if (!TestTrue(TEXT("Missing Generated Output Log Path"), Context.OutputLogFilename.IsSet()))
 		{
 			return false;
 		}
 
-		FWorldPartitionHandle ActorHandle(ActorDescMainContainer, FGuid(TEXT("0D2B04D240BE5DE58FE437A8D2DBF5C9")));
-		if (!TestTrue(TEXT("Invalid Actor Handle"), ActorHandle.IsValid()))
+		// Read reference output log
+		FString ReferenceOuputLog;
+		const FString ReferenceOutputLogPath(FPaths::GetPath(World->GetPackage()->GetLoadedPath().GetLocalFullPath()) / FPaths::GetCleanFilename(*Context.OutputLogFilename));
+		if (!TestTrue(TEXT("Error Reading Reference Output Log File"), FFileHelper::LoadFileToString(ReferenceOuputLog, *ReferenceOutputLogPath)))
 		{
 			return false;
 		}
 
-		if (!TestNull(TEXT("Actor Handle Not Loaded"), ActorHandle->GetActor()))
+		// Read generated output log
+		FString GeneratedOuputLog;
+		const FString GeneratedOutputLogPath(*Context.OutputLogFilename);
+		if (!TestTrue(TEXT("Error Reading Generated Output Log File"), FFileHelper::LoadFileToString(GeneratedOuputLog, *GeneratedOutputLogPath)))
 		{
 			return false;
 		}
 
-		UObject* ResolvedObject = ActorHandle->GetActorSoftPath().TryLoad();
-		if (!TestNotNull(TEXT("Actor Handle Loaded"), ActorHandle->GetActor()))
+		if (!TestTrue(TEXT("Generated Output Log Don't Match Reference Output Log"), GeneratedOuputLog == ReferenceOuputLog))
 		{
-			return false;
-		}
-
-		if (!TestTrue(TEXT("Resolving Runtime Actor From Editor Path Failed"), ResolvedObject == ActorHandle->GetActor()))
-		{
-			return false;
-		}
-
-		auto TestEditorRuntimePathConversion = [this](const FSoftObjectPath& InEditorPath) -> bool
-		{
-			FSoftObjectPath RuntimePath;
-			if (!TestTrue(TEXT("Path Editor to Runtime Conversion Failed"), FWorldPartitionHelpers::ConvertEditorPathToRuntimePath(InEditorPath, RuntimePath)))
-			{
-				return false;
-			}
-
-			FSoftObjectPath EditorPath;
-			if (!TestTrue(TEXT("Path Runtime to Editor Conversion Failed"), FWorldPartitionHelpers::ConvertRuntimePathToEditorPath(RuntimePath, EditorPath)))
-			{
-				return false;
-			}
-
-			if (!TestTrue(TEXT("Path Editor to Runtime to Editor Conversion Failed"), EditorPath == InEditorPath))
+			if (!FParse::Param(FCommandLine::Get(), TEXT("unattended")))
 			{
-				return false;
+				FString DiffCommand = GetDefault<UEditorLoadingSavingSettings>()->TextDiffToolPath.FilePath;
+				FAssetToolsModule& AssetToolsModule = FModuleManager::Get().LoadModuleChecked<FAssetToolsModule>("AssetTools");
+				AssetToolsModule.Get().CreateDiffProcess(DiffCommand, ReferenceOutputLogPath, GeneratedOutputLogPath);
 			}
-
-			return true;
-		};
-
-		// Test path conversions to an actor
-		if (!TestEditorRuntimePathConversion(FSoftObjectPath(Actor)))
-		{
-			return false;
-		}
-
-		// Test path conversions to an actor component
-		if (!TestEditorRuntimePathConversion(FSoftObjectPath(Actor->GetRootComponent())))
-		{
-			return false;
-		}
-
-		// Test path conversions to an instanced level
-		UPackage* WorldPackage = World->GetPackage();
-		WorldPackage->Rename(*FString(WorldPackage->GetName() + TEXT("_LevelInstance_1")));
-
-		// Test path conversions to an instanced actor
-		if (!TestEditorRuntimePathConversion(FSoftObjectPath(Actor)))
-		{
-			return false;
-		}
-
-		// Test path conversions to an instanced actor component
-		if (!TestEditorRuntimePathConversion(FSoftObjectPath(Actor->GetRootComponent())))
-		{
-			return false;
-		}
-
-		// Test path conversions to a newly spawned actor
-		AActor* NewActor = World->SpawnActor(AActor::StaticClass());
-		if (!TestNotNull(TEXT("Sapwning Actor Failed"), NewActor))
-		{
-			return false;
-		}
-
-		if (!TestEditorRuntimePathConversion(FSoftObjectPath(NewActor)))
-		{
 			return false;
 		}
 #endif
@@ -163,4 +97,4 @@ namespace WorldPartitionTests
 
 #undef TEST_NAME_ROOT
 
-#endif 
+#endif
\ No newline at end of file
diff --git a/Engine/Source/Runtime/Engine/Private/WorldPartition/ContentBundle/ContentBundleEditor.cpp b/Engine/Source/Runtime/Engine/Private/WorldPartition/ContentBundle/ContentBundleEditor.cpp
--- a/Engine/Source/Runtime/Engine/Private/WorldPartition/ContentBundle/ContentBundleEditor.cpp
+++ b/Engine/Source/Runtime/Engine/Private/WorldPartition/ContentBundle/ContentBundleEditor.cpp
@@ -313,8 +313,14 @@ void FContentBundleEditor::GenerateStreaming(TArray<FString>* OutPackageToGenera
 		return;
 	}
 
+	UWorldPartition::FGenerateStreamingParams Params = UWorldPartition::FGenerateStreamingParams()
+		.SetActorDescContainer(ActorDescContainer.Get());
+
+	UWorldPartition::FGenerateStreamingContext Context = UWorldPartition::FGenerateStreamingContext()
+		.SetPackagesToGenerate(OutPackageToGenerate);
+
 	UWorldPartition* WorldPartition = GetInjectedWorld()->GetWorldPartition();
-	WorldPartition->GenerateContainerStreaming(ActorDescContainer.Get(), OutPackageToGenerate);
+	WorldPartition->GenerateContainerStreaming(Params, Context);
 
 	ExternalStreamingObject = WorldPartition->RuntimeHash->StoreToExternalStreamingObject(WorldPartition, *GetExternalStreamingObjectName());
 
diff --git a/Engine/Source/Runtime/Engine/Private/WorldPartition/WorldPartition.cpp b/Engine/Source/Runtime/Engine/Private/WorldPartition/WorldPartition.cpp
--- a/Engine/Source/Runtime/Engine/Private/WorldPartition/WorldPartition.cpp
+++ b/Engine/Source/Runtime/Engine/Private/WorldPartition/WorldPartition.cpp
@@ -378,8 +378,13 @@ void UWorldPartition::OnPrePIEEnded(bool bWasSimulatingInEditor)
 
 void UWorldPartition::OnBeginPlay()
 {
+	FGenerateStreamingParams Params;
+
 	TArray<FString> OutGeneratedStreamingPackageNames;
-	GenerateStreaming((bIsPIE || IsRunningGame()) ? &OutGeneratedStreamingPackageNames : nullptr);
+	FGenerateStreamingContext Context = FGenerateStreamingContext()
+		.SetPackagesToGenerate((bIsPIE || IsRunningGame()) ? &OutGeneratedStreamingPackageNames : nullptr);
+
+	GenerateStreaming(Params, Context);
 
 	// Prepare GeneratedStreamingPackages
 	check(GeneratedStreamingPackageNames.IsEmpty());
@@ -1600,8 +1605,14 @@ void UWorldPartition::BeginCook(IWorldPartitionCookPackageContext& CookContext)
 
 bool UWorldPartition::GatherPackagesToCook(IWorldPartitionCookPackageContext& CookContext)
 {
+	FGenerateStreamingParams Params = FGenerateStreamingParams()
+		.SetActorDescContainer(ActorDescContainer);
+
 	TArray<FString> PackagesToCook;
-	if (GenerateContainerStreaming(ActorDescContainer, &PackagesToCook))
+	FGenerateStreamingContext Context = FGenerateStreamingContext()
+		.SetPackagesToGenerate(&PackagesToCook);
+
+	if (GenerateContainerStreaming(Params, Context))
 	{
 		FString PackageName = GetPackage()->GetName();
 		for (const FString& PackageToCook : PackagesToCook)
diff --git a/Engine/Source/Runtime/Engine/Private/WorldPartition/WorldPartitionRuntimeHash.cpp b/Engine/Source/Runtime/Engine/Private/WorldPartition/WorldPartitionRuntimeHash.cpp
--- a/Engine/Source/Runtime/Engine/Private/WorldPartition/WorldPartitionRuntimeHash.cpp
+++ b/Engine/Source/Runtime/Engine/Private/WorldPartition/WorldPartitionRuntimeHash.cpp
@@ -330,31 +330,34 @@ bool UWorldPartitionRuntimeHash::PopulateGeneratorPackageForCook(const TArray<FW
 
 void UWorldPartitionRuntimeHash::DumpStateLog(FHierarchicalLogArchive& Ar)
 {
-	Ar.Printf(TEXT("----------------------------------------------------------------------------------------------------------------"));
-	Ar.Printf(TEXT("%s - Persistent Level"), *GetWorld()->GetName());
-	Ar.Printf(TEXT("----------------------------------------------------------------------------------------------------------------"));
+	if (!GIsAutomationTesting)
 	{
-		FHierarchicalLogArchive::FIndentScope CellIndentScope = Ar.PrintfIndent(TEXT("Content of %s Persistent Level"), *GetWorld()->GetName());
-
-		TArray<const AActor*> Actors;
-		for (const AActor* Actor : GetWorld()->PersistentLevel->Actors)
+		Ar.Printf(TEXT("----------------------------------------------------------------------------------------------------------------"));
+		Ar.Printf(TEXT("%s - Persistent Level"), *GetWorld()->GetName());
+		Ar.Printf(TEXT("----------------------------------------------------------------------------------------------------------------"));
 		{
-			if (Actor)
+			FHierarchicalLogArchive::FIndentScope CellIndentScope = Ar.PrintfIndent(TEXT("Content of %s Persistent Level"), *GetWorld()->GetName());
+
+			TArray<const AActor*> Actors;
+			for (const AActor* Actor : GetWorld()->PersistentLevel->Actors)
 			{
-				Actors.Add(Actor);
+				if (Actor)
+				{
+					Actors.Add(Actor);
+				}
 			}
-		}
 
-		Actors.Sort([this](const AActor& A, const AActor& B) { return A.GetFName().LexicalLess(B.GetFName()); });
+			Actors.Sort([this](const AActor& A, const AActor& B) { return A.GetFName().LexicalLess(B.GetFName()); });
 
-		Ar.Printf(TEXT("Always loaded Actor Count: %d "), Actors.Num());
-		for (const AActor* Actor : Actors)
-		{
-			Ar.Printf(TEXT("Actor Path: %s"), *Actor->GetPathName());
-			Ar.Printf(TEXT("Actor Package: %s"), *Actor->GetPackage()->GetName());
+			Ar.Printf(TEXT("Always loaded Actor Count: %d "), Actors.Num());
+			for (const AActor* Actor : Actors)
+			{
+				Ar.Printf(TEXT("Actor Path: %s"), *Actor->GetPathName());
+				Ar.Printf(TEXT("Actor Package: %s"), *Actor->GetPackage()->GetName());
+			}
 		}
+		Ar.Printf(TEXT(""));
 	}
-	Ar.Printf(TEXT(""));
 }
 
 void UWorldPartitionRuntimeHash::ForceExternalActorLevelReference(bool bForceExternalActorLevelReferenceForPIE)
diff --git a/Engine/Source/Runtime/Engine/Private/WorldPartition/WorldPartitionStreamingGeneration.cpp b/Engine/Source/Runtime/Engine/Private/WorldPartition/WorldPartitionStreamingGeneration.cpp
--- a/Engine/Source/Runtime/Engine/Private/WorldPartition/WorldPartitionStreamingGeneration.cpp
+++ b/Engine/Source/Runtime/Engine/Private/WorldPartition/WorldPartitionStreamingGeneration.cpp
@@ -33,7 +33,9 @@ static FAutoConsoleCommand DumpStreamingGenerationLog(
 		{
 			if (UWorldPartition* WorldPartition = World->GetWorldPartition())
 			{
-				WorldPartition->GenerateStreaming();
+				UWorldPartition::FGenerateStreamingParams Params;
+				UWorldPartition::FGenerateStreamingContext Context;
+				WorldPartition->GenerateStreaming(Params, Context);
 				WorldPartition->FlushStreaming();
 			}
 		}
@@ -929,12 +931,17 @@ public:
 		StreamingGenerationContext = MakeUnique<FStreamingGenerationContext>(this, Container);
 	}
 
-	static TUniquePtr<FArchive> CreateDumpStateLogArchive(const TCHAR* Suffix)
+	static TUniquePtr<FArchive> CreateDumpStateLogArchive(const TCHAR* Suffix, bool bTimeStamped = true)
 	{
 		if (!GIsBuildMachine)
 		{
 			FString StateLogOutputFilename = FPaths::ProjectSavedDir() / TEXT("WorldPartition") / FString::Printf(TEXT("StreamingGeneration-%s"), Suffix);
-			StateLogOutputFilename += FString::Printf(TEXT("-%08x-%s"), FPlatformProcess::GetCurrentProcessId(), *FDateTime::Now().ToIso8601().Replace(TEXT(":"), TEXT(".")));
+
+			if (bTimeStamped)
+			{			
+				StateLogOutputFilename += FString::Printf(TEXT("-%08x-%s"), FPlatformProcess::GetCurrentProcessId(), *FDateTime::Now().ToIso8601().Replace(TEXT(":"), TEXT(".")));
+			}
+
 			StateLogOutputFilename += TEXT(".log");
 			return TUniquePtr<FArchive>(IFileManager::Get().CreateFileWriter(*StateLogOutputFilename));
 		}
@@ -1068,12 +1075,29 @@ private:
 
 bool UWorldPartition::GenerateStreaming(TArray<FString>* OutPackagesToGenerate)
 {
-	OnPreGenerateStreaming.Broadcast(OutPackagesToGenerate);
+	FGenerateStreamingParams Params = FGenerateStreamingParams().SetActorDescContainer(nullptr);
+	FGenerateStreamingContext Context = FGenerateStreamingContext().SetPackagesToGenerate(OutPackagesToGenerate);
+	return GenerateStreaming(Params, Context);
+}
 
-	return GenerateContainerStreaming(ActorDescContainer, OutPackagesToGenerate);
+bool UWorldPartition::GenerateContainerStreaming(const UActorDescContainer* InActorDescContainer, TArray<FString>* OutPackagesToGenerate)
+{
+	FGenerateStreamingParams Params = FGenerateStreamingParams().SetActorDescContainer(InActorDescContainer);
+	FGenerateStreamingContext Context = FGenerateStreamingContext().SetPackagesToGenerate(OutPackagesToGenerate);
+	return GenerateContainerStreaming(Params, Context);
 }
 
-bool UWorldPartition::GenerateContainerStreaming(const UActorDescContainer* InActorDescContainer, TArray<FString>* OutPackagesToGenerate /* = nullptr */)
+bool UWorldPartition::GenerateStreaming(const FGenerateStreamingParams& InParams, FGenerateStreamingContext& InContext)
+{
+	check(!InParams.ActorDescContainer);
+	FGenerateStreamingParams Params = FGenerateStreamingParams(InParams).SetActorDescContainer(ActorDescContainer);
+
+	OnPreGenerateStreaming.Broadcast(InContext.PackagesToGenerate);
+
+	return GenerateContainerStreaming(Params, InContext);
+}
+
+bool UWorldPartition::GenerateContainerStreaming(const FGenerateStreamingParams& InParams, FGenerateStreamingContext& InContext)
 {
 	TRACE_CPUPROFILER_EVENT_SCOPE(UWorldPartition::GenerateContainerStreaming);
 
@@ -1091,7 +1115,7 @@ bool UWorldPartition::GenerateContainerStreaming(const UActorDescContainer* InAc
 		ErrorHandler = &MapCheckErrorHandler;
 	}
 
-	const FString ContainerPackageName = InActorDescContainer->ContainerPackageName.ToString();
+	const FString ContainerPackageName = InParams.ActorDescContainer->ContainerPackageName.ToString();
 	FString ContainerShortName = FPackageName::GetShortName(ContainerPackageName);
 	if (!ContainerPackageName.StartsWith(TEXT("/Game/")))
 	{
@@ -1107,12 +1131,17 @@ bool UWorldPartition::GenerateContainerStreaming(const UActorDescContainer* InAc
 
 	// Dump state log
 	TStringBuilder<256> StateLogSuffix;
-	StateLogSuffix += bIsPIE ? TEXT("PIE") : (IsRunningGame() ? TEXT("Game") : (IsRunningCookCommandlet() ? TEXT("Cook") : TEXT("Manual")));
+	StateLogSuffix += bIsPIE ? TEXT("PIE") : (IsRunningGame() ? TEXT("Game") : (IsRunningCookCommandlet() ? TEXT("Cook") : (GIsAutomationTesting ? TEXT("UnitTest") : TEXT("Manual"))));
 	StateLogSuffix += TEXT("_");
 	StateLogSuffix += ContainerShortName;
-	TUniquePtr<FArchive> LogFileAr = FWorldPartitionStreamingGenerator::CreateDumpStateLogArchive(*StateLogSuffix);
+	TUniquePtr<FArchive> LogFileAr = FWorldPartitionStreamingGenerator::CreateDumpStateLogArchive(*StateLogSuffix, !InParams.OutputLogPath);
 	FHierarchicalLogArchive HierarchicalLogAr(*LogFileAr);
 
+	if (LogFileAr)
+	{
+		InContext.OutputLogFilename = LogFileAr->GetArchiveName();
+	}
+
 	FWorldPartitionStreamingGenerator::FWorldPartitionStreamingGeneratorParams StreamingGeneratorParams;
 	StreamingGeneratorParams.WorldPartitionContext = this;
 	StreamingGeneratorParams.ModifiedActorsDescList = ModifiedActorsDescList;
@@ -1123,7 +1152,7 @@ bool UWorldPartition::GenerateContainerStreaming(const UActorDescContainer* InAc
 	FWorldPartitionStreamingGenerator StreamingGenerator(StreamingGeneratorParams);
 
 	// Preparation Phase
-	StreamingGenerator.PreparationPhase(InActorDescContainer);
+	StreamingGenerator.PreparationPhase(InParams.ActorDescContainer);
 
 	StreamingGenerator.DumpStateLog(HierarchicalLogAr);
 
@@ -1132,7 +1161,7 @@ bool UWorldPartition::GenerateContainerStreaming(const UActorDescContainer* InAc
 	StreamingPolicy = NewObject<UWorldPartitionStreamingPolicy>(const_cast<UWorldPartition*>(this), WorldPartitionStreamingPolicyClass.Get(), NAME_None, bIsPIE ? RF_Transient : RF_NoFlags);
 
 	check(RuntimeHash);
-	if (RuntimeHash->GenerateStreaming(StreamingPolicy, StreamingGenerator.GetStreamingGenerationContext(), OutPackagesToGenerate))
+	if (RuntimeHash->GenerateStreaming(StreamingPolicy, StreamingGenerator.GetStreamingGenerationContext(), InContext.PackagesToGenerate))
 	{
 		//if (IsRunningCookCommandlet())
 		{
diff --git a/Engine/Source/Runtime/Engine/Public/WorldPartition/WorldPartition.h b/Engine/Source/Runtime/Engine/Public/WorldPartition/WorldPartition.h
--- a/Engine/Source/Runtime/Engine/Public/WorldPartition/WorldPartition.h
+++ b/Engine/Source/Runtime/Engine/Public/WorldPartition/WorldPartition.h
@@ -159,8 +159,39 @@ public:
 
 	// Streaming generation
 	bool CanGenerateStreaming() const { return !StreamingPolicy; }
+
+	UE_DEPRECATED(5.3, "GenerateStreaming is deprecated, use GenerateStreaming with a param struct instead")
 	bool GenerateStreaming(TArray<FString>* OutPackagesToGenerate = nullptr);
+
+	UE_DEPRECATED(5.3, "GenerateContainerStreaming is deprecated, use GenerateContainerStreaming with a param struct instead")
 	bool GenerateContainerStreaming(const UActorDescContainer* ActorDescContainer, TArray<FString>* OutPackagesToGenerate = nullptr);
+
+	struct FGenerateStreamingParams
+	{
+		FGenerateStreamingParams()
+		{}
+
+		const UActorDescContainer* ActorDescContainer = nullptr;
+		TOptional<const FString> OutputLogPath;
+
+		FGenerateStreamingParams& SetActorDescContainer(const UActorDescContainer* InActorDescContainer) { ActorDescContainer = InActorDescContainer; return *this; }
+		FGenerateStreamingParams& SetOutputLogPath(const FString& InOutputLogPath) { OutputLogPath = InOutputLogPath; return *this; }
+	};
+
+	struct FGenerateStreamingContext
+	{
+		FGenerateStreamingContext()
+		{}
+
+		 TArray<FString>* PackagesToGenerate = nullptr;
+		 TOptional<FString> OutputLogFilename;
+
+		FGenerateStreamingContext& SetPackagesToGenerate(TArray<FString>* InPackagesToGenerate) { PackagesToGenerate = InPackagesToGenerate; return *this; }
+	};
+
+	bool GenerateStreaming(const FGenerateStreamingParams& InParams, FGenerateStreamingContext& InContext);
+	bool GenerateContainerStreaming(const FGenerateStreamingParams& InParams, FGenerateStreamingContext& InContext);
+
 	void FlushStreaming();
 
 	DECLARE_MULTICAST_DELEGATE_OneParam(FWorldPartitionGenerateStreamingDelegate, TArray<FString>*);
