World Partition: remove code deprecated for the 5.1 release.

diff --git a/Engine/Source/Runtime/Engine/Private/WorldPartition/WorldPartition.cpp b/Engine/Source/Runtime/Engine/Private/WorldPartition/WorldPartition.cpp
--- a/Engine/Source/Runtime/Engine/Private/WorldPartition/WorldPartition.cpp
+++ b/Engine/Source/Runtime/Engine/Private/WorldPartition/WorldPartition.cpp
@@ -641,9 +641,6 @@ void UWorldPartition::Initialize(UWorld* InWorld, const FTransform& InTransform)
 #endif
 
 	FWorldPartitionEvents::BroadcastWorldPartitionInitialized(World, this);
-PRAGMA_DISABLE_DEPRECATION_WARNINGS
-	OnWorldPartitionInitialized.Broadcast(this);
-PRAGMA_ENABLE_DEPRECATION_WARNINGS
 }
 
 void UWorldPartition::Uninitialize()
@@ -718,9 +715,6 @@ void UWorldPartition::Uninitialize()
 		InitState = EWorldPartitionInitState::Uninitialized;
 
 		FWorldPartitionEvents::BroadcastWorldPartitionUninitialized(World, this);
-PRAGMA_DISABLE_DEPRECATION_WARNINGS
-		OnWorldPartitionUninitialized.Broadcast(this);
-PRAGMA_ENABLE_DEPRECATION_WARNINGS
 
 		World = nullptr;
 	}
diff --git a/Engine/Source/Runtime/Engine/Public/WorldPartition/WorldPartition.h b/Engine/Source/Runtime/Engine/Public/WorldPartition/WorldPartition.h
--- a/Engine/Source/Runtime/Engine/Public/WorldPartition/WorldPartition.h
+++ b/Engine/Source/Runtime/Engine/Public/WorldPartition/WorldPartition.h
@@ -99,14 +99,6 @@ public:
 	static bool RemoveWorldPartition(AWorldSettings* WorldSettings);
 #endif
 
-	DECLARE_MULTICAST_DELEGATE_OneParam(FWorldPartitionInitializeDelegate, UWorldPartition*);
-	
-	UE_DEPRECATED(5.1, "Please use FWorldPartitionInitializedEvent& UWorld::OnWorldPartitionInitialized() instead.")
-	FWorldPartitionInitializeDelegate OnWorldPartitionInitialized;
-
-	UE_DEPRECATED(5.1, "Please use FWorldPartitionInitializedEvent& UWorld::OnWorldPartitionUninitialized() instead.")
-	FWorldPartitionInitializeDelegate OnWorldPartitionUninitialized;
-
 #if WITH_EDITOR
 	TArray<FBox> GetUserLoadedEditorRegions() const;
 
