Revert to refreshing the entire graph when a node changes, this reintroduces some ensures, but fixes a crash when dragging off nodes with orphaned pins
#rnx

diff --git a/Engine/Source/Runtime/Engine/Private/EdGraph/EdGraph.cpp b/Engine/Source/Runtime/Engine/Private/EdGraph/EdGraph.cpp
--- a/Engine/Source/Runtime/Engine/Private/EdGraph/EdGraph.cpp
+++ b/Engine/Source/Runtime/Engine/Private/EdGraph/EdGraph.cpp
@@ -285,11 +285,7 @@ void UEdGraph::NotifyGraphChanged()
 
 void UEdGraph::NotifyNodeChanged(const UEdGraphNode* Node)
 {
-	FEdGraphEditAction Action;
-	Action.Graph = this;
-	Action.Action = GRAPHACTION_EditNode;
-	Action.Nodes.Add(Node);
-	NotifyGraphChanged(Action);
+	NotifyGraphChanged();
 }
 
 void UEdGraph::NotifyGraphChanged(const FEdGraphEditAction& InAction)
