By default, disable level streaming profiling on server.

#tests Tested PIE generation of level streaming profile files still works.

diff --git a/Engine/Source/Runtime/Engine/Private/ProfilingDebugging/LevelStreamingProfilingSubsystem.cpp b/Engine/Source/Runtime/Engine/Private/ProfilingDebugging/LevelStreamingProfilingSubsystem.cpp
--- a/Engine/Source/Runtime/Engine/Private/ProfilingDebugging/LevelStreamingProfilingSubsystem.cpp
+++ b/Engine/Source/Runtime/Engine/Private/ProfilingDebugging/LevelStreamingProfilingSubsystem.cpp
@@ -20,7 +20,7 @@
 DEFINE_LOG_CATEGORY_STATIC(LogLevelStreamingProfiling, Log, All);
 CSV_DEFINE_CATEGORY_MODULE(ENGINE_API, LevelStreaming, true);
 
-static bool GLevelStreamingProfilingEnabled = !UE_BUILD_SHIPPING;
+static bool GLevelStreamingProfilingEnabled = !UE_BUILD_SHIPPING && !UE_SERVER;
 static FAutoConsoleVariableRef CVarLevelStreamingProfilingEnabled(
 #if UE_BUILD_SHIPPING
 	TEXT("LevelStreaming.Profiling.Enabled.Shipping"),
