Fix browse to asset not showing up for blueprints in cooked editor

diff --git a/Engine/Source/Runtime/Engine/Private/ActorEditor.cpp b/Engine/Source/Runtime/Engine/Private/ActorEditor.cpp
--- a/Engine/Source/Runtime/Engine/Private/ActorEditor.cpp
+++ b/Engine/Source/Runtime/Engine/Private/ActorEditor.cpp
@@ -1533,11 +1533,14 @@ void AActor::CheckForErrors()
 
 bool AActor::GetReferencedContentObjects( TArray<UObject*>& Objects ) const
 {
-	UBlueprint* Blueprint = UBlueprint::GetBlueprintFromClass( GetClass() );
-	if (Blueprint)
+	if (UBlueprint* Blueprint = UBlueprint::GetBlueprintFromClass(GetClass()))
 	{
 		Objects.AddUnique(Blueprint);
 	}
+	else if (UBlueprintGeneratedClass* BlueprintGeneratedClass = Cast<UBlueprintGeneratedClass>(GetClass()))
+	{
+		Objects.AddUnique(BlueprintGeneratedClass);
+	}
 	return true;
 }
 
