Chaos - fix transforms for welded objects with complex collision

diff --git a/Engine/Source/Runtime/Engine/Private/PhysicsEngine/Experimental/ChaosInterfaceUtils.cpp b/Engine/Source/Runtime/Engine/Private/PhysicsEngine/Experimental/ChaosInterfaceUtils.cpp
--- a/Engine/Source/Runtime/Engine/Private/PhysicsEngine/Experimental/ChaosInterfaceUtils.cpp
+++ b/Engine/Source/Runtime/Engine/Private/PhysicsEngine/Experimental/ChaosInterfaceUtils.cpp
@@ -278,6 +278,8 @@ namespace ChaosInterface
 
 		if (bMakeComplexGeometry)
 		{
+			const FTransform MeshTransform = FTransform(InParams.LocalTransform.GetRotation(), Scale * InParams.LocalTransform.GetTranslation(), FVector(1, 1, 1));
+			const bool bHasTranslationOrRotation = !MeshTransform.GetTranslation().IsNearlyZero() || !MeshTransform.GetRotation().IsIdentity();
 			const bool bNoScale = Scale == FVector(1);
 			for (auto& ChaosTriMesh : InParams.ChaosTriMeshes)
 			{
@@ -291,6 +293,12 @@ namespace ChaosInterface
 					Implicit = TUniquePtr<Chaos::FImplicitObject>(new Chaos::TImplicitObjectScaled<Chaos::FTriangleMeshImplicitObject>(ChaosTriMesh, Scale));
 				}
 
+				// Wrap the mesh in a non-scaled transform if necessary (the scale is pulled out above)
+				if (bHasTranslationOrRotation)
+				{
+					Implicit = TUniquePtr<Chaos::FImplicitObject>(new Chaos::TImplicitObjectTransformed<Chaos::FReal, 3>(MoveTemp(Implicit), MeshTransform));
+				}
+
 				ChaosTriMesh->SetCullsBackFaceRaycast(!InParams.bDoubleSided);
 
 				TUniquePtr<Chaos::FPerShapeData> NewShape = NewShapeHelper(MakeSerializable(Implicit),Shapes.Num(), nullptr, ECollisionEnabled::QueryAndPhysics, true);
diff --git a/Engine/Source/Runtime/Engine/Private/PhysicsEngine/Experimental/PhysInterface_Chaos.cpp b/Engine/Source/Runtime/Engine/Private/PhysicsEngine/Experimental/PhysInterface_Chaos.cpp
--- a/Engine/Source/Runtime/Engine/Private/PhysicsEngine/Experimental/PhysInterface_Chaos.cpp
+++ b/Engine/Source/Runtime/Engine/Private/PhysicsEngine/Experimental/PhysInterface_Chaos.cpp
@@ -1174,10 +1174,10 @@ bool FPhysInterface_Chaos::GetSquaredDistanceToBody(const FBodyInstance* InInsta
 	return bFoundValidBody;
 }
 
-uint32 GetTriangleMeshExternalFaceIndex(const FPhysicsShape& Shape, uint32 InternalFaceIndex)
+uint32 GetTriangleMeshExternalFaceIndex(const Chaos::FImplicitObject* Geom, uint32 InternalFaceIndex)
 {
 	using namespace Chaos;
-	uint8 OuterType = Shape.GetGeometry()->GetType();
+	uint8 OuterType = Geom->GetType();
 	uint8 InnerType = GetInnerType(OuterType);
 	if (ensure(InnerType == ImplicitObjectType::TriangleMesh))
 	{
@@ -1185,16 +1185,16 @@ uint32 GetTriangleMeshExternalFaceIndex(const FPhysicsShape& Shape, uint32 Inter
 
 		if (IsScaled(OuterType))
 		{
-			const TImplicitObjectScaled<FTriangleMeshImplicitObject>& ScaledTriangleMesh = Shape.GetGeometry()->GetObjectChecked<TImplicitObjectScaled<FTriangleMeshImplicitObject>>();
+			const TImplicitObjectScaled<FTriangleMeshImplicitObject>& ScaledTriangleMesh = Geom->GetObjectChecked<TImplicitObjectScaled<FTriangleMeshImplicitObject>>();
 			TriangleMesh = ScaledTriangleMesh.GetUnscaledObject();
 		}
 		else if(IsInstanced(OuterType))
 		{
-			TriangleMesh = Shape.GetGeometry()->GetObjectChecked<TImplicitObjectInstanced<FTriangleMeshImplicitObject>>().GetInstancedObject();
+			TriangleMesh = Geom->GetObjectChecked<TImplicitObjectInstanced<FTriangleMeshImplicitObject>>().GetInstancedObject();
 		}
 		else
 		{
-			TriangleMesh = &Shape.GetGeometry()->GetObjectChecked<FTriangleMeshImplicitObject>();
+			TriangleMesh = &Geom->GetObjectChecked<FTriangleMeshImplicitObject>();
 		}
 
 		return TriangleMesh->GetExternalFaceIndexFromInternal(InternalFaceIndex);
@@ -1203,6 +1203,12 @@ uint32 GetTriangleMeshExternalFaceIndex(const FPhysicsShape& Shape, uint32 Inter
 	return -1;
 }
 
+uint32 GetTriangleMeshExternalFaceIndex(const FPhysicsShape& Shape, uint32 InternalFaceIndex)
+{
+	// NOTE: GetLeafGeometry will strip Transformed and Instanced wrappers (but not Scaled)
+	return GetTriangleMeshExternalFaceIndex(Shape.GetLeafGeometry(), InternalFaceIndex);
+}
+
 void FPhysInterface_Chaos::CalculateMassPropertiesFromShapeCollection(Chaos::FMassProperties& OutProperties,const TArray<FPhysicsShapeHandle>& InShapes,float InDensityKGPerCM)
 {
 	ChaosInterface::CalculateMassPropertiesFromShapeCollection(OutProperties,InShapes,InDensityKGPerCM);
