Remove camera shakes from the expired shakes pool when they are marked for GC.

diff --git a/Engine/Source/Runtime/Engine/Classes/Camera/CameraModifier_CameraShake.h b/Engine/Source/Runtime/Engine/Classes/Camera/CameraModifier_CameraShake.h
--- a/Engine/Source/Runtime/Engine/Classes/Camera/CameraModifier_CameraShake.h
+++ b/Engine/Source/Runtime/Engine/Classes/Camera/CameraModifier_CameraShake.h
@@ -142,6 +142,8 @@ class ENGINE_API UCameraModifier_CameraShake : public UCameraModifier
 public:
 	UCameraModifier_CameraShake(const FObjectInitializer& ObjectInitializer);
 
+	virtual void BeginDestroy() override;
+
 	/** 
 	 * Adds a new active screen shake to be applied. 
 	 * @param NewShake - The class of camera shake to instantiate.
@@ -234,5 +236,9 @@ protected:
 	/** Whether we are currently recording camera shake debug info */
 	bool bRecordDebugData;
 #endif
+
+private:
+	void OnPreGarbageCollect();
+	void RemoveInvalidObjectsFromExpiredPool();
 };
 
diff --git a/Engine/Source/Runtime/Engine/Private/Camera/CameraModifier_CameraShake.cpp b/Engine/Source/Runtime/Engine/Private/Camera/CameraModifier_CameraShake.cpp
--- a/Engine/Source/Runtime/Engine/Private/Camera/CameraModifier_CameraShake.cpp
+++ b/Engine/Source/Runtime/Engine/Private/Camera/CameraModifier_CameraShake.cpp
@@ -49,8 +49,19 @@ UCameraModifier_CameraShake::UCameraModifier_CameraShake(const FObjectInitialize
 #endif
 {
 	SplitScreenShakeScale = 0.5f;
+	FCoreUObjectDelegates::GetPreGarbageCollectDelegate().AddUObject(this, &UCameraModifier_CameraShake::OnPreGarbageCollect);
 }
 
+void UCameraModifier_CameraShake::BeginDestroy()
+{
+	FCoreUObjectDelegates::GetPreGarbageCollectDelegate().RemoveAll(this);
+	Super::BeginDestroy();
+}
+
+void UCameraModifier_CameraShake::OnPreGarbageCollect()
+{
+	RemoveInvalidObjectsFromExpiredPool();
+}
 
 bool UCameraModifier_CameraShake::ModifyCamera(float DeltaTime, FMinimalViewInfo& InOutPOV)
 {
@@ -345,6 +356,26 @@ UCameraShakeBase* UCameraModifier_CameraShake::ReclaimShakeFromExpiredPool(TSubc
 	return nullptr;
 }
 
+void UCameraModifier_CameraShake::RemoveInvalidObjectsFromExpiredPool()
+{
+	for (auto CamShakeItr = ExpiredPooledShakesMap.CreateIterator(); CamShakeItr; ++CamShakeItr)
+	{
+		if (!IsValid(CamShakeItr.Key()))
+		{
+			CamShakeItr.RemoveCurrent();
+			continue;
+		}
+		TArray<TObjectPtr<UCameraShakeBase>>& PooledShakes = CamShakeItr.Value().PooledShakes;
+		for (int32 i = PooledShakes.Num() - 1; i >= 0; --i)
+		{
+			if (!IsValid(PooledShakes[i]))
+			{
+				PooledShakes.RemoveAtSwap(i);
+			}
+		}
+	}
+}
+
 void UCameraModifier_CameraShake::GetActiveCameraShakes(TArray<FActiveCameraShakeInfo>& ActiveCameraShakes) const
 {
 	ActiveCameraShakes.Append(ActiveShakes);
