FSparseVolumeTextureData now stores each mip map in its own data structure. Null tiles are no longer explicitly stored in (on-disk) memory. Fix for SVT not generating whole mip chain. Fix for CPU implementation of SVT addressing mode causing out of bound memory reads resulting in messed up border pixels for lower resolution mip levels (in the best case).

#rnx

diff --git a/Engine/Source/Runtime/Engine/Classes/SparseVolumeTexture/SparseVolumeTexture.h b/Engine/Source/Runtime/Engine/Classes/SparseVolumeTexture/SparseVolumeTexture.h
--- a/Engine/Source/Runtime/Engine/Classes/SparseVolumeTexture/SparseVolumeTexture.h
+++ b/Engine/Source/Runtime/Engine/Classes/SparseVolumeTexture/SparseVolumeTexture.h
@@ -37,6 +37,7 @@ struct ENGINE_API FSparseVolumeTextureHeader
 	FIntVector3 PageTableVolumeAABBMax = FIntVector3(INT32_MIN, INT32_MIN, INT32_MIN);
 	TStaticArray<EPixelFormat, 2> AttributesFormats = TStaticArray<EPixelFormat, 2>(InPlace, PF_Unknown);
 	TStaticArray<FVector4f, 2> NullTileValues = TStaticArray<FVector4f, 2>(InPlace, FVector4f());
+	TStaticArray<FVector4f, 2> NullTileValuesQuantized = TStaticArray<FVector4f, 2>(InPlace, FVector4f()); // Values after converting/quantizing original fallback values to the given attribute formats
 
 	void Serialize(FArchive& Ar);
 };
diff --git a/Engine/Source/Runtime/Engine/Private/SparseVolumeTexture/SparseVolumeTexture.cpp b/Engine/Source/Runtime/Engine/Private/SparseVolumeTexture/SparseVolumeTexture.cpp
--- a/Engine/Source/Runtime/Engine/Private/SparseVolumeTexture/SparseVolumeTexture.cpp
+++ b/Engine/Source/Runtime/Engine/Private/SparseVolumeTexture/SparseVolumeTexture.cpp
@@ -54,6 +54,8 @@ void FSparseVolumeTextureHeader::Serialize(FArchive& Ar)
 		UE::SVT::Private::SerializeEnumAs<uint8>(Ar, AttributesFormats[1]);
 		Ar << NullTileValues[0];
 		Ar << NullTileValues[1];
+		Ar << NullTileValuesQuantized[0];
+		Ar << NullTileValuesQuantized[1];
 	}
 	else
 	{
@@ -94,7 +96,7 @@ bool FSparseVolumeTextureFrame::BuildDerivedData(const FIntVector3& VolumeResolu
 		AddressingInfo.AddressY = AddressY;
 		AddressingInfo.AddressZ = AddressZ;
 
-		const int32 NumMipLevels = 1; // generate entire mip chain
+		const int32 NumMipLevels = -1; // generate entire mip chain
 		const bool bMoveMip0FromSource = true; // we have no need to keep TextureData around
 		if (!TextureData.BuildDerivedData(AddressingInfo, NumMipLevels, bMoveMip0FromSource, *OutMippedTextureData))
 		{
diff --git a/Engine/Source/Runtime/Engine/Private/SparseVolumeTexture/SparseVolumeTextureData.cpp b/Engine/Source/Runtime/Engine/Private/SparseVolumeTexture/SparseVolumeTextureData.cpp
--- a/Engine/Source/Runtime/Engine/Private/SparseVolumeTexture/SparseVolumeTextureData.cpp
+++ b/Engine/Source/Runtime/Engine/Private/SparseVolumeTexture/SparseVolumeTextureData.cpp
@@ -12,49 +12,31 @@ DEFINE_LOG_CATEGORY_STATIC(LogSparseVolumeTextureData, Log, All);
 
 ////////////////////////////////////////////////////////////////////////////////////////////////
 
-void FSparseVolumeTextureDataHeader::Serialize(FArchive& Ar)
+void FSparseVolumeTextureData::Serialize(FArchive& Ar)
 {
-	FSparseVolumeTextureHeader::Serialize(Ar);
+	Header.Serialize(Ar);
 
 	Ar << Version;
 
 	if (Version == 0)
 	{
-		int32 NumMipLevels = MipInfo.Num();
+		int32 NumMipLevels = MipMaps.Num();
 		Ar << NumMipLevels;
 		if (Ar.IsLoading())
 		{
-			MipInfo.Reset(NumMipLevels);
-			MipInfo.SetNum(NumMipLevels);
+			MipMaps.Reset(NumMipLevels);
+			MipMaps.SetNum(NumMipLevels);
 		}
 		for (int32 MipLevel = 0; MipLevel < NumMipLevels; ++MipLevel)
 		{
-			Ar << MipInfo[MipLevel].TileOffset;
-			Ar << MipInfo[MipLevel].TileCount;
+			FMipMap& MipMap = MipMaps[MipLevel];
+			Ar << MipMap.PageTable;
+			Ar << MipMap.PhysicalTileDataA;
+			Ar << MipMap.PhysicalTileDataB;
+			Ar << MipMap.NumPhysicalTiles;
 		}
 	}
 	else
-	{
-		// FSparseVolumeTextureDataHeader needs to account for new version
-		check(false);
-	}
-}
-
-////////////////////////////////////////////////////////////////////////////////////////////////
-
-void FSparseVolumeTextureData::Serialize(FArchive& Ar)
-{
-	Header.Serialize(Ar);
-
-	Ar << Version;
-
-	if (Version == 0)
-	{
-		Ar << PageTable;
-		Ar << PhysicalTileDataA;
-		Ar << PhysicalTileDataB;
-	}
-	else
 	{
 		// FSparseVolumeTextureData needs to account for new version
 		check(false);
@@ -63,14 +45,10 @@ void FSparseVolumeTextureData::Serialize(FArchive& Ar)
 
 bool FSparseVolumeTextureData::Construct(const ISparseVolumeTextureDataConstructionAdapter& Adapter)
 {
-	using namespace UE::SVT;
-	using namespace UE::SVT::Private;
-
 	ISparseVolumeTextureDataConstructionAdapter::FAttributesInfo AttributesInfo[2];
 	Adapter.GetAttributesInfo(AttributesInfo[0], AttributesInfo[1]);
 
-	Header.MipInfo.SetNum(1);
-	FSparseVolumeTextureMipInfo& MipInfo = Header.MipInfo[0];
+	MipMaps.SetNum(1);
 	Header.AttributesFormats[0] = AttributesInfo[0].Format;
 	Header.AttributesFormats[1] = AttributesInfo[1].Format;
 
@@ -89,12 +67,12 @@ bool FSparseVolumeTextureData::Construct(const ISparseVolumeTextureDataConstruct
 
 	Header.PageTableVolumeAABBMax = Header.PageTableVolumeAABBMin + PageTableVolumeResolution;
 
-	if (PageTableVolumeResolution.X > SVTMaxVolumeTextureDim
-		|| PageTableVolumeResolution.Y > SVTMaxVolumeTextureDim
-		|| PageTableVolumeResolution.Z > SVTMaxVolumeTextureDim)
+	if (PageTableVolumeResolution.X > UE::SVT::SVTMaxVolumeTextureDim
+		|| PageTableVolumeResolution.Y > UE::SVT::SVTMaxVolumeTextureDim
+		|| PageTableVolumeResolution.Z > UE::SVT::SVTMaxVolumeTextureDim)
 	{
 		UE_LOG(LogSparseVolumeTextureData, Warning, TEXT("SparseVolumeTexture page table texture dimensions exceed limit (%ix%ix%i): %ix%ix%i"),
-			SVTMaxVolumeTextureDim, SVTMaxVolumeTextureDim, SVTMaxVolumeTextureDim,
+			UE::SVT::SVTMaxVolumeTextureDim, UE::SVT::SVTMaxVolumeTextureDim, UE::SVT::SVTMaxVolumeTextureDim,
 			PageTableVolumeResolution.X, PageTableVolumeResolution.Y, PageTableVolumeResolution.Z);
 		return false;
 	}
@@ -174,20 +152,16 @@ bool FSparseVolumeTextureData::Construct(const ISparseVolumeTextureDataConstruct
 	}
 	LinearAllocatedPages.SetNum(NumAllocatedPages);
 
-	const int32 EffectivelyAllocatedTiles = NumAllocatedPages + 1 /*null tile*/;
-
 	// Initialize the SparseVolumeTexture page and tile.
 	const int32 FormatSize[] = { GPixelFormats[Header.AttributesFormats[0]].BlockBytes, GPixelFormats[Header.AttributesFormats[1]].BlockBytes };
-	PageTable.SetNum(1);
-	PageTable[0].SetNumZeroed(Header.PageTableVolumeResolution.X * Header.PageTableVolumeResolution.Y * Header.PageTableVolumeResolution.Z);
-	PhysicalTileDataA.SetNumZeroed((int64)EffectivelyAllocatedTiles * (int64)SVTNumVoxelsPerPaddedTile * (int64)FormatSize[0]);
-	PhysicalTileDataB.SetNumZeroed((int64)EffectivelyAllocatedTiles * (int64)SVTNumVoxelsPerPaddedTile * (int64)FormatSize[1]);
+	MipMaps[0].PageTable.SetNumZeroed(Header.PageTableVolumeResolution.X * Header.PageTableVolumeResolution.Y * Header.PageTableVolumeResolution.Z);
+	MipMaps[0].PhysicalTileDataA.SetNum((int64)NumAllocatedPages * (int64)UE::SVT::SVTNumVoxelsPerPaddedTile * (int64)FormatSize[0]);
+	MipMaps[0].PhysicalTileDataB.SetNum((int64)NumAllocatedPages * (int64)UE::SVT::SVTNumVoxelsPerPaddedTile * (int64)FormatSize[1]);
+	MipMaps[0].NumPhysicalTiles = NumAllocatedPages;
 
 	// Generate page table and tile volume data by splatting the data
 	{
-		uint32 DstTileIndex = 0;
-
-		// Add an empty tile, reserve slot at coord 0 and fill all tiles with null tile data
+		// Fill all tiles with null tile data. See below for an explanation of why this is necessary.
 		{
 			// Compute potentially normalized fallback values
 			FVector4f FallbackValues[] = { AttributesInfo[0].FallbackValue, AttributesInfo[1].FallbackValue };
@@ -201,8 +175,9 @@ bool FSparseVolumeTextureData::Construct(const ISparseVolumeTextureDataConstruct
 			Header.NullTileValues[0] = FallbackValues[0];
 			Header.NullTileValues[1] = FallbackValues[1];
 
-			// Fill all tiles with fallback values since we can't guarantee that all voxels of a page will be written to in the next step.
-			ParallelFor(EffectivelyAllocatedTiles, [&FallbackValues, this](int32 TileIndex)
+			// Fill all tiles with fallback values/null tile values. We can't guarantee that all voxels of a page will be written to in the next step, so if we don't do this,
+			// we might end up with physical tiles that have only been partially filled with valid voxels.
+			ParallelFor(NumAllocatedPages, [&FallbackValues, this](int32 TileIndex)
 				{
 					for (int32 Z = 0; Z < SPARSE_VOLUME_TILE_RES_PADDED; ++Z)
 					{
@@ -211,34 +186,24 @@ bool FSparseVolumeTextureData::Construct(const ISparseVolumeTextureDataConstruct
 							for (int32 X = 0; X < SPARSE_VOLUME_TILE_RES_PADDED; ++X)
 							{
 								const FIntVector3 WriteCoord = FIntVector3(X - 1, Y - 1, Z - 1);
-								WriteTileDataVoxel(TileIndex, WriteCoord, 0 /*AttributesIndex*/, FallbackValues[0]);
-								WriteTileDataVoxel(TileIndex, WriteCoord, 1 /*AttributesIndex*/, FallbackValues[1]);
+								WriteTileDataVoxel(TileIndex, WriteCoord, 0 /*MipLevel*/, 0 /*AttributesIndex*/, FallbackValues[0]);
+								WriteTileDataVoxel(TileIndex, WriteCoord, 0 /*MipLevel*/, 1 /*AttributesIndex*/, FallbackValues[1]);
 							}
 						}
 					}
 				});
 
-			// PageTable is all cleared to zero, simply skip a tile
-			++DstTileIndex;
+			Header.NullTileValuesQuantized[0] = ReadTileDataVoxel(0, FIntVector::ZeroValue, 0 /*MipLevel*/, 0 /*AttributesIndex*/);
+			Header.NullTileValuesQuantized[1] = ReadTileDataVoxel(0, FIntVector::ZeroValue, 0 /*MipLevel*/, 1 /*AttributesIndex*/);
 		}
 
+		// Set up the page table entries
 		for (int32 i = 0; i < NumAllocatedPages; ++i)
 		{
 			const FIntVector3 PageCoordToSplat = LinearAllocatedPages[i];
-
-			// Setup the page table entry
-			PageTable[0]
-				[
-					PageCoordToSplat.Z * Header.PageTableVolumeResolution.X * Header.PageTableVolumeResolution.Y +
-					PageCoordToSplat.Y * Header.PageTableVolumeResolution.X +
-					PageCoordToSplat.X
-				] = DstTileIndex;
-
-			// Set the next tile to be written to
-			++DstTileIndex;
+			const int32 TileIndex = (PageCoordToSplat.Z * Header.PageTableVolumeResolution.X * Header.PageTableVolumeResolution.Y) + (PageCoordToSplat.Y * Header.PageTableVolumeResolution.X) + PageCoordToSplat.X;
+			MipMaps[0].PageTable[TileIndex] = (uint32)(i + 1); // +1 because 0 is reserved for the null tile which we don't actually store in memory
 		}
-		MipInfo.TileOffset = 1; // Don't include the null tile in the list of physical tiles of mip0. The runtime SVT will create a null of its own.
-		MipInfo.TileCount = NumAllocatedPages;
 	}
 
 	// Write physical tile data
@@ -264,9 +229,9 @@ bool FSparseVolumeTextureData::Construct(const ISparseVolumeTextureDataConstruct
 				WriteValue = WriteValue * AttributesInfo[AttributesIdx].NormalizeScale[ComponentIdx] + AttributesInfo[AttributesIdx].NormalizeBias[ComponentIdx];
 			}
 
-			const int32 TileIndex = (int32)PageTable[0][PageIndex];
+			const int32 TileIndex = (int32)MipMaps[0].PageTable[PageIndex];
 			const FIntVector3 TileLocalCoord = GridCoord % SPARSE_VOLUME_TILE_RES;
-			WriteTileDataVoxel(TileIndex, TileLocalCoord, AttributesIdx, FVector4f(WriteValue, WriteValue, WriteValue, WriteValue), ComponentIdx);
+			WriteTileDataVoxel(TileIndex - 1, TileLocalCoord, 0 /*MipLevel*/, AttributesIdx, FVector4f(WriteValue, WriteValue, WriteValue, WriteValue), ComponentIdx);  // -1 because 0 is reserved for the null tile which we don't actually store in memory
 		});
 
 	return true;
@@ -274,13 +239,13 @@ bool FSparseVolumeTextureData::Construct(const ISparseVolumeTextureDataConstruct
 
 uint32 FSparseVolumeTextureData::ReadPageTable(const FIntVector3& PageTableCoord, int32 MipLevel) const
 {
-	if (!Header.MipInfo.IsValidIndex(MipLevel))
+	if (!MipMaps.IsValidIndex(MipLevel))
 	{
 		return INDEX_NONE;
 	}
 	const FIntVector3 PageTableResolution = FIntVector3(
-		FMath::Max(1, Header.PageTableVolumeResolution.X >> MipLevel), 
-		FMath::Max(1, Header.PageTableVolumeResolution.Y >> MipLevel), 
+		FMath::Max(1, Header.PageTableVolumeResolution.X >> MipLevel),
+		FMath::Max(1, Header.PageTableVolumeResolution.Y >> MipLevel),
 		FMath::Max(1, Header.PageTableVolumeResolution.Z >> MipLevel));
 	if (PageTableCoord.X < 0 || PageTableCoord.Y < 0 || PageTableCoord.Z < 0
 		|| PageTableCoord.X >= PageTableResolution.X
@@ -291,17 +256,17 @@ uint32 FSparseVolumeTextureData::ReadPageTable(const FIntVector3& PageTableCoord
 	}
 
 	const int32 PageIndex = PageTableCoord.Z * (PageTableResolution.Y * PageTableResolution.X) + PageTableCoord.Y * PageTableResolution.X + PageTableCoord.X;
-	if (PageTable[MipLevel].IsValidIndex(PageIndex))
+	if (MipMaps[MipLevel].PageTable.IsValidIndex(PageIndex))
 	{
-		return PageTable[MipLevel][PageIndex];
+		return MipMaps[MipLevel].PageTable[PageIndex];
 	}
 	return INDEX_NONE;
 }
 
-FVector4f FSparseVolumeTextureData::ReadTileDataVoxel(int32 TileIndex, const FIntVector3& TileDataCoord, int32 AttributesIdx) const
+FVector4f FSparseVolumeTextureData::ReadTileDataVoxel(int32 TileIndex, const FIntVector3& TileDataCoord, int32 MipLevel, int32 AttributesIdx) const
 {
 	check(AttributesIdx >= 0 && AttributesIdx <= 1);
-	if (TileIndex == INDEX_NONE || (AttributesIdx == 0 && PhysicalTileDataA.IsEmpty()) || (AttributesIdx == 1 && PhysicalTileDataB.IsEmpty()))
+	if (TileIndex == INDEX_NONE || (AttributesIdx == 0 && MipMaps[MipLevel].PhysicalTileDataA.IsEmpty()) || (AttributesIdx == 1 && MipMaps[MipLevel].PhysicalTileDataB.IsEmpty()))
 	{
 		return FVector4f();
 	}
@@ -318,7 +283,7 @@ FVector4f FSparseVolumeTextureData::ReadTileDataVoxel(int32 TileIndex, const FIn
 
 	const FInt64Vector3 CoordPadded = FInt64Vector3(TileDataCoord) + FInt64Vector3(SPARSE_VOLUME_TILE_BORDER);
 	const int64 VoxelIndex = (int64)TileIndex * (int64)UE::SVT::SVTNumVoxelsPerPaddedTile + CoordPadded.Z * (SPARSE_VOLUME_TILE_RES_PADDED * SPARSE_VOLUME_TILE_RES_PADDED) + CoordPadded.Y * SPARSE_VOLUME_TILE_RES_PADDED + CoordPadded.X;
-	const uint8* TileData = AttributesIdx == 0 ? PhysicalTileDataA.GetData() : PhysicalTileDataB.GetData();
+	const uint8* TileData = AttributesIdx == 0 ? MipMaps[MipLevel].PhysicalTileDataA.GetData() : MipMaps[MipLevel].PhysicalTileDataB.GetData();
 	const EPixelFormat Format = Header.AttributesFormats[AttributesIdx];
 	return UE::SVT::ReadVoxel(VoxelIndex, TileData, Format);
 }
@@ -340,25 +305,29 @@ FVector4f FSparseVolumeTextureData::Load(const FIntVector3& VolumeCoord, int32 M
 		}
 		return x;
 	};
-	if (!Header.MipInfo.IsValidIndex(MipLevel))
+	if (!MipMaps.IsValidIndex(MipLevel))
 	{
 		return FVector4f();
 	}
 	const FIntVector3 AddressModeVolumeCoord = FIntVector3(
-		ApplyAddressMode(VolumeCoord.X, AddressingInfo.VolumeResolution.X, AddressingInfo.AddressX),
-		ApplyAddressMode(VolumeCoord.Y, AddressingInfo.VolumeResolution.Y, AddressingInfo.AddressY), 
-		ApplyAddressMode(VolumeCoord.Z, AddressingInfo.VolumeResolution.Z, AddressingInfo.AddressZ));
+		ApplyAddressMode(VolumeCoord.X, FMath::Max(1, AddressingInfo.VolumeResolution.X >> MipLevel), AddressingInfo.AddressX),
+		ApplyAddressMode(VolumeCoord.Y, FMath::Max(1, AddressingInfo.VolumeResolution.Y >> MipLevel), AddressingInfo.AddressY), 
+		ApplyAddressMode(VolumeCoord.Z, FMath::Max(1, AddressingInfo.VolumeResolution.Z >> MipLevel), AddressingInfo.AddressZ));
 	const FIntVector3 PageTableCoord = (AddressModeVolumeCoord / SPARSE_VOLUME_TILE_RES) - (Header.PageTableVolumeAABBMin >> MipLevel);
 	const uint32 TileIndex = ReadPageTable(PageTableCoord, MipLevel);
+	if (TileIndex == 0)
+	{
+		return Header.NullTileValuesQuantized[AttributesIdx];
+	}
 	const FIntVector3 VoxelCoord = AddressModeVolumeCoord % SPARSE_VOLUME_TILE_RES;
-	const FVector4f Sample = ReadTileDataVoxel(TileIndex, VoxelCoord, AttributesIdx);
+	const FVector4f Sample = ReadTileDataVoxel(TileIndex - 1, VoxelCoord, MipLevel, AttributesIdx); // -1 because we don't store the null tile in memory.
 	return Sample;
 }
 
-void FSparseVolumeTextureData::WriteTileDataVoxel(int32 TileIndex, const FIntVector3& TileDataCoord, int32 AttributesIdx, const FVector4f& Value, int32 DstComponent)
+void FSparseVolumeTextureData::WriteTileDataVoxel(int32 TileIndex, const FIntVector3& TileDataCoord, int32 MipLevel, int32 AttributesIdx, const FVector4f& Value, int32 DstComponent)
 {
 	check(AttributesIdx >= 0 && AttributesIdx <= 1);
-	if ((AttributesIdx == 0 && PhysicalTileDataA.IsEmpty()) || (AttributesIdx == 1 && PhysicalTileDataB.IsEmpty()))
+	if ((AttributesIdx == 0 && MipMaps[MipLevel].PhysicalTileDataA.IsEmpty()) || (AttributesIdx == 1 && MipMaps[MipLevel].PhysicalTileDataB.IsEmpty()))
 	{
 		return;
 	}
@@ -367,8 +336,8 @@ void FSparseVolumeTextureData::WriteTileDataVoxel(int32 TileIndex, const FIntVec
 
 	// Allow TileDataCoord to extend past SPARSE_VOLUME_TILE_RES by SPARSE_VOLUME_TILE_BORDER.
 	// This allows us to write the border voxels without changing how TileDataCoord needs to be computed for other use cases.
-	if (TileDataCoord.X < -SPARSE_VOLUME_TILE_BORDER 
-		|| TileDataCoord.Y < -SPARSE_VOLUME_TILE_BORDER 
+	if (TileDataCoord.X < -SPARSE_VOLUME_TILE_BORDER
+		|| TileDataCoord.Y < -SPARSE_VOLUME_TILE_BORDER
 		|| TileDataCoord.Z < -SPARSE_VOLUME_TILE_BORDER
 		|| TileDataCoord.X >= (SPARSE_VOLUME_TILE_RES + SPARSE_VOLUME_TILE_BORDER)
 		|| TileDataCoord.Y >= (SPARSE_VOLUME_TILE_RES + SPARSE_VOLUME_TILE_BORDER)
@@ -379,7 +348,7 @@ void FSparseVolumeTextureData::WriteTileDataVoxel(int32 TileIndex, const FIntVec
 
 	const FInt64Vector3 CoordPadded = FInt64Vector3(TileDataCoord) + FInt64Vector3(SPARSE_VOLUME_TILE_BORDER);
 	const int64 VoxelIndex = (int64)TileIndex * (int64)UE::SVT::SVTNumVoxelsPerPaddedTile + CoordPadded.Z * (SPARSE_VOLUME_TILE_RES_PADDED * SPARSE_VOLUME_TILE_RES_PADDED) + CoordPadded.Y * SPARSE_VOLUME_TILE_RES_PADDED + CoordPadded.X;
-	uint8* TileData = AttributesIdx == 0 ? PhysicalTileDataA.GetData() : PhysicalTileDataB.GetData();
+	uint8* TileData = AttributesIdx == 0 ? MipMaps[MipLevel].PhysicalTileDataA.GetData() : MipMaps[MipLevel].PhysicalTileDataB.GetData();
 	const EPixelFormat Format = Header.AttributesFormats[AttributesIdx];
 	UE::SVT::WriteVoxel(VoxelIndex, TileData, Format, Value, DstComponent);
 }
@@ -388,16 +357,14 @@ bool FSparseVolumeTextureData::GenerateMipMaps(const FSparseVolumeTextureDataAdd
 {
 	const int32 FormatSize[] = { GPixelFormats[Header.AttributesFormats[0]].BlockBytes, GPixelFormats[Header.AttributesFormats[1]].BlockBytes };
 
-	check(!Header.MipInfo.IsEmpty());
-	check(!PageTable.IsEmpty());
+	check(!MipMaps.IsEmpty());
 	check(FMath::IsPowerOfTwo(Header.PageTableVolumeResolution.X)
 		&& FMath::IsPowerOfTwo(Header.PageTableVolumeResolution.Y)
 		&& FMath::IsPowerOfTwo(Header.PageTableVolumeResolution.Z));
-	check(Header.MipInfo[0].TileOffset == 1);
-	const int32 NumMip0Tiles = Header.MipInfo[0].TileOffset + Header.MipInfo[0].TileCount;
-	check(PhysicalTileDataA.Num() >= ((int64)NumMip0Tiles * (int64)UE::SVT::SVTNumVoxelsPerPaddedTile * (int64)FormatSize[0]));
-	check(PhysicalTileDataB.Num() >= ((int64)NumMip0Tiles * (int64)UE::SVT::SVTNumVoxelsPerPaddedTile * (int64)FormatSize[1]));
-	check(PageTable[0].Num() == (Header.PageTableVolumeResolution.X * Header.PageTableVolumeResolution.Y * Header.PageTableVolumeResolution.Z));
+	const int32 NumMip0Tiles = MipMaps[0].NumPhysicalTiles;
+	check(MipMaps[0].PhysicalTileDataA.Num() >= ((int64)NumMip0Tiles * (int64)UE::SVT::SVTNumVoxelsPerPaddedTile * (int64)FormatSize[0]));
+	check(MipMaps[0].PhysicalTileDataB.Num() >= ((int64)NumMip0Tiles * (int64)UE::SVT::SVTNumVoxelsPerPaddedTile * (int64)FormatSize[1]));
+	check(MipMaps[0].PageTable.Num() == (Header.PageTableVolumeResolution.X * Header.PageTableVolumeResolution.Y * Header.PageTableVolumeResolution.Z));
 	
 	if (NumMipLevels == -1)
 	{
@@ -409,16 +376,14 @@ bool FSparseVolumeTextureData::GenerateMipMaps(const FSparseVolumeTextureDataAdd
 			++NumMipLevels;
 		}
 	}
-	Header.MipInfo.SetNum(1);
-	PageTable.SetNum(1);
-	PhysicalTileDataA.SetNum((int64)NumMip0Tiles * (int64)UE::SVT::SVTNumVoxelsPerPaddedTile * (int64)FormatSize[0]);
-	PhysicalTileDataB.SetNum((int64)NumMip0Tiles * (int64)UE::SVT::SVTNumVoxelsPerPaddedTile * (int64)FormatSize[1]);
+
+	MipMaps.SetNum(1);
+	MipMaps[0].PhysicalTileDataA.SetNum((int64)NumMip0Tiles * (int64)UE::SVT::SVTNumVoxelsPerPaddedTile * (int64)FormatSize[0]);
+	MipMaps[0].PhysicalTileDataB.SetNum((int64)NumMip0Tiles * (int64)UE::SVT::SVTNumVoxelsPerPaddedTile * (int64)FormatSize[1]);
 
 	for (int32 MipLevel = 1; MipLevel < NumMipLevels; ++MipLevel)
 	{
-		Header.MipInfo.AddDefaulted();
-		FSparseVolumeTextureMipInfo& MipInfo = Header.MipInfo[MipLevel];
-		MipInfo.TileOffset = Header.MipInfo[MipLevel - 1].TileOffset + Header.MipInfo[MipLevel - 1].TileCount;
+		MipMaps.AddDefaulted();
 
 		const FIntVector3 PageTableVolumeAABBMin = (Header.VirtualVolumeAABBMin / SPARSE_VOLUME_TILE_RES) >> MipLevel;
 		const FIntVector3 ParentPageTableVolumeAABBMin = (Header.VirtualVolumeAABBMin / SPARSE_VOLUME_TILE_RES) >> (MipLevel - 1);
@@ -468,20 +433,20 @@ bool FSparseVolumeTextureData::GenerateMipMaps(const FSparseVolumeTextureDataAdd
 		LinearAllocatedPages.SetNum(NumAllocatedPages);
 
 		// Initialize the SparseVolumeTexture page and tile.
-		PageTable.AddDefaulted();
-		PageTable[MipLevel].SetNumZeroed(PageTableVolumeResolution.X * PageTableVolumeResolution.Y * PageTableVolumeResolution.Z);
-		PhysicalTileDataA.AddZeroed((int64)NumAllocatedPages * (int64)UE::SVT::SVTNumVoxelsPerPaddedTile * (int64)FormatSize[0]);
-		PhysicalTileDataB.AddZeroed((int64)NumAllocatedPages * (int64)UE::SVT::SVTNumVoxelsPerPaddedTile * (int64)FormatSize[1]);
+		MipMaps[MipLevel].PageTable.SetNumZeroed(PageTableVolumeResolution.X * PageTableVolumeResolution.Y * PageTableVolumeResolution.Z);
+		MipMaps[MipLevel].PhysicalTileDataA.SetNum((int64)NumAllocatedPages * (int64)UE::SVT::SVTNumVoxelsPerPaddedTile * (int64)FormatSize[0]);
+		MipMaps[MipLevel].PhysicalTileDataB.SetNum((int64)NumAllocatedPages * (int64)UE::SVT::SVTNumVoxelsPerPaddedTile * (int64)FormatSize[1]);
+		MipMaps[MipLevel].NumPhysicalTiles = NumAllocatedPages;
 
 		// Generate page table and tile volume data by splatting the data
 		{
-			ParallelFor(NumAllocatedPages, [this, &LinearAllocatedPages, &PageTableVolumeResolution, &PageTableVolumeAABBMin, &AddressingInfo, MipLevel, BaseDstTileIndex = (uint32)MipInfo.TileOffset](int32 PageIndex)
+			ParallelFor(NumAllocatedPages, [this, &LinearAllocatedPages, &PageTableVolumeResolution, &PageTableVolumeAABBMin, &AddressingInfo, MipLevel](int32 PageIndex)
 				{
 					const FIntVector3 PageCoordToSplat = LinearAllocatedPages[PageIndex];
-					const uint32 DstTileIndex = BaseDstTileIndex + PageIndex;
+					const uint32 DstTileIndex = PageIndex + 1; // Page table entries are relative to the mip level, but 0 is reserved for null tile
 
 					// Setup the page table entry
-					PageTable[MipLevel]
+					MipMaps[MipLevel].PageTable
 						[
 							PageCoordToSplat.Z * PageTableVolumeResolution.X * PageTableVolumeResolution.Y +
 							PageCoordToSplat.Y * PageTableVolumeResolution.X +
@@ -492,7 +457,7 @@ bool FSparseVolumeTextureData::GenerateMipMaps(const FSparseVolumeTextureDataAdd
 					const FIntVector3 ParentVolumeCoordBase = (PageTableVolumeAABBMin + PageCoordToSplat) * SPARSE_VOLUME_TILE_RES * 2;
 					for (int32 AttributesIdx = 0; AttributesIdx < 2; ++AttributesIdx)
 					{
-						if ((AttributesIdx == 0 && PhysicalTileDataA.IsEmpty()) || (AttributesIdx == 1 && PhysicalTileDataB.IsEmpty()))
+						if (Header.AttributesFormats[AttributesIdx] == PF_Unknown)
 						{
 							continue;
 						}
@@ -514,14 +479,12 @@ bool FSparseVolumeTextureData::GenerateMipMaps(const FSparseVolumeTextureDataAdd
 									DownsampledValue /= 8.0f;
 
 									const FIntVector3 WriteCoord = FIntVector3(X, Y, Z);
-									WriteTileDataVoxel(DstTileIndex, WriteCoord, AttributesIdx, DownsampledValue);
+									WriteTileDataVoxel(DstTileIndex - 1, WriteCoord, MipLevel, AttributesIdx, DownsampledValue); // -1 because we don't store the null tile in memory.
 								}
 							}
 						}
 					}
 				});
-
-			MipInfo.TileCount = NumAllocatedPages;
 		}
 
 		if (!GenerateBorderVoxels(AddressingInfo, MipLevel, LinearAllocatedPages))
@@ -543,12 +506,8 @@ bool FSparseVolumeTextureData::GenerateBorderVoxels(const FSparseVolumeTextureDa
 	ParallelFor(PageCoords.Num(), [this, MipLevel, &PageCoords, &PageTableVolumeResolution, &AddressingInfo](int32 PageIndex)
 		{
 			const FIntVector3& PageCoord = PageCoords[PageIndex];
-			const uint32 DstTileIndex = PageTable[MipLevel]
-			[
-				PageCoord.Z * PageTableVolumeResolution.X * PageTableVolumeResolution.Y +
-				PageCoord.Y * PageTableVolumeResolution.X +
-				PageCoord.X
-			];
+			const int32 PageTableIndex = (PageCoord.Z * PageTableVolumeResolution.X * PageTableVolumeResolution.Y) + (PageCoord.Y * PageTableVolumeResolution.X) + PageCoord.X;
+			const int32 DstTileIndex = (int32)MipMaps[MipLevel].PageTable[PageTableIndex] - 1; // -1 because we don't actually store the null tile in memory
 			const FIntVector3 PageTableOffset = (Header.PageTableVolumeAABBMin >> MipLevel);
 			for (int32 Z = -SPARSE_VOLUME_TILE_BORDER; Z < (SPARSE_VOLUME_TILE_RES + SPARSE_VOLUME_TILE_BORDER); ++Z)
 			{
@@ -562,8 +521,8 @@ bool FSparseVolumeTextureData::GenerateBorderVoxels(const FSparseVolumeTextureDa
 							const FIntVector3 VolumeCoord = (PageTableOffset + PageCoord) * SPARSE_VOLUME_TILE_RES + LocalCoord;
 							const FVector4f Border0 = Load(VolumeCoord, MipLevel, 0, AddressingInfo);
 							const FVector4f Border1 = Load(VolumeCoord, MipLevel, 1, AddressingInfo);
-							WriteTileDataVoxel(DstTileIndex, LocalCoord, 0, Border0);
-							WriteTileDataVoxel(DstTileIndex, LocalCoord, 1, Border1);
+							WriteTileDataVoxel(DstTileIndex, LocalCoord, MipLevel, 0, Border0);
+							WriteTileDataVoxel(DstTileIndex, LocalCoord, MipLevel, 1, Border1);
 						}
 					}
 				}
@@ -575,100 +534,99 @@ bool FSparseVolumeTextureData::GenerateBorderVoxels(const FSparseVolumeTextureDa
 
 bool FSparseVolumeTextureData::DeduplicateTiles()
 {
-	const int32 NumMipLevels = Header.MipInfo.Num();
+	const int32 NumMipLevels = MipMaps.Num();
 	const int32 FormatSize[] = { GPixelFormats[Header.AttributesFormats[0]].BlockBytes, GPixelFormats[Header.AttributesFormats[1]].BlockBytes };
 	const int32 TileByteSize[] = { UE::SVT::SVTNumVoxelsPerPaddedTile * FormatSize[0], UE::SVT::SVTNumVoxelsPerPaddedTile * FormatSize[1] };
-	uint8* PhysicalTileData[] = { PhysicalTileDataA.GetData(), PhysicalTileDataB.GetData() };
 
-	TArray<uint32> TileRemap;
-	TileRemap.Empty(Header.MipInfo.Last().TileOffset + Header.MipInfo.Last().TileCount);
-	TileRemap.Add(0); // null tile
-	int32 TileWritePos = 1; // null tile
-	int32 NumProcessedTiles = 1;
+	int32 NumProcessedTiles = 0;
+	int32 NumSurvivingTiles = 0;
 
 	TArray<uint64> TileHashes; // Reused between mip level iterations
+	TArray<uint32> TileRemap; // Reused between mip level iterations
 
 	// Deduplicate each mip level individually
 	for (int32 MipLevel = 0; MipLevel < NumMipLevels; ++MipLevel)
 	{
-		const FSparseVolumeTextureMipInfo& MipInfo = Header.MipInfo[MipLevel];
-		const int32 MipTileOffset = TileWritePos;
-		TileHashes.Reset();
+		TileHashes.Reset(MipMaps[MipLevel].NumPhysicalTiles);
+		TileRemap.Reset(MipMaps[MipLevel].NumPhysicalTiles);
+		uint8* PhysicalTileData[] = { MipMaps[MipLevel].PhysicalTileDataA.GetData(), MipMaps[MipLevel].PhysicalTileDataB.GetData() };
+		int32 TileWritePos = 0;
 
-		// Compute a hash per tile and compare against previous hashes
-		for (int32 TileIdx = 0; TileIdx < MipInfo.TileCount; ++TileIdx)
+		// Compute a hash per tile and compare against previous hashes.
+		for (int32 TileIdx = 0; TileIdx < MipMaps[MipLevel].NumPhysicalTiles; ++TileIdx)
 		{
-			const int32 GlobalTileIdx = TileIdx + MipInfo.TileOffset;
-			
 			uint64 Hash = 0;
 			if (FormatSize[0])
 			{
-				Hash = CityHash64((const char*)PhysicalTileData[0] + (SIZE_T)GlobalTileIdx * (SIZE_T)TileByteSize[0], (uint32)TileByteSize[0]);
+				Hash = CityHash64((const char*)PhysicalTileData[0] + (SIZE_T)TileIdx * (SIZE_T)TileByteSize[0], (uint32)TileByteSize[0]);
 			}
 			if (FormatSize[1])
 			{
-				Hash = CityHash64WithSeed((const char*)PhysicalTileData[1] + (SIZE_T)GlobalTileIdx * (SIZE_T)TileByteSize[1], (uint32)TileByteSize[1], Hash);
+				Hash = CityHash64WithSeed((const char*)PhysicalTileData[1] + (SIZE_T)TileIdx * (SIZE_T)TileByteSize[1], (uint32)TileByteSize[1], Hash);
 			}
 
 			const int32 OldNumUniqueTiles = TileHashes.Num();
-			const uint32 RemapIndex = TileHashes.AddUnique(Hash);
-			TileRemap.Add(MipTileOffset + RemapIndex);
+			const uint32 RemapIndex = TileHashes.AddUnique(Hash) + 1; // +1 because of implicit null tile at 0
+			TileRemap.Add(RemapIndex);
 			
 			const bool bAddedNewUniqueTile = TileHashes.Num() > OldNumUniqueTiles;
 			
 			// Move tile to compacted position
 			if (bAddedNewUniqueTile)
 			{
-				if (TileWritePos < GlobalTileIdx)
+				if (TileWritePos < TileIdx)
 				{
 					if (FormatSize[0])
 					{
-						FMemory::Memcpy(PhysicalTileData[0] + (SIZE_T)TileWritePos * (SIZE_T)TileByteSize[0], PhysicalTileData[0] + (SIZE_T)GlobalTileIdx * (SIZE_T)TileByteSize[0], (SIZE_T)TileByteSize[0]);
+						FMemory::Memcpy(PhysicalTileData[0] + (SIZE_T)TileWritePos * (SIZE_T)TileByteSize[0], PhysicalTileData[0] + (SIZE_T)TileIdx * (SIZE_T)TileByteSize[0], (SIZE_T)TileByteSize[0]);
 					}
 					if (FormatSize[1])
 					{
-						FMemory::Memcpy(PhysicalTileData[1] + (SIZE_T)TileWritePos * (SIZE_T)TileByteSize[1], PhysicalTileData[1] + (SIZE_T)GlobalTileIdx * (SIZE_T)TileByteSize[1], (SIZE_T)TileByteSize[1]);
+						FMemory::Memcpy(PhysicalTileData[1] + (SIZE_T)TileWritePos * (SIZE_T)TileByteSize[1], PhysicalTileData[1] + (SIZE_T)TileIdx * (SIZE_T)TileByteSize[1], (SIZE_T)TileByteSize[1]);
 					}
 				}
 				else
 				{
-					check(TileWritePos == GlobalTileIdx);
+					check(TileWritePos == TileIdx);
 				}
 				++TileWritePos;
+				++NumSurvivingTiles;
 			}
 
 			++NumProcessedTiles;
 		}
 
 		// Update mip info
-		Header.MipInfo[MipLevel].TileOffset = MipTileOffset;
-		Header.MipInfo[MipLevel].TileCount = TileWritePos - MipTileOffset;
+		MipMaps[MipLevel].NumPhysicalTiles = TileWritePos;
 
 		// Fix up page table
-		for (uint32& PageTableEntry : PageTable[MipLevel])
+		for (uint32& PageTableEntry : MipMaps[MipLevel].PageTable)
 		{
-			PageTableEntry = TileRemap[PageTableEntry];
+			if (PageTableEntry != 0)
+			{
+				PageTableEntry = TileRemap[PageTableEntry - 1]; // -1 because of implicit null tile at 0
+			}
 		}
-	}
 
-	const uint64 DeduplicatedTileSizeBytes[] = { (uint64)TileWritePos * (uint64)TileByteSize[0], (uint64)TileWritePos * (uint64)TileByteSize[1] };
-	if (DeduplicatedTileSizeBytes[0] > INT32_MAX || DeduplicatedTileSizeBytes[1] > INT32_MAX)
-	{
-		UE_LOG(LogSparseVolumeTextureData, Error, TEXT("SparseVolumeTexture still requires too much memory after tile deduplication! Physical tile data A: %llu bytes. Physical tile data B: %llu bytes"), 
-			(unsigned long long)DeduplicatedTileSizeBytes[0], (unsigned long long)DeduplicatedTileSizeBytes[1]);
-		return false;
-	}
+		const uint64 DeduplicatedTileSizeBytes[] = { (uint64)TileWritePos * (uint64)TileByteSize[0], (uint64)TileWritePos * (uint64)TileByteSize[1] };
+		if (DeduplicatedTileSizeBytes[0] > INT32_MAX || DeduplicatedTileSizeBytes[1] > INT32_MAX)
+		{
+			UE_LOG(LogSparseVolumeTextureData, Error, TEXT("SparseVolumeTexture still requires too much memory after tile deduplication! Physical tile data A: %llu bytes. Physical tile data B: %llu bytes"),
+				(unsigned long long)DeduplicatedTileSizeBytes[0], (unsigned long long)DeduplicatedTileSizeBytes[1]);
+			return false;
+		}
 
-	if (!PhysicalTileDataA.IsEmpty())
-	{
-		PhysicalTileDataA.SetNum(TileWritePos * TileByteSize[0]);
-	}
-	if (!PhysicalTileDataB.IsEmpty())
-	{
-		PhysicalTileDataB.SetNum(TileWritePos * TileByteSize[1]);
+		if (!MipMaps[MipLevel].PhysicalTileDataA.IsEmpty())
+		{
+			MipMaps[MipLevel].PhysicalTileDataA.SetNum(TileWritePos * TileByteSize[0]);
+		}
+		if (!MipMaps[MipLevel].PhysicalTileDataB.IsEmpty())
+		{
+			MipMaps[MipLevel].PhysicalTileDataB.SetNum(TileWritePos * TileByteSize[1]);
+		}
 	}
 
-	const int32 NumRemovedTiles = NumProcessedTiles - TileWritePos;
+	const int32 NumRemovedTiles = NumProcessedTiles - NumSurvivingTiles;
 	const float RemovedTilesPercentage = ((float)NumRemovedTiles / (float)NumProcessedTiles) * 100.0f;
 	UE_LOG(LogSparseVolumeTextureData, Display, TEXT("SparseVolumeTexture tile deduplication removed %i tiles out of %i (%f%%)"), NumRemovedTiles, NumProcessedTiles, RemovedTilesPercentage);
 
@@ -679,22 +637,15 @@ bool FSparseVolumeTextureData::BuildDerivedData(const FSparseVolumeTextureDataAd
 {
 	OutDerivedData = FSparseVolumeTextureData{};
 	OutDerivedData.Header = Header;
-	OutDerivedData.Header.MipInfo.SetNum(1);
-	OutDerivedData.PageTable.SetNum(1);
+	OutDerivedData.MipMaps.SetNum(1);
 	if (bMoveMip0FromThis)
 	{
-		OutDerivedData.PageTable[0] = MoveTemp(PageTable[0]);
-		OutDerivedData.PhysicalTileDataA = MoveTemp(PhysicalTileDataA);
-		OutDerivedData.PhysicalTileDataB = MoveTemp(PhysicalTileDataB);
-		PageTable.Empty();
-		PhysicalTileDataA.Empty();
-		PhysicalTileDataA.Empty();
+		OutDerivedData.MipMaps[0] = MoveTemp(MipMaps[0]);
+		MipMaps[0] = {};
 	}
 	else
 	{
-		OutDerivedData.PageTable[0] = PageTable[0];
-		OutDerivedData.PhysicalTileDataA = PhysicalTileDataA;
-		OutDerivedData.PhysicalTileDataB = PhysicalTileDataB;
+		OutDerivedData.MipMaps[0] = MipMaps[0];
 	}
 
 	// Generate border voxels of mip0
@@ -708,7 +659,7 @@ bool FSparseVolumeTextureData::BuildDerivedData(const FSparseVolumeTextureDataAd
 				for (int32 PageX = 0; PageX < Header.PageTableVolumeResolution.X; ++PageX)
 				{
 					const int32 PageIndex = PageZ * (Header.PageTableVolumeResolution.Y * Header.PageTableVolumeResolution.X) + PageY * Header.PageTableVolumeResolution.X + PageX;
-					const uint32 PageTableEntry = OutDerivedData.PageTable[0][PageIndex];
+					const uint32 PageTableEntry = OutDerivedData.MipMaps[0].PageTable[PageIndex];
 					if (PageTableEntry != 0)
 					{
 						PageCoords.Add(FIntVector3(PageX, PageY, PageZ));
diff --git a/Engine/Source/Runtime/Engine/Private/SparseVolumeTexture/SparseVolumeTextureSceneProxy.cpp b/Engine/Source/Runtime/Engine/Private/SparseVolumeTexture/SparseVolumeTextureSceneProxy.cpp
--- a/Engine/Source/Runtime/Engine/Private/SparseVolumeTexture/SparseVolumeTextureSceneProxy.cpp
+++ b/Engine/Source/Runtime/Engine/Private/SparseVolumeTexture/SparseVolumeTextureSceneProxy.cpp
@@ -52,27 +52,25 @@ namespace Private
 
 bool FSparseVolumeTextureRuntime::Create(const FSparseVolumeTextureData& TextureData)
 {
-	bool bSuccess = Create(TextureData.Header, TextureData.Header.MipInfo.Num());
+	bool bSuccess = Create(TextureData.Header, TextureData.MipMaps.Num());
 
 	if (!bSuccess)
 	{
 		return false;
 	}
 
-	const int32 NumMipLevels = TextureData.Header.MipInfo.Num();
-	const int32 TileSizeBytesA = GPixelFormats[Header.AttributesFormats[0]].BlockBytes * UE::SVT::SVTNumVoxelsPerPaddedTile;
-	const int32 TileSizeBytesB = GPixelFormats[Header.AttributesFormats[1]].BlockBytes * UE::SVT::SVTNumVoxelsPerPaddedTile;
-
+	const int32 NumMipLevels = TextureData.MipMaps.Num();
+	
 	TArray<FSparseVolumeTextureTileMapping> Mappings;
 	Mappings.Empty(NumMipLevels);
 	for (int32 MipLevel = 0; MipLevel < NumMipLevels; ++MipLevel)
 	{
 		FSparseVolumeTextureTileMapping Mapping;
-		Mapping.TileIndices = TextureData.PageTable[MipLevel].GetData();
-		Mapping.TileDataA = TextureData.PhysicalTileDataA.GetData() + TileSizeBytesA * TextureData.Header.MipInfo[MipLevel].TileOffset;
-		Mapping.TileDataB = TextureData.PhysicalTileDataB.GetData() + TileSizeBytesB * TextureData.Header.MipInfo[MipLevel].TileOffset;
-		Mapping.NumPhysicalTiles = TextureData.Header.MipInfo[MipLevel].TileCount;
-		Mapping.TileIndicesOffset = -TextureData.Header.MipInfo[MipLevel].TileOffset;
+		Mapping.TileIndices = TextureData.MipMaps[MipLevel].PageTable.GetData();
+		Mapping.TileDataA = TextureData.MipMaps[MipLevel].PhysicalTileDataA.GetData();
+		Mapping.TileDataB = TextureData.MipMaps[MipLevel].PhysicalTileDataB.GetData();
+		Mapping.NumPhysicalTiles = TextureData.MipMaps[MipLevel].NumPhysicalTiles;
+		Mapping.TileIndicesOffset = 0;
 
 		Mappings.Add(Mapping);
 	}
@@ -162,10 +160,10 @@ bool FSparseVolumeTextureRuntime::SetTileMappings(const TArrayView<const FSparse
 	using namespace UE::SVT::Private;
 
 	const int32 NumMipLevels = Header.NumMipLevels;
-	check(Mappings.Num() == NumMipLevels)
+	check(Mappings.Num() == NumMipLevels);
 
-		// Compute number of required tiles
-		int32 NumTiles = 1; // always need a null tile
+	// Compute number of required tiles
+	int32 NumTiles = 1; // always need a null tile
 	for (const FSparseVolumeTextureTileMapping& Mapping : Mappings)
 	{
 		NumTiles += Mapping.NumPhysicalTiles;
@@ -196,8 +194,8 @@ bool FSparseVolumeTextureRuntime::SetTileMappings(const TArrayView<const FSparse
 
 	// Allocate memory for tile data
 	const int32 FormatSize[] = { GPixelFormats[Header.AttributesFormats[0]].BlockBytes, GPixelFormats[Header.AttributesFormats[1]].BlockBytes };
-	PhysicalTileDataA.SetNumZeroed(Header.TileDataVolumeResolution.X * Header.TileDataVolumeResolution.Y * Header.TileDataVolumeResolution.Z * FormatSize[0]);
-	PhysicalTileDataB.SetNumZeroed(Header.TileDataVolumeResolution.X * Header.TileDataVolumeResolution.Y * Header.TileDataVolumeResolution.Z * FormatSize[1]);
+	PhysicalTileDataA.SetNum(Header.TileDataVolumeResolution.X * Header.TileDataVolumeResolution.Y * Header.TileDataVolumeResolution.Z * FormatSize[0]);
+	PhysicalTileDataB.SetNum(Header.TileDataVolumeResolution.X * Header.TileDataVolumeResolution.Y * Header.TileDataVolumeResolution.Z * FormatSize[1]);
 	uint8* PhysicalTileDataPtrs[] = { PhysicalTileDataA.GetData(), PhysicalTileDataB.GetData() };
 
 	int32 NumWrittenTiles = 0;
@@ -247,16 +245,15 @@ bool FSparseVolumeTextureRuntime::SetTileMappings(const TArrayView<const FSparse
 				for (int32 PageX = 0; PageX < PageTableRes.X; ++PageX)
 				{
 					const int32 PageIndex = PageZ * (PageTableRes.Y * PageTableRes.X) + PageY * PageTableRes.X + PageX;
-					uint32 TileIndex = Mapping.TileIndices[PageIndex];
-					if (TileIndex == 0) // points to null tile
+					const uint32 MipLocalTileIndex = Mapping.TileIndices[PageIndex];
+					if (MipLocalTileIndex == 0) // points to null tile
 					{
 						PageTablePtr[PageIndex] = 0;
 					}
 					else // points to actual physical tile
 					{
 						// Make index relative to the start index of this mip level
-						TileIndex += Mapping.TileIndicesOffset;
-						TileIndex += MipTileOffset;
+						const uint32 TileIndex = MipTileOffset + MipLocalTileIndex - 1; // -1 because MipLocalTileIndex uses 0 to refer to the null tile, but we only want a single null tile for the entire resource.
 
 						FIntVector3 TileCoord;
 						TileCoord.X = TileIndex % TileCoordSpace.X;
diff --git a/Engine/Source/Runtime/Engine/Public/SparseVolumeTexture/SparseVolumeTextureData.h b/Engine/Source/Runtime/Engine/Public/SparseVolumeTexture/SparseVolumeTextureData.h
--- a/Engine/Source/Runtime/Engine/Public/SparseVolumeTexture/SparseVolumeTextureData.h
+++ b/Engine/Source/Runtime/Engine/Public/SparseVolumeTexture/SparseVolumeTextureData.h
@@ -25,21 +25,6 @@ public:
 	virtual ~ISparseVolumeTextureDataConstructionAdapter() = default;
 };
 
-struct ENGINE_API FSparseVolumeTextureMipInfo
-{
-	int32 TileOffset = 0;
-	int32 TileCount = 0;
-};
-
-struct ENGINE_API FSparseVolumeTextureDataHeader : public FSparseVolumeTextureHeader
-{
-	static const uint32 kVersion = 0; // The current data format version for the header.
-	uint32 Version = kVersion; // This version can be used to convert existing header to new version later.
-	TArray<FSparseVolumeTextureMipInfo> MipInfo;
-
-	void Serialize(FArchive& Ar);
-};
-
 struct ENGINE_API FSparseVolumeTextureDataAddressingInfo
 {
 	FIntVector3 VolumeResolution;
@@ -49,26 +34,37 @@ struct ENGINE_API FSparseVolumeTextureDataAddressingInfo
 };
 
 // Holds the data for a SparseVolumeTexture that is stored on disk. It only has a single mip after importing a source asset. The mip chain is built during cooking.
-// Tiles are addressed by a flat index; unlike the runtime representation, this one stores all tiles in a 1D array and doesn't have the concept of a 3D physical tile texture.
+// Tiles are addressed by a flat index; unlike the runtime representation, this one stores all tiles in a 1D array (per mip level) and doesn't have the concept of a 3D physical tile texture.
 // The page table itself is 3D though.
 struct ENGINE_API FSparseVolumeTextureData
 {
+	struct FMipMap
+	{
+		TArray<uint32> PageTable;
+		TArray64<uint8> PhysicalTileDataA;
+		TArray64<uint8> PhysicalTileDataB;
+		int32 NumPhysicalTiles;
+	};
+
 	static const uint32 kVersion = 0; // The current data format version for the raw source data.
 	uint32 Version = kVersion; // This version can be used to convert existing source data to new version later.
 
-	FSparseVolumeTextureDataHeader Header = {};
-	TArray<TArray<uint32>> PageTable; // PageTable[MipLevel][PageCoord]
-	TArray64<uint8> PhysicalTileDataA;
-	TArray64<uint8> PhysicalTileDataB;
+	FSparseVolumeTextureHeader Header = {};
+	TArray<FMipMap> MipMaps;
 
 	void Serialize(FArchive& Ar);
+	
 	bool Construct(const ISparseVolumeTextureDataConstructionAdapter& Adapter);
-	uint32 ReadPageTable(const FIntVector3& PageTableCoord, int32 MipLevel) const;
-	FVector4f ReadTileDataVoxel(int32 TileIndex, const FIntVector3& TileDataCoord, int32 AttributesIdx) const;
 	FVector4f Load(const FIntVector3& VolumeCoord, int32 MipLevel, int32 AttributesIdx, const FSparseVolumeTextureDataAddressingInfo& AddressingInfo) const;
-	void WriteTileDataVoxel(int32 TileIndex, const FIntVector3& TileDataCoord, int32 AttributesIdx, const FVector4f& Value, int32 DstComponent = -1);
+	bool BuildDerivedData(const FSparseVolumeTextureDataAddressingInfo& AddressingInfo, int32 NumMipLevels, bool bMoveMip0FromThis, FSparseVolumeTextureData& OutDerivedData);
+
+private:
+
 	bool GenerateMipMaps(const FSparseVolumeTextureDataAddressingInfo& AddressingInfo, int32 NumMipLevels = -1);
 	bool GenerateBorderVoxels(const FSparseVolumeTextureDataAddressingInfo& AddressingInfo, int32 MipLevel, const TArray<FIntVector3>& PageCoords);
 	bool DeduplicateTiles();
-	bool BuildDerivedData(const FSparseVolumeTextureDataAddressingInfo& AddressingInfo, int32 NumMipLevels, bool bMoveMip0FromThis, FSparseVolumeTextureData& OutDerivedData);
+
+	uint32 ReadPageTable(const FIntVector3& PageTableCoord, int32 MipLevel) const;
+	FVector4f ReadTileDataVoxel(int32 TileIndex, const FIntVector3& TileDataCoord, int32 MipLevel, int32 AttributesIdx) const;
+	void WriteTileDataVoxel(int32 TileIndex, const FIntVector3& TileDataCoord, int32 MipLevel, int32 AttributesIdx, const FVector4f& Value, int32 DstComponent = -1);
 };
