World Partition: use FActorContainerID to assign deterministic guids to actors in instanced level instances.

diff --git a/Engine/Source/Runtime/Core/Public/UObject/FortniteMainBranchObjectVersion.h b/Engine/Source/Runtime/Core/Public/UObject/FortniteMainBranchObjectVersion.h
--- a/Engine/Source/Runtime/Core/Public/UObject/FortniteMainBranchObjectVersion.h
+++ b/Engine/Source/Runtime/Core/Public/UObject/FortniteMainBranchObjectVersion.h
@@ -280,6 +280,9 @@ struct CORE_API FFortniteMainBranchObjectVersion
 		// Serialize actor class descriptors
 		WorldPartitionActorClassDescSerialize,
 
+		// FActorContainerID is now an FGuid instead of a uint64
+		WorldPartitionFActorContainerIDu64ToGuid,
+
 		// -----<new versions can be added above this line>-------------------------------------------------
 		VersionPlusOne,
 		LatestVersion = VersionPlusOne - 1
diff --git a/Engine/Source/Runtime/Engine/Private/WorldPartition/HLOD/HLODActor.cpp b/Engine/Source/Runtime/Engine/Private/WorldPartition/HLOD/HLODActor.cpp
--- a/Engine/Source/Runtime/Engine/Private/WorldPartition/HLOD/HLODActor.cpp
+++ b/Engine/Source/Runtime/Engine/Private/WorldPartition/HLOD/HLODActor.cpp
@@ -13,6 +13,7 @@
 
 #if WITH_EDITOR
 #include "Misc/ArchiveMD5.h"
+#include "UObject/FortniteMainBranchObjectVersion.h"
 #include "WorldPartition/HLOD/HLODActorDesc.h"
 #include "WorldPartition/HLOD/HLODLayer.h"
 #include "WorldPartition/HLOD/HLODStats.h"
@@ -32,8 +33,19 @@ static FAutoConsoleVariableRef CVarWorldPartitionHLODForceDisableShadows(
 #if WITH_EDITORONLY_DATA
 FArchive& operator<<(FArchive& Ar, FHLODSubActorDesc& SubActor)
 {
+	Ar.UsingCustomVersion(FFortniteMainBranchObjectVersion::GUID);
+
 	Ar << SubActor.ActorGuid;
-	Ar << SubActor.ContainerID.ID;
+
+	if (Ar.CustomVer(FFortniteMainBranchObjectVersion::GUID) < FFortniteMainBranchObjectVersion::WorldPartitionFActorContainerIDu64ToGuid)
+	{
+		uint64 ContainerID;
+		Ar << ContainerID;
+	}
+	else
+	{
+		Ar << SubActor.ContainerID;
+	}
 
 	return Ar;
 }
diff --git a/Engine/Source/Runtime/Engine/Private/WorldPartition/WorldPartitionLevelHelper.cpp b/Engine/Source/Runtime/Engine/Private/WorldPartition/WorldPartitionLevelHelper.cpp
--- a/Engine/Source/Runtime/Engine/Private/WorldPartition/WorldPartitionLevelHelper.cpp
+++ b/Engine/Source/Runtime/Engine/Private/WorldPartition/WorldPartitionLevelHelper.cpp
@@ -194,11 +194,11 @@ FString FWorldPartitionLevelHelper::AddActorContainerIDToSubPathString(const FAc
 			const int32 SubObjectPos = InSubPathString.Find(TEXT("."), ESearchCase::IgnoreCase, ESearchDir::FromStart, DotPos);
 			if (SubObjectPos == INDEX_NONE)
 			{
-				return InSubPathString + TEXT("_") + InContainerID.ToString();
+				return InSubPathString + TEXT("_") + InContainerID.ToShortString();
 			}
 			else
 			{
-				return InSubPathString.Mid(0, SubObjectPos) + TEXT("_") + InContainerID.ToString() + InSubPathString.Mid(SubObjectPos);
+				return InSubPathString.Mid(0, SubObjectPos) + TEXT("_") + InContainerID.ToShortString() + InSubPathString.Mid(SubObjectPos);
 			}
 		}
 	}
@@ -213,7 +213,7 @@ FString FWorldPartitionLevelHelper::GetContainerPackage(const FActorContainerID&
 	// InDestLevelPackageName will distinguish between instances of the same top level World Partition world (Only needed in PIE, In Cook we always cook the source WP and not an instance and Actor packages no longer exist at runtime)
 	uint64 DestLevelID = 0;
 	TStringBuilder<512> PackageNameBuilder;
-	PackageNameBuilder.Appendf(TEXT("/Temp%s_%s"), *InPackageName, *InContainerID.ToString());
+	PackageNameBuilder.Appendf(TEXT("/Temp%s_%s"), *InPackageName, *InContainerID.ToShortString());
 		
 	if (!InDestLevelPackageName.IsEmpty())
 	{
@@ -431,13 +431,14 @@ bool FWorldPartitionLevelHelper::LoadActors(UWorld* InOwningWorld, ULevel* InDes
 					ContainerWorld->GetSoftObjectPathMapping(SourceWorldPath, RemappedWorldPath);
 
 					// Rename through UObject to avoid changing Actor's external packaging and folder properties
-					Actor->UObject::Rename(*FString::Printf(TEXT("%s_%s"), *Actor->GetName(), *PackageObjectMapping->ContainerID.ToString()), InDestLevel, REN_NonTransactional | REN_ForceNoResetLoaders | REN_DoNotDirty | REN_DontCreateRedirectors);
+					Actor->UObject::Rename(*FString::Printf(TEXT("%s_%s"), *Actor->GetName(), *PackageObjectMapping->ContainerID.ToShortString()), InDestLevel, REN_NonTransactional | REN_ForceNoResetLoaders | REN_DoNotDirty | REN_DontCreateRedirectors);
+
 					// Handle child actors
 					Actor->ForEachComponent<UChildActorComponent>(true, [InDestLevel, PackageObjectMapping](UChildActorComponent* ChildActorComponent)
 					{
 						if (AActor* ChildActor = ChildActorComponent->GetChildActor())
 						{
-							ChildActor->UObject::Rename(*FString::Printf(TEXT("%s_%s"), *ChildActor->GetName(), *PackageObjectMapping->ContainerID.ToString()), InDestLevel, REN_NonTransactional | REN_ForceNoResetLoaders | REN_DoNotDirty | REN_DontCreateRedirectors);
+							ChildActor->UObject::Rename(*FString::Printf(TEXT("%s_%s"), *ChildActor->GetName(), *PackageObjectMapping->ContainerID.ToShortString()), InDestLevel, REN_NonTransactional | REN_ForceNoResetLoaders | REN_DoNotDirty | REN_DontCreateRedirectors);
 						}
 					});
 					
diff --git a/Engine/Source/Runtime/Engine/Private/WorldPartition/WorldPartitionStreamingGeneration.cpp b/Engine/Source/Runtime/Engine/Private/WorldPartition/WorldPartitionStreamingGeneration.cpp
--- a/Engine/Source/Runtime/Engine/Private/WorldPartition/WorldPartitionStreamingGeneration.cpp
+++ b/Engine/Source/Runtime/Engine/Private/WorldPartition/WorldPartitionStreamingGeneration.cpp
@@ -169,7 +169,7 @@ class FWorldPartitionStreamingGenerator
 						ActorSetInstance.ActorsSetGuids.Reserve(ActorSet.Actors.Num());
 						for (const FGuid& ActorGuid : ActorSet.Actors)
 						{
-							ActorSetInstance.ActorsSetGuids.Add(ActorGuid, FGuid::Combine(ContainerInstanceDescriptor.Guid, ActorGuid));
+							ActorSetInstance.ActorsSetGuids.Add(ActorGuid, ActorSetInstance.ContainerID.GetActorGuid(ActorGuid));
 						}
 					}
 
@@ -275,7 +275,6 @@ class FWorldPartitionStreamingGenerator
 		FString OwnerName;
 		FActorContainerID ID;
 		FActorContainerID ParentID;
-		FGuid Guid;
 		TMap<FActorContainerID, TSet<FGuid>> FilteredActors;
 	};
 
@@ -520,11 +519,6 @@ class FWorldPartitionStreamingGenerator
 			SubContainerInstanceDescriptor.ParentID = InContainerInstanceDescriptor.ID;
 			SubContainerInstanceDescriptor.OwnerName = *ContainerInstanceView.GetActorLabelOrName().ToString();
 
-			FContainerInstanceDescriptor& ParentContainer = ContainerInstanceDescriptorsMap.FindChecked(InContainerInstanceDescriptor.ParentID);
-			check(ParentContainer.ID.IsMainContainer() || ParentContainer.Guid.IsValid());
-
-			SubContainerInstanceDescriptor.Guid = FGuid::Combine(ParentContainer.Guid, ActorGuid);
-
 			// Inherit parent container properties
 			FName InheritedRuntimeGrid = InContainerInstanceDescriptor.RuntimeGrid;
 			bool bInheritedbIsSpatiallyLoaded = InContainerInstanceDescriptor.bIsSpatiallyLoaded;
@@ -1010,7 +1004,6 @@ public:
 
 					Ar.Printf(TEXT("       ID: %s"), *ContainerInstanceDescriptor.ID.ToString());
 					Ar.Printf(TEXT(" ParentID: %s"), *ContainerInstanceDescriptor.ParentID.ToString());
-					Ar.Printf(TEXT("     GUID: %s"), *ContainerInstanceDescriptor.Guid.ToString());
 					Ar.Printf(TEXT("   Bounds: %s"), *ContainerInstanceDescriptor.Bounds.ToString());
 					Ar.Printf(TEXT("Transform: %s"), *ContainerInstanceDescriptor.Transform.ToString());
 					Ar.Printf(TEXT("Container: %s"), *ContainerInstanceDescriptor.Container->GetContainerPackage().ToString());
@@ -1018,7 +1011,7 @@ public:
 
 				TArray<FActorContainerID> ChildContainersIDs;
 				InvertedContainersHierarchy.MultiFind(ContainerID, ChildContainersIDs);
-				ChildContainersIDs.Sort([](const FActorContainerID& ActorContainerIDA, const FActorContainerID& ActorContainerIDB) { return ActorContainerIDA.ID < ActorContainerIDB.ID; });
+				ChildContainersIDs.Sort();
 
 				if (ChildContainersIDs.Num())
 				{
diff --git a/Engine/Source/Runtime/Engine/Public/WorldPartition/HLOD/HLODSubActor.h b/Engine/Source/Runtime/Engine/Public/WorldPartition/HLOD/HLODSubActor.h
--- a/Engine/Source/Runtime/Engine/Public/WorldPartition/HLOD/HLODSubActor.h
+++ b/Engine/Source/Runtime/Engine/Public/WorldPartition/HLOD/HLODSubActor.h
@@ -35,7 +35,7 @@ struct FHLODSubActorDesc
 	{
 		if (Lhs.ActorGuid == Rhs.ActorGuid)
 		{
-			return Lhs.ContainerID.ID < Rhs.ContainerID.ID;
+			return Lhs.ContainerID < Rhs.ContainerID;
 		}
 
 		return Lhs.ActorGuid < Rhs.ActorGuid;
@@ -96,7 +96,7 @@ struct FHLODSubActor
 	{
 		if (Lhs.ActorGuid == Rhs.ActorGuid)
 		{
-			return Lhs.ContainerID.ID < Rhs.ContainerID.ID;
+			return Lhs.ContainerID < Rhs.ContainerID;
 		}
 
 		return Lhs.ActorGuid < Rhs.ActorGuid;
diff --git a/Engine/Source/Runtime/Engine/Public/WorldPartition/WorldPartitionActorContainerID.h b/Engine/Source/Runtime/Engine/Public/WorldPartition/WorldPartitionActorContainerID.h
--- a/Engine/Source/Runtime/Engine/Public/WorldPartition/WorldPartitionActorContainerID.h
+++ b/Engine/Source/Runtime/Engine/Public/WorldPartition/WorldPartitionActorContainerID.h
@@ -10,45 +10,69 @@ struct FActorContainerID
 	GENERATED_USTRUCT_BODY()
 
 	FActorContainerID()
-	: ID(0)
 	{}
 
 	FActorContainerID(const FActorContainerID& InOther)
-	: ID(InOther.ID)
+		: Guid(InOther.Guid)
+	{}
+
+	FActorContainerID(const FActorContainerID& InParent, const FActorContainerID& InOther)
+		: Guid(FGuid::Combine(InParent.Guid, InOther.Guid))
 	{}
 
 	FActorContainerID(const FActorContainerID& InParent, const FGuid& InActorGuid)
-	: ID(CityHash64WithSeed((const char*)&InActorGuid, sizeof(InActorGuid), InParent.ID))
+		: Guid(FGuid::Combine(InParent.Guid, InActorGuid))
 	{}
 
 	void operator=(const FActorContainerID& InOther)
 	{
-		ID = InOther.ID;
+		Guid = InOther.Guid;
 	}
 
 	bool operator==(const FActorContainerID& InOther) const
 	{
-		return ID == InOther.ID;
+		return Guid == InOther.Guid;
 	}
 
 	bool operator!=(const FActorContainerID& InOther) const
 	{
-		return ID != InOther.ID;
+		return Guid != InOther.Guid;
+	}
+
+	bool operator<(const FActorContainerID& InOther) const
+	{
+		return Guid < InOther.Guid;
+	}
+
+	friend FArchive& operator<<(FArchive& Ar, FActorContainerID& InActorContainer)
+	{
+		return Ar << InActorContainer.Guid;
 	}
 
 	bool IsMainContainer() const
 	{
-		return !ID;
+		return !Guid.IsValid();
 	}
 
 	FString ToString() const
 	{
+		return Guid.ToString();
+	}
+
+	FString ToShortString() const
+	{
+		const uint64 ID = CityHash64((const char*)&Guid, sizeof(FGuid));
 		return FString::Printf(TEXT("%016llx"), ID);
 	}
 
+	FGuid GetActorGuid(const FGuid& InActorGuid) const
+	{
+		return FGuid::Combine(Guid, InActorGuid);
+	}
+
 	friend FORCEINLINE uint32 GetTypeHash(const FActorContainerID& InContainerID)
 	{
-		return GetTypeHash(InContainerID.ID);
+		return GetTypeHash(InContainerID.Guid);
 	}
 
 	static FActorContainerID GetMainContainerID()
@@ -56,6 +80,7 @@ struct FActorContainerID
 		return FActorContainerID();
 	}
 
+private:
 	UPROPERTY()
-	uint64 ID;
+	FGuid Guid;
 };
