Substrate - updated expresions F0/F90 tooltips.

diff --git a/Engine/Source/Runtime/Engine/Classes/Materials/MaterialExpressionStrata.h b/Engine/Source/Runtime/Engine/Classes/Materials/MaterialExpressionStrata.h
--- a/Engine/Source/Runtime/Engine/Classes/Materials/MaterialExpressionStrata.h
+++ b/Engine/Source/Runtime/Engine/Classes/Materials/MaterialExpressionStrata.h
@@ -207,13 +207,13 @@ class UMaterialExpressionStrataSlabBSDF : public UMaterialExpressionStrataBSDF
 	FExpressionInput DiffuseAlbedo;
 
 	/**
-	 * Defines F0, the percentage of light reflected as specular from a surface when the view is perpendicular to the surface. (type = float3, unit = unitless, defaults to plastic 0.04)
+	 * Defines the color and brightness of the specular highlight where the surface is facing the camera. Each specular contribution will fade to black as F0 drops below 0.02. (type = float3, unit = unitless, defaults to plastic 0.04)
 	 */
 	UPROPERTY()
 	FExpressionInput F0;
 
 	/**
-	 * Defines F90, the percentage of light reflected as specular from a surface when the view is tangent to the surface. (type = float3, unit = unitless, defaults to 1.0f)
+	 * Defines the color of the specular highlight where the surface normal is 90 degrees from the view direction. Only the hue and saturation are preserved, the brightness is fixed at 1.0. Fades to black as F0 drops below 0.02. (type = float3, unit = unitless, defaults to 1.0f)
 	 */
 	UPROPERTY()
 	FExpressionInput F90;
@@ -341,7 +341,7 @@ class UMaterialExpressionStrataSimpleClearCoatBSDF : public UMaterialExpressionS
 	FExpressionInput DiffuseAlbedo;
 
 	/**
-	 * Defines F0, the percentage of light reflected as specular from a surface when the view is perpendicular to the surface. (type = float3, unit = unitless, defaults to plastic 0.04)
+	 * Defines the color and brightness of the specular highlight where the surface is facing the camera. Each specular contribution will fade to black as F0 drops below 0.02. (type = float3, unit = unitless, defaults to plastic 0.04)
 	 */
 	UPROPERTY()
 	FExpressionInput F0;
@@ -1063,13 +1063,13 @@ class UMaterialExpressionStrataThinFilm : public UMaterialExpressionStrataUtilit
 	FExpressionInput Normal;
 
 	/**
-	 * Defines F0, the percentage of light reflected as specular from a surface when the view is perpendicular to the surface. (type = float3, unit = unitless, defaults to plastic 0.04)
+	 * Defines the color and brightness of the specular highlight where the surface is facing the camera. Each specular contribution will fade to black as F0 drops below 0.02. (type = float3, unit = unitless, defaults to plastic 0.04)
 	 */
 	UPROPERTY()
 	FExpressionInput F0;
 
 	/**
-	 * Defines F90, the percentage of light reflected as specular from a surface when the view is tangent to the surface. (type = float3, unit = unitless, defaults to 1.0f)
+	 * Defines the color of the specular highlight where the surface normal is 90 degrees from the view direction. Only the hue and saturation are preserved, the brightness is fixed at 1.0. Fades to black as F0 drops below 0.02. (type = float3, unit = unitless, defaults to 1.0f)
 	 */
 	UPROPERTY()
 	FExpressionInput F90;
