LevelInstanceFilter: fix filter crash caused by [REDACTED_PERSON] actor now having a unique Guid.

#rnx

diff --git a/Engine/Source/Runtime/Engine/Private/LevelInstance/LevelInstanceSubsystem.cpp b/Engine/Source/Runtime/Engine/Private/LevelInstance/LevelInstanceSubsystem.cpp
--- a/Engine/Source/Runtime/Engine/Private/LevelInstance/LevelInstanceSubsystem.cpp
+++ b/Engine/Source/Runtime/Engine/Private/LevelInstance/LevelInstanceSubsystem.cpp
@@ -2585,7 +2585,7 @@ bool ULevelInstanceSubsystem::PassLevelInstanceFilter(UWorld* World, const FWorl
 			TArray<FGuid> AncestorGuids;
 			ForEachLevelInstanceAncestorsAndSelf(Cast<AActor>(LevelInstanceInterface), [&AncestorGuids, &LevelInstanceInterface](const ILevelInstanceInterface* Ancestor)
 			{
-				AncestorGuids.Insert(Cast<AActor>(Ancestor)->GetActorGuid(), 0);
+				AncestorGuids.Insert(Ancestor->GetLevelInstanceID().GetOriginalGuid(), 0);
 				LevelInstanceInterface = Ancestor;
 				return true;
 			});
diff --git a/Engine/Source/Runtime/Engine/Public/LevelInstance/LevelInstanceTypes.h b/Engine/Source/Runtime/Engine/Public/LevelInstance/LevelInstanceTypes.h
--- a/Engine/Source/Runtime/Engine/Public/LevelInstance/LevelInstanceTypes.h
+++ b/Engine/Source/Runtime/Engine/Public/LevelInstance/LevelInstanceTypes.h
@@ -35,6 +35,13 @@ struct FLevelInstanceID
 
 	inline const FGuid& GetGuid() const { return Guid; }
 
+protected:
+#if WITH_EDITOR
+	friend ULevelInstanceSubsystem;
+	// - this relies on the Guids[0] to be the OriginalGuid of the LevelInstance Actor before it gets overriden in ULevelInstanceLevelStreaming.
+	// - this is only used for filtering actors, we compare a filter built from a Level Package (with original actor guids) to the filter of a loaded Level Instance (with unique actor guids)
+	inline const FGuid& GetOriginalGuid() const { return Guids[0]; }
+#endif
 private:
 	uint64 Hash = 0;
 	FGuid Guid;
