LevelInstance: Cache IsContainerInstance (optim)

#rnx

diff --git a/Engine/Source/Runtime/Engine/Private/WorldPartition/LevelInstance/LevelInstanceActorDesc.cpp b/Engine/Source/Runtime/Engine/Private/WorldPartition/LevelInstance/LevelInstanceActorDesc.cpp
--- a/Engine/Source/Runtime/Engine/Private/WorldPartition/LevelInstance/LevelInstanceActorDesc.cpp
+++ b/Engine/Source/Runtime/Engine/Private/WorldPartition/LevelInstance/LevelInstanceActorDesc.cpp
@@ -26,6 +26,7 @@ static FAutoConsoleVariableRef CVarForceLevelStreaming(
 FLevelInstanceActorDesc::FLevelInstanceActorDesc()
 	: DesiredRuntimeBehavior(ELevelInstanceRuntimeBehavior::Partitioned)
 	, LevelInstanceContainer(nullptr)
+	, bIsContainerInstance(false)
 {}
 
 FLevelInstanceActorDesc::~FLevelInstanceActorDesc()
@@ -44,6 +45,8 @@ void FLevelInstanceActorDesc::Init(const AActor* InActor)
 	LevelInstanceTransform = InActor->GetActorTransform();
 	DesiredRuntimeBehavior = LevelInstance->GetDesiredRuntimeBehavior();
 	Filter = LevelInstance->GetFilter();
+	
+	bIsContainerInstance = IsContainerInstanceInternal();
 }
 
 void FLevelInstanceActorDesc::Init(const FWorldPartitionActorDescInitData& DescData)
@@ -53,6 +56,8 @@ void FLevelInstanceActorDesc::Init(const FWorldPartitionActorDescInitData& DescD
 	DesiredRuntimeBehavior = LevelInstanceCDO->GetDefaultRuntimeBehavior();
 
 	FWorldPartitionActorDesc::Init(DescData);
+	
+	bIsContainerInstance = IsContainerInstanceInternal();
 }
 
 bool FLevelInstanceActorDesc::Equals(const FWorldPartitionActorDesc* Other) const
@@ -139,6 +144,11 @@ void FLevelInstanceActorDesc::SetContainer(UActorDescContainer* InContainer, UWo
 }
 
 bool FLevelInstanceActorDesc::IsContainerInstance() const
+{
+	return bIsContainerInstance;
+}
+
+bool FLevelInstanceActorDesc::IsContainerInstanceInternal() const
 {
 	if (DesiredRuntimeBehavior != ELevelInstanceRuntimeBehavior::Partitioned)
 	{
diff --git a/Engine/Source/Runtime/Engine/Public/WorldPartition/LevelInstance/LevelInstanceActorDesc.h b/Engine/Source/Runtime/Engine/Public/WorldPartition/LevelInstance/LevelInstanceActorDesc.h
--- a/Engine/Source/Runtime/Engine/Public/WorldPartition/LevelInstance/LevelInstanceActorDesc.h
+++ b/Engine/Source/Runtime/Engine/Public/WorldPartition/LevelInstance/LevelInstanceActorDesc.h
@@ -47,8 +47,9 @@ protected:
 	TWeakObjectPtr<UWorld> LevelInstanceContainerWorldContext;
 
 	FWorldPartitionActorFilter Filter;
-
+	bool bIsContainerInstance;
 private:
+	bool IsContainerInstanceInternal() const;
 	void RegisterContainerInstance(UWorld* InWorldContext);
 	void UnregisterContainerInstance();
 	void UpdateBounds();
