Setup a way of slicing up the amount of Actors that can change culling distance in a single frame

diff --git a/Engine/Source/Runtime/Engine/Private/WorldPartition/HLOD/HLODSubsystem.cpp b/Engine/Source/Runtime/Engine/Private/WorldPartition/HLOD/HLODSubsystem.cpp
--- a/Engine/Source/Runtime/Engine/Private/WorldPartition/HLOD/HLODSubsystem.cpp
+++ b/Engine/Source/Runtime/Engine/Private/WorldPartition/HLOD/HLODSubsystem.cpp
@@ -71,6 +71,16 @@ static TAutoConsoleVariable<int32> CVarHLODWarmupVTSizeClamp(
 	2048,
 	TEXT("Clamp VT warmup requests for safety."));
 
+static TAutoConsoleVariable<int32> CVarCullDistanceWorkSlicer(
+	TEXT("wp.Runtime.HLOD.CullDistanceWorkSlicer"),
+	60,
+	TEXT("Maximum amount of actors that can have it's CullDistance changed per frames."));
+
+static TAutoConsoleVariable<int32> CVarCullDistanceDifferenceNeeded(
+	TEXT("wp.Runtime.HLOD.CVarCullDistanceDifferenceNeeded"),
+	300,
+	TEXT("The difference needed for the CullDistance changed to be processed"));
+
 static void HLODSubsystemCVarSinkFunction()
 {
 	for (UWorld* World : TObjectRange<UWorld>(RF_ClassDefaultObject | RF_ArchetypeObject, true, EInternalObjectFlags::Garbage))
@@ -139,6 +149,11 @@ UHLODSubsystem::UHLODSubsystem()
 
 UHLODSubsystem::~UHLODSubsystem()
 {
+	if (TickHandle.IsValid())
+	{
+		FTSTicker::GetCoreTicker().RemoveTicker(TickHandle);
+		TickHandle.Reset();
+	}
 }
 
 void UHLODSubsystem::SetHLODAlwaysLoadedCullDistance(int32 InCullDistance)
@@ -146,20 +161,47 @@ void UHLODSubsystem::SetHLODAlwaysLoadedCullDistance(int32 InCullDistance)
 	const float MaxDrawDistance = FMath::Max(InCullDistance, 0);
 	const bool bNeverDistanceCull = InCullDistance <= 0;
 
-	for (AWorldPartitionHLOD* HLODActor : AlwaysLoadedHLODActors)
+	//If the WorkSlicer is set to process 0 items, revert back to processing everything at once
+	if (CVarCullDistanceWorkSlicer.GetValueOnAnyThread() > 0)
 	{
-		HLODActor->ForEachComponent<UPrimitiveComponent>(false, [MaxDrawDistance, bNeverDistanceCull](UPrimitiveComponent* HLODComponent)
+		int32 CullDiff = FMath::Abs(LastSetCullDistance - InCullDistance);
+		if (!OperationQueue.IsEmpty() || CullDiff < CVarCullDistanceDifferenceNeeded.GetValueOnAnyThread())
 		{
-			HLODComponent->bNeverDistanceCull = bNeverDistanceCull;
-			HLODComponent->SetCachedMaxDrawDistance(MaxDrawDistance);
+			return;
+		}
+		LastSetCullDistance = InCullDistance;
 
-			// Enable per instance culling to avoid ISM components dissapearing as a whole
-			if (UInstancedStaticMeshComponent* ISMComponent = Cast<UInstancedStaticMeshComponent>(HLODComponent))
+		OperationQueue.Reserve(AlwaysLoadedHLODActors.Num());
+		for (AWorldPartitionHLOD* HLODActor : AlwaysLoadedHLODActors)
+		{
+			FDrawDistanceQueue Item = FDrawDistanceQueue();
+			Item.DrawDistance = MaxDrawDistance;
+			Item.HLODActor = HLODActor;
+			OperationQueue.Add(Item);
+		}
+
+		if (!TickHandle.IsValid())
+		{
+			TickHandle = FTSTicker::GetCoreTicker().AddTicker(FTickerDelegate::CreateUObject(this, &UHLODSubsystem::Tick), 0.0f);
+		}
+	}
+	else
+	{
+		for (AWorldPartitionHLOD* HLODActor : AlwaysLoadedHLODActors)
+		{
+			HLODActor->ForEachComponent<UPrimitiveComponent>(false, [MaxDrawDistance, bNeverDistanceCull](UPrimitiveComponent* HLODComponent)
 			{
-				ISMComponent->SetCullDistances(0, MaxDrawDistance);
-			}
-		});
+				HLODComponent->bNeverDistanceCull = bNeverDistanceCull;
+				HLODComponent->SetCachedMaxDrawDistance(MaxDrawDistance);
+				// Enable per instance culling to avoid ISM components dissapearing as a whole
+				if (UInstancedStaticMeshComponent* ISMComponent = Cast<UInstancedStaticMeshComponent>(HLODComponent))
+				{
+					ISMComponent->SetCullDistances(0, MaxDrawDistance);
+				}
+			});
+		}
 	}
+	
 }
 
 void UHLODSubsystem::OnWorldBeginPlay(UWorld& InWorld)
@@ -178,6 +220,37 @@ void UHLODSubsystem::OnWorldBeginPlay(UWorld& InWorld)
 			AlwaysLoadedHLODActors.Add(HLODActor);
 		}
 	}
+
+}
+
+bool UHLODSubsystem::Tick(float DeltaTime)
+{
+	//If the CVarCullDistanceAllowWorkSlicer is false we will never add this function to FTSTicker
+	if (OperationQueue.Num() > 0)
+	{
+		int32 Amount = FMath::Min(OperationQueue.Num(), CVarCullDistanceWorkSlicer.GetValueOnAnyThread());	
+		for (int32 i = OperationQueue.Num() - Amount; i < OperationQueue.Num(); i++)
+		{
+			const float MaxDrawDistance = OperationQueue[i].DrawDistance;
+			TObjectPtr<AWorldPartitionHLOD> HLODActor = OperationQueue[i].HLODActor;
+			if (HLODActor != nullptr)
+			{
+				HLODActor->ForEachComponent<UPrimitiveComponent>(false, [MaxDrawDistance](UPrimitiveComponent* HLODComponent)
+				{
+					HLODComponent->bNeverDistanceCull = false;
+					HLODComponent->SetCachedMaxDrawDistance(MaxDrawDistance);
+					// Enable per instance culling to avoid ISM components dissapearing as a whole
+					if (UInstancedStaticMeshComponent* ISMComponent = Cast<UInstancedStaticMeshComponent>(HLODComponent))
+					{
+						ISMComponent->SetCullDistances(0, MaxDrawDistance);
+					}
+				});
+			}
+			
+		}
+		OperationQueue.RemoveAt(OperationQueue.Num() - Amount, Amount);
+	}
+	return true;
 }
 
 bool UHLODSubsystem::WorldPartitionHLODEnabled = true;
diff --git a/Engine/Source/Runtime/Engine/Public/WorldPartition/HLOD/HLODSubsystem.h b/Engine/Source/Runtime/Engine/Public/WorldPartition/HLOD/HLODSubsystem.h
--- a/Engine/Source/Runtime/Engine/Public/WorldPartition/HLOD/HLODSubsystem.h
+++ b/Engine/Source/Runtime/Engine/Public/WorldPartition/HLOD/HLODSubsystem.h
@@ -5,6 +5,7 @@
 #include "Containers/Array.h"
 #include "Containers/Map.h"
 #include "UObject/ObjectKey.h"
+#include "Containers/Ticker.h"
 #include "HLODSubsystem.generated.h"
 
 
@@ -42,6 +43,8 @@ public:
 	virtual void OnWorldBeginPlay(UWorld& InWorld) override;
 	//~ End UWorldSubsystem Interface.
 
+	bool Tick(float DeltaTime);
+
 	void RegisterHLODActor(AWorldPartitionHLOD* InWorldPartitionHLOD);
 	void UnregisterHLODActor(AWorldPartitionHLOD* InWorldPartitionHLOD);
 
@@ -76,6 +79,12 @@ private:
 	{
 		TMap<FGuid, FCellData> CellsData;
 	};
+
+	struct FDrawDistanceQueue
+	{
+		float DrawDistance;
+		TObjectPtr<AWorldPartitionHLOD> HLODActor;
+	};
 	
 	TMap<TObjectPtr<UWorldPartition>, FWorldPartitionHLODRuntimeData> WorldPartitionsHLODRuntimeData;
 	const FCellData* GetCellData(const UWorldPartitionRuntimeCell* InCell) const;
@@ -100,6 +109,10 @@ private:
 	bool ShouldPerformWarmupForCell(const UWorldPartitionRuntimeCell* InCell) const;
 	bool bCachedShouldPerformWarmup;
 
+	TArray<FDrawDistanceQueue> OperationQueue;
+	FTSTicker::FDelegateHandle TickHandle;
+	int32 LastSetCullDistance;
+
 	/** Console command used to turn on/off loading & rendering of world partition HLODs */
 	static class FAutoConsoleCommand EnableHLODCommand;
 
