Fully resolve the redirector before validating world.

diff --git a/Engine/Source/Runtime/CoreUObject/Private/UObject/SavePackage2.cpp b/Engine/Source/Runtime/CoreUObject/Private/UObject/SavePackage2.cpp
--- a/Engine/Source/Runtime/CoreUObject/Private/UObject/SavePackage2.cpp
+++ b/Engine/Source/Runtime/CoreUObject/Private/UObject/SavePackage2.cpp
@@ -556,7 +556,7 @@ ESavePackageResult ValidateExports(FSaveContext& SaveContext)
 		{
 			UObject* Object = Export.Obj;
 			// Consider redirectors to world/levels as map packages too.
-			if (UObjectRedirector* Redirector = Cast<UObjectRedirector>(Object))
+			while (UObjectRedirector* Redirector = Cast<UObjectRedirector>(Object))
 			{
 				Object = Redirector->DestinationObject;
 			}
diff --git a/Engine/Source/Runtime/Engine/Private/Level.cpp b/Engine/Source/Runtime/Engine/Private/Level.cpp
--- a/Engine/Source/Runtime/Engine/Private/Level.cpp
+++ b/Engine/Source/Runtime/Engine/Private/Level.cpp
@@ -3971,6 +3971,8 @@ TArray<UPackage*> ULevel::GetLoadedExternalObjectPackages() const
 
 	// We also need to provide empty packages (for example, actors that were converted to non-external)
 	UWorld* World = GetTypedOuter<UWorld>();
+	check(IsValid(World));
+
 	TArray<FString> ExternalObjectsPaths;
 	ExternalObjectsPaths.Add(SanitizeExternalPath(ULevel::GetExternalActorsPath(World->GetPackage(), (World->OriginalWorldName == NAME_None) ? World->GetName() : World->OriginalWorldName.ToString())));
 	ExternalObjectsPaths.Add(SanitizeExternalPath(FExternalPackageHelper::GetExternalObjectsPath(World->GetPackage(), (World->OriginalWorldName == NAME_None) ? World->GetName() : World->OriginalWorldName.ToString())));
