WorldPartition - Helpers: Removed logic forcing FWorldPartitionHelpers::ShouldCollectGarbage() to return true at a fixed frequency (30 sec)
* Greatly speed up iterative WP builders (HLOD/minimap/etc)

diff --git a/Engine/Source/Runtime/Engine/Private/WorldPartition/WorldPartitionHelpers.cpp b/Engine/Source/Runtime/Engine/Private/WorldPartition/WorldPartitionHelpers.cpp
--- a/Engine/Source/Runtime/Engine/Private/WorldPartition/WorldPartitionHelpers.cpp
+++ b/Engine/Source/Runtime/Engine/Private/WorldPartition/WorldPartitionHelpers.cpp
@@ -238,15 +238,6 @@ bool FWorldPartitionHelpers::HasExceededMaxMemory()
 
 bool FWorldPartitionHelpers::ShouldCollectGarbage()
 {
-	// Even if we're not exhausting memory, GC should be run at periodic intervals
-	static double LastTime = FPlatformTime::Seconds();
-	double ThisTime = FPlatformTime::Seconds();
-	if ((ThisTime - LastTime) > 30)
-	{
-		LastTime = ThisTime;
-		return true;
-	}
-
 	return HasExceededMaxMemory();
 }
 
