WorldPartition: Add missing ContainerInstanceViews validation

#rnx

diff --git a/Engine/Source/Runtime/Engine/Private/WorldPartition/WorldPartitionStreamingGeneration.cpp b/Engine/Source/Runtime/Engine/Private/WorldPartition/WorldPartitionStreamingGeneration.cpp
--- a/Engine/Source/Runtime/Engine/Private/WorldPartition/WorldPartitionStreamingGeneration.cpp
+++ b/Engine/Source/Runtime/Engine/Private/WorldPartition/WorldPartitionStreamingGeneration.cpp
@@ -642,6 +642,11 @@ class FWorldPartitionStreamingGenerator
 			ActorDescView.CheckForErrors(ErrorHandler);
 		});
 
+		for (FWorldPartitionActorDescView& ContainerInstanceView : ContainerDescriptor.ContainerInstanceViews)
+		{
+			ContainerInstanceView.CheckForErrors(ErrorHandler);
+		}
+
 		// Perform various adjustements based on validations and report errors
 		//
 		// The first validation pass is used to report errors, subsequent passes are used to make corrections to the actor descriptor views.
