Modifications to Heterogeneous Volumes component

diff --git a/Engine/Source/Runtime/Engine/Classes/Components/HeterogeneousVolumeComponent.h b/Engine/Source/Runtime/Engine/Classes/Components/HeterogeneousVolumeComponent.h
--- a/Engine/Source/Runtime/Engine/Classes/Components/HeterogeneousVolumeComponent.h
+++ b/Engine/Source/Runtime/Engine/Classes/Components/HeterogeneousVolumeComponent.h
@@ -23,7 +23,7 @@ class UHeterogeneousVolumeComponent : public UMeshComponent
 	FIntVector VolumeResolution;
 
 	UPROPERTY(EditAnywhere, BlueprintReadOnly, Category = Volume)
-	float MinimumVoxelSize;
+	float Framerate;
 
 	UPROPERTY(EditAnywhere, BlueprintReadOnly, Category = Volume)
 	uint32 bAnimate : 1;
@@ -47,7 +47,6 @@ public:
 
 private:
 	float Time;
-	float Framerate;
 };
 
 /**
diff --git a/Engine/Source/Runtime/Engine/Private/Components/HeterogeneousVolumeComponent.cpp b/Engine/Source/Runtime/Engine/Private/Components/HeterogeneousVolumeComponent.cpp
--- a/Engine/Source/Runtime/Engine/Private/Components/HeterogeneousVolumeComponent.cpp
+++ b/Engine/Source/Runtime/Engine/Private/Components/HeterogeneousVolumeComponent.cpp
@@ -70,7 +70,14 @@ FHeterogeneousVolumeSceneProxy::FHeterogeneousVolumeSceneProxy(UHeterogeneousVol
 	bIsHeterogeneousVolume = true;
 
 	HeterogeneousVolumeData.VoxelResolution = InComponent->VolumeResolution;
-	HeterogeneousVolumeData.MinimumVoxelSize = InComponent->MinimumVoxelSize;
+
+	// Infer minimum voxel size from bounds and resolution
+	FVector VoxelSize = 2.0 * InComponent->Bounds.BoxExtent;
+	VoxelSize.X /= InComponent->VolumeResolution.X;
+	VoxelSize.Y /= InComponent->VolumeResolution.Y;
+	VoxelSize.Z /= InComponent->VolumeResolution.Z;
+	HeterogeneousVolumeData.MinimumVoxelSize = FMath::Max(VoxelSize.GetMin(), 0.001);
+
 	HeterogeneousVolumeData.LightingDownsampleFactor = InComponent->LightingDownsampleFactor;
 
 	// Update material assignment to include Heterogeneous Volumes
@@ -109,19 +116,20 @@ FHeterogeneousVolumeSceneProxy::FHeterogeneousVolumeSceneProxy(UHeterogeneousVol
 		[Self](FRHICommandListImmediate& RHICmdList)
 		{
 			Self->StaticMeshVertexBuffers.PositionVertexBuffer.InitResource();
-	Self->StaticMeshVertexBuffers.StaticMeshVertexBuffer.InitResource();
-	Self->StaticMeshVertexBuffers.ColorVertexBuffer.InitResource();
-
-	FLocalVertexFactory::FDataType Data;
-	Self->StaticMeshVertexBuffers.PositionVertexBuffer.BindPositionVertexBuffer(&Self->VertexFactory, Data);
-	Self->StaticMeshVertexBuffers.StaticMeshVertexBuffer.BindTangentVertexBuffer(&Self->VertexFactory, Data);
-	Self->StaticMeshVertexBuffers.StaticMeshVertexBuffer.BindPackedTexCoordVertexBuffer(&Self->VertexFactory, Data);
-	Self->StaticMeshVertexBuffers.StaticMeshVertexBuffer.BindLightMapVertexBuffer(&Self->VertexFactory, Data, 0);
-	Self->StaticMeshVertexBuffers.ColorVertexBuffer.BindColorVertexBuffer(&Self->VertexFactory, Data);
-	Self->VertexFactory.SetData(Data);
-
-	Self->VertexFactory.InitResource();
-		});
+			Self->StaticMeshVertexBuffers.StaticMeshVertexBuffer.InitResource();
+			Self->StaticMeshVertexBuffers.ColorVertexBuffer.InitResource();
+
+			FLocalVertexFactory::FDataType Data;
+			Self->StaticMeshVertexBuffers.PositionVertexBuffer.BindPositionVertexBuffer(&Self->VertexFactory, Data);
+			Self->StaticMeshVertexBuffers.StaticMeshVertexBuffer.BindTangentVertexBuffer(&Self->VertexFactory, Data);
+			Self->StaticMeshVertexBuffers.StaticMeshVertexBuffer.BindPackedTexCoordVertexBuffer(&Self->VertexFactory, Data);
+			Self->StaticMeshVertexBuffers.StaticMeshVertexBuffer.BindLightMapVertexBuffer(&Self->VertexFactory, Data, 0);
+			Self->StaticMeshVertexBuffers.ColorVertexBuffer.BindColorVertexBuffer(&Self->VertexFactory, Data);
+			Self->VertexFactory.SetData(Data);
+
+			Self->VertexFactory.InitResource();
+		}
+	);
 }
 
 /** Virtual destructor. */
@@ -198,12 +206,6 @@ void FHeterogeneousVolumeSceneProxy::GetDynamicMeshElements(
 
 UHeterogeneousVolumeComponent::UHeterogeneousVolumeComponent(const FObjectInitializer& ObjectInitializer)
 	: Super(ObjectInitializer)
-	, VolumeResolution(128)
-	, MinimumVoxelSize(0.1f)
-	, bAnimate(false)
-	, LightingDownsampleFactor(1.0f)
-	, Time(0.0f)
-	, Framerate(24.0f)
 {
 	PrimaryComponentTick.bCanEverTick = true;
 	// What is this?
@@ -212,6 +214,12 @@ UHeterogeneousVolumeComponent::UHeterogeneousVolumeComponent(const FObjectInitia
 #if WITH_EDITORONLY_DATA
 	bTickInEditor = true;
 #endif // WITH_EDITORONLY_DATA
+
+	VolumeResolution = FIntVector(128);
+	Framerate = 24.0f;
+	bAnimate = false;
+	LightingDownsampleFactor = 1.0f;
+	Time = 0.0f;
 }
 
 FPrimitiveSceneProxy* UHeterogeneousVolumeComponent::CreateSceneProxy()
@@ -263,12 +271,12 @@ void UHeterogeneousVolumeComponent::TickComponent(float DeltaTime, ELevelTick Ti
 
 			if (bAnimate)
 			{
-			int32 FrameIndex = int32(Time * (1000.0f / Framerate)) % DefaultSVT->GetNumFrames();
-			int32 MipLevel = 0;
-			USparseVolumeTextureFrame* SVTFrame = USparseVolumeTextureFrame::CreateFrame(DefaultSVT, FrameIndex, MipLevel);
-			MaterialInstanceDynamic->SetSparseVolumeTextureParameterValue(SVTParameterInfo[MaterialIndex].Name, SVTFrame);
+				int32 FrameIndex = int32(Time * Framerate) % DefaultSVT->GetNumFrames();
+				int32 MipLevel = 0;
+				USparseVolumeTextureFrame* SVTFrame = USparseVolumeTextureFrame::CreateFrame(DefaultSVT, FrameIndex, MipLevel);
+				MaterialInstanceDynamic->SetSparseVolumeTextureParameterValue(SVTParameterInfo[MaterialIndex].Name, SVTFrame);
 
-				VolumeResolution = SVTFrame->GetVolumeResolution();
+				VolumeResolution = SVTFrame ? SVTFrame->GetVolumeResolution() : FIntVector(1);
 			}
 			else
 			{
