Added UInstancedStaticMeshComponent::GetCullDistances and UHierarchicalInstancedStaticMeshComponent::GetLODDistanceScale

diff --git a/Engine/Source/Runtime/Engine/Classes/Components/HierarchicalInstancedStaticMeshComponent.h b/Engine/Source/Runtime/Engine/Classes/Components/HierarchicalInstancedStaticMeshComponent.h
--- a/Engine/Source/Runtime/Engine/Classes/Components/HierarchicalInstancedStaticMeshComponent.h
+++ b/Engine/Source/Runtime/Engine/Classes/Components/HierarchicalInstancedStaticMeshComponent.h
@@ -253,6 +253,9 @@ public:
 	UFUNCTION(BlueprintCallable, Category = "Components|HierarchicalInstancedStaticMesh")
 	void SetLODDistanceScale(float InLODDistanceScale);
 
+	UFUNCTION(BlueprintCallable, Category = "Components|HierarchicalInstancedStaticMesh")
+	float GetLODDistanceScale() const { return InstanceLODDistanceScale; }
+
 	/** Get the number of instances that overlap a given sphere */
 	int32 GetOverlappingSphereCount(const FSphere& Sphere) const;
 	/** Get the number of instances that overlap a given box */
diff --git a/Engine/Source/Runtime/Engine/Classes/Components/InstancedStaticMeshComponent.h b/Engine/Source/Runtime/Engine/Classes/Components/InstancedStaticMeshComponent.h
--- a/Engine/Source/Runtime/Engine/Classes/Components/InstancedStaticMeshComponent.h
+++ b/Engine/Source/Runtime/Engine/Classes/Components/InstancedStaticMeshComponent.h
@@ -276,6 +276,10 @@ class ENGINE_API UInstancedStaticMeshComponent : public UStaticMeshComponent, pu
 	UFUNCTION(BlueprintCallable, Category = "Components|InstancedStaticMesh")
 	void SetCullDistances(int32 StartCullDistance, int32 EndCullDistance);
 
+	/** Gets the fading start and culling end distances for this component. */
+	UFUNCTION(BlueprintCallable, Category = "Components|InstancedStaticMesh")
+	void GetCullDistances(int32& OutStartCullDistance, int32& OutEndCullDistance) const { OutStartCullDistance = InstanceStartCullDistance; OutEndCullDistance = InstanceEndCullDistance; }
+
 	/** Returns the instances with instance bounds overlapping the specified sphere. The return value is an array of instance indices. */
 	UFUNCTION(BlueprintCallable, Category = "Components|InstancedStaticMesh")
 	virtual TArray<int32> GetInstancesOverlappingSphere(const FVector& Center, float Radius, bool bSphereInWorldSpace=true) const;
