[shaders] add a template helper to set a define only if not already set; eliminates a few usages of FShaderCompilerEnvironment::GetDefinitions in engine code

diff --git a/Engine/Plugins/FX/Niagara/Source/NiagaraVertexFactories/Private/NiagaraMeshVertexFactory.cpp b/Engine/Plugins/FX/Niagara/Source/NiagaraVertexFactories/Private/NiagaraMeshVertexFactory.cpp
--- a/Engine/Plugins/FX/Niagara/Source/NiagaraVertexFactories/Private/NiagaraMeshVertexFactory.cpp
+++ b/Engine/Plugins/FX/Niagara/Source/NiagaraVertexFactories/Private/NiagaraMeshVertexFactory.cpp
@@ -187,10 +187,9 @@ void FNiagaraMeshVertexFactory::ModifyCompilationEnvironment(const FVertexFactor
 	}
 #endif
 
-	const bool ContainsManualVertexFetch = OutEnvironment.GetDefinitions().Contains("MANUAL_VERTEX_FETCH");
-	if (!ContainsManualVertexFetch && RHISupportsManualVertexFetch(Parameters.Platform))
+	if (RHISupportsManualVertexFetch(Parameters.Platform))
 	{
-		OutEnvironment.SetDefine(TEXT("MANUAL_VERTEX_FETCH"), TEXT("1"));
+		OutEnvironment.SetDefineIfUnset(TEXT("MANUAL_VERTEX_FETCH"), TEXT("1"));
 	}
 }
 
diff --git a/Engine/Source/Runtime/Engine/Private/GeometryCollectionRendering.cpp b/Engine/Source/Runtime/Engine/Private/GeometryCollectionRendering.cpp
--- a/Engine/Source/Runtime/Engine/Private/GeometryCollectionRendering.cpp
+++ b/Engine/Source/Runtime/Engine/Private/GeometryCollectionRendering.cpp
@@ -60,11 +60,9 @@ void FGeometryCollectionVertexFactory::ModifyCompilationEnvironment(const FVerte
 
 	OutEnvironment.SetDefine(TEXT("VF_SUPPORTS_PRIMITIVE_SCENE_DATA"), bSupportsPrimitiveIdStream && bUseGPUScene);
 
-	bool ContainsManualVertexFetch = OutEnvironment.GetDefinitions().Contains("MANUAL_VERTEX_FETCH");
-	if (!ContainsManualVertexFetch && RHISupportsManualVertexFetch(Parameters.Platform))
+	if (RHISupportsManualVertexFetch(Parameters.Platform))
 	{
-		OutEnvironment.SetDefine(TEXT("MANUAL_VERTEX_FETCH"), TEXT("1"));
-		ContainsManualVertexFetch = true;
+		OutEnvironment.SetDefineIfUnset(TEXT("MANUAL_VERTEX_FETCH"), TEXT("1"));
 	}
 
 	// Geometry collections use a custom hit proxy per bone
diff --git a/Engine/Source/Runtime/Engine/Private/InstancedStaticMesh.cpp b/Engine/Source/Runtime/Engine/Private/InstancedStaticMesh.cpp
--- a/Engine/Source/Runtime/Engine/Private/InstancedStaticMesh.cpp
+++ b/Engine/Source/Runtime/Engine/Private/InstancedStaticMesh.cpp
@@ -871,10 +871,9 @@ bool FInstancedStaticMeshVertexFactory::ShouldCompilePermutation(const FVertexFa
 
 void FInstancedStaticMeshVertexFactory::ModifyCompilationEnvironment(const FVertexFactoryShaderPermutationParameters& Parameters, FShaderCompilerEnvironment& OutEnvironment)
 {
-	const bool ContainsManualVertexFetch = OutEnvironment.GetDefinitions().Contains("MANUAL_VERTEX_FETCH");
-	if (!ContainsManualVertexFetch && RHISupportsManualVertexFetch(Parameters.Platform))
+	if (RHISupportsManualVertexFetch(Parameters.Platform))
 	{
-		OutEnvironment.SetDefine(TEXT("MANUAL_VERTEX_FETCH"), TEXT("1"));
+		OutEnvironment.SetDefineIfUnset(TEXT("MANUAL_VERTEX_FETCH"), TEXT("1"));
 	}
 
 	if (UseGPUScene(Parameters.Platform))
diff --git a/Engine/Source/Runtime/Engine/Private/LocalVertexFactory.cpp b/Engine/Source/Runtime/Engine/Private/LocalVertexFactory.cpp
--- a/Engine/Source/Runtime/Engine/Private/LocalVertexFactory.cpp
+++ b/Engine/Source/Runtime/Engine/Private/LocalVertexFactory.cpp
@@ -293,15 +293,11 @@ bool FLocalVertexFactory::ShouldCompilePermutation(const FVertexFactoryShaderPer
 void FLocalVertexFactory::ModifyCompilationEnvironment(const FVertexFactoryShaderPermutationParameters& Parameters, FShaderCompilerEnvironment& OutEnvironment)
 {
 	// Don't override e.g. SplineMesh's opt-out
-	if (!OutEnvironment.GetDefinitions().Contains("VF_SUPPORTS_SPEEDTREE_WIND"))
-	{
-		OutEnvironment.SetDefine(TEXT("VF_SUPPORTS_SPEEDTREE_WIND"), TEXT("1"));
-	}
+	OutEnvironment.SetDefineIfUnset(TEXT("VF_SUPPORTS_SPEEDTREE_WIND"), TEXT("1"));
 
-	const bool ContainsManualVertexFetch = OutEnvironment.GetDefinitions().Contains("MANUAL_VERTEX_FETCH");
-	if (!ContainsManualVertexFetch && RHISupportsManualVertexFetch(Parameters.Platform))
+	if (RHISupportsManualVertexFetch(Parameters.Platform))
 	{
-		OutEnvironment.SetDefine(TEXT("MANUAL_VERTEX_FETCH"), TEXT("1"));
+		OutEnvironment.SetDefineIfUnset(TEXT("MANUAL_VERTEX_FETCH"), TEXT("1"));
 	}
 
 	const bool bVFSupportsPrimtiveSceneData = Parameters.VertexFactoryType->SupportsPrimitiveIdStream() && UseGPUScene(Parameters.Platform, GetMaxSupportedFeatureLevel(Parameters.Platform));
diff --git a/Engine/Source/Runtime/Engine/Private/MeshParticleVertexFactory.cpp b/Engine/Source/Runtime/Engine/Private/MeshParticleVertexFactory.cpp
--- a/Engine/Source/Runtime/Engine/Private/MeshParticleVertexFactory.cpp
+++ b/Engine/Source/Runtime/Engine/Private/MeshParticleVertexFactory.cpp
@@ -264,10 +264,9 @@ void FMeshParticleVertexFactory::ModifyCompilationEnvironment(const FVertexFacto
 	OutEnvironment.SetDefine(TEXT("PARTICLE_MESH_FACTORY"), TEXT("1"));
 	OutEnvironment.SetDefine(TEXT("PARTICLE_MESH_INSTANCED"), TEXT("1"));
 
-	const bool ContainsManualVertexFetch = OutEnvironment.GetDefinitions().Contains("MANUAL_VERTEX_FETCH");
-	if (!ContainsManualVertexFetch && RHISupportsManualVertexFetch(Parameters.Platform))
+	if (RHISupportsManualVertexFetch(Parameters.Platform))
 	{
-		OutEnvironment.SetDefine(TEXT("MANUAL_VERTEX_FETCH"), TEXT("1"));
+		OutEnvironment.SetDefineIfUnset(TEXT("MANUAL_VERTEX_FETCH"), TEXT("1"));
 	}
 }
 
diff --git a/Engine/Source/Runtime/RenderCore/Public/ShaderCore.h b/Engine/Source/Runtime/RenderCore/Public/ShaderCore.h
--- a/Engine/Source/Runtime/RenderCore/Public/ShaderCore.h
+++ b/Engine/Source/Runtime/RenderCore/Public/ShaderCore.h
@@ -559,6 +559,14 @@ struct FShaderCompilerEnvironment
 	void SetDefine(const TCHAR* Name, bool Value)			{ Definitions.SetDefine(Name, Value); }
 	void SetDefine(const TCHAR* Name, float Value)			{ Definitions.SetFloatDefine(Name, Value); }
 
+	template <typename ValueType> void SetDefineIfUnset(const TCHAR* Name, ValueType Value)
+	{
+		if (!Definitions.GetDefinitionMap().Contains(Name))
+		{
+			SetDefine(Name, Value);
+		}
+	}
+
 	const TMap<FString,FString>& GetDefinitions() const
 	{
 		return Definitions.GetDefinitionMap();
