Use correct skeleton bind pose in decompression context.
Improve documentation regarding the decompression context bind pose.

diff --git a/Engine/Source/Runtime/Engine/Private/Animation/AnimCompressionTypes.cpp b/Engine/Source/Runtime/Engine/Private/Animation/AnimCompressionTypes.cpp
--- a/Engine/Source/Runtime/Engine/Private/Animation/AnimCompressionTypes.cpp
+++ b/Engine/Source/Runtime/Engine/Private/Animation/AnimCompressionTypes.cpp
@@ -1497,7 +1497,7 @@ void DecompressPose(FCompactPose& OutPose, const FCompressedAnimSequence& Compre
 {
 	const TArray<FTransform>& RetargetTransforms = SourceSkeleton->GetRefLocalPoses(RetargetSource);
 	PRAGMA_DISABLE_DEPRECATION_WARNINGS
-	FAnimSequenceDecompressionContext DecompressionContext(SequenceLength, Interpolation, SourceName, *CompressedData.CompressedDataStructure.Get());
+	FAnimSequenceDecompressionContext DecompressionContext(SequenceLength, Interpolation, SourceName, *CompressedData.CompressedDataStructure.Get(), SourceSkeleton->GetRefLocalPoses(), CompressedData.CompressedTrackToSkeletonMapTable);
 	PRAGMA_ENABLE_DEPRECATION_WARNINGS
 	UE::Anim::Decompression::DecompressPose(OutPose, CompressedData, ExtractionContext, DecompressionContext, RetargetTransforms, RootMotionReset);
 }
@@ -1505,7 +1505,7 @@ void DecompressPose(FCompactPose& OutPose, const FCompressedAnimSequence& Compre
 void DecompressPose(FCompactPose& OutPose, const FCompressedAnimSequence& CompressedData, const FAnimExtractContext& ExtractionContext, USkeleton* SourceSkeleton, float SequenceLength, EAnimInterpolationType Interpolation, bool bIsBakedAdditive, const TArray<FTransform>& RetargetTransforms, FName SourceName, const FRootMotionReset& RootMotionReset)
 {
 	PRAGMA_DISABLE_DEPRECATION_WARNINGS
-	FAnimSequenceDecompressionContext DecompressionContext(SequenceLength, Interpolation, SourceName, *CompressedData.CompressedDataStructure.Get());
+	FAnimSequenceDecompressionContext DecompressionContext(SequenceLength, Interpolation, SourceName, *CompressedData.CompressedDataStructure.Get(), SourceSkeleton->GetRefLocalPoses(), CompressedData.CompressedTrackToSkeletonMapTable);
 	PRAGMA_ENABLE_DEPRECATION_WARNINGS
 	UE::Anim::Decompression::DecompressPose(OutPose, CompressedData, ExtractionContext, DecompressionContext, RetargetTransforms, RootMotionReset);
 }
diff --git a/Engine/Source/Runtime/Engine/Public/Animation/AnimSequenceDecompressionContext.h b/Engine/Source/Runtime/Engine/Public/Animation/AnimSequenceDecompressionContext.h
--- a/Engine/Source/Runtime/Engine/Public/Animation/AnimSequenceDecompressionContext.h
+++ b/Engine/Source/Runtime/Engine/Public/Animation/AnimSequenceDecompressionContext.h
@@ -41,7 +41,7 @@ struct FAnimSequenceDecompressionContext
 		PRAGMA_ENABLE_DEPRECATION_WARNINGS
 	{
 	}
-	
+
 	FAnimSequenceDecompressionContext(const FFrameRate& InSamplingRate, const int32 InNumberOfFrames, EAnimInterpolationType Interpolation_, const FName& AnimName_, const ICompressedAnimData& CompressedAnimData_, const TArray<FTransform>& RefPoses_, const TArray<FTrackToSkeletonMap>& TrackToSkeletonMap_, const USkeleton* InSourceSkeleton, bool bIsBakedAdditive)
 		:
 		PRAGMA_DISABLE_DEPRECATION_WARNINGS
@@ -68,10 +68,16 @@ struct FAnimSequenceDecompressionContext
 	const ICompressedAnimData& CompressedAnimData;
 
 private:
-	// Immutable reference/bind pose, maps a bone index to its reference transform
+	// Immutable reference/bind pose from the source skeleton, maps a bone index to its reference transform.
+	// The same skeleton reference pose must be provided during compression and decompression. This is the
+	// source skeleton reference pose (the one assigned to an anim sequence). Re-targeting is taken into account
+	// elsewhere. Additionally, if the source skeleton reference pose isn't used, the TrackToSkeletonMap indices
+	// may not match the provided pose. This could cause us to read incorrect bone transforms or possibly out
+	// of bounds.
 	TArrayView<const FTransform> RefPoses;
 
-	// Immutable array that maps from a track index to a bone index
+	// Immutable array that maps from a track index to a bone index. This is built during compression using
+	// the source skeleton (the one assigned to an anim sequence).
 	TArrayView<const FTrackToSkeletonMap> TrackToSkeletonMap;
 
 public:
