[shaders] remove ODSC guard on querying for completed shadermaps from DDC. the reason this was introduced (synchronous blocking on DDC queries) is no longer the case since these were async'd, and this is causing unnecessary shader compilation when creating projects from templates.

diff --git a/Engine/Source/Runtime/Engine/Private/Materials/MaterialShared.cpp b/Engine/Source/Runtime/Engine/Private/Materials/MaterialShared.cpp
--- a/Engine/Source/Runtime/Engine/Private/Materials/MaterialShared.cpp
+++ b/Engine/Source/Runtime/Engine/Private/Materials/MaterialShared.cpp
@@ -2666,11 +2666,8 @@ bool FMaterial::CacheShaders(const FMaterialShaderMapId& ShaderMapId, EShaderPla
 			// If we have an incomplete shadermap, continue with it to prevent creation of duplicate shadermaps for the same ShaderMapId
 			if (!ShaderMap && !FPlatformProperties::RequiresCookedData())
 			{
-				if (bSkipCompilationOnPostLoad == false || IsRequiredComplete())
-				{
-					TRefCountPtr<FMaterialShaderMap> LoadedShaderMap;
-					CacheShadersPending = FMaterialShaderMap::BeginLoadFromDerivedDataCache(this, ShaderMapId, Platform, TargetPlatform, LoadedShaderMap, DDCKeyHash);
-				}
+				TRefCountPtr<FMaterialShaderMap> LoadedShaderMap;
+				CacheShadersPending = FMaterialShaderMap::BeginLoadFromDerivedDataCache(this, ShaderMapId, Platform, TargetPlatform, LoadedShaderMap, DDCKeyHash);
 			}
 
 			check(!ShaderMap || ShaderMap->GetFrozenContentSize() > 0u);
