Fix anim node functions not being called in versioned builds when multiple are bound

Enum serialization does not fully support bitfields, so converting to underlying type.

diff --git a/Engine/Source/Runtime/Engine/Public/Animation/AnimNodeData.h b/Engine/Source/Runtime/Engine/Public/Animation/AnimNodeData.h
--- a/Engine/Source/Runtime/Engine/Public/Animation/AnimNodeData.h
+++ b/Engine/Source/Runtime/Engine/Public/Animation/AnimNodeData.h
@@ -119,14 +119,14 @@ public:
 	int32 GetNodeIndex() const { return NodeIndex; }
 
 	/** Check whether the node has the specified flags */
-	bool HasNodeAnyFlags(EAnimNodeDataFlags InFlags) const { return EnumHasAnyFlags<EAnimNodeDataFlags>(Flags, InFlags); }
+	bool HasNodeAnyFlags(EAnimNodeDataFlags InFlags) const { return EnumHasAnyFlags(static_cast<EAnimNodeDataFlags>(Flags), InFlags); }
 	
 private:
 	friend class FAnimBlueprintCompilerContext;
 	friend struct UE::Anim::FNodeDataId;
 
 	/** Set whether the node has the specified flags */
-	void SetNodeFlags(EAnimNodeDataFlags InFlags) { return EnumAddFlags<EAnimNodeDataFlags>(Flags, InFlags); }
+	void SetNodeFlags(EAnimNodeDataFlags InFlags) { return EnumAddFlags(*reinterpret_cast<EAnimNodeDataFlags*>(&Flags), InFlags); }
 	
 	/** The class we are part of */
 	UPROPERTY()
@@ -148,7 +148,7 @@ private:
 
 	/** Common flags for this node */
 	UPROPERTY()
-	EAnimNodeDataFlags Flags = EAnimNodeDataFlags::None;
+	uint32 Flags = static_cast<uint32>(EAnimNodeDataFlags::None);
 };
 
 /**
