Fix reset custom time step and timecode not working with Media Profiles
Resetting custom time step would not take into account the overriden custom time step and timecode, so we now default to shutting down and restarting a provider if it's present.
At the moment restarting the timecode or genlock will sleep momentarily to give time for the asynchronous shut down process to complete.
In a future release, timestep and timecode providers should call a delegate when their shut down process is complete in order to avoid sleeping.

diff --git a/Engine/Source/Runtime/Engine/Private/UnrealEngine.cpp b/Engine/Source/Runtime/Engine/Private/UnrealEngine.cpp
--- a/Engine/Source/Runtime/Engine/Private/UnrealEngine.cpp
+++ b/Engine/Source/Runtime/Engine/Private/UnrealEngine.cpp
@@ -2707,7 +2707,7 @@ void UEngine::ReinitializeCustomTimeStep()
 	{
 		if (bIsCurrentCustomTimeStepInitialized)
 		{
-		CustomTimeStep->Shutdown(this);
+			CustomTimeStep->Shutdown(this);
 		}
 		bIsCurrentCustomTimeStepInitialized = CustomTimeStep->Initialize(this);
 	}
@@ -2715,22 +2715,23 @@ void UEngine::ReinitializeCustomTimeStep()
 
 bool UEngine::SetCustomTimeStep(UEngineCustomTimeStep* InCustomTimeStep)
 {
-	if (InCustomTimeStep != CustomTimeStep)
+	if (CustomTimeStep && bIsCurrentCustomTimeStepInitialized)
 	{
-		if (CustomTimeStep && bIsCurrentCustomTimeStepInitialized)
-		{
-			CustomTimeStep->Shutdown(this);
-		}
+		CustomTimeStep->Shutdown(this);
 
-		bIsCurrentCustomTimeStepInitialized = false;
-		CustomTimeStep = IsValid(InCustomTimeStep) ? InCustomTimeStep : nullptr;
+		/** Some custom timesteps need some time to release their resources when we shut them down.
+			Until a delegate is added to the custom time steps to know when their shutdown is complete we need to sleep in order for the hardware to complete its uninitialization. **/
+		FPlatformProcess::Sleep(0.2);
+	}
 
-		if (CustomTimeStep)
-			{
-			bIsCurrentCustomTimeStepInitialized = CustomTimeStep->Initialize(this);
-		}
-		OnCustomTimeStepChanged().Broadcast();
+	bIsCurrentCustomTimeStepInitialized = false;
+	CustomTimeStep = IsValid(InCustomTimeStep) ? InCustomTimeStep : nullptr;
+
+	if (CustomTimeStep)
+	{
+		bIsCurrentCustomTimeStepInitialized = CustomTimeStep->Initialize(this);
 	}
+	OnCustomTimeStepChanged().Broadcast();
 	return bIsCurrentCustomTimeStepInitialized;
 }
 
@@ -2741,7 +2742,7 @@ void UEngine::ReinitializeTimecodeProvider()
 	{
 		if (bIsCurrentTimecodeProviderInitialized)
 		{
-	Provider->Shutdown(this);
+			Provider->Shutdown(this);
 		}
 		bIsCurrentTimecodeProviderInitialized = Provider->Initialize(this);
 	}
@@ -2749,22 +2750,24 @@ void UEngine::ReinitializeTimecodeProvider()
 
 bool UEngine::SetTimecodeProvider(UTimecodeProvider* InTimecodeProvider)
 {
-	if (InTimecodeProvider != TimecodeProvider)
-			{
-		if (TimecodeProvider && bIsCurrentTimecodeProviderInitialized)
-		{
-			TimecodeProvider->Shutdown(this);
-		}
+	if (TimecodeProvider && bIsCurrentTimecodeProviderInitialized)
+	{
+		TimecodeProvider->Shutdown(this);
+		/** Some timecode providers need some time to release their resources when we shut them down.
+			Until a delegate is added to the timecode providers to know when their shutdown is complete, we need to sleep in order for the hardware to complete its uninitialization. **/
+		// 
+		FPlatformProcess::Sleep(0.2);
+	}
 
-		bIsCurrentTimecodeProviderInitialized = false;
-		TimecodeProvider = IsValid(InTimecodeProvider) ? InTimecodeProvider : nullptr;
+	bIsCurrentTimecodeProviderInitialized = false;
+	TimecodeProvider = IsValid(InTimecodeProvider) ? InTimecodeProvider : nullptr;
 
-		if (TimecodeProvider)
-			{
-			bIsCurrentTimecodeProviderInitialized = TimecodeProvider->Initialize(this);
-		}
-		OnTimecodeProviderChanged().Broadcast();
+	if (TimecodeProvider)
+	{
+		bIsCurrentTimecodeProviderInitialized = TimecodeProvider->Initialize(this);
 	}
+	OnTimecodeProviderChanged().Broadcast();
+
 	return bIsCurrentTimecodeProviderInitialized;
 }
 
@@ -2933,7 +2936,32 @@ static void LoadCustomTimeStep(UEngine* Engine)
 
 static void ResetCustomTimeStep()
 {
-	LoadCustomTimeStep(GEngine);
+	if (UEngineCustomTimeStep* CustomTimeStep = GEngine->GetCustomTimeStep())
+	{
+#if WITH_EDITOR
+		if (CustomTimeStep->GetClass()->HasAnyClassFlags(CLASS_NewerVersionExists) || CustomTimeStep->GetClass()->HasAnyClassFlags(CLASS_CompiledFromBlueprint))
+		{
+			// Class was recompiled so we need to reload it.
+			UClass* CustomTimeStepClass = CustomTimeStep->GetClass();
+
+			if (UBlueprint* Blueprint = Cast<UBlueprint>(CustomTimeStepClass->ClassGeneratedBy))
+			{
+				if (UClass* NewTimeStepClass = Blueprint->GeneratedClass)
+				{
+					UEngineCustomTimeStep* NewCustomTimeStep = NewObject<UEngineCustomTimeStep>(GEngine, NewTimeStepClass);
+					GEngine->SetCustomTimeStep(NewCustomTimeStep);
+					return;
+				}
+			}
+
+		}
+#endif
+		GEngine->SetCustomTimeStep(CustomTimeStep);
+	}
+	else
+	{
+		LoadCustomTimeStep(GEngine);
+	}
 }
 
 static FAutoConsoleCommand GResetCustomTimeStep(
@@ -2947,10 +2975,10 @@ static void LoadTimecodeProvider(UEngine* Engine)
 	if (Engine->TimecodeProviderClassName.IsValid())
 {
 		if (UClass* TimecodeProviderClass = Engine->TimecodeProviderClassName.TryLoadClass<UTimecodeProvider>())
-	{
+		{
 			UTimecodeProvider* NewTimecodeProvider = NewObject<UTimecodeProvider>(Engine, TimecodeProviderClass);
 			if (!Engine->SetTimecodeProvider(NewTimecodeProvider))
-		{
+			{
 				UE_LOG(LogEngine, Warning, TEXT("InitializeTimecodeProvider - Failed to intialize TimecodeProvider '%s'."), *NewTimecodeProvider->GetPathName());
 			}
 		}
@@ -2983,7 +3011,31 @@ static void LoadTimecodeProvider(UEngine* Engine)
 
 static void ResetTimecode()
 {
-	LoadTimecodeProvider(GEngine);
+	if (UTimecodeProvider* TimecodeProvider = GEngine->GetTimecodeProvider())
+	{
+#if WITH_EDITOR
+		if (TimecodeProvider->GetClass()->HasAnyClassFlags(CLASS_NewerVersionExists) || TimecodeProvider->GetClass()->HasAnyClassFlags(CLASS_CompiledFromBlueprint))
+		{
+			// Class was recompiled so we need to reload it.
+			UClass* TCProviderClass = TimecodeProvider->GetClass();
+
+			if (UBlueprint* Blueprint = Cast<UBlueprint>(TCProviderClass->ClassGeneratedBy))
+			{
+				if (UClass* NewTimecodeClass = Blueprint->GeneratedClass)
+				{
+					UTimecodeProvider* NewTimecodeProvider = NewObject<UTimecodeProvider>(GEngine, NewTimecodeClass);
+					GEngine->SetTimecodeProvider(NewTimecodeProvider);
+					return;
+				}
+			}
+		}
+#endif
+		GEngine->SetTimecodeProvider(TimecodeProvider);
+	}
+	else
+	{
+		LoadTimecodeProvider(GEngine);
+	}
 }
 
 static FAutoConsoleCommand GResetTimecode(
