Improve comments around FSaveGameHeader to explain the limitations of the versioning system and encourage use of new ULocalPlayerSaveGame

diff --git a/Engine/Source/Runtime/Engine/Private/GameplayStatics.cpp b/Engine/Source/Runtime/Engine/Private/GameplayStatics.cpp
--- a/Engine/Source/Runtime/Engine/Private/GameplayStatics.cpp
+++ b/Engine/Source/Runtime/Engine/Private/GameplayStatics.cpp
@@ -104,6 +104,8 @@ DECLARE_CYCLE_STAT(TEXT("SpawnTime"), STAT_SpawnTime, STATGROUP_Game);
 //////////////////////////////////////////////////////////////////////////
 // FSaveGameHeader
 
+// This is the engine-level header for save game versioning and is not useful for game-specific version changes
+// To implement those, you would need to save the version number into the save game object using something like ULocalPlayerSaveGame
 struct FSaveGameHeader
 {
 	FSaveGameHeader();
@@ -164,15 +166,10 @@ void FSaveGameHeader::Read(FMemoryReader& MemoryReader)
 
 	if (FileTypeTag != UE_SAVEGAME_FILE_TYPE_TAG)
 	{
-		// this is an old saved game, back up the file pointer to the beginning and assume version 1
+		// This is a very old saved game, back up the file pointer to the beginning and assume version 1
+		// This is unlikely to work without additional licensee-specific modifications to this code
 		MemoryReader.Seek(0);
 		SaveGameFileVersion = FSaveGameFileVersion::InitialVersion;
-
-		// Note for 4.8 and beyond: if you get a crash loading a pre-4.8 version of your savegame file and 
-		// you don't want to delete it, try uncommenting these lines and changing them to use the version 
-		// information from your previous build. Then load and resave your savegame file.
-		//MemoryReader.SetUEVer(MyPreviousUEVersion);				// @see GPackageFileUEVersion
-		//MemoryReader.SetEngineVer(MyPreviousEngineVersion);		// @see FEngineVersion::Current()
 	}
 	else
 	{
@@ -204,6 +201,10 @@ void FSaveGameHeader::Read(FMemoryReader& MemoryReader)
 			CustomVersions.Serialize(MemoryReader, static_cast<ECustomVersionSerializationFormat::Type>(CustomVersionFormat));
 			MemoryReader.SetCustomVersions(CustomVersions);
 		}
+
+		// This code does not handle SetLicenseeUEVer because save games are not expected to work across major licensee changes to the engine
+		// If your game wants to support advanced backward compatibility, you will probably want to implement a version number on the object itself
+		// ULocalPlayerSaveGame has an example of how to implement a version number in a way that can be accessed after serialization
 	}
 
 	// Get the class name
