Adding a debug command to simulate a GPU crash for crash report testing.

diff --git a/Engine/Source/Runtime/Engine/Private/UnrealEngine.cpp b/Engine/Source/Runtime/Engine/Private/UnrealEngine.cpp
--- a/Engine/Source/Runtime/Engine/Private/UnrealEngine.cpp
+++ b/Engine/Source/Runtime/Engine/Private/UnrealEngine.cpp
@@ -10039,6 +10039,8 @@ private:
 
 #endif
 
+extern CORE_API bool GIsGPUCrashed;
+
 bool UEngine::PerformError(const TCHAR* Cmd, FOutputDevice& Ar)
 {
 #if !UE_BUILD_SHIPPING
@@ -10361,6 +10363,22 @@ bool UEngine::PerformError(const TCHAR* Cmd, FOutputDevice& Ar)
 		UE_LOG(LogEngine, Fatal, TEXT("%s"), TEXT("Crashing the gamethread at your request"));
 		return true;
 	}
+	else if (FParse::Command(&Cmd, TEXT("GPUCRASH")))
+	{
+		struct FRender
+		{
+			static void GpuCrash(FRHICommandList&)
+			{
+				UE_LOG(LogEngine, Warning, TEXT("Printed warning to log."));
+				FGenericCrashContext::SetCrashTrigger(ECrashTrigger::Debug);
+				GIsGPUCrashed = true;
+				UE_LOG(LogEngine, Fatal, TEXT("Simulating a GPU crash at your request"));
+				DebugDummyWorkAfterFailure();
+			}
+		};
+		ENQUEUE_RENDER_COMMAND(CauseGpuCrash)(&FRender::GpuCrash);
+		return true;
+	}
 	else if (FParse::Command(&Cmd, TEXT("TERMINATE")))
 	{
 		UE_LOG(LogEngine, Warning, TEXT("Printed warning to log."));
