After saving a newly created world, its OriginalWorldName was not updated, which resulted in PIE not working propertly.

diff --git a/Engine/Source/Runtime/Engine/Private/World.cpp b/Engine/Source/Runtime/Engine/Private/World.cpp
--- a/Engine/Source/Runtime/Engine/Private/World.cpp
+++ b/Engine/Source/Runtime/Engine/Private/World.cpp
@@ -1308,6 +1308,12 @@ void UWorld::PostSaveRoot( FObjectPostSaveRootContext ObjectSaveContext )
 	{
 		PersistentLevel->DetachAttachAllActorsPackages(/*bReattach*/true);
 	}
+
+	if (!ObjectSaveContext.IsProceduralSave() && !(ObjectSaveContext.GetSaveFlags() & SAVE_FromAutosave))
+	{
+		// Once saved, OriginalWorldName must match World's name
+		OriginalWorldName = GetFName();
+	}
 #endif
 }
 
