Substrate - Fix for incorrect operator and BSDF / slab count
Should be working but there is a problem with operator BSDFIndex

diff --git a/Engine/Shaders/Private/Strata/Strata.ush b/Engine/Shaders/Private/Strata/Strata.ush
--- a/Engine/Shaders/Private/Strata/Strata.ush
+++ b/Engine/Shaders/Private/Strata/Strata.ush
@@ -1116,7 +1116,7 @@ struct FStrataTree
 	uint OperatorCount;
 
 	FStrataBSDF BSDFs[STRATA_MAX_BSDF_COUNT];
-	FStrataOperator Operators[STRATA_MAX_OPERATOR_COUNT]; // STRATA_TODO assert when too many ops on the compiler side. Shoud already be done for BSDF count.
+	FStrataOperator Operators[STRATA_MAX_OPERATOR_COUNT];
 };
 
 FStrataTree GetInitialisedStrataTree()
diff --git a/Engine/Source/Runtime/Engine/Private/Materials/HLSLMaterialTranslator.cpp b/Engine/Source/Runtime/Engine/Private/Materials/HLSLMaterialTranslator.cpp
--- a/Engine/Source/Runtime/Engine/Private/Materials/HLSLMaterialTranslator.cpp
+++ b/Engine/Source/Runtime/Engine/Private/Materials/HLSLMaterialTranslator.cpp
@@ -288,6 +288,9 @@ FHLSLMaterialTranslator::FStrataCompilationContext::FStrataCompilationContext(ES
 	CompilationContextIndex = InCompilationContext;
 }
 
+// This limitation is required so that the operator array is never reallocated, invalidating references to it while parsing the Strata tree within StrataGenerateMaterialTopologyTree for instance.
+#define STRATA_MAX_COMPILER_REGISTERED_OPERATOR_COUNT 128
+
 void FHLSLMaterialTranslator::FStrataCompilationContext::Initialise()
 {
 	CompilationContextIndex = EStrataCompilationContext::SCC_MAX;
@@ -295,8 +298,8 @@ void FHLSLMaterialTranslator::FStrataCompilationContext::Initialise()
 	NextFreeStrataShaderNormalIndex = 0;
 	FinalUsedSharedLocalBasesCount = 0;
 	StrataMaterialRootOperator = nullptr;
-	StrataMaterialExpressionRegisteredOperators.Reserve(STRATA_MAX_OPERATOR_COUNT);
-	StrataMaterialExpressionToOperatorIndex.Reserve(STRATA_MAX_OPERATOR_COUNT);
+	StrataMaterialExpressionRegisteredOperators.Reserve(STRATA_MAX_COMPILER_REGISTERED_OPERATOR_COUNT);
+	StrataMaterialExpressionToOperatorIndex.Reserve(STRATA_MAX_COMPILER_REGISTERED_OPERATOR_COUNT);
 	StrataMaterialBSDFCount = 0;
 	StrataMaterialRequestedSizeByte = 0;
 	bStrataMaterialIsSimple = false;
@@ -10498,9 +10501,9 @@ FStrataOperator& FHLSLMaterialTranslator::StrataCompilationRegisterOperator(int3
 	}
 
 	const uint32 NewOperatorIndex = StrataCtx.StrataMaterialExpressionToOperatorIndex.Num();
-	if (NewOperatorIndex >= STRATA_MAX_OPERATOR_COUNT)
+	if (NewOperatorIndex >= STRATA_MAX_COMPILER_REGISTERED_OPERATOR_COUNT)
 	{
-		Errorf(TEXT("Material %s have too many Strata Operators (asset: %s).\r\n"), *Material->GetDebugName(), *Material->GetAssetPath().ToString());
+		Errorf(TEXT("Material %s have too many Strata Operators: the compiler is failing (asset: %s).\r\n"), *Material->GetDebugName(), *Material->GetAssetPath().ToString());
 		return DefaultOperatorOnError;
 	}
 
@@ -10522,7 +10525,7 @@ FStrataOperator& FHLSLMaterialTranslator::StrataCompilationRegisterOperator(int3
 	NewOperator.LeftIndex   = INDEX_NONE;
 	NewOperator.RightIndex   = INDEX_NONE;
 
-	NewOperator.BSDFIndex = INDEX_NONE;	// Allocated later to be albedo account for inline
+	NewOperator.BSDFIndex = INDEX_NONE;	// Allocated later to be able to account for inline
 
 	NewOperator.MaxDistanceFromLeaves = 0;
 	NewOperator.bIsBottom = false;
@@ -10535,7 +10538,7 @@ FStrataOperator& FHLSLMaterialTranslator::StrataCompilationGetOperator(FGuid Str
 {
 	FStrataCompilationContext& StrataCtx = StrataCompilationContext[CurrentStrataCompilationContext];
 	auto* OperatorIndex = StrataCtx.StrataMaterialExpressionToOperatorIndex.Find(StrataExpressionGuid);
-	if (!(OperatorIndex && *OperatorIndex >= 0 && *OperatorIndex < STRATA_MAX_OPERATOR_COUNT))
+	if (!(OperatorIndex && *OperatorIndex >= 0 && *OperatorIndex < STRATA_MAX_COMPILER_REGISTERED_OPERATOR_COUNT))
 	{
 		static FStrataOperator DefaultOperatorOnError = FStrataOperator();
 		return DefaultOperatorOnError;
@@ -10793,7 +10796,7 @@ bool FHLSLMaterialTranslator::FStrataCompilationContext::StrataGenerateDerivedMa
 	const uint32 StrataBytePerPixel = Strata::GetBytePerPixel(Compiler->GetShaderPlatform());
 	do 
 	{
-		// Reset some data for simplifiucation iteration
+		// Reset some data
 		StrataMaterialBSDFCount = 0;
 
 		//
@@ -10904,10 +10907,6 @@ bool FHLSLMaterialTranslator::FStrataCompilationContext::StrataGenerateDerivedMa
 			}
 		}
 
-		// STRATA_TODO: operation using parameter blending are not actually discarded so they still occupy a spot in the operation array in the compiler and in the shader also.
-		// Even thought the compiler will remove unused operators from the shader code, we will still be limited to STRATA_MAX_OPERATOR_COUNT even with parameter blending.
-		// This can be fixed by remapping all operation index and make sure this is the index that is always used in the compiler and specified to the shader.
-
 		//
 		// Make sure all the types have valid children operator indices
 		//
@@ -11136,11 +11135,12 @@ bool FHLSLMaterialTranslator::FStrataCompilationContext::StrataGenerateDerivedMa
 				StrataMaterialRequestedSizeByte += UsedSharedLocalBasesCount * STRATA_PACKED_SHAREDLOCALBASIS_STRIDE_BYTES;
 			}
 			// Note:
-			//  - We do not need to account for the Top Normal texture when evaluating the material byte count for the optimisation algorithm.
-			//  - This is because we only need to optimise for the Strata uint matertial buffer.
+			//  - We do not need to account for the Top Normal texture when evaluating the material byte count for the optimization algorithm.
+			//  - This is because we only need to optimize for the Strata uint material buffer.
 
 
-			// 2. The list of BSDFs
+			// 2. Process the list of BSDFs for worst case memory usage and count operators.
+			uint32 OperatorCount = 0;
 			for (auto& It : StrataMaterialExpressionRegisteredOperators)
 			{
 				if (It.IsDiscarded())
@@ -11148,6 +11148,10 @@ bool FHLSLMaterialTranslator::FStrataCompilationContext::StrataGenerateDerivedMa
 					continue; // ignore discarded operations in sub tree using parameter blending
 				}
 
+				// Operators are weight, vertical layering, etc. 
+				// Be aware that BSDFs also count as Operators when they are promoted from parameter blending!
+				OperatorCount++;
+
 				switch (It.OperatorType)
 				{
 				case STRATA_OPERATOR_BSDF:
@@ -11256,6 +11260,13 @@ bool FHLSLMaterialTranslator::FStrataCompilationContext::StrataGenerateDerivedMa
 				} // switch (It.OperatorType)
 			}
 
+			if (OperatorCount > STRATA_MAX_OPERATOR_COUNT)
+			{
+				// Why do we have an operator limit: due to the size of the array of Operator in FStrataTre and the way the strata tree is exported for advanced debug purpose. Parameter blending can help working around that. 
+				Compiler->Errorf(TEXT("Material %s have too many Strata Operators (asset: %s): %d / %d. Please note that BSDFs also count as an operator. Use parameter blending to workaround that limitation.\r\n"), *CompilerMaterial->GetDebugName(), *CompilerMaterial->GetAssetPath().ToString(), OperatorCount, STRATA_MAX_OPERATOR_COUNT);
+				return false;
+			}
+
 			StrataSimplificationStatus.bMaterialFitsInMemoryBudget = StrataMaterialRequestedSizeByte <= StrataBytePerPixel;
 			if (!StrataSimplificationStatus.bMaterialFitsInMemoryBudget && StrataSimplificationStatus.bRunFullSimplification)
 			{
