Added missing definition for UAnimInstance's GetRelevantAssetPlayerInterfaceFromState
GitHub #10085

Contributed by [REDACTED_PERSON] ([REDACTED_PERSON])

diff --git a/Engine/Source/Runtime/Engine/Private/Animation/AnimInstance.cpp b/Engine/Source/Runtime/Engine/Private/Animation/AnimInstance.cpp
--- a/Engine/Source/Runtime/Engine/Private/Animation/AnimInstance.cpp
+++ b/Engine/Source/Runtime/Engine/Private/Animation/AnimInstance.cpp
@@ -3875,6 +3875,11 @@ const FBakedAnimationStateMachine* UAnimInstance::GetStateMachineInstanceDesc(FN
 	return GetProxyOnAnyThread<FAnimInstanceProxy>().GetStateMachineInstanceDesc(MachineName);
 }
 
+const FAnimNode_AssetPlayerRelevancyBase* UAnimInstance::GetRelevantAssetPlayerInterfaceFromState(int32 MachineIndex, int32 StateIndex) const
+{
+	return GetProxyOnAnyThread<FAnimInstanceProxy>().GetRelevantAssetPlayerInterfaceFromState(MachineIndex, StateIndex);
+}
+
 int32 UAnimInstance::GetStateMachineIndex(FName MachineName) const
 {
 	return GetProxyOnAnyThread<FAnimInstanceProxy>().GetStateMachineIndex(MachineName);
