Fix a crash when enable transmission on directional light with spot light shadow enabled on mobile.

diff --git a/Engine/Source/Runtime/Engine/Private/LightSceneProxy.cpp b/Engine/Source/Runtime/Engine/Private/LightSceneProxy.cpp
--- a/Engine/Source/Runtime/Engine/Private/LightSceneProxy.cpp
+++ b/Engine/Source/Runtime/Engine/Private/LightSceneProxy.cpp
@@ -137,6 +137,8 @@ FLightSceneProxy::FLightSceneProxy(const ULightComponent* InLightComponent)
 		{
 			bCastDynamicShadow = false;
 		}
+
+		bTransmission = false;
 	}
 }
 
diff --git a/Engine/Source/Runtime/Engine/Public/LightSceneProxy.h b/Engine/Source/Runtime/Engine/Public/LightSceneProxy.h
--- a/Engine/Source/Runtime/Engine/Public/LightSceneProxy.h
+++ b/Engine/Source/Runtime/Engine/Public/LightSceneProxy.h
@@ -388,7 +388,7 @@ protected:
 	const uint8 bCastTranslucentShadows : 1;
 
 	/** Whether light from this light transmits through surfaces with subsurface scattering profiles. Requires light to be movable. */
-	const uint8 bTransmission : 1;
+	uint8 bTransmission : 1;
 
 	const uint8 bCastVolumetricShadow : 1;
 	const uint8 bCastHairStrandsDeepShadow : 1;
