Fix TSAN warning by using proper atomics in StudioAnalytics

#rnx

diff --git a/Engine/Source/Runtime/Engine/Private/Analytics/StudioAnalytics.cpp b/Engine/Source/Runtime/Engine/Private/Analytics/StudioAnalytics.cpp
--- a/Engine/Source/Runtime/Engine/Private/Analytics/StudioAnalytics.cpp
+++ b/Engine/Source/Runtime/Engine/Private/Analytics/StudioAnalytics.cpp
@@ -14,7 +14,7 @@
 #endif
 
 bool FStudioAnalytics::bInitialized = false;
-volatile double FStudioAnalytics::TimeEstimation = 0;
+std::atomic<double> FStudioAnalytics::TimeEstimation { 0 };
 FThread FStudioAnalytics::TimerThread;
 TSharedPtr<IAnalyticsProviderET> FStudioAnalytics::Analytics;
 TArray<FAnalyticsEventAttribute> FStudioAnalytics::DefaultAttributes;
@@ -88,11 +88,11 @@ void FStudioAnalytics::RunTimer_Concurrent()
 
 		if (DeltaTime > BreakpointHitchTime)
 		{
-			TimeEstimation += FixedInterval;
+			TimeEstimation.store(TimeEstimation.load(std::memory_order_relaxed) + FixedInterval);
 		}
 		else
 		{
-			TimeEstimation += DeltaTime;
+			TimeEstimation.store(TimeEstimation.load(std::memory_order_relaxed) + DeltaTime);
 		}
 	}
 }
@@ -117,7 +117,7 @@ void FStudioAnalytics::Shutdown()
 
 double FStudioAnalytics::GetAnalyticSeconds()
 {
-	return bInitialized ? TimeEstimation : FPlatformTime::Seconds();
+	return bInitialized ? TimeEstimation.load(std::memory_order_relaxed) : FPlatformTime::Seconds();
 }
 
 void FStudioAnalytics::RecordEvent(const FString& EventName)
diff --git a/Engine/Source/Runtime/Engine/Public/StudioAnalytics.h b/Engine/Source/Runtime/Engine/Public/StudioAnalytics.h
--- a/Engine/Source/Runtime/Engine/Public/StudioAnalytics.h
+++ b/Engine/Source/Runtime/Engine/Public/StudioAnalytics.h
@@ -5,6 +5,8 @@
 #include "CoreMinimal.h"
 #include "HAL/Thread.h"
 
+#include <atomic>
+
 class IAnalyticsProvider;
 class IAnalyticsProviderET;
 struct FAnalyticsEventAttribute;
@@ -56,6 +58,6 @@ private:
 	static ENGINE_API TSharedPtr<IAnalyticsProviderET> Analytics;
 	static TArray<FAnalyticsEventAttribute> DefaultAttributes;
 	static FThread TimerThread;
-	static volatile double TimeEstimation;
+	static std::atomic<double> TimeEstimation;
 };
 
