Solved incorrect material editor only data name message spamming during cooking.

diff --git a/Engine/Source/Runtime/Engine/Private/Materials/MaterialInterface.cpp b/Engine/Source/Runtime/Engine/Private/Materials/MaterialInterface.cpp
--- a/Engine/Source/Runtime/Engine/Private/Materials/MaterialInterface.cpp
+++ b/Engine/Source/Runtime/Engine/Private/Materials/MaterialInterface.cpp
@@ -1545,7 +1545,7 @@ void UMaterialInterface::PreSave(FObjectPreSaveContext ObjectSaveContext)
 	// the editor running on cooked data.
 	UMaterialInterfaceEditorOnlyData* EditorOnly = GetEditorOnlyData();
 	FString EditorOnlyDataName = MaterialInterface::GetEditorOnlyDataName(*GetName());
-	if (EditorOnly && EditorOnly->GetName() != EditorOnlyDataName)
+	if (!ObjectSaveContext.IsCooking() && EditorOnly && EditorOnly->GetName() != EditorOnlyDataName)
 	{
 		UE_LOG(LogMaterial, Display, TEXT("MaterialInterface %s has a incorrectly name EditorOnlyData '%s'. This may cause issues when running the editor on cooked data. Trying to rename it to the correct name '%s'."), *GetName(), *EditorOnly->GetName(), *EditorOnlyDataName);
 		UObject* EditorOnlyExisting = StaticFindObject(/*Class=*/ nullptr, this, *EditorOnlyDataName, true);
