Level Instance: change how we store the deterministic guids to actors in instanced level instances, by storing the instance guid in a separate actor property and keeping the original actor guid intact, since there was existing code relying on this.

diff --git a/Engine/Source/Runtime/Engine/Classes/GameFramework/Actor.h b/Engine/Source/Runtime/Engine/Classes/GameFramework/Actor.h
--- a/Engine/Source/Runtime/Engine/Classes/GameFramework/Actor.h
+++ b/Engine/Source/Runtime/Engine/Classes/GameFramework/Actor.h
@@ -882,14 +882,23 @@ private:
 #if WITH_EDITORONLY_DATA
 protected:
 	/**
-	 * The GUID for this actor.
-	 *
-	 * Note: Don't use VisibleAnywhere here to avoid getting the CPF_Edit flag and get this property reset when resetting to defaults.
-	 *       See FActorDetails::AddActorCategory and EditorUtilities::CopySingleProperty for details.
+	 * The GUID for this actor; this guid will be the same for actors from instanced streaming levels.
+	 * @see		ActorInstanceGuid
+	 * @note	Don't use VisibleAnywhere here to avoid getting the CPF_Edit flag and get this property reset when resetting to defaults.
+	 *			See FActorDetails::AddActorCategory and EditorUtilities::CopySingleProperty for details.
 	 */
 	UPROPERTY(BluePrintReadOnly, AdvancedDisplay, Category=Actor, NonPIEDuplicateTransient, TextExportTransient, NonTransactional)
 	FGuid ActorGuid;
 
+	/**
+	 * The instance GUID for this actor; this guid will be unique for actors from instanced streaming levels.
+	 * @see		ActorGuid
+	 * @note	Don't use VisibleAnywhere here to avoid getting the CPF_Edit flag and get this property reset when resetting to defaults.
+	 *			See FActorDetails::AddActorCategory and EditorUtilities::CopySingleProperty for details.
+	 */
+	UPROPERTY(BluePrintReadOnly, AdvancedDisplay, Category=Actor, Transient, NonTransactional)
+	FGuid ActorInstanceGuid;
+
 	/**
 	 * The GUID for this actor's content bundle.
 	 */
@@ -955,6 +964,9 @@ public:
 	/** Returns this actor's Guid. Actor Guids are only available in editor builds. */
 	inline const FGuid& GetActorGuid() const { return ActorGuid; }
 
+	/** Returns this actor's instance Guid. Actor Guids are only available in editor builds. */
+	inline const FGuid& GetActorInstanceGuid() const { return ActorInstanceGuid.IsValid() ? ActorInstanceGuid : ActorGuid; }
+
 	/** Returns this actor's content bundle Guid. */
 	inline const FGuid& GetContentBundleGuid() const { return ContentBundleGuid; }
 
@@ -3953,6 +3965,7 @@ private:
 
 	friend struct FSetActorHiddenInSceneOutliner;
 	friend struct FSetActorGuid;
+	friend struct FSetActorInstanceGuid;
 	friend struct FSetActorContentBundleGuid;
 	friend struct FAssignActorDataLayer;
 	friend struct FSetActorSelectable;
@@ -4209,8 +4222,16 @@ private:
 	friend class UEngine;
 	friend class UExternalActorsCommandlet;
 	friend class UWorldPartitionConvertCommandlet;
+};
+
+struct FSetActorInstanceGuid
+{
+private:
+	FSetActorInstanceGuid(AActor* InActor, const FGuid& InActorInstanceGuid)
+	{
+		InActor->ActorInstanceGuid = InActorInstanceGuid;
+	}
 	friend class ULevelStreamingLevelInstance;
-	friend class UWorldPartitionLevelStreamingDynamic;
 	friend class FWorldPartitionLevelHelper;
 };
 
diff --git a/Engine/Source/Runtime/Engine/Private/LevelInstance/LevelInstanceLevelStreaming.cpp b/Engine/Source/Runtime/Engine/Private/LevelInstance/LevelInstanceLevelStreaming.cpp
--- a/Engine/Source/Runtime/Engine/Private/LevelInstance/LevelInstanceLevelStreaming.cpp
+++ b/Engine/Source/Runtime/Engine/Private/LevelInstance/LevelInstanceLevelStreaming.cpp
@@ -151,6 +151,9 @@ void ULevelStreamingLevelInstance::PrepareLevelInstanceLoadedActor(AActor& InAct
 	if (InLevelInstance)
 	{
 		InActor.PushLevelInstanceEditingStateToProxies(CastChecked<AActor>(InLevelInstance)->IsInEditingLevelInstance());
+
+		const FGuid& LevelInstanceGuid = InLevelInstance->GetLevelInstanceID().GetGuid();
+		FSetActorInstanceGuid SetActorInstanceGuid(&InActor, FGuid::Combine(LevelInstanceGuid, InActor.GetActorGuid()));
 	}
 
 	if (LevelInstanceEditorInstanceActor.IsValid())
@@ -221,8 +224,6 @@ ULevelStreamingLevelInstance* ULevelStreamingLevelInstance::LoadInstance(ILevelI
 		{
 			GEngine->BlockTillLevelStreamingCompleted(LevelInstanceActor->GetWorld());
 
-			const FGuid& LevelInstanceGuid = LevelInstance->GetLevelInstanceID().GetGuid();
-
 			// Most of the code here is meant to allow partial support for undo/redo of LevelInstance Instance Loading:
 			// by setting the objects RF_Transient and !RF_Transactional we can check when unloading if those flags
 			// have been changed and figure out if we need to clear the transaction buffer or not.
@@ -252,8 +253,6 @@ ULevelStreamingLevelInstance* ULevelStreamingLevelInstance::LoadInstance(ILevelI
 					if (LevelActor)
 					{
 						LevelStreaming->PrepareLevelInstanceLoadedActor(*LevelActor, LevelInstance, false);
-
-						FSetActorGuid SetActorGuid(LevelActor, FGuid::Combine(LevelInstanceGuid, LevelActor->GetActorGuid()));
 					}
 				}
 
diff --git a/Engine/Source/Runtime/Engine/Private/LevelInstance/LevelInstanceSubsystem.cpp b/Engine/Source/Runtime/Engine/Private/LevelInstance/LevelInstanceSubsystem.cpp
--- a/Engine/Source/Runtime/Engine/Private/LevelInstance/LevelInstanceSubsystem.cpp
+++ b/Engine/Source/Runtime/Engine/Private/LevelInstance/LevelInstanceSubsystem.cpp
@@ -2585,7 +2585,7 @@ bool ULevelInstanceSubsystem::PassLevelInstanceFilter(UWorld* World, const FWorl
 			TArray<FGuid> AncestorGuids;
 			ForEachLevelInstanceAncestorsAndSelf(Cast<AActor>(LevelInstanceInterface), [&AncestorGuids, &LevelInstanceInterface](const ILevelInstanceInterface* Ancestor)
 			{
-				AncestorGuids.Insert(Ancestor->GetLevelInstanceID().GetOriginalGuid(), 0);
+				AncestorGuids.Insert(Cast<AActor>(Ancestor)->GetActorGuid(), 0);
 				LevelInstanceInterface = Ancestor;
 				return true;
 			});
diff --git a/Engine/Source/Runtime/Engine/Private/WorldPartition/WorldPartitionLevelHelper.cpp b/Engine/Source/Runtime/Engine/Private/WorldPartition/WorldPartitionLevelHelper.cpp
--- a/Engine/Source/Runtime/Engine/Private/WorldPartition/WorldPartitionLevelHelper.cpp
+++ b/Engine/Source/Runtime/Engine/Private/WorldPartition/WorldPartitionLevelHelper.cpp
@@ -447,8 +447,8 @@ bool FWorldPartitionLevelHelper::LoadActors(UWorld* InOwningWorld, ULevel* InDes
 					TransformParams.bDoPostEditMove = false;
 					FLevelUtils::ApplyLevelTransform(TransformParams);
 
-					// Set the actor's guid
-					FSetActorGuid SetActorGuid(Actor, PackageObjectMapping->ActorGuid);
+					// Set the actor's instance guid
+					FSetActorInstanceGuid SetActorInstanceGuid(Actor, PackageObjectMapping->ActorInstanceGuid);
 						
 					// Path to use when searching for this actor in MoveExternalActorsToLevel
 					PackageObjectMapping->LoadedPath = *Actor->GetPathName();
diff --git a/Engine/Source/Runtime/Engine/Private/WorldPartition/WorldPartitionRuntimeHash.cpp b/Engine/Source/Runtime/Engine/Private/WorldPartition/WorldPartitionRuntimeHash.cpp
--- a/Engine/Source/Runtime/Engine/Private/WorldPartition/WorldPartitionRuntimeHash.cpp
+++ b/Engine/Source/Runtime/Engine/Private/WorldPartition/WorldPartitionRuntimeHash.cpp
@@ -222,10 +222,8 @@ void UWorldPartitionRuntimeHash::PopulateRuntimeCell(UWorldPartitionRuntimeCell*
 	FBox CellContentBounds(ForceInit);
 	for (const IStreamingGenerationContext::FActorInstance& ActorInstance : ActorInstances)
 	{
-		check(ActorInstance.GetActorGuid().IsValid());
-
 		const FWorldPartitionActorDescView& ActorDescView = ActorInstance.GetActorDescView();
-		RuntimeCell->AddActorToCell(ActorDescView, ActorInstance.GetContainerID(), ActorInstance.GetTransform(), ActorInstance.GetActorDescContainer(), ActorInstance.GetActorGuid());
+		RuntimeCell->AddActorToCell(ActorDescView, ActorInstance.GetContainerID(), ActorInstance.GetTransform(), ActorInstance.GetActorDescContainer());
 		CellContentBounds += ActorDescView.GetRuntimeBounds().TransformBy(ActorInstance.GetTransform());
 					
 		if (ActorInstance.GetContainerID().IsMainContainer() && RuntimeCell->UnsavedActorsContainer)
diff --git a/Engine/Source/Runtime/Engine/Private/WorldPartition/WorldPartitionRuntimeLevelStreamingCell.cpp b/Engine/Source/Runtime/Engine/Private/WorldPartition/WorldPartitionRuntimeLevelStreamingCell.cpp
--- a/Engine/Source/Runtime/Engine/Private/WorldPartition/WorldPartitionRuntimeLevelStreamingCell.cpp
+++ b/Engine/Source/Runtime/Engine/Private/WorldPartition/WorldPartitionRuntimeLevelStreamingCell.cpp
@@ -202,7 +202,7 @@ void UWorldPartitionRuntimeLevelStreamingCell::SetIsAlwaysLoaded(bool bInIsAlway
 }
 
 #if WITH_EDITOR
-void UWorldPartitionRuntimeLevelStreamingCell::AddActorToCell(const FWorldPartitionActorDescView& ActorDescView, const FActorContainerID& InContainerID, const FTransform& InContainerTransform, const UActorDescContainer* InContainer, const FGuid& InActorGuid)
+void UWorldPartitionRuntimeLevelStreamingCell::AddActorToCell(const FWorldPartitionActorDescView& ActorDescView, const FActorContainerID& InContainerID, const FTransform& InContainerTransform, const UActorDescContainer* InContainer)
 {
 	check(!ActorDescView.GetActorIsEditorOnly());
 
@@ -213,7 +213,7 @@ void UWorldPartitionRuntimeLevelStreamingCell::AddActorToCell(const FWorldPartit
 		InContainer->GetContainerPackage(), 
 		GetWorld()->GetPackage()->GetFName(), 
 		ActorDescView.GetContentBundleGuid(),
-		InActorGuid
+		InContainerID.GetActorGuid(ActorDescView.GetGuid())
 	);
 }
 
@@ -322,9 +322,9 @@ void UWorldPartitionRuntimeLevelStreamingCell::DumpStateLog(FHierarchicalLogArch
 
 	for (const FWorldPartitionRuntimeCellObjectMapping& Mapping : SortedPackages)
 	{
-		Ar.Printf(TEXT("   Actor Path: %s"), *Mapping.Path.ToString());
-		Ar.Printf(TEXT("Actor Package: %s"), *Mapping.Package.ToString());
-		Ar.Printf(TEXT("   Actor Guid: %s"), *Mapping.ActorGuid.ToString());
+		Ar.Printf(TEXT("         Actor Path: %s"), *Mapping.Path.ToString());
+		Ar.Printf(TEXT("      Actor Package: %s"), *Mapping.Package.ToString());
+		Ar.Printf(TEXT("Actor Instance Guid: %s"), *Mapping.ActorInstanceGuid.ToString());
 	}
 }
 #endif
diff --git a/Engine/Source/Runtime/Engine/Private/WorldPartition/WorldPartitionStreamingGeneration.cpp b/Engine/Source/Runtime/Engine/Private/WorldPartition/WorldPartitionStreamingGeneration.cpp
--- a/Engine/Source/Runtime/Engine/Private/WorldPartition/WorldPartitionStreamingGeneration.cpp
+++ b/Engine/Source/Runtime/Engine/Private/WorldPartition/WorldPartitionStreamingGeneration.cpp
@@ -163,16 +163,6 @@ class FWorldPartitionStreamingGenerator
 					ActorSetInstance.ContainerID = ContainerInstanceDescriptor.ID;
 					ActorSetInstance.Transform = ContainerInstanceDescriptor.Transform;
 
-					// Apply the per-container actor guids
-					if (!ContainerID.IsMainContainer())
-					{
-						ActorSetInstance.ActorsSetGuids.Reserve(ActorSet.Actors.Num());
-						for (const FGuid& ActorGuid : ActorSet.Actors)
-						{
-							ActorSetInstance.ActorsSetGuids.Add(ActorGuid, ActorSetInstance.ContainerID.GetActorGuid(ActorGuid));
-						}
-					}
-
 					// Apply AND logic on spatially loaded flag
 					ActorSetInstance.bIsSpatiallyLoaded = ReferenceActorDescView.GetIsSpatiallyLoaded() && ContainerInstanceDescriptor.bIsSpatiallyLoaded;
 
diff --git a/Engine/Source/Runtime/Engine/Private/WorldPartition/WorldPartitionStreamingGenerationContext.cpp b/Engine/Source/Runtime/Engine/Private/WorldPartition/WorldPartitionStreamingGenerationContext.cpp
--- a/Engine/Source/Runtime/Engine/Private/WorldPartition/WorldPartitionStreamingGenerationContext.cpp
+++ b/Engine/Source/Runtime/Engine/Private/WorldPartition/WorldPartitionStreamingGenerationContext.cpp
@@ -14,15 +14,6 @@ const FActorContainerID& IStreamingGenerationContext::FActorInstance::GetContain
 	return ActorSetInstance->ContainerID;
 }
 
-const FGuid& IStreamingGenerationContext::FActorInstance::GetActorGuid() const
-{
-	if (const FGuid* RemappedGuid = ActorSetInstance->ActorsSetGuids.Find(ActorGuid))
-	{
-		return *RemappedGuid;
-	}
-	return ActorGuid;
-}
-
 const FTransform& IStreamingGenerationContext::FActorInstance::GetTransform() const
 {
 	return ActorSetInstance->Transform;
diff --git a/Engine/Source/Runtime/Engine/Public/LevelInstance/LevelInstanceTypes.h b/Engine/Source/Runtime/Engine/Public/LevelInstance/LevelInstanceTypes.h
--- a/Engine/Source/Runtime/Engine/Public/LevelInstance/LevelInstanceTypes.h
+++ b/Engine/Source/Runtime/Engine/Public/LevelInstance/LevelInstanceTypes.h
@@ -35,13 +35,6 @@ struct FLevelInstanceID
 
 	inline const FGuid& GetGuid() const { return Guid; }
 
-protected:
-#if WITH_EDITOR
-	friend ULevelInstanceSubsystem;
-	// - this relies on the Guids[0] to be the OriginalGuid of the LevelInstance Actor before it gets overriden in ULevelInstanceLevelStreaming.
-	// - this is only used for filtering actors, we compare a filter built from a Level Package (with original actor guids) to the filter of a loaded Level Instance (with unique actor guids)
-	inline const FGuid& GetOriginalGuid() const { return Guids[0]; }
-#endif
 private:
 	uint64 Hash = 0;
 	FGuid Guid;
diff --git a/Engine/Source/Runtime/Engine/Public/WorldPartition/WorldPartitionActorContainerID.h b/Engine/Source/Runtime/Engine/Public/WorldPartition/WorldPartitionActorContainerID.h
--- a/Engine/Source/Runtime/Engine/Public/WorldPartition/WorldPartitionActorContainerID.h
+++ b/Engine/Source/Runtime/Engine/Public/WorldPartition/WorldPartitionActorContainerID.h
@@ -67,7 +67,8 @@ struct FActorContainerID
 
 	FGuid GetActorGuid(const FGuid& InActorGuid) const
 	{
-		return FGuid::Combine(Guid, InActorGuid);
+		// Preserve the original actor guid for top level actors
+		return Guid.IsValid() ? FGuid::Combine(Guid, InActorGuid) : InActorGuid;
 	}
 
 	friend FORCEINLINE uint32 GetTypeHash(const FActorContainerID& InContainerID)
diff --git a/Engine/Source/Runtime/Engine/Public/WorldPartition/WorldPartitionRuntimeCell.h b/Engine/Source/Runtime/Engine/Public/WorldPartition/WorldPartitionRuntimeCell.h
--- a/Engine/Source/Runtime/Engine/Public/WorldPartition/WorldPartitionRuntimeCell.h
+++ b/Engine/Source/Runtime/Engine/Public/WorldPartition/WorldPartitionRuntimeCell.h
@@ -42,7 +42,7 @@ struct FWorldPartitionRuntimeCellObjectMapping
 #endif
 	{}
 
-	FWorldPartitionRuntimeCellObjectMapping(FName InPackage, FName InPath, const FActorContainerID& InContainerID, const FTransform& InContainerTransform, FName InContainerPackage, FName InWorldPackage, const FGuid& InContentBundleGuid, const FGuid& InActorGuid)
+	FWorldPartitionRuntimeCellObjectMapping(FName InPackage, FName InPath, const FActorContainerID& InContainerID, const FTransform& InContainerTransform, FName InContainerPackage, FName InWorldPackage, const FGuid& InContentBundleGuid, const FGuid& InActorInstanceGuid)
 #if WITH_EDITORONLY_DATA
 		: Package(InPackage)
 		, Path(InPath)
@@ -51,7 +51,7 @@ struct FWorldPartitionRuntimeCellObjectMapping
 		, ContainerPackage(InContainerPackage)
 		, WorldPackage(InWorldPackage)
 		, ContentBundleGuid(InContentBundleGuid)
-		, ActorGuid(InActorGuid)
+		, ActorInstanceGuid(InActorInstanceGuid)
 		, LoadedPath(InPath)
 #endif
 	{}
@@ -103,7 +103,7 @@ struct FWorldPartitionRuntimeCellObjectMapping
 	 * GUID of the actor instance
 	 */
 	UPROPERTY()
-	FGuid ActorGuid;
+	FGuid ActorInstanceGuid;
 
 	/**
 	* Loaded actor path (when cooking or pie)
@@ -227,7 +227,7 @@ class ENGINE_API UWorldPartitionRuntimeCell : public UObject, public IWorldParti
 	virtual TSet<FName> GetActorPackageNames() const override { return TSet<FName>(); }
 	//~End IWorldPartitionCell Interface
 
-	virtual void AddActorToCell(const FWorldPartitionActorDescView& ActorDescView, const FActorContainerID& InContainerID, const FTransform& InContainerTransform, const UActorDescContainer* InContainer, const FGuid& InActorGuid) PURE_VIRTUAL(UWorldPartitionRuntimeCell::AddActorToCell,);
+	virtual void AddActorToCell(const FWorldPartitionActorDescView& ActorDescView, const FActorContainerID& InContainerID, const FTransform& InContainerTransform, const UActorDescContainer* InContainer) PURE_VIRTUAL(UWorldPartitionRuntimeCell::AddActorToCell,);
 	virtual int32 GetActorCount() const PURE_VIRTUAL(UWorldPartitionRuntimeCell::GetActorCount, return 0;);
 
 	// Cook methods
diff --git a/Engine/Source/Runtime/Engine/Public/WorldPartition/WorldPartitionRuntimeLevelStreamingCell.h b/Engine/Source/Runtime/Engine/Public/WorldPartition/WorldPartitionRuntimeLevelStreamingCell.h
--- a/Engine/Source/Runtime/Engine/Public/WorldPartition/WorldPartitionRuntimeLevelStreamingCell.h
+++ b/Engine/Source/Runtime/Engine/Public/WorldPartition/WorldPartitionRuntimeLevelStreamingCell.h
@@ -48,7 +48,7 @@ class ENGINE_API UWorldPartitionRuntimeLevelStreamingCell : public UWorldPartiti
 
 #if WITH_EDITOR
 	//~Begin UWorldPartitionRuntimeCell Interface
-	virtual void AddActorToCell(const FWorldPartitionActorDescView& ActorDescView, const FActorContainerID& InContainerID, const FTransform& InContainerTransform, const UActorDescContainer* InContainer, const FGuid& InActorGuid) override;
+	virtual void AddActorToCell(const FWorldPartitionActorDescView& ActorDescView, const FActorContainerID& InContainerID, const FTransform& InContainerTransform, const UActorDescContainer* InContainer) override;
 	virtual int32 GetActorCount() const override;
 	virtual void DumpStateLog(FHierarchicalLogArchive& Ar) override;
 	// Cook methods
diff --git a/Engine/Source/Runtime/Engine/Public/WorldPartition/WorldPartitionStreamingGenerationContext.h b/Engine/Source/Runtime/Engine/Public/WorldPartition/WorldPartitionStreamingGenerationContext.h
--- a/Engine/Source/Runtime/Engine/Public/WorldPartition/WorldPartitionStreamingGenerationContext.h
+++ b/Engine/Source/Runtime/Engine/Public/WorldPartition/WorldPartitionStreamingGenerationContext.h
@@ -43,7 +43,6 @@ public:
 		FActorContainerID ContainerID;
 		FTransform Transform;
 		const FActorSet* ActorSet;
-		TMap<FGuid, FGuid> ActorsSetGuids;
 	};
 
 	struct ENGINE_API FActorInstance
@@ -58,7 +57,6 @@ public:
 
 		const FWorldPartitionActorDescView& GetActorDescView() const;
 		const FActorContainerID& GetContainerID() const;
-		const FGuid& GetActorGuid() const;
 		const FTransform& GetTransform() const;
 		const UActorDescContainer* GetActorDescContainer() const;
 	};
