WorldPartitionActorFilter: Cleanup + Caching

- Move Registered containers into WorldPartitionSubsystem
- Move Filtering code into WorldPartitionSubsystem
- Share filtering code between Editor Filter and Generate Streaming

#rnx

diff --git a/Engine/Source/Editor/LevelInstanceEditor/Private/LevelInstanceFilterPropertyTypeCustomization.cpp b/Engine/Source/Editor/LevelInstanceEditor/Private/LevelInstanceFilterPropertyTypeCustomization.cpp
--- a/Engine/Source/Editor/LevelInstanceEditor/Private/LevelInstanceFilterPropertyTypeCustomization.cpp
+++ b/Engine/Source/Editor/LevelInstanceEditor/Private/LevelInstanceFilterPropertyTypeCustomization.cpp
@@ -2,8 +2,8 @@
 
 #include "LevelInstanceFilterPropertyTypeCustomization.h"
 #include "LevelInstance/LevelInstanceInterface.h"
-#include "LevelInstance/LevelInstanceSubsystem.h"
 #include "LevelInstance/LevelInstanceComponent.h"
+#include "WorldPartition/WorldPartitionSubsystem.h"
 #include "WorldPartition/Filter/WorldPartitionActorFilter.h"
 #include "Engine/World.h"
 #include "GameFramework/Actor.h"
@@ -68,10 +68,10 @@ TSharedPtr<FWorldPartitionActorFilterMode::FFilter> FLevelInstanceFilterProperty
 		return nullptr;
 	}
 
-	ULevelInstanceSubsystem* LevelInstanceSubsystem = World->GetSubsystem<ULevelInstanceSubsystem>();
+	UWorldPartitionSubsystem* WorldPartitionSubsystem = World->GetSubsystem<UWorldPartitionSubsystem>();
 
 	// Get the Default Filter for selected WorldAssetPackage
-	TSharedPtr<FWorldPartitionActorFilter> Filter = MakeShared<FWorldPartitionActorFilter>(LevelInstanceSubsystem->GetLevelInstanceFilter(WorldAssetPackage));
+	TSharedPtr<FWorldPartitionActorFilter> Filter = MakeShared<FWorldPartitionActorFilter>(WorldPartitionSubsystem->GetWorldPartitionActorFilter(WorldAssetPackage));
 
 	// Set its name based on single/multi selection (root node name in outliner)
 	Filter->DisplayName = LevelInstances.Num() == 1 ? ActorLabel : TEXT("(Multiple Actors)");
diff --git a/Engine/Source/Runtime/Engine/Private/LevelInstance/LevelInstanceComponent.cpp b/Engine/Source/Runtime/Engine/Private/LevelInstance/LevelInstanceComponent.cpp
--- a/Engine/Source/Runtime/Engine/Private/LevelInstance/LevelInstanceComponent.cpp
+++ b/Engine/Source/Runtime/Engine/Private/LevelInstance/LevelInstanceComponent.cpp
@@ -2,6 +2,7 @@
 #include "LevelInstance/LevelInstanceComponent.h"
 #include "LevelInstance/LevelInstanceInterface.h"
 #include "LevelInstance/LevelInstanceSubsystem.h"
+#include "WorldPartition/WorldPartitionSubsystem.h"
 #include "LevelInstance/LevelInstanceEditorInstanceActor.h"
 #include "Components/BillboardComponent.h"
 #include "Engine/Texture2D.h"
@@ -63,7 +64,7 @@ bool ULevelInstanceComponent::ShouldShowSpriteComponent() const
 
 void ULevelInstanceComponent::PreEditUndo()
 {
-	CachedFilter = Filter;
+	UndoRedoCachedFilter = Filter;
 }
 
 void ULevelInstanceComponent::PostEditUndo()
@@ -73,12 +74,12 @@ void ULevelInstanceComponent::PostEditUndo()
 	UpdateComponentToWorld();
 	UpdateEditorInstanceActor();
 
-	if (Filter != CachedFilter)
+	if (Filter != UndoRedoCachedFilter)
 	{
 		FWorldPartitionActorFilter::RequestFilterRefresh(false);
 		FWorldPartitionActorFilter::GetOnWorldPartitionActorFilterChanged().Broadcast();
 	}
-	CachedFilter = FWorldPartitionActorFilter();
+	UndoRedoCachedFilter = FWorldPartitionActorFilter();
 }
 
 void ULevelInstanceComponent::PostEditChangeProperty(FPropertyChangedEvent& PropertyChangedEvent)
@@ -142,4 +143,31 @@ void ULevelInstanceComponent::OnCommit()
 	}
 }
 
+const TMap<FActorContainerID, TSet<FGuid>>& ULevelInstanceComponent::GetFilteredActorsPerContainer() const
+{
+	if (CachedFilter != Filter)
+	{
+		CachedFilteredActorsPerContainer.Reset();
+	}
+
+	if (CachedFilteredActorsPerContainer.IsSet())
+	{
+		return CachedFilteredActorsPerContainer.GetValue();
+	}
+
+	TMap<FActorContainerID, TSet<FGuid>> FilteredActors;
+	if (ILevelInstanceInterface* LevelInstance = Cast<ILevelInstanceInterface>(GetOwner()))
+	{
+		// Fill Container Instance Filter
+		UWorldPartitionSubsystem* WorldPartitionSubsystem = GetWorld()->GetSubsystem<UWorldPartitionSubsystem>();
+		check(WorldPartitionSubsystem);
+				
+		FilteredActors = WorldPartitionSubsystem->GetFilteredActorsPerContainer(LevelInstance->GetLevelInstanceID().GetContainerID(), LevelInstance->GetWorldAssetPackage(), Filter);
+	}
+	
+	CachedFilteredActorsPerContainer = MoveTemp(FilteredActors);
+	CachedFilter = Filter;
+	return CachedFilteredActorsPerContainer.GetValue();
+}
+
 #endif
diff --git a/Engine/Source/Runtime/Engine/Private/LevelInstance/LevelInstanceInterface.cpp b/Engine/Source/Runtime/Engine/Private/LevelInstance/LevelInstanceInterface.cpp
--- a/Engine/Source/Runtime/Engine/Private/LevelInstance/LevelInstanceInterface.cpp
+++ b/Engine/Source/Runtime/Engine/Private/LevelInstance/LevelInstanceInterface.cpp
@@ -248,6 +248,17 @@ const FWorldPartitionActorFilter& ILevelInstanceInterface::GetFilter() const
 	return NoFilter;
 }
 
+const TMap<FActorContainerID, TSet<FGuid>>& ILevelInstanceInterface::GetFilteredActorsPerContainer() const
+{
+	if (ULevelInstanceComponent* LevelInstanceComponent = GetLevelInstanceComponent())
+	{
+		return LevelInstanceComponent->GetFilteredActorsPerContainer();
+	}
+
+	static TMap<FActorContainerID, TSet<FGuid>> NoFilteredActors;
+	return NoFilteredActors;
+}
+
 #endif
 
 
diff --git a/Engine/Source/Runtime/Engine/Private/LevelInstance/LevelInstanceLevelStreaming.cpp b/Engine/Source/Runtime/Engine/Private/LevelInstance/LevelInstanceLevelStreaming.cpp
--- a/Engine/Source/Runtime/Engine/Private/LevelInstance/LevelInstanceLevelStreaming.cpp
+++ b/Engine/Source/Runtime/Engine/Private/LevelInstance/LevelInstanceLevelStreaming.cpp
@@ -152,8 +152,8 @@ void ULevelStreamingLevelInstance::PrepareLevelInstanceLoadedActor(AActor& InAct
 	{
 		InActor.PushLevelInstanceEditingStateToProxies(CastChecked<AActor>(InLevelInstance)->IsInEditingLevelInstance());
 
-		const FGuid& LevelInstanceGuid = InLevelInstance->GetLevelInstanceID().GetGuid();
-		FSetActorInstanceGuid SetActorInstanceGuid(&InActor, FGuid::Combine(LevelInstanceGuid, InActor.GetActorGuid()));
+		const FActorContainerID& ContainerID = InLevelInstance->GetLevelInstanceID().GetContainerID();
+		FSetActorInstanceGuid SetActorInstanceGuid(&InActor, ContainerID.GetActorGuid(InActor.GetActorGuid()));
 	}
 
 	if (LevelInstanceEditorInstanceActor.IsValid())
diff --git a/Engine/Source/Runtime/Engine/Private/LevelInstance/LevelInstanceSubsystem.cpp b/Engine/Source/Runtime/Engine/Private/LevelInstance/LevelInstanceSubsystem.cpp
--- a/Engine/Source/Runtime/Engine/Private/LevelInstance/LevelInstanceSubsystem.cpp
+++ b/Engine/Source/Runtime/Engine/Private/LevelInstance/LevelInstanceSubsystem.cpp
@@ -33,7 +33,6 @@
 #include "Editor.h"
 #include "Editor/Transactor.h"
 #include "EditorLevelUtils.h"
-#include "Engine/LevelBounds.h"
 #include "Modules/ModuleManager.h"
 #include "PackedLevelActor/PackedLevelActor.h"
 #include "PackedLevelActor/PackedLevelActorBuilder.h"
@@ -73,8 +72,6 @@ void ULevelInstanceSubsystem::AddReferencedObjects(UObject* InThis, FReferenceCo
 	{
 		This->LevelInstanceEdit->AddReferencedObjects(Collector);
 	}
-
-	This->ActorDescContainerInstanceManager.AddReferencedObjects(Collector);
 #endif
 }
 
@@ -136,11 +133,11 @@ FLevelInstanceID::FLevelInstanceID(ULevelInstanceSubsystem* LevelInstanceSubsyst
 	
 	Hash = CityHash64WithSeed((const char*)Guids.GetData(), Guids.Num() * sizeof(FGuid), NameHash);
 
-	// Make sure to start with a zero guid as the top-level container ID is zero.
-	Guid.Invalidate();
+	// Make sure to start main container id
+	ContainerID = FActorContainerID();
 	for (int32 GuidIndex = Guids.Num() - 1; GuidIndex >= 0; GuidIndex--)
 	{
-		Guid = FGuid::Combine(Guid, Guids[GuidIndex]);
+		ContainerID = FActorContainerID(ContainerID, Guids[GuidIndex]);
 	}
 }
 
@@ -705,7 +702,9 @@ bool ULevelInstanceSubsystem::GetLevelInstanceBounds(const ILevelInstanceInterfa
 
 		FString LevelPackage = LevelInstance->GetWorldAssetPackage();
 
-		if (FBox ContainerBounds = ActorDescContainerInstanceManager.GetContainerBounds(*LevelPackage); ContainerBounds.IsValid)
+		UWorldPartitionSubsystem* WorldPartitionSubsystem = GetWorld()->GetSubsystem<UWorldPartitionSubsystem>();
+		check(WorldPartitionSubsystem);
+		if (FBox ContainerBounds = WorldPartitionSubsystem->GetContainerBounds(*LevelPackage); ContainerBounds.IsValid)
 		{
 			FTransform LevelInstancePivotOffsetTransform = FTransform(ULevel::GetLevelInstancePivotOffsetFromPackage(*LevelPackage));
 			FTransform LevelTransform = LevelInstancePivotOffsetTransform * CastChecked<AActor>(LevelInstance)->GetActorTransform();
@@ -1482,83 +1481,6 @@ void ULevelInstanceSubsystem::RemoveLevelsFromWorld(const TArray<ULevel*>& InLev
 	}
 }
 
-void ULevelInstanceSubsystem::FActorDescContainerInstanceManager::FActorDescContainerInstance::AddReferencedObjects(FReferenceCollector& Collector)
-{
-	Collector.AddReferencedObject(Container);
-}
-
-void ULevelInstanceSubsystem::FActorDescContainerInstanceManager::FActorDescContainerInstance::UpdateBounds()
-{
-	Bounds.Init();
-	for (FActorDescList::TIterator<> ActorDescIt(Container); ActorDescIt; ++ActorDescIt)
-	{
-		if (ActorDescIt->GetActorNativeClass()->IsChildOf<ALevelBounds>())
-		{
-			continue;
-		}
-		Bounds += ActorDescIt->GetRuntimeBounds();
-	}
-}
-
-void ULevelInstanceSubsystem::FActorDescContainerInstanceManager::AddReferencedObjects(FReferenceCollector& Collector)
-{
-	for (auto& [Name, ContainerInstance] : ActorDescContainers)
-	{
-		ContainerInstance.AddReferencedObjects(Collector);
-	}
-}
-
-UActorDescContainer* ULevelInstanceSubsystem::FActorDescContainerInstanceManager::RegisterContainer(FName PackageName, UWorld* InWorld)
-{
-	FActorDescContainerInstance* ExistingContainerInstance = &ActorDescContainers.FindOrAdd(PackageName);
-	UActorDescContainer* ActorDescContainer = ExistingContainerInstance->Container;
-	
-	if (ExistingContainerInstance->RefCount++ == 0)
-	{
-		ActorDescContainer = NewObject<UActorDescContainer>(GetTransientPackage());
-		ExistingContainerInstance->Container = ActorDescContainer;
-		
-		// This will potentially invalidate ExistingContainerInstance due to ActorDescContainers reallocation
-		ActorDescContainer->Initialize({ InWorld, PackageName });
-			
-		ExistingContainerInstance = &ActorDescContainers.FindChecked(PackageName);
-		ExistingContainerInstance->UpdateBounds();
-	}
-
-	check(ActorDescContainer->IsTemplateContainer());
-	return ActorDescContainer;
-}
-
-void ULevelInstanceSubsystem::FActorDescContainerInstanceManager::UnregisterContainer(UActorDescContainer* Container)
-{
-	FName PackageName = Container->GetContainerPackage();
-	FActorDescContainerInstance& ExistingContainerInstance = ActorDescContainers.FindChecked(PackageName);
-
-	if (--ExistingContainerInstance.RefCount == 0)
-	{
-		ExistingContainerInstance.Container->Uninitialize();
-		ActorDescContainers.FindAndRemoveChecked(PackageName);
-	}
-}
-
-FBox ULevelInstanceSubsystem::FActorDescContainerInstanceManager::GetContainerBounds(FName PackageName) const
-{
-	if (const FActorDescContainerInstance* ActorDescContainerInstance = ActorDescContainers.Find(PackageName))
-	{
-		return ActorDescContainerInstance->Bounds;
-	}
-	return FBox(ForceInit);
-}
-
-void ULevelInstanceSubsystem::FActorDescContainerInstanceManager::OnLevelInstanceActorCommitted(ILevelInstanceInterface* LevelInstance)
-{
-	const FName PackageName = *LevelInstance->GetWorldAssetPackage();
-	if (FActorDescContainerInstance* ActorDescContainerInstance = ActorDescContainers.Find(PackageName))
-	{
-		ActorDescContainerInstance->UpdateBounds();
-	}
-}
-
 ULevelInstanceSubsystem::FLevelsToRemoveScope::FLevelsToRemoveScope(ULevelInstanceSubsystem* InOwner)
 	: Owner(InOwner)
 	, bIsBeingDestroyed(false)
@@ -2179,8 +2101,11 @@ bool ULevelInstanceSubsystem::CommitLevelInstanceInternal(TUniquePtr<FLevelInsta
 	// Update pointer since BP Compilation might have invalidated LevelInstance
 	LevelInstance = GetLevelInstance(LevelInstanceID);
 
-	ActorDescContainerInstanceManager.OnLevelInstanceActorCommitted(LevelInstance);
-
+	// Update Registered Container Bounds
+	UWorldPartitionSubsystem* WorldPartitionSubsystem = GetWorld()->GetSubsystem<UWorldPartitionSubsystem>();
+	check(WorldPartitionSubsystem);
+	WorldPartitionSubsystem->UpdateContainerBounds(*LevelInstance->GetWorldAssetPackage());
+	
 	TArray<TPair<ULevelInstanceSubsystem*, FLevelInstanceID>> LevelInstancesToUpdate;
 	// Gather list to update
 	for (TObjectIterator<UWorld> It(RF_ClassDefaultObject | RF_ArchetypeObject, true); It; ++It)
@@ -2480,142 +2405,27 @@ bool ULevelInstanceSubsystem::CanUseWorldAsset(const ILevelInstanceInterface* Le
 	return true;
 }
 
-FWorldPartitionActorFilter ULevelInstanceSubsystem::GetLevelInstanceFilter(const FString& LevelPackage) const
-{
-	TSet<FString> VisitedPackages;
-	return GetLevelInstanceFilterInternal(LevelPackage, VisitedPackages);
-}
-
-// @todo_ow: Move this code out of ULevelInstanceSubsystem into UWorldPartitionSubsystem now that it is generic (not relying on LevelInstances)
-FWorldPartitionActorFilter ULevelInstanceSubsystem::GetLevelInstanceFilterInternal(const FString& LevelPackage, TSet<FString>& VisitedPackages) const
-{
-	if (VisitedPackages.Contains(LevelPackage))
-	{
-		return FWorldPartitionActorFilter(LevelPackage);
-	}
-
-	VisitedPackages.Add(LevelPackage);
-		
-	// Most of the time if the LevelInstance Actor already exists this will return an existing Container but when loading a new LevelInstance (Content Browser Drag&Drop, Create LI) 
-	// This will make sure Container exists.
-	UActorDescContainer* LevelContainer = ActorDescContainerInstanceManager.RegisterContainer(*LevelPackage, GetWorld());
-	check(LevelContainer);
-	ON_SCOPE_EXIT{ ActorDescContainerInstanceManager.UnregisterContainer(LevelContainer); };
-
-	// Lazy create filter for now
-	TArray<const FWorldPartitionActorDesc*> ContainerActorDescs;
-	const FWorldDataLayersActorDesc* WorldDataLayersActorDesc = nullptr;
-
-	for (FActorDescList::TConstIterator<> ActorDescIt(LevelContainer); ActorDescIt; ++ActorDescIt)
-	{
-		if (ActorDescIt->GetActorNativeClass()->IsChildOf<AWorldDataLayers>())
-		{
-			check(!WorldDataLayersActorDesc);
-			WorldDataLayersActorDesc = static_cast<const FWorldDataLayersActorDesc*>(*ActorDescIt);
-		}
-		else if (ActorDescIt->IsContainerInstance())
-		{
-			ContainerActorDescs.Add(*ActorDescIt);
-		}
-	}
-	
-	FWorldPartitionActorFilter Filter(LevelPackage);
-
-	if (WorldDataLayersActorDesc)
-	{
-		for (const FDataLayerInstanceDesc& DataLayerInstanceDesc : WorldDataLayersActorDesc->GetDataLayerInstances())
-		{
-			// For now consider all DataLayerInstances using Assets as filters that are included by default
-			if (DataLayerInstanceDesc.IsUsingAsset() && DataLayerInstanceDesc.GetAsset() && DataLayerInstanceDesc.GetAsset()->SupportsActorFilters())
-			{
-				Filter.DataLayerFilters.Add(FSoftObjectPath(DataLayerInstanceDesc.GetAssetPath().ToString()), FWorldPartitionActorFilter::FDataLayerFilter(DataLayerInstanceDesc.GetShortName(), DataLayerInstanceDesc.IsIncludedInActorFilterDefault()));
-			}
-		}
-	}
-
-	for (const FWorldPartitionActorDesc* ContainerActorDesc : ContainerActorDescs)
-	{
-		TSet<FString> VisitedPackagesCopy(VisitedPackages);
-
-		// Get World Default Filter
-		FWorldPartitionActorFilter* ChildFilter = new FWorldPartitionActorFilter(GetLevelInstanceFilterInternal(ContainerActorDesc->GetLevelPackage().ToString(), VisitedPackagesCopy));
-		ChildFilter->DisplayName = ContainerActorDesc->GetActorLabelOrName().ToString();
-
-		// Apply Filter to Default
-		if (const FWorldPartitionActorFilter* ContainerFilter = ContainerActorDesc->GetContainerFilter())
-		{
-			ChildFilter->Override(*ContainerFilter);
-		}
-
-		Filter.AddChildFilter(ContainerActorDesc->GetGuid(), ChildFilter);
-	}
-
-	return Filter;
-}
-
 bool ULevelInstanceSubsystem::PassLevelInstanceFilter(UWorld* World, const FWorldPartitionHandle& Actor) const
 {
-	if (Actor->IsUsingDataLayerAsset() && Actor->GetDataLayers().Num() > 0)
+	if (const ILevelInstanceInterface* TopAncestor = GetOwningLevelInstance(World->PersistentLevel))
 	{
-		// Make sure Actor is using a filtered Data Layer Instance
-		TSet<FSoftObjectPath> FilteredDataLayerAssets;
-		if (UDataLayerManager* DataLayerManager = UDataLayerManager::GetDataLayerManager(World))
-		{
-			bool bFoundFilteredDataLayer = false;
-			for (FName AssetPath : Actor->GetDataLayers())
-			{
-				FSoftObjectPath SoftAssetPath(AssetPath.ToString());
-				UDataLayerAsset* Asset = Cast<UDataLayerAsset>(SoftAssetPath.TryLoad());
-				if (const UDataLayerInstanceWithAsset* DataLayerInstance = Cast<UDataLayerInstanceWithAsset>(DataLayerManager->GetDataLayerInstanceFromAsset(Asset)); DataLayerInstance && DataLayerInstance->SupportsActorFilters())
-				{
-					FilteredDataLayerAssets.Add(SoftAssetPath);
-				}
-			}
-		}
-
-		// Actor is not using a filtered Data Layer Instance
-		if (FilteredDataLayerAssets.IsEmpty())
+		FActorContainerID ContainerID = TopAncestor->GetLevelInstanceID().GetContainerID();
+		ForEachLevelInstanceAncestors(Cast<AActor>(TopAncestor), [&TopAncestor](const ILevelInstanceInterface* Ancestor)
 		{
+			TopAncestor = Ancestor;
 			return true;
-		}
+		});
 
-		// Build Filter
-		if (const ILevelInstanceInterface* LevelInstanceInterface = GetOwningLevelInstance(World->PersistentLevel))
+		const TMap<FActorContainerID, TSet<FGuid>>& FilteredActors = TopAncestor->GetFilteredActorsPerContainer();
+		if (const TSet<FGuid>* FilteredActorsForContainer = FilteredActors.Find(ContainerID))
 		{
-			TArray<FGuid> AncestorGuids;
-			ForEachLevelInstanceAncestorsAndSelf(Cast<AActor>(LevelInstanceInterface), [&AncestorGuids, &LevelInstanceInterface](const ILevelInstanceInterface* Ancestor)
-			{
-				AncestorGuids.Insert(Cast<AActor>(Ancestor)->GetActorGuid(), 0);
-				LevelInstanceInterface = Ancestor;
-				return true;
-			});
-
-			check(LevelInstanceInterface);
-
-			FWorldPartitionActorFilter ActorFilter = GetLevelInstanceFilter(LevelInstanceInterface->GetWorldAssetPackage());
-			ActorFilter.Override(LevelInstanceInterface->GetFilter());
-
-			FWorldPartitionActorFilter* FoundFilter = &ActorFilter;
-			for (int32 AncestorIndex = 1; AncestorIndex < AncestorGuids.Num(); ++AncestorIndex)
+			if (FilteredActorsForContainer->Contains(Actor->GetGuid()))
 			{
-				FoundFilter = FoundFilter->GetChildFilters().FindChecked(AncestorGuids[AncestorIndex]);
-			}
-
-			for (const FSoftObjectPath& AssetPath : FilteredDataLayerAssets)
-			{
-				if (FWorldPartitionActorFilter::FDataLayerFilter* DataLayerFilter = FoundFilter->DataLayerFilters.Find(AssetPath))
-				{
-					if (DataLayerFilter->bIncluded)
-					{
-						return true;
-					}
-				}				
+				return false;
 			}
-
-			return false;
 		}
 	}
-
+	
 	return true;
 }
 
diff --git a/Engine/Source/Runtime/Engine/Private/WorldPartition/LevelInstance/LevelInstanceActorDesc.cpp b/Engine/Source/Runtime/Engine/Private/WorldPartition/LevelInstance/LevelInstanceActorDesc.cpp
--- a/Engine/Source/Runtime/Engine/Private/WorldPartition/LevelInstance/LevelInstanceActorDesc.cpp
+++ b/Engine/Source/Runtime/Engine/Private/WorldPartition/LevelInstance/LevelInstanceActorDesc.cpp
@@ -13,6 +13,7 @@
 #include "WorldPartition/ActorDescContainer.h"
 #include "WorldPartition/ErrorHandling/WorldPartitionStreamingGenerationErrorHandler.h"
 #include "WorldPartition/WorldPartitionActorDescArchive.h"
+#include "WorldPartition/WorldPartitionSubsystem.h"
 #include "UObject/UE5ReleaseStreamObjectVersion.h"
 #include "UObject/UE5MainStreamObjectVersion.h"
 #include "UObject/FortniteMainBranchObjectVersion.h"
@@ -79,12 +80,12 @@ void FLevelInstanceActorDesc::UpdateBounds()
 {
 	check(LevelInstanceContainerWorldContext.IsValid());
 
-	ULevelInstanceSubsystem* LevelInstanceSubsystem = UWorld::GetSubsystem<ULevelInstanceSubsystem>(LevelInstanceContainerWorldContext.Get());
-	check(LevelInstanceSubsystem);
+	UWorldPartitionSubsystem* WorldPartitionSubsystem = UWorld::GetSubsystem<UWorldPartitionSubsystem>(LevelInstanceContainerWorldContext.Get());
+	check(WorldPartitionSubsystem);
 
 	FTransform LevelInstancePivotOffsetTransform = FTransform(ULevel::GetLevelInstancePivotOffsetFromPackage(LevelPackage));
 	FTransform FinalLevelTransform = LevelInstancePivotOffsetTransform * LevelInstanceTransform;
-	FBox ContainerBounds = LevelInstanceSubsystem->GetContainerBounds(LevelPackage).TransformBy(FinalLevelTransform);
+	FBox ContainerBounds = WorldPartitionSubsystem->GetContainerBounds(LevelPackage).TransformBy(FinalLevelTransform);
 
 	ContainerBounds.GetCenterAndExtents(BoundsLocation, BoundsExtent);
 }
@@ -101,10 +102,10 @@ void FLevelInstanceActorDesc::RegisterContainerInstance(UWorld* InWorldContext)
 			check(InWorldContext);
 			LevelInstanceContainerWorldContext = InWorldContext;
 
-			ULevelInstanceSubsystem* LevelInstanceSubsystem = UWorld::GetSubsystem<ULevelInstanceSubsystem>(InWorldContext);
-			check(LevelInstanceSubsystem);
+			UWorldPartitionSubsystem* WorldPartitionSubsystem = UWorld::GetSubsystem<UWorldPartitionSubsystem>(InWorldContext);
+			check(WorldPartitionSubsystem);
 
-			LevelInstanceContainer = LevelInstanceSubsystem->RegisterContainer(LevelPackage);
+			LevelInstanceContainer = WorldPartitionSubsystem->RegisterContainer(LevelPackage);
 			check(LevelInstanceContainer.IsValid());
 
 			// Should only be called on RegisterContainerInstance before ActorDesc is hashed
@@ -119,10 +120,10 @@ void FLevelInstanceActorDesc::UnregisterContainerInstance()
 	{
 		check(LevelInstanceContainerWorldContext.IsValid());
 
-		ULevelInstanceSubsystem* LevelInstanceSubsystem = UWorld::GetSubsystem<ULevelInstanceSubsystem>(LevelInstanceContainerWorldContext.Get());
-		check(LevelInstanceSubsystem);
+		UWorldPartitionSubsystem* WorldPartitionSubsystem = UWorld::GetSubsystem<UWorldPartitionSubsystem>(LevelInstanceContainerWorldContext.Get());
+		check(WorldPartitionSubsystem);
 
-		LevelInstanceSubsystem->UnregisterContainer(LevelInstanceContainer.Get());
+		WorldPartitionSubsystem->UnregisterContainer(LevelInstanceContainer.Get());
 		LevelInstanceContainer.Reset();
 	}
 
@@ -187,65 +188,10 @@ bool FLevelInstanceActorDesc::GetContainerInstance(const FGetContainerInstancePa
 		if (InParams.bBuildFilter)
 		{
 			// Fill Container Instance Filter
-			ULevelInstanceSubsystem* LevelInstanceSubsystem = UWorld::GetSubsystem<ULevelInstanceSubsystem>(LevelInstanceContainerWorldContext.Get());
-			check(LevelInstanceSubsystem);
+			UWorldPartitionSubsystem* WorldPartitionSubsystem = UWorld::GetSubsystem<UWorldPartitionSubsystem>(LevelInstanceContainerWorldContext.Get());
+			check(WorldPartitionSubsystem);
 
-			FWorldPartitionActorFilter ContainerFilter = LevelInstanceSubsystem->GetLevelInstanceFilter(LevelInstanceContainer->GetContainerPackage().ToString());
-			ContainerFilter.Override(Filter);
-
-			// Flatten Filter to FActorContainerID map
-			TMap<FActorContainerID, TSet<FSoftObjectPath>> FilteredOutDataLayersPerContainer;
-			TFunction<void(const FActorContainerID&, const FWorldPartitionActorFilter&)> ProcessFilter = [&FilteredOutDataLayersPerContainer, &ProcessFilter](const FActorContainerID& InContainerID, const FWorldPartitionActorFilter& InContainerFilter)
-			{
-				check(!FilteredOutDataLayersPerContainer.Contains(InContainerID));
-				TSet<FSoftObjectPath>& Filtered = FilteredOutDataLayersPerContainer.Add(InContainerID);
-
-				for (auto& [AssetPath, DataLayerFilter] : InContainerFilter.DataLayerFilters)
-				{
-					if (!DataLayerFilter.bIncluded)
-					{
-						Filtered.Add(AssetPath);
-					}
-				}
-
-				for (auto& [ActorGuid, WorldPartitionActorFilter] : InContainerFilter.GetChildFilters())
-				{
-					ProcessFilter(FActorContainerID(InContainerID, ActorGuid), *WorldPartitionActorFilter);
-				}
-			};
-
-			ProcessFilter(InParams.ContainerID, ContainerFilter);
-
-			TFunction<void(const FActorContainerID&, const UActorDescContainer*)> ProcessContainers = [&FilteredOutDataLayersPerContainer, &OutContainerInstance, &ProcessContainers, &InParams](const FActorContainerID& InContainerID, const UActorDescContainer* InContainer)
-			{
-				const TSet<FSoftObjectPath>& Filtered = FilteredOutDataLayersPerContainer.FindChecked(InContainerID);
-				for (FActorDescList::TConstIterator<> ActorDescIt(InContainer); ActorDescIt; ++ActorDescIt)
-				{
-					if (ActorDescIt->GetDataLayers().Num() > 0 && ActorDescIt->IsUsingDataLayerAsset())
-					{
-						for (FName DataLayerName : ActorDescIt->GetDataLayers())
-						{
-							FSoftObjectPath DataLayerAsset(DataLayerName.ToString());
-							if (Filtered.Contains(DataLayerAsset))
-							{
-								OutContainerInstance.FilteredActors.FindOrAdd(InContainerID).Add(ActorDescIt->GetGuid());
-								break;
-							}
-						}
-					}
-
-					if (ActorDescIt->IsContainerInstance())
-					{
-						FWorldPartitionActorDesc::FContainerInstance ChildContainerInstance;
-						if (ActorDescIt->GetContainerInstance(FWorldPartitionActorDesc::FGetContainerInstanceParams(), ChildContainerInstance))
-						{
-							ProcessContainers(FActorContainerID(InContainerID, ActorDescIt->GetGuid()), ChildContainerInstance.Container);
-						}
-					}
-				}
-			};
-
-			ProcessContainers(InParams.ContainerID, OutContainerInstance.Container);
+			OutContainerInstance.FilteredActors = WorldPartitionSubsystem->GetFilteredActorsPerContainer(InParams.ContainerID, LevelInstanceContainer->GetContainerPackage().ToString(), Filter);
 		}
 		return true;
 	}
diff --git a/Engine/Source/Runtime/Engine/Private/WorldPartition/WorldPartitionSubsystem.cpp b/Engine/Source/Runtime/Engine/Private/WorldPartition/WorldPartitionSubsystem.cpp
--- a/Engine/Source/Runtime/Engine/Private/WorldPartition/WorldPartitionSubsystem.cpp
+++ b/Engine/Source/Runtime/Engine/Private/WorldPartition/WorldPartitionSubsystem.cpp
@@ -4,11 +4,15 @@
 #include "SceneView.h"
 #include "UnrealEngine.h"
 #include "WorldPartition/WorldPartition.h"
+#include "WorldPartition/WorldPartitionActorDesc.h"
 #include "WorldPartition/WorldPartitionDebugHelper.h"
+#include "WorldPartition/DataLayer/WorldDataLayersActorDesc.h"
 #include "WorldPartition/DataLayer/DataLayerManager.h"
+#include "WorldPartition/DataLayer/DataLayerAsset.h"
 #include "WorldPartition/ContentBundle/ContentBundleWorldSubsystem.h"
 #include "Engine/Canvas.h"
 #include "Engine/CoreSettings.h"
+#include "Engine/LevelBounds.h"
 #include "Debug/DebugDrawService.h"
 #include "WorldPartition/WorldPartitionLog.h"
 #include "WorldPartition/WorldPartitionRuntimeHash.h"
@@ -120,12 +124,235 @@ const UWorldPartition* UWorldPartitionSubsystem::GetWorldPartition() const
 }
 
 #if WITH_EDITOR
+void UWorldPartitionSubsystem::AddReferencedObjects(UObject* InThis, FReferenceCollector& Collector)
+{
+	UWorldPartitionSubsystem* This = CastChecked<UWorldPartitionSubsystem>(InThis);
+
+	This->ActorDescContainerInstanceManager.AddReferencedObjects(Collector);
+}
+
+FWorldPartitionActorFilter UWorldPartitionSubsystem::GetWorldPartitionActorFilter(const FString& InWorldPackage) const
+{
+	TSet<FString> VisitedPackages;
+	return GetWorldPartitionActorFilterInternal(InWorldPackage, VisitedPackages);
+}
+
+
+FWorldPartitionActorFilter UWorldPartitionSubsystem::GetWorldPartitionActorFilterInternal(const FString& InWorldPackage, TSet<FString>& InOutVisitedPackages) const
+{
+	if (InOutVisitedPackages.Contains(InWorldPackage))
+	{
+		return FWorldPartitionActorFilter(InWorldPackage);
+	}
+
+	InOutVisitedPackages.Add(InWorldPackage);
+
+	// Most of the time if this will return an existing Container but when loading a new LevelInstance (Content Browser Drag&Drop, Create LI) 
+	// This will make sure Container exists.
+	UActorDescContainer* LevelContainer = ActorDescContainerInstanceManager.RegisterContainer(*InWorldPackage, GetWorld());
+	check(LevelContainer);
+	ON_SCOPE_EXIT{ ActorDescContainerInstanceManager.UnregisterContainer(LevelContainer); };
+
+	// Lazy create filter for now
+	TArray<const FWorldPartitionActorDesc*> ContainerActorDescs;
+	const FWorldDataLayersActorDesc* WorldDataLayersActorDesc = nullptr;
+
+	for (FActorDescList::TConstIterator<> ActorDescIt(LevelContainer); ActorDescIt; ++ActorDescIt)
+	{
+		if (ActorDescIt->GetActorNativeClass()->IsChildOf<AWorldDataLayers>())
+		{
+			check(!WorldDataLayersActorDesc);
+			WorldDataLayersActorDesc = static_cast<const FWorldDataLayersActorDesc*>(*ActorDescIt);
+		}
+		else if (ActorDescIt->IsContainerInstance())
+		{
+			ContainerActorDescs.Add(*ActorDescIt);
+		}
+	}
+
+	FWorldPartitionActorFilter Filter(InWorldPackage);
+
+	if (WorldDataLayersActorDesc)
+	{
+		for (const FDataLayerInstanceDesc& DataLayerInstanceDesc : WorldDataLayersActorDesc->GetDataLayerInstances())
+		{
+			// For now consider all DataLayerInstances using Assets as filters that are included by default
+			if (DataLayerInstanceDesc.IsUsingAsset() && DataLayerInstanceDesc.GetAsset() && DataLayerInstanceDesc.GetAsset()->SupportsActorFilters())
+			{
+				Filter.DataLayerFilters.Add(FSoftObjectPath(DataLayerInstanceDesc.GetAssetPath().ToString()), FWorldPartitionActorFilter::FDataLayerFilter(DataLayerInstanceDesc.GetShortName(), DataLayerInstanceDesc.IsIncludedInActorFilterDefault()));
+			}
+		}
+	}
+
+	for (const FWorldPartitionActorDesc* ContainerActorDesc : ContainerActorDescs)
+	{
+		TSet<FString> VisitedPackagesCopy(InOutVisitedPackages);
+
+		// Get World Default Filter
+		FWorldPartitionActorFilter* ChildFilter = new FWorldPartitionActorFilter(GetWorldPartitionActorFilterInternal(ContainerActorDesc->GetLevelPackage().ToString(), VisitedPackagesCopy));
+		ChildFilter->DisplayName = ContainerActorDesc->GetActorLabelOrName().ToString();
+
+		// Apply Filter to Default
+		if (const FWorldPartitionActorFilter* ContainerFilter = ContainerActorDesc->GetContainerFilter())
+		{
+			ChildFilter->Override(*ContainerFilter);
+		}
+
+		Filter.AddChildFilter(ContainerActorDesc->GetGuid(), ChildFilter);
+	}
+
+	return Filter;
+}
+
+TMap<FActorContainerID, TSet<FGuid>> UWorldPartitionSubsystem::GetFilteredActorsPerContainer(const FActorContainerID& InContainerID, const FString& InWorldPackage, const FWorldPartitionActorFilter& InActorFilter)
+{
+	TMap<FActorContainerID, TSet<FGuid>> FilteredActors;
+
+	FWorldPartitionActorFilter ContainerFilter = GetWorldPartitionActorFilter(InWorldPackage);
+	ContainerFilter.Override(InActorFilter);
+
+	// Flatten Filter to FActorContainerID map
+	TMap<FActorContainerID, TSet<FSoftObjectPath>> FilteredOutDataLayersPerContainer;
+	TFunction<void(const FActorContainerID&, const FWorldPartitionActorFilter&)> ProcessFilter = [&FilteredOutDataLayersPerContainer, &ProcessFilter](const FActorContainerID& InContainerID, const FWorldPartitionActorFilter& InContainerFilter)
+	{
+		check(!FilteredOutDataLayersPerContainer.Contains(InContainerID));
+		TSet<FSoftObjectPath>& Filtered = FilteredOutDataLayersPerContainer.Add(InContainerID);
+
+		for (auto& [AssetPath, DataLayerFilter] : InContainerFilter.DataLayerFilters)
+		{
+			if (!DataLayerFilter.bIncluded)
+			{
+				Filtered.Add(AssetPath);
+			}
+		}
+
+		for (auto& [ActorGuid, WorldPartitionActorFilter] : InContainerFilter.GetChildFilters())
+		{
+			ProcessFilter(FActorContainerID(InContainerID, ActorGuid), *WorldPartitionActorFilter);
+		}
+	};
+
+	ProcessFilter(InContainerID, ContainerFilter);
+
+	TFunction<void(const FActorContainerID&, const UActorDescContainer*)> ProcessContainers = [&FilteredOutDataLayersPerContainer, &FilteredActors, &ProcessContainers](const FActorContainerID& InContainerID, const UActorDescContainer* InContainer)
+	{
+		const TSet<FSoftObjectPath>& Filtered = FilteredOutDataLayersPerContainer.FindChecked(InContainerID);
+		for (FActorDescList::TConstIterator<> ActorDescIt(InContainer); ActorDescIt; ++ActorDescIt)
+		{
+			if (ActorDescIt->GetDataLayers().Num() > 0 && ActorDescIt->IsUsingDataLayerAsset())
+			{
+				for (FName DataLayerName : ActorDescIt->GetDataLayers())
+				{
+					FSoftObjectPath DataLayerAsset(DataLayerName.ToString());
+					if (Filtered.Contains(DataLayerAsset))
+					{
+						FilteredActors.FindOrAdd(InContainerID).Add(ActorDescIt->GetGuid());
+						break;
+					}
+				}
+			}
+
+			if (ActorDescIt->IsContainerInstance())
+			{
+				FWorldPartitionActorDesc::FContainerInstance ChildContainerInstance;
+				if (ActorDescIt->GetContainerInstance(FWorldPartitionActorDesc::FGetContainerInstanceParams(), ChildContainerInstance))
+				{
+					ProcessContainers(FActorContainerID(InContainerID, ActorDescIt->GetGuid()), ChildContainerInstance.Container);
+				}
+			}
+		}
+	};
+
+	UActorDescContainer* Container = RegisterContainer(*InWorldPackage);
+	ProcessContainers(InContainerID, Container);
+	UnregisterContainer(Container);
+
+	return FilteredActors;
+}
+
+
 bool UWorldPartitionSubsystem::IsRunningConvertWorldPartitionCommandlet()
 {
 	static UClass* WorldPartitionConvertCommandletClass = FindObject<UClass>(nullptr, TEXT("/Script/UnrealEd.WorldPartitionConvertCommandlet"), true);
 	check(WorldPartitionConvertCommandletClass);
 	return GetRunningCommandletClass() && GetRunningCommandletClass()->IsChildOf(WorldPartitionConvertCommandletClass);
 }
+
+void UWorldPartitionSubsystem::FActorDescContainerInstanceManager::FActorDescContainerInstance::AddReferencedObjects(FReferenceCollector& Collector)
+{
+	Collector.AddReferencedObject(Container);
+}
+
+void UWorldPartitionSubsystem::FActorDescContainerInstanceManager::FActorDescContainerInstance::UpdateBounds()
+{
+	Bounds.Init();
+	for (FActorDescList::TIterator<> ActorDescIt(Container); ActorDescIt; ++ActorDescIt)
+	{
+		if (ActorDescIt->GetActorNativeClass()->IsChildOf<ALevelBounds>())
+		{
+			continue;
+		}
+		Bounds += ActorDescIt->GetRuntimeBounds();
+	}
+}
+
+void UWorldPartitionSubsystem::FActorDescContainerInstanceManager::AddReferencedObjects(FReferenceCollector& Collector)
+{
+	for (auto& [Name, ContainerInstance] : ActorDescContainers)
+	{
+		ContainerInstance.AddReferencedObjects(Collector);
+	}
+}
+
+UActorDescContainer* UWorldPartitionSubsystem::FActorDescContainerInstanceManager::RegisterContainer(FName PackageName, UWorld* InWorld)
+{
+	FActorDescContainerInstance* ExistingContainerInstance = &ActorDescContainers.FindOrAdd(PackageName);
+	UActorDescContainer* ActorDescContainer = ExistingContainerInstance->Container;
+
+	if (ExistingContainerInstance->RefCount++ == 0)
+	{
+		ActorDescContainer = NewObject<UActorDescContainer>(GetTransientPackage());
+		ExistingContainerInstance->Container = ActorDescContainer;
+
+		// This will potentially invalidate ExistingContainerInstance due to ActorDescContainers reallocation
+		ActorDescContainer->Initialize({ InWorld, PackageName });
+
+		ExistingContainerInstance = &ActorDescContainers.FindChecked(PackageName);
+		ExistingContainerInstance->UpdateBounds();
+	}
+
+	check(ActorDescContainer->IsTemplateContainer());
+	return ActorDescContainer;
+}
+
+void UWorldPartitionSubsystem::FActorDescContainerInstanceManager::UnregisterContainer(UActorDescContainer* Container)
+{
+	FName PackageName = Container->GetContainerPackage();
+	FActorDescContainerInstance& ExistingContainerInstance = ActorDescContainers.FindChecked(PackageName);
+
+	if (--ExistingContainerInstance.RefCount == 0)
+	{
+		ExistingContainerInstance.Container->Uninitialize();
+		ActorDescContainers.FindAndRemoveChecked(PackageName);
+	}
+}
+
+FBox UWorldPartitionSubsystem::FActorDescContainerInstanceManager::GetContainerBounds(FName PackageName) const
+{
+	if (const FActorDescContainerInstance* ActorDescContainerInstance = ActorDescContainers.Find(PackageName))
+	{
+		return ActorDescContainerInstance->Bounds;
+	}
+	return FBox(ForceInit);
+}
+
+void UWorldPartitionSubsystem::FActorDescContainerInstanceManager::UpdateContainerBounds(FName PackageName)
+{
+	if (FActorDescContainerInstance* ActorDescContainerInstance = ActorDescContainers.Find(PackageName))
+	{
+		ActorDescContainerInstance->UpdateBounds();
+	}
+}
 #endif
 
 void UWorldPartitionSubsystem::Initialize(FSubsystemCollectionBase& Collection)
diff --git a/Engine/Source/Runtime/Engine/Public/LevelInstance/LevelInstanceComponent.h b/Engine/Source/Runtime/Engine/Public/LevelInstance/LevelInstanceComponent.h
--- a/Engine/Source/Runtime/Engine/Public/LevelInstance/LevelInstanceComponent.h
+++ b/Engine/Source/Runtime/Engine/Public/LevelInstance/LevelInstanceComponent.h
@@ -5,6 +5,7 @@
 #include "CoreMinimal.h"
 #include "Components/SceneComponent.h"
 #include "WorldPartition/Filter/WorldPartitionActorFilter.h"
+#include "WorldPartition/WorldPartitionActorContainerID.h"
 
 #include "LevelInstanceComponent.generated.h"
 
@@ -38,6 +39,7 @@ public:
 
 	const FWorldPartitionActorFilter& GetFilter() const { return Filter; }
 	void SetFilter(const FWorldPartitionActorFilter& InFilter);
+	const TMap<FActorContainerID, TSet<FGuid>>& GetFilteredActorsPerContainer() const;
 private:
 	bool ShouldShowSpriteComponent() const;
 
@@ -48,6 +50,9 @@ private:
 	UPROPERTY(EditAnywhere, Category = Filter, meta=(LevelInstanceFilter))
 	FWorldPartitionActorFilter Filter;
 
-	FWorldPartitionActorFilter CachedFilter;
+	FWorldPartitionActorFilter UndoRedoCachedFilter;
+
+	mutable FWorldPartitionActorFilter CachedFilter;
+	mutable TOptional<TMap<FActorContainerID, TSet<FGuid>>> CachedFilteredActorsPerContainer;
 #endif
 };
\ No newline at end of file
diff --git a/Engine/Source/Runtime/Engine/Public/LevelInstance/LevelInstanceInterface.h b/Engine/Source/Runtime/Engine/Public/LevelInstance/LevelInstanceInterface.h
--- a/Engine/Source/Runtime/Engine/Public/LevelInstance/LevelInstanceInterface.h
+++ b/Engine/Source/Runtime/Engine/Public/LevelInstance/LevelInstanceInterface.h
@@ -6,6 +6,7 @@
 #include "UObject/Interface.h"
 #include "UObject/SoftObjectPtr.h"
 #include "LevelInstance/LevelInstanceTypes.h"
+#include "WorldPartition/WorldPartitionActorContainerID.h"
 #include "LevelInstanceInterface.generated.h"
 
 class ULevelInstanceComponent;
@@ -117,5 +118,7 @@ class ENGINE_API ILevelInstanceInterface
 	virtual bool MoveActorsTo(const TArray<AActor*>& ActorsToMove);
 
 	virtual const FWorldPartitionActorFilter& GetFilter() const;
+
+	virtual const TMap<FActorContainerID, TSet<FGuid>>& GetFilteredActorsPerContainer() const;
 #endif
 };
\ No newline at end of file
diff --git a/Engine/Source/Runtime/Engine/Public/LevelInstance/LevelInstanceSubsystem.h b/Engine/Source/Runtime/Engine/Public/LevelInstance/LevelInstanceSubsystem.h
--- a/Engine/Source/Runtime/Engine/Public/LevelInstance/LevelInstanceSubsystem.h
+++ b/Engine/Source/Runtime/Engine/Public/LevelInstance/LevelInstanceSubsystem.h
@@ -10,6 +10,7 @@
 #include "LevelInstance/LevelInstanceTypes.h"
 #include "WorldPartition/Filter/WorldPartitionActorFilter.h"
 #include "WorldPartition/WorldPartitionHandle.h"
+#include "WorldPartition/WorldPartitionActorContainerID.h"
 
 #include "LevelInstanceSubsystem.generated.h"
 
@@ -112,10 +113,6 @@ public:
 	TArray<AActor*> GetParentLevelInstanceActors(const ULevel* Level) const;
 	FString PrefixWithParentLevelInstanceActorLabels(const FString& ActorLabel, const ULevel* Level) const;
 
-	UActorDescContainer* RegisterContainer(FName PackageName) { return ActorDescContainerInstanceManager.RegisterContainer(PackageName, GetWorld()); }
-	void UnregisterContainer(UActorDescContainer* Container) { ActorDescContainerInstanceManager.UnregisterContainer(Container); }
-	FBox GetContainerBounds(FName PackageName) const { return ActorDescContainerInstanceManager.GetContainerBounds(PackageName); }
-
 	static bool CheckForLoop(const ILevelInstanceInterface* LevelInstance, TArray<TPair<FText, TSoftObjectPtr<UWorld>>>* LoopInfo = nullptr, const ILevelInstanceInterface** LoopStart = nullptr);
 	static bool CheckForLoop(const ILevelInstanceInterface* LevelInstance, TSoftObjectPtr<UWorld> WorldAsset, TArray<TPair<FText, TSoftObjectPtr<UWorld>>>* LoopInfo = nullptr, const ILevelInstanceInterface** LoopStart = nullptr);
 	static bool CanUseWorldAsset(const ILevelInstanceInterface* LevelInstance, TSoftObjectPtr<UWorld> WorldAsset, FString* OutReason);
@@ -127,7 +124,6 @@ public:
 
 	static void ResetLoadersForWorldAsset(const FString& WorldAsset);
 
-	FWorldPartitionActorFilter GetLevelInstanceFilter(const FString& LevelPackage) const;
 	bool PassLevelInstanceFilter(UWorld* World, const FWorldPartitionHandle& Actor) const;
 #endif
 
@@ -199,39 +195,6 @@ private:
 	
 	void RemoveLevelsFromWorld(const TArray<ULevel*>& Levels, bool bResetTrans = true);
 
-	class FActorDescContainerInstanceManager
-	{
-		friend class ULevelInstanceSubsystem;
-
-		struct FActorDescContainerInstance
-		{
-			FActorDescContainerInstance()
-			: Container(nullptr)
-			, RefCount(0)
-			, Bounds(ForceInit)
-			{}
-
-			void AddReferencedObjects(FReferenceCollector& Collector);
-			void UpdateBounds();
-
-			UActorDescContainer* Container;
-			uint32 RefCount;
-			FBox Bounds;
-		};
-
-		void AddReferencedObjects(FReferenceCollector& Collector);
-
-	public:
-		UActorDescContainer* RegisterContainer(FName PackageName, UWorld* InWorld);
-		void UnregisterContainer(UActorDescContainer* Container);
-
-		FBox GetContainerBounds(FName PackageName) const;
-
-	private:
-		void OnLevelInstanceActorCommitted(ILevelInstanceInterface* LevelInstance);
-
-		TMap<FName, FActorDescContainerInstance> ActorDescContainers;
-	};
 private:
 #endif
 	friend ULevelStreamingLevelInstance;
@@ -263,8 +226,6 @@ private:
 
 	TMap<FLevelInstanceID, int32> ChildEdits;
 
-	mutable FActorDescContainerInstanceManager ActorDescContainerInstanceManager;
-
 	FWorldPartitionReference CurrentEditLevelInstanceActor;
 #endif
 };
\ No newline at end of file
diff --git a/Engine/Source/Runtime/Engine/Public/LevelInstance/LevelInstanceTypes.h b/Engine/Source/Runtime/Engine/Public/LevelInstance/LevelInstanceTypes.h
--- a/Engine/Source/Runtime/Engine/Public/LevelInstance/LevelInstanceTypes.h
+++ b/Engine/Source/Runtime/Engine/Public/LevelInstance/LevelInstanceTypes.h
@@ -5,6 +5,7 @@
 #include "CoreMinimal.h"
 #include "Misc/Guid.h"
 #include "Templates/SubclassOf.h"
+#include "WorldPartition/WorldPartitionActorContainerID.h"
 #include "LevelInstanceTypes.generated.h"
 
 // FLevelInstanceID is a runtime unique id that is computed from the Hash of LevelInstance Actor Guid and all its ancestor LevelInstance Actor Guids.
@@ -33,11 +34,11 @@ struct FLevelInstanceID
 
 	inline uint64 GetHash() const { return Hash; }
 
-	inline const FGuid& GetGuid() const { return Guid; }
+	inline const FActorContainerID& GetContainerID() const { return ContainerID; }
 
 private:
 	uint64 Hash = 0;
-	FGuid Guid;
+	FActorContainerID ContainerID;
 	TArray<FGuid> Guids;
 	FName ActorName;
 };
diff --git a/Engine/Source/Runtime/Engine/Public/WorldPartition/WorldPartitionSubsystem.h b/Engine/Source/Runtime/Engine/Public/WorldPartition/WorldPartitionSubsystem.h
--- a/Engine/Source/Runtime/Engine/Public/WorldPartition/WorldPartitionSubsystem.h
+++ b/Engine/Source/Runtime/Engine/Public/WorldPartition/WorldPartitionSubsystem.h
@@ -4,9 +4,15 @@
 
 #include "WorldPartitionStreamingSource.h"
 #include "Subsystems/WorldSubsystem.h"
+#include "WorldPartition/Filter/WorldPartitionActorFilter.h"
+#include "WorldPartition/WorldPartitionActorContainerID.h"
+#include "Containers/Map.h"
+#include "Containers/Set.h"
+#include "Misc/Guid.h"
 #include "WorldPartitionSubsystem.generated.h"
 
 class UWorldPartition;
+class UActorDescContainer;
 class FWorldPartitionActorDesc;
 
 enum class EWorldPartitionRuntimeCellState : uint8;
@@ -23,6 +29,12 @@ class ENGINE_API UWorldPartitionSubsystem : public UTickableWorldSubsystem
 public:
 	UWorldPartitionSubsystem();
 
+	//~ Begin UObject Interface
+#if WITH_EDITOR
+	static void AddReferencedObjects(UObject* InThis, FReferenceCollector& Collector);
+#endif
+	//~ End UObject Interface
+
 	//~ Begin USubsystem Interface.
 	virtual void Initialize(FSubsystemCollectionBase& Collection) override;
 	virtual void Deinitialize() override;
@@ -65,9 +77,52 @@ public:
 	void ForEachWorldPartition(TFunctionRef<bool(UWorldPartition*)> Func);
 
 #if WITH_EDITOR
+	FWorldPartitionActorFilter GetWorldPartitionActorFilter(const FString& InWorldPackage) const;
+	TMap<FActorContainerID, TSet<FGuid>> GetFilteredActorsPerContainer(const FActorContainerID& InContainerID, const FString& InWorldPackage, const FWorldPartitionActorFilter& InActorFilter);
+
 	static bool IsRunningConvertWorldPartitionCommandlet();
 
+	UActorDescContainer* RegisterContainer(FName PackageName) { return ActorDescContainerInstanceManager.RegisterContainer(PackageName, GetWorld()); }
+	void UnregisterContainer(UActorDescContainer* Container) { ActorDescContainerInstanceManager.UnregisterContainer(Container); }
+	FBox GetContainerBounds(FName PackageName) const { return ActorDescContainerInstanceManager.GetContainerBounds(PackageName); }
+	void UpdateContainerBounds(FName PackageName) { ActorDescContainerInstanceManager.UpdateContainerBounds(PackageName); }
+
 	TSet<FWorldPartitionActorDesc*> SelectedActorDescs;
+
+	class FActorDescContainerInstanceManager
+	{
+		friend class UWorldPartitionSubsystem;
+
+		struct FActorDescContainerInstance
+		{
+			FActorDescContainerInstance()
+				: Container(nullptr)
+				, RefCount(0)
+				, Bounds(ForceInit)
+			{}
+
+			void AddReferencedObjects(FReferenceCollector& Collector);
+			void UpdateBounds();
+
+			UActorDescContainer* Container;
+			uint32 RefCount;
+			FBox Bounds;
+		};
+
+		void AddReferencedObjects(FReferenceCollector& Collector);
+
+	public:
+		UActorDescContainer* RegisterContainer(FName PackageName, UWorld* InWorld);
+		void UnregisterContainer(UActorDescContainer* Container);
+
+		FBox GetContainerBounds(FName PackageName) const;
+		void UpdateContainerBounds(FName PackageName);
+
+	private:
+		TMap<FName, FActorDescContainerInstance> ActorDescContainers;
+	};
+private:
+	FWorldPartitionActorFilter GetWorldPartitionActorFilterInternal(const FString& InWorldPackage, TSet<FString>& InOutVisitedPackages) const;
 #endif
 
 protected:
@@ -99,5 +154,6 @@ private:
 
 #if WITH_EDITOR
 	bool bIsRunningConvertWorldPartitionCommandlet;
+	mutable FActorDescContainerInstanceManager ActorDescContainerInstanceManager;
 #endif
 };
