fix for determinism issue with UAnimCurveCompressionCodec instance guid.

UAnimCurveCompressionSettings creates a default subobject with a fixed name. the code de-serialized from disk has a different name
both objects are cooked out because the default subobject created in the constructor is public

https://p4-swarm.epicgames.net/reviews/24738963

diff --git a/Engine/Source/Runtime/Engine/Private/Animation/AnimCurveCompressionSettings.cpp b/Engine/Source/Runtime/Engine/Private/Animation/AnimCurveCompressionSettings.cpp
--- a/Engine/Source/Runtime/Engine/Private/Animation/AnimCurveCompressionSettings.cpp
+++ b/Engine/Source/Runtime/Engine/Private/Animation/AnimCurveCompressionSettings.cpp
@@ -11,6 +11,7 @@ UAnimCurveCompressionSettings::UAnimCurveCompressionSettings(const FObjectInitia
 {
 	Codec = CreateDefaultSubobject<UAnimCurveCompressionCodec_CompressedRichCurve>(TEXT("CurveCompressionCodec"));
 	Codec->SetFlags(RF_Transactional);
+	Codec->ClearFlags(RF_Public);
 }
 
 UAnimCurveCompressionCodec* UAnimCurveCompressionSettings::GetCodec(const FString& Path)
