#UE Add option for PriamryAssetLabels to not include/follow redirectors

diff --git a/Engine/Source/Runtime/Engine/Classes/Engine/PrimaryAssetLabel.h b/Engine/Source/Runtime/Engine/Classes/Engine/PrimaryAssetLabel.h
--- a/Engine/Source/Runtime/Engine/Classes/Engine/PrimaryAssetLabel.h
+++ b/Engine/Source/Runtime/Engine/Classes/Engine/PrimaryAssetLabel.h
@@ -35,6 +35,10 @@ public:
 	UPROPERTY(EditAnywhere, Category = PrimaryAssetLabel)
 	uint32 bIsRuntimeLabel : 1;
 
+	/** Set to true if the label asset itself should be cooked and available at runtime. This does not affect the assets that are labeled, they are set with cook rule */
+	UPROPERTY(EditAnywhere, Category = PrimaryAssetLabel)
+	uint32 bIncludeRedirectors : 1;
+
 	/** List of manually specified assets to label */
 	UPROPERTY(EditAnywhere, Category = PrimaryAssetLabel, meta = (AssetBundles = "Explicit"))
 	TArray<TSoftObjectPtr<UObject>> ExplicitAssets;
diff --git a/Engine/Source/Runtime/Engine/Private/PrimaryAssetLabel.cpp b/Engine/Source/Runtime/Engine/Private/PrimaryAssetLabel.cpp
--- a/Engine/Source/Runtime/Engine/Private/PrimaryAssetLabel.cpp
+++ b/Engine/Source/Runtime/Engine/Private/PrimaryAssetLabel.cpp
@@ -21,6 +21,7 @@ UPrimaryAssetLabel::UPrimaryAssetLabel()
 {
 	bLabelAssetsInMyDirectory = false;
 	bIsRuntimeLabel = false;
+	bIncludeRedirectors = true;
 
 	// By default have low priority and don't recurse
 	Rules.bApplyRecursively = false;
@@ -53,7 +54,7 @@ void UPrimaryAssetLabel::UpdateAssetBundleData()
 		{
 			FSoftObjectPath AssetRef = Manager.GetAssetPathForData(AssetData);
 
-			if (!AssetRef.IsNull())
+			if (!AssetRef.IsNull() && (bIncludeRedirectors || !AssetData.IsRedirector()))
 			{
 				NewPaths.Add(AssetRef.GetAssetPath());
 			}
@@ -74,7 +75,7 @@ void UPrimaryAssetLabel::UpdateAssetBundleData()
 			FAssetData FoundAsset = Manager.GetAssetRegistry().GetAssetByObjectPath(CollectionAssets[Index]);
 			FSoftObjectPath AssetRef = Manager.GetAssetPathForData(FoundAsset);
 
-			if (!AssetRef.IsNull())
+			if (!AssetRef.IsNull() && (bIncludeRedirectors || !FoundAsset.IsRedirector()))
 			{
 				NewPaths.Add(AssetRef.GetAssetPath());
 			}
