Bumping the schema version for Loading Event analytics to 2

diff --git a/Engine/Source/Runtime/Engine/Private/Analytics/StudioAnalytics.cpp b/Engine/Source/Runtime/Engine/Private/Analytics/StudioAnalytics.cpp
--- a/Engine/Source/Runtime/Engine/Private/Analytics/StudioAnalytics.cpp
+++ b/Engine/Source/Runtime/Engine/Private/Analytics/StudioAnalytics.cpp
@@ -135,7 +135,7 @@ void FStudioAnalytics::RecordEvent(const FString& EventName, const TArray<FAnaly
 
 void FStudioAnalytics::FireEvent_Loading(const FString& LoadingName, double SecondsSpentLoading, const TArray<FAnalyticsEventAttribute>& InAttributes)
 {
-	const int SchemaVersion = 1;
+	const int SchemaVersion = 2;
 
 	// Ignore anything less than a 1/4th a second.
 	if (SecondsSpentLoading < 0.250)
