Fix crash when duplicating object with a within class in component instance data cache

diff --git a/Engine/Source/Runtime/Engine/Private/ComponentInstanceDataCache.cpp b/Engine/Source/Runtime/Engine/Private/ComponentInstanceDataCache.cpp
--- a/Engine/Source/Runtime/Engine/Private/ComponentInstanceDataCache.cpp
+++ b/Engine/Source/Runtime/Engine/Private/ComponentInstanceDataCache.cpp
@@ -79,7 +79,7 @@ public:
 			}
 			else if (Object->GetOuter() == CacheObject)
 			{
-				Result = DuplicateObject(Object, InstanceData.GetUniqueTransientPackage());
+				Result = DuplicateObject(Object, InstanceData.GetUniqueTransientObject(Object->GetClass()));
 				InstanceData.DuplicatedObjects.Emplace(Result);
 			}
 			else
@@ -556,11 +556,20 @@ void FActorComponentInstanceData::ApplyToComponent(UActorComponent* Component, c
 	}
 }
 
-UObject* FInstanceCacheDataBase::GetUniqueTransientPackage()
+UObject* FInstanceCacheDataBase::GetUniqueTransientObject(UClass* ForClass)
 {
-	if (UniqueTransientPackage.DuplicatedObject == nullptr)
+	if (ForClass->ClassWithin && ForClass->ClassWithin != ForClass && ForClass->ClassWithin != UObject::StaticClass())
 	{
-		UniqueTransientPackage = FDataCacheDuplicatedObjectData(NewObject<UActorComponentInstanceDataTransientOuter>(GetTransientPackage()));
+		// If this results in a lot of new objects, we could do a GetObjectsWithOuter on the result of
+		// GetUniqueTransientObject(ForClass->ClassWithin) and reuse an object.But that doesn't seem likely
+		// to be necessary as this is a pretty edge case situation
+		FScopedAllowAbstractClassAllocation AllowAbstract;
+		return NewObject<UObject>(GetUniqueTransientObject(ForClass->ClassWithin), ForClass->ClassWithin, NAME_None, RF_Transient);
+	}
+	else if (UniqueTransientPackage.DuplicatedObject == nullptr)
+	{
+		FScopedAllowAbstractClassAllocation AllowAbstract;
+		UniqueTransientPackage = FDataCacheDuplicatedObjectData(NewObject<UObject>(GetTransientPackage()));
 	}
 	return UniqueTransientPackage.DuplicatedObject;
 }
diff --git a/Engine/Source/Runtime/Engine/Public/ComponentInstanceDataCache.h b/Engine/Source/Runtime/Engine/Public/ComponentInstanceDataCache.h
--- a/Engine/Source/Runtime/Engine/Public/ComponentInstanceDataCache.h
+++ b/Engine/Source/Runtime/Engine/Public/ComponentInstanceDataCache.h
@@ -63,13 +63,6 @@ struct TStructOpsTypeTraits<FDataCacheDuplicatedObjectData> : public TStructOpsT
 	};
 };
 
-/** Dummy class to use as an outer as we can instantiate a UObject since it is abstract class. */
-UCLASS()
-class UActorComponentInstanceDataTransientOuter : public UObject
-{
-	GENERATED_BODY()
-};
-
 /** Base class for instance cached data of a particular type. */
 USTRUCT()
 struct ENGINE_API FInstanceCacheDataBase
@@ -80,8 +73,12 @@ struct ENGINE_API FInstanceCacheDataBase
 
 	virtual void AddReferencedObjects(FReferenceCollector& Collector);
 
+	UE_DEPRECATED(5.3, "Use GetUniqueTransientObject instead")
+	/** Get (or create) the unique transient outer for the duplicated objects created for this object */
+	UObject* GetUniqueTransientPackage() { return GetUniqueTransientObject(UObject::StaticClass()); }
+
 	/** Get (or create) the unique transient outer for the duplicated objects created for this object */
-	UObject* GetUniqueTransientPackage();
+	UObject* GetUniqueTransientObject(UClass* Class);
 
 	const TArray<FDataCacheDuplicatedObjectData>& GetDuplicatedObjects() const { return DuplicatedObjects; }
 	const TArray<TObjectPtr<UObject>>& GetReferencedObjects() const { return ReferencedObjects; }
