Resetting replay delta checkpoint tracking when a new replay starts
Fix fault handler init when using InitConnection

diff --git a/Engine/Source/Runtime/Engine/Classes/Engine/NetworkObjectList.h b/Engine/Source/Runtime/Engine/Classes/Engine/NetworkObjectList.h
--- a/Engine/Source/Runtime/Engine/Classes/Engine/NetworkObjectList.h
+++ b/Engine/Source/Runtime/Engine/Classes/Engine/NetworkObjectList.h
@@ -188,6 +188,8 @@ public:
 	/** Marks this object dirty for replays using delta checkpoints */
 	void MarkDirtyForReplay(AActor* const Actor);
 
+	void ResetReplayDirtyTracking();
+
 	/** Removes the recently dormant status from the passed in connection */
 	void ClearRecentlyDormantConnection(AActor* const Actor, UNetConnection* const Connection, UNetDriver* NetDriver);
 
diff --git a/Engine/Source/Runtime/Engine/Classes/Engine/PackageMapClient.h b/Engine/Source/Runtime/Engine/Classes/Engine/PackageMapClient.h
--- a/Engine/Source/Runtime/Engine/Classes/Engine/PackageMapClient.h
+++ b/Engine/Source/Runtime/Engine/Classes/Engine/PackageMapClient.h
@@ -219,6 +219,8 @@ public:
 	 */
 	void ReportSyncLoadedGUIDs();
 
+	void ResetReplayDirtyTracking();
+
 	TMap< FNetworkGUID, FNetGuidCacheObject >		ObjectLookup;
 	TMap< TWeakObjectPtr< UObject >, FNetworkGUID >	NetGUIDLookup;
 
diff --git a/Engine/Source/Runtime/Engine/Private/NetConnection.cpp b/Engine/Source/Runtime/Engine/Private/NetConnection.cpp
--- a/Engine/Source/Runtime/Engine/Private/NetConnection.cpp
+++ b/Engine/Source/Runtime/Engine/Private/NetConnection.cpp
@@ -570,6 +570,8 @@ void UNetConnection::InitBase(UNetDriver* InDriver,class FSocket* InSocket, cons
  */
 void UNetConnection::InitConnection(UNetDriver* InDriver, EConnectionState InState, const FURL& InURL, int32 InConnectionSpeed, int32 InMaxPacket)
 {
+	using namespace UE::Net;
+
 	Driver = InDriver;
 
 	InitChannelData();
@@ -610,6 +612,9 @@ void UNetConnection::InitConnection(UNetDriver* InDriver, EConnectionState InSta
 		}
 	}
 
+	FaultRecovery = MakeUnique<FNetConnectionFaultRecovery>();
+	FaultRecovery->InitDefaults((Driver != nullptr ? Driver->GetNetDriverDefinition().ToString() : TEXT("")), this);
+
 	// Create package map.
 	auto PackageMapClient = NewObject<UPackageMapClient>(this);
 	PackageMapClient->Initialize(this, Driver->GuidCache);
diff --git a/Engine/Source/Runtime/Engine/Private/NetworkObjectList.cpp b/Engine/Source/Runtime/Engine/Private/NetworkObjectList.cpp
--- a/Engine/Source/Runtime/Engine/Private/NetworkObjectList.cpp
+++ b/Engine/Source/Runtime/Engine/Private/NetworkObjectList.cpp
@@ -400,6 +400,17 @@ void FNetworkObjectList::MarkDirtyForReplay(AActor* const Actor)
 	}
 }
 
+void FNetworkObjectList::ResetReplayDirtyTracking()
+{
+	for (auto It = AllNetworkObjects.CreateIterator(); It; ++It)
+	{
+		if (FNetworkObjectInfo* ObjectInfo = (*It).Get())
+		{
+			ObjectInfo->bDirtyForReplay = false;
+		}
+	}
+}
+
 void FNetworkObjectList::ClearRecentlyDormantConnection(AActor* const Actor, UNetConnection* const Connection, UNetDriver* NetDriver)
 {
 	TSharedPtr<FNetworkObjectInfo>* NetworkObjectInfoPtr = FindOrAdd(Actor, NetDriver);
diff --git a/Engine/Source/Runtime/Engine/Private/PackageMapClient.cpp b/Engine/Source/Runtime/Engine/Private/PackageMapClient.cpp
--- a/Engine/Source/Runtime/Engine/Private/PackageMapClient.cpp
+++ b/Engine/Source/Runtime/Engine/Private/PackageMapClient.cpp
@@ -4199,3 +4199,28 @@ bool FNetGUIDCache::FIsOwnerOrPawnHelper::IsOwnerOrPawn() const
 }
 #endif
 
+void FNetGUIDCache::ResetReplayDirtyTracking()
+{
+	// reset guids
+	for (auto It = ObjectLookup.CreateIterator(); It; ++It)
+	{
+		FNetGuidCacheObject& CacheObject = It.Value();
+		CacheObject.bDirtyForReplay = true;
+	}
+
+	// reset net export groups
+	for (auto It = NetFieldExportGroupMap.CreateIterator(); It; ++It)
+	{
+		TSharedPtr<FNetFieldExportGroup>& ExportGroup = It.Value();
+		if (ExportGroup.IsValid())
+		{
+			// reset the exports in each group
+			for (FNetFieldExport& Export : ExportGroup->NetFieldExports)
+			{
+				Export.bDirtyForReplay = true;
+			}
+
+			ExportGroup->bDirtyForReplay = true;
+		}
+	}
+}
\ No newline at end of file
diff --git a/Engine/Source/Runtime/Engine/Private/ReplayHelper.cpp b/Engine/Source/Runtime/Engine/Private/ReplayHelper.cpp
--- a/Engine/Source/Runtime/Engine/Private/ReplayHelper.cpp
+++ b/Engine/Source/Runtime/Engine/Private/ReplayHelper.cpp
@@ -179,6 +179,11 @@ void FReplayHelper::StartRecording(UNetConnection* Connection)
 
 	ActiveReplayName = DemoURL.Map;
 
+	if (bHasDeltaCheckpoints)
+	{
+		ResetDeltaCheckpointTracking(Connection);
+	}
+
 	FStartStreamingParameters Params;
 	Params.CustomName = DemoURL.Map;
 	Params.FriendlyName = FriendlyNameOption != nullptr ? FString(FriendlyNameOption) : World->GetMapName();
@@ -2344,6 +2349,36 @@ void FReplayHelper::NotifyReplayError(UE::Net::TNetResult<EReplayResult>&& Resul
 	ResultManager.HandleNetResult(MoveTemp(Result));
 }
 
+void FReplayHelper::ResetDeltaCheckpointTracking(UNetConnection* Connection)
+{
+	DECLARE_SCOPE_CYCLE_COUNTER(TEXT("ReplayResetDeltaCheckpoint time"), STAT_ReplayResetDeltaCheckpoint, STATGROUP_Net);
+
+	if (Connection)
+	{
+		if (UNetDriver* Driver = Connection->GetDriver())
+		{
+			// reset object list
+			FNetworkObjectList& NetworkObjects = Driver->GetNetworkObjectList();
+			NetworkObjects.ResetReplayDirtyTracking();
+
+			// reset guid cache
+			if (FNetGUIDCache* GuidCache = Connection->Driver->GuidCache.Get())
+			{
+				GuidCache->ResetReplayDirtyTracking();
+			}
+
+			// reset object replicators
+			for (FObjectReplicator* Replicator : Driver->AllOwnedReplicators)
+			{
+				if (Replicator)
+				{
+					Replicator->ResetReplayDirtyTracking();
+				}
+			}
+		}
+	}
+}
+
 void FReplayResultHandler::InitResultHandler(FReplayHelper* InReplayHelper)
 {
 	ReplayHelper = InReplayHelper;
@@ -2375,5 +2410,4 @@ UE::Net::EHandleNetResult FReplayResultHandler::HandleNetResult(UE::Net::FNetRes
 	}
 
 	return EHandleNetResult::Handled;
-}
-
+}
\ No newline at end of file
diff --git a/Engine/Source/Runtime/Engine/Public/Net/DataReplication.h b/Engine/Source/Runtime/Engine/Public/Net/DataReplication.h
--- a/Engine/Source/Runtime/Engine/Public/Net/DataReplication.h
+++ b/Engine/Source/Runtime/Engine/Public/Net/DataReplication.h
@@ -251,6 +251,7 @@ public:
 	bool CanSkipUpdate(FReplicationFlags Flags);
 
 	bool IsDirtyForReplay() const { return bDirtyForReplay; }
+	void ResetReplayDirtyTracking() { bDirtyForReplay = true; }
 
 public:
 
diff --git a/Engine/Source/Runtime/Engine/Public/ReplayHelper.h b/Engine/Source/Runtime/Engine/Public/ReplayHelper.h
--- a/Engine/Source/Runtime/Engine/Public/ReplayHelper.h
+++ b/Engine/Source/Runtime/Engine/Public/ReplayHelper.h
@@ -201,6 +201,8 @@ private:
 
 	bool SetExternalDataForObject(UNetConnection* Connection, UObject* OwningObject, const uint8* Src, const int32 NumBits);
 
+	void ResetDeltaCheckpointTracking(UNetConnection* Connection);
+
 	// Cached replay URL
 	FURL DemoURL;
 
