Fix crash during box selection of static mesh components without LOD0 source data

diff --git a/Engine/Source/Runtime/Engine/Private/Components/StaticMeshComponent.cpp b/Engine/Source/Runtime/Engine/Private/Components/StaticMeshComponent.cpp
--- a/Engine/Source/Runtime/Engine/Private/Components/StaticMeshComponent.cpp
+++ b/Engine/Source/Runtime/Engine/Private/Components/StaticMeshComponent.cpp
@@ -3024,6 +3024,48 @@ void UStaticMeshComponent::PostStaticMeshCompilation()
 	}
 }
 
+namespace ComponentIsTouchingSelectionHelpers
+{
+
+const FStaticMeshLODResources& GetMinLODModel(const UStaticMesh& StaticMesh)
+{
+	const int32 MinLODIdx = StaticMesh.GetMinLODIdx();
+	const FStaticMeshRenderData* RenderData = StaticMesh.GetRenderData();
+	const int32 LODIdx = RenderData->GetCurrentFirstLODIdx(MinLODIdx);
+	return RenderData->LODResources[LODIdx];
+}
+
+enum class ECheckSectionBoundsResult
+{
+	Valid,           // Section data is available, continue with triangle tests.
+	InvalidContinue, // Section data is not available, but continue checking the other sections.
+	InvalidFail      // Section data is not available, and the selection test needs to fail because of it.
+};
+
+ECheckSectionBoundsResult CheckSectionBounds(const FStaticMeshSection& Section, const FIndexArrayView& Indices, const FPositionVertexBuffer& Vertices,
+                                             const bool bMustEncompassEntireComponent)
+{
+	// Not sure if empty sections are valid, but if they are then we need to ignore them for the triangle tests.
+	// Otherwise the checks for encompassing the entire component are producing false negatives and/or positives.
+	if (Section.NumTriangles == 0)
+	{
+		return ECheckSectionBoundsResult::InvalidContinue;
+	}
+
+	// Vertex and Index buffers are not guaranteed to be present and/or complete.
+	// Thus, we need to check their bounds before we use them to test triangle intersections.
+	if (static_cast<int32>(Section.FirstIndex + Section.NumTriangles * 3) > Indices.Num() || Section.MaxVertexIndex >= Vertices.GetNumVertices())
+	{
+		// If the entire component must to be encompassed, not having data for a section means we cannot verify that the section is included, and
+		// this entire selection test fails. Otherwise, we just continue and check the other sections.
+		return bMustEncompassEntireComponent ? ECheckSectionBoundsResult::InvalidFail : ECheckSectionBoundsResult::InvalidContinue;
+	}
+
+	return ECheckSectionBoundsResult::Valid;
+}
+
+} // namespace StaticMeshComponent_SelectionHelpers
+
 bool UStaticMeshComponent::ComponentIsTouchingSelectionBox(const FBox& InSelBBox, const bool bConsiderOnlyBSP, const bool bMustEncompassEntireComponent) const
 {
 	if (!bConsiderOnlyBSP && GetStaticMesh() != nullptr && GetStaticMesh()->HasValidRenderData())
@@ -3037,34 +3079,46 @@ bool UStaticMeshComponent::ComponentIsTouchingSelectionBox(const FBox& InSelBBox
 		// Check if we are even inside it's bounding box, if we are not, there is no way we colliding via the more advanced checks we will do.
 		if (Super::ComponentIsTouchingSelectionBox(InSelBBox, bConsiderOnlyBSP, false))
 		{
-			TArray<FVector> Vertex;
+			TArray<FVector> TriVertices;
+			TriVertices.SetNumUninitialized(3);
 
-			FStaticMeshLODResources& LODModel = GetStaticMesh()->GetRenderData()->LODResources[0];
-			FIndexArrayView Indices = LODModel.IndexBuffer.GetArrayView();
+			const FStaticMeshLODResources& LODModel = ComponentIsTouchingSelectionHelpers::GetMinLODModel(*GetStaticMesh());
+
+			const FIndexArrayView Indices = LODModel.IndexBuffer.GetArrayView();
+			const FPositionVertexBuffer& Vertices = LODModel.VertexBuffers.PositionVertexBuffer;
 
 			for (const FStaticMeshSection& Section : LODModel.Sections)
 			{
-				// Iterate over each triangle.
-				for (int32 TriangleIndex = 0; TriangleIndex < (int32)Section.NumTriangles; TriangleIndex++)
+				switch (ComponentIsTouchingSelectionHelpers::CheckSectionBounds(Section, Indices, Vertices, bMustEncompassEntireComponent))
 				{
-					Vertex.Empty(3);
+				case ComponentIsTouchingSelectionHelpers::ECheckSectionBoundsResult::Valid:
+					/* Proceed with triangle tests. */ ;
+				case ComponentIsTouchingSelectionHelpers::ECheckSectionBoundsResult::InvalidContinue:
+					continue;
+				case ComponentIsTouchingSelectionHelpers::ECheckSectionBoundsResult::InvalidFail:
+					return false;
+				}
 
-					int32 FirstIndex = TriangleIndex * 3 + Section.FirstIndex;
+				// Iterate over each triangle.
+				const int32 SectionIndicesEnd = static_cast<int32>(Section.FirstIndex + Section.NumTriangles * 3);
+				for (int32 TriFirstVertexIndex = Section.FirstIndex; TriFirstVertexIndex < SectionIndicesEnd; TriFirstVertexIndex += 3)
+				{
 					for (int32 i = 0; i < 3; i++)
 					{
-						int32 VertexIndex = Indices[FirstIndex + i];
-						FVector LocalPosition(LODModel.VertexBuffers.PositionVertexBuffer.VertexPosition(VertexIndex));
-						Vertex.Emplace(GetComponentTransform().TransformPosition(LocalPosition));
+						const int32 VertexIndex = Indices[TriFirstVertexIndex + i];
+						const FVector LocalPosition(Vertices.VertexPosition(VertexIndex));
+						TriVertices[i] = GetComponentTransform().TransformPosition(LocalPosition);
 					}
 
 					// Check if the triangle is colliding with the bounding box.
-					FSeparatingAxisPointCheck ThePointCheck(Vertex, InSelBBox.GetCenter(), InSelBBox.GetExtent(), false);
+					const FSeparatingAxisPointCheck ThePointCheck(TriVertices, InSelBBox.GetCenter(), InSelBBox.GetExtent(), false);
+
 					if (!bMustEncompassEntireComponent && ThePointCheck.bHit)
 					{
 						// Needn't encompass entire component: any intersection, we consider as touching
 						return true;
 					}
-					else if (bMustEncompassEntireComponent && !ThePointCheck.bHit)
+					if (bMustEncompassEntireComponent && !ThePointCheck.bHit)
 					{
 						// Must encompass entire component: any non intersection, we consider as not touching
 						return false;
@@ -3096,20 +3150,31 @@ bool UStaticMeshComponent::ComponentIsTouchingSelectionFrustum(const FConvexVolu
 				return true;
 			}
 
-			FStaticMeshLODResources& LODModel = GetStaticMesh()->GetRenderData()->LODResources[0];
-			FIndexArrayView Indices = LODModel.IndexBuffer.GetArrayView();
+			const FStaticMeshLODResources& LODModel = ComponentIsTouchingSelectionHelpers::GetMinLODModel(*GetStaticMesh());
+
+			const FIndexArrayView Indices = LODModel.IndexBuffer.GetArrayView();
+			const FPositionVertexBuffer& Vertices = LODModel.VertexBuffers.PositionVertexBuffer;
 
 			const FTransform& TransformToWorld = GetComponentTransform();
 			for (const FStaticMeshSection& Section : LODModel.Sections)
 			{
+				switch (ComponentIsTouchingSelectionHelpers::CheckSectionBounds(Section, Indices, Vertices, bMustEncompassEntireComponent))
+				{
+				case ComponentIsTouchingSelectionHelpers::ECheckSectionBoundsResult::Valid:
+					/* Proceed with triangle tests. */ ;
+				case ComponentIsTouchingSelectionHelpers::ECheckSectionBoundsResult::InvalidContinue:
+					continue;
+				case ComponentIsTouchingSelectionHelpers::ECheckSectionBoundsResult::InvalidFail:
+					return false;
+				}
+
 				// Iterate over each triangle.
-				for (int32 TriangleIndex = 0; TriangleIndex < (int32)Section.NumTriangles; TriangleIndex++)
+				const int32 SectionIndicesEnd = static_cast<int32>(Section.FirstIndex + Section.NumTriangles * 3);
+				for (int32 TriFirstVertexIndex = Section.FirstIndex; TriFirstVertexIndex < SectionIndicesEnd; TriFirstVertexIndex += 3)
 				{
-					int32 Index = TriangleIndex * 3 + Section.FirstIndex;
-				
-					FVector PointA(LODModel.VertexBuffers.PositionVertexBuffer.VertexPosition(Indices[Index]));
-					FVector PointB(LODModel.VertexBuffers.PositionVertexBuffer.VertexPosition(Indices[Index + 1]));
-					FVector PointC(LODModel.VertexBuffers.PositionVertexBuffer.VertexPosition(Indices[Index + 2]));
+					FVector PointA(Vertices.VertexPosition(Indices[TriFirstVertexIndex]));
+					FVector PointB(Vertices.VertexPosition(Indices[TriFirstVertexIndex + 1]));
+					FVector PointC(Vertices.VertexPosition(Indices[TriFirstVertexIndex + 2]));
 
 					PointA = TransformToWorld.TransformPosition(PointA);
 					PointB = TransformToWorld.TransformPosition(PointB);
@@ -3121,15 +3186,15 @@ bool UStaticMeshComponent::ComponentIsTouchingSelectionFrustum(const FConvexVolu
 					if (!bMustEncompassEntireComponent && bIntersect)
 					{
 						// Needn't encompass entire component: any intersection, we consider as touching
-					return true;
-				}
-					else if (bMustEncompassEntireComponent && !bFullyContained)
-				{
+						return true;
+					}
+					if (bMustEncompassEntireComponent && !bFullyContained)
+					{
 						// Must encompass entire component: any non intersection, we consider as not touching
-					return false;
+						return false;
+					}
 				}
 			}
-			}
 
 			// Either:
 			// a) It must encompass the entire component and all points were intersected (return true), or;
