Skeleton rendering: expose axis options (thickness + length)  for future use

diff --git a/Engine/Plugins/Animation/ControlRig/Source/ControlRigEditor/Private/Editor/ControlRigEditorEditMode.cpp b/Engine/Plugins/Animation/ControlRig/Source/ControlRigEditor/Private/Editor/ControlRigEditorEditMode.cpp
--- a/Engine/Plugins/Animation/ControlRig/Source/ControlRigEditor/Private/Editor/ControlRigEditorEditMode.cpp
+++ b/Engine/Plugins/Animation/ControlRig/Source/ControlRigEditor/Private/Editor/ControlRigEditorEditMode.cpp
@@ -148,6 +148,7 @@ void FControlRigEditorEditMode::Render(const FSceneView* View, FViewport* Viewpo
 	});
 
 	// get size setting for bone drawing
+	FBoneAxisDrawConfig AxisConfig;
 	float BoneRadius = 1.0f;
 	if (FAnimationViewportClient* AnimViewportClient = static_cast<FAnimationViewportClient*>(Viewport->GetClient()))
 	{
@@ -177,7 +178,7 @@ void FControlRigEditorEditMode::Render(const FSceneView* View, FViewport* Viewpo
 		BoneColor = bIsSelected ? PersonaOptions->SelectedBoneColor : BoneColor;
 
 		// draw the little coordinate frame inside the bone ONLY if selected or affected
-		const bool bDrawAxesInsideBone = bIsAffected|| bIsSelected;
+		AxisConfig.bDraw = bIsAffected|| bIsSelected;
 
 		// draw cone to each child
 		// but use a different color if this bone is NOT selected, but the child IS selected
@@ -208,7 +209,7 @@ void FControlRigEditorEditMode::Render(const FSceneView* View, FViewport* Viewpo
 			BoneColor,
 			SDPG_Foreground,
 			BoneRadius,
-			bDrawAxesInsideBone);
+			AxisConfig);
 		if (Hierarchy->GetFirstParent(ElementIndex) == INDEX_NONE)
 		{
 			SkeletalDebugRendering::DrawConeConnection(PDI, BoneTransform.GetLocation(), FVector::Zero(), BoneRadius, FLinearColor::Red);
diff --git a/Engine/Source/Runtime/Engine/Private/SkeletalDebugRendering.cpp b/Engine/Source/Runtime/Engine/Private/SkeletalDebugRendering.cpp
--- a/Engine/Source/Runtime/Engine/Private/SkeletalDebugRendering.cpp
+++ b/Engine/Source/Runtime/Engine/Private/SkeletalDebugRendering.cpp
@@ -66,7 +66,7 @@ void DrawWireBoneAdvanced(
 	const FLinearColor& InColor,
 	ESceneDepthPriorityGroup InDepthPriority,
 	const float SphereRadius,
-	const bool bDrawAxes)
+	const FBoneAxisDrawConfig& InAxisConfig)
 {
 #if ENABLE_DRAW_DEBUG
 
@@ -76,9 +76,11 @@ void DrawWireBoneAdvanced(
 	DrawWireSphere(PDI, InBoneTransform, InColor, SphereRadius, NumSphereSides, InDepthPriority, 0.0f, 1.0f);
 
 	// draw axes at joint location
-	if (bDrawAxes)
+	if (InAxisConfig.bDraw)
 	{
-		SkeletalDebugRendering::DrawAxes(PDI, InBoneTransform, SDPG_Foreground, 0.f , SphereRadius);
+		const float Thickness = InAxisConfig.Thickness > 0.f ? InAxisConfig.Thickness : 0.f;
+		const float Length = InAxisConfig.Length > 0.f ? InAxisConfig.Length : SphereRadius;
+		SkeletalDebugRendering::DrawAxes(PDI, InBoneTransform, SDPG_Foreground, Thickness, Length);
 	}
 
 	// draw wire cones to each child
@@ -363,6 +365,8 @@ void DrawBonesInternal(
 		}
 	}
 
+	FBoneAxisDrawConfig AxisConfig = DrawConfig.AxisConfig;
+	
 	// spin through all required bones and render them
 	const float BoneRadius = DrawConfig.BoneDrawSize;
 	for (int32 Index = 0; Index < RequiredBones.Num(); ++Index)
@@ -389,7 +393,7 @@ void DrawBonesInternal(
 		BoneColor = bIsSelected ? DrawConfig.SelectedBoneColor : BoneColor;
 
 		// draw the little coordinate frame inside the bone ONLY if selected or affected
-		const bool bDrawAxesInsideBone = bIsAffected || bIsSelected;
+		AxisConfig.bDraw = bIsAffected || bIsSelected;
 
 		// draw cone to each child
 		// but use a different color if this bone is NOT selected, but the child IS selected
@@ -428,7 +432,7 @@ void DrawBonesInternal(
 			BoneColor,
 			SDPG_Foreground,
 			BoneRadius,
-			bDrawAxesInsideBone);
+			AxisConfig);
 		
 		// special case for root connection to origin
 		if (GetParentIndex(BoneIndex) == INDEX_NONE)
diff --git a/Engine/Source/Runtime/Engine/Public/SkeletalDebugRendering.h b/Engine/Source/Runtime/Engine/Public/SkeletalDebugRendering.h
--- a/Engine/Source/Runtime/Engine/Public/SkeletalDebugRendering.h
+++ b/Engine/Source/Runtime/Engine/Public/SkeletalDebugRendering.h
@@ -26,6 +26,13 @@ namespace EBoneDrawMode
 	};
 };
 
+struct FBoneAxisDrawConfig
+{
+	bool bDraw = false;
+	float Thickness = 0.f;
+	float Length = 0.f;
+};
+
 struct FSkelDebugDrawConfig
 {
 public:
@@ -37,6 +44,7 @@ public:
 	FLinearColor AffectedBoneColor;
 	FLinearColor SelectedBoneColor;
 	FLinearColor ParentOfSelectedBoneColor;
+	FBoneAxisDrawConfig AxisConfig;
 };
 
 namespace SkeletalDebugRendering
@@ -85,7 +93,7 @@ ENGINE_API void DrawAxes(
  * @param	InColor				The color to use for the bone
  * @param	InDepthPriority		The scene depth priority group to use
  * @param	SphereRadius		Radius of the ball drawn at the bone location
- * @param	bDrawAxes			If true, will draw small coordinate axes inside the joint sphere
+ * @param	InAxisConfig		Draw configuration for small coordinate axes inside the joint sphere
  */
 ENGINE_API	void DrawWireBoneAdvanced(
 	FPrimitiveDrawInterface* PDI,
@@ -95,7 +103,7 @@ ENGINE_API	void DrawWireBoneAdvanced(
 	const FLinearColor& InColor,
 	ESceneDepthPriorityGroup InDepthPriority,
 	const float SphereRadius,
-	const bool bDrawAxes);
+	const FBoneAxisDrawConfig& InAxisConfig);
 
 /**
  * Draw a cone showing offset from origin position to a given bone transform
