Level Instance Editor Partial Loading
- Loading regions will now recurse into actor containers (Level Instances) and partially load them instead of loading them completly.
- This can be disabled with the the wp.Editor.DisableLevelInstanceEditorPartialLoading cvar.

diff --git a/Engine/Source/Editor/WorldPartitionEditor/Private/WorldPartition/SWorldPartitionEditorGrid2D.cpp b/Engine/Source/Editor/WorldPartitionEditor/Private/WorldPartition/SWorldPartitionEditorGrid2D.cpp
--- a/Engine/Source/Editor/WorldPartitionEditor/Private/WorldPartition/SWorldPartitionEditorGrid2D.cpp
+++ b/Engine/Source/Editor/WorldPartitionEditor/Private/WorldPartition/SWorldPartitionEditorGrid2D.cpp
@@ -1045,8 +1045,8 @@ void SWorldPartitionEditorGrid2D::Tick(const FGeometry& AllottedGeometry, const
 		}
 	}
 
-	UWorldPartitionEditorHash::FForEachIntersectingActorParams ForEachIntersectingActorParams;
-	ForEachIntersectingActorParams.MinimumBox = FBox(FVector::ZeroVector, FVector(GetSelectionSnap()));
+	UWorldPartitionEditorHash::FForEachIntersectingActorParams ForEachIntersectingActorParams = UWorldPartitionEditorHash::FForEachIntersectingActorParams()
+		.SetMinimumBox(FBox(FVector::ZeroVector, FVector(GetSelectionSnap())));
 
 	GetWorldPartition()->EditorHash->ForEachIntersectingActor(ViewRectWorld, [&](FWorldPartitionActorDesc* ActorDesc)
 	{
diff --git a/Engine/Source/Runtime/Engine/Private/LevelInstance/LevelInstanceInterface.cpp b/Engine/Source/Runtime/Engine/Private/LevelInstance/LevelInstanceInterface.cpp
--- a/Engine/Source/Runtime/Engine/Private/LevelInstance/LevelInstanceInterface.cpp
+++ b/Engine/Source/Runtime/Engine/Private/LevelInstance/LevelInstanceInterface.cpp
@@ -3,6 +3,7 @@
 #include "LevelInstance/LevelInstanceInterface.h"
 #include "LevelInstance/LevelInstanceSubsystem.h"
 #include "LevelInstance/LevelInstanceLevelStreaming.h"
+#include "GameFramework/Actor.h"
 #include "Engine/Level.h"
 #include "Engine/World.h"
 #include "UObject/Package.h"
@@ -17,6 +18,13 @@ ULevelInstanceInterface::ULevelInstanceInterface(const FObjectInitializer& Objec
 	: Super(ObjectInitializer)
 {}
 
+#if WITH_EDITOR
+bool ILevelInstanceInterface::SupportsPartialEditorLoading() const
+{
+	return CastChecked<AActor>(this)->GetIsSpatiallyLoaded() && !IsEditing();
+}
+#endif
+
 ULevelInstanceSubsystem* ILevelInstanceInterface::GetLevelInstanceSubsystem() const
 {
 	const AActor* Actor = CastChecked<AActor>(this);
diff --git a/Engine/Source/Runtime/Engine/Private/LevelInstance/LevelInstanceLevelStreaming.cpp b/Engine/Source/Runtime/Engine/Private/LevelInstance/LevelInstanceLevelStreaming.cpp
--- a/Engine/Source/Runtime/Engine/Private/LevelInstance/LevelInstanceLevelStreaming.cpp
+++ b/Engine/Source/Runtime/Engine/Private/LevelInstance/LevelInstanceLevelStreaming.cpp
@@ -15,6 +15,7 @@
 #include "ProfilingDebugging/ScopedTimers.h"
 #include "UObject/Linker.h"
 #include "UObject/Package.h"
+#include "WorldPartition/WorldPartition.h"
 
 #include UE_INLINE_GENERATED_CPP_BY_NAME(LevelInstanceLevelStreaming)
 
@@ -23,6 +24,13 @@
 #include "LevelUtils.h"
 #include "ActorFolder.h"
 #include "UObject/LinkerLoad.h"
+
+static bool GDisableLevelInstanceEditorPartialLoading = false;
+FAutoConsoleVariableRef CVarDisableLevelInstanceEditorPartialLoading(
+	TEXT("wp.Editor.DisableLevelInstanceEditorPartialLoading"),
+	GDisableLevelInstanceEditorPartialLoading,
+	TEXT("Allow disabling partial loading of level instances in the editor."),
+	ECVF_Default);
 #endif
 
 ULevelStreamingLevelInstance::ULevelStreamingLevelInstance(const FObjectInitializer& ObjectInitializer)
@@ -329,10 +337,23 @@ void ULevelStreamingLevelInstance::OnLevelLoadedChanged(ULevel* InLevel)
 			LevelTransform = FTransform(WorldSettings->LevelInstancePivotOffset) * LevelTransform;
 		}
 
-		if (ULevelInstanceSubsystem* LevelInstanceSubsystem = GetWorld()->GetSubsystem<ULevelInstanceSubsystem>())
+#if WITH_EDITOR
+		if (!GDisableLevelInstanceEditorPartialLoading)
 		{
-			LevelInstanceSubsystem->RegisterLoadedLevelStreamingLevelInstance(this);
+			if (ULevelInstanceSubsystem* LevelInstanceSubsystem = GetWorld()->GetSubsystem<ULevelInstanceSubsystem>())
+			{
+				LevelInstanceSubsystem->RegisterLoadedLevelStreamingLevelInstance(this);
+
+				if (UWorldPartition* WorldPartition = NewLoadedLevel->GetWorldPartition())
+				{
+					check(!WorldPartition->IsInitialized());
+					if (ILevelInstanceInterface* LevelInstance = LevelInstanceSubsystem->GetLevelInstance(LevelInstanceID))
+					{
+						WorldPartition->bForceEnableStreamingInEditor = LevelInstance->SupportsPartialEditorLoading();
+					}
+				}
+			}
 		}
+#endif
 	}
-}
-
+}
\ No newline at end of file
diff --git a/Engine/Source/Runtime/Engine/Private/LevelInstance/LevelInstanceSubsystem.cpp b/Engine/Source/Runtime/Engine/Private/LevelInstance/LevelInstanceSubsystem.cpp
--- a/Engine/Source/Runtime/Engine/Private/LevelInstance/LevelInstanceSubsystem.cpp
+++ b/Engine/Source/Runtime/Engine/Private/LevelInstance/LevelInstanceSubsystem.cpp
@@ -22,6 +22,7 @@
 #include "WorldPartition/LevelInstance/LevelInstanceActorDesc.h"
 #include "WorldPartition/WorldPartition.h"
 #include "WorldPartition/WorldPartitionSubsystem.h"
+#include "WorldPartition/WorldPartitionActorLoaderInterface.h"
 #include "WorldPartition/DataLayer/DataLayerManager.h"
 #include "WorldPartition/DataLayer/DataLayerInstanceWithAsset.h"
 #include "WorldPartition/DataLayer/WorldDataLayersActorDesc.h"
@@ -280,6 +281,7 @@ void ULevelInstanceSubsystem::UpdateStreamingState()
 
 #if WITH_EDITOR
 	LevelsToRemoveScope.Reset();
+	IWorldPartitionActorLoaderInterface::RefreshLoadedState(true);
 #endif
 }
 
diff --git a/Engine/Source/Runtime/Engine/Private/WorldPartition/LevelInstance/LevelInstanceActorDesc.cpp b/Engine/Source/Runtime/Engine/Private/WorldPartition/LevelInstance/LevelInstanceActorDesc.cpp
--- a/Engine/Source/Runtime/Engine/Private/WorldPartition/LevelInstance/LevelInstanceActorDesc.cpp
+++ b/Engine/Source/Runtime/Engine/Private/WorldPartition/LevelInstance/LevelInstanceActorDesc.cpp
@@ -178,12 +178,15 @@ bool FLevelInstanceActorDesc::GetContainerInstance(const FGetContainerInstancePa
 {
 	if (LevelInstanceContainer.IsValid())
 	{
+		const ILevelInstanceInterface* LevelInstance = ActorPtr.IsValid() ? CastChecked<const ILevelInstanceInterface>(ActorPtr.Get()) : nullptr;
 		OutContainerInstance.Container = LevelInstanceContainer.Get();
+		OutContainerInstance.LoadedLevel = LevelInstance ? LevelInstance->GetLoadedLevel() : nullptr;
+		OutContainerInstance.bSupportsPartialEditorLoading = LevelInstance ? LevelInstance->SupportsPartialEditorLoading() : false;
 		OutContainerInstance.ClusterMode = EContainerClusterMode::Partitioned;
 
 		// Apply level instance pivot offset
 		FTransform LevelInstancePivotOffsetTransform = FTransform(ULevel::GetLevelInstancePivotOffsetFromPackage(LevelInstanceContainer->GetContainerPackage()));
-		OutContainerInstance.Transform = LevelInstancePivotOffsetTransform * LevelInstanceTransform;
+		OutContainerInstance.Transform = LevelInstancePivotOffsetTransform * LevelInstanceTransform;		
 
 		if (InParams.bBuildFilter)
 		{
diff --git a/Engine/Source/Runtime/Engine/Private/WorldPartition/LoaderAdapter/LoaderAdapterActorList.cpp b/Engine/Source/Runtime/Engine/Private/WorldPartition/LoaderAdapter/LoaderAdapterActorList.cpp
--- a/Engine/Source/Runtime/Engine/Private/WorldPartition/LoaderAdapter/LoaderAdapterActorList.cpp
+++ b/Engine/Source/Runtime/Engine/Private/WorldPartition/LoaderAdapter/LoaderAdapterActorList.cpp
@@ -2,6 +2,7 @@
 
 #include "WorldPartition/LoaderAdapter/LoaderAdapterActorList.h"
 #include "WorldPartition/WorldPartition.h"
+#include "Engine/Level.h"
 #include "Engine/World.h"
 
 #if WITH_EDITOR
@@ -60,6 +61,29 @@ void FLoaderAdapterActorList::RemoveActors(const TArray<FGuid>& ActorGuids)
 void FLoaderAdapterActorList::RemoveActors(const TArray<FWorldPartitionHandle>& ActorHandles)
 {
 	ActorsToRemove = TSet<FWorldPartitionHandle>(ActorHandles);
+
+	for (const FWorldPartitionHandle& ActorHandle : ActorHandles)
+	{
+		if (ActorHandle->IsContainerInstance())
+		{
+			FWorldPartitionActorDesc::FContainerInstance ContainerInstance;
+			if (ActorHandle->GetContainerInstance(FWorldPartitionActorDesc::FGetContainerInstanceParams(), ContainerInstance))
+			{
+				if (ContainerInstance.bSupportsPartialEditorLoading)
+				{
+					if (UWorldPartition* ContainerWorldPartition = ContainerInstance.LoadedLevel ? ContainerInstance.LoadedLevel->GetWorldPartition() : nullptr)
+					{
+						for (FActorDescContainerCollection::TIterator<> ActorDescIterator(ContainerWorldPartition); ActorDescIterator; ++ActorDescIterator)
+						{
+							FWorldPartitionHandle SubActorHandle(ContainerWorldPartition, ActorDescIterator->GetGuid());
+							ActorsToRemove.Add(SubActorHandle);
+						}
+					}
+				}
+			}
+		}
+	}
+
 	RefreshLoadedState();
 	Actors = Actors.Difference(ActorsToRemove);
 	ActorsToRemove.Empty();
diff --git a/Engine/Source/Runtime/Engine/Private/WorldPartition/LoaderAdapter/LoaderAdapterList.cpp b/Engine/Source/Runtime/Engine/Private/WorldPartition/LoaderAdapter/LoaderAdapterList.cpp
--- a/Engine/Source/Runtime/Engine/Private/WorldPartition/LoaderAdapter/LoaderAdapterList.cpp
+++ b/Engine/Source/Runtime/Engine/Private/WorldPartition/LoaderAdapter/LoaderAdapterList.cpp
@@ -1,6 +1,10 @@
 // Copyright Epic Games, Inc. All Rights Reserved.
 
 #include "WorldPartition/LoaderAdapter/LoaderAdapterList.h"
+#include "WorldPartition/ActorDescContainerCollection.h"
+#include "WorldPartition/WorldPartition.h"
+#include "WorldPartition/WorldPartitionActorDesc.h"
+#include "Engine/Level.h"
 
 #if WITH_EDITOR
 FLoaderAdapterList::FLoaderAdapterList(UWorld* InWorld)
@@ -14,6 +18,30 @@ void FLoaderAdapterList::ForEachActor(TFunctionRef<void(const FWorldPartitionHan
 		if (Actor.IsValid())
 		{
 			InOperation(Actor);
+			HandleActorContainer(Actor, InOperation);
+		}
+	}
+}
+
+void FLoaderAdapterList::HandleActorContainer(const FWorldPartitionHandle& InActor, TFunctionRef<void(const FWorldPartitionHandle&)> InOperation) const
+{
+	if (InActor->IsContainerInstance())
+	{
+		FWorldPartitionActorDesc::FContainerInstance ContainerInstance;
+		if (InActor->GetContainerInstance(FWorldPartitionActorDesc::FGetContainerInstanceParams(), ContainerInstance))
+		{
+			if (ContainerInstance.bSupportsPartialEditorLoading)
+			{
+				if (UWorldPartition* ContainerWorldPartition = ContainerInstance.LoadedLevel ? ContainerInstance.LoadedLevel->GetWorldPartition() : nullptr)
+				{
+					for (FActorDescContainerCollection::TIterator<> ActorDescIterator(ContainerWorldPartition); ActorDescIterator; ++ActorDescIterator)
+					{
+						FWorldPartitionHandle ActorHandle(ContainerWorldPartition, ActorDescIterator->GetGuid());
+						InOperation(ActorHandle);
+						HandleActorContainer(ActorHandle, InOperation);
+					}
+				}
+			}
 		}
 	}
 }
diff --git a/Engine/Source/Runtime/Engine/Private/WorldPartition/LoaderAdapter/LoaderAdapterSpatial.cpp b/Engine/Source/Runtime/Engine/Private/WorldPartition/LoaderAdapter/LoaderAdapterSpatial.cpp
--- a/Engine/Source/Runtime/Engine/Private/WorldPartition/LoaderAdapter/LoaderAdapterSpatial.cpp
+++ b/Engine/Source/Runtime/Engine/Private/WorldPartition/LoaderAdapter/LoaderAdapterSpatial.cpp
@@ -3,6 +3,7 @@
 #include "WorldPartition/LoaderAdapter/LoaderAdapterSpatial.h"
 #include "WorldPartition/WorldPartition.h"
 #include "WorldPartition/WorldPartitionEditorHash.h"
+#include "Engine/Level.h"
 #include "Engine/World.h"
 
 #if WITH_EDITOR
@@ -16,18 +17,42 @@ void ILoaderAdapterSpatial::ForEachActor(TFunctionRef<void(const FWorldPartition
 {
 	if (UWorldPartition* WorldPartition = GetWorld()->GetWorldPartition())
 	{
-		UWorldPartitionEditorHash::FForEachIntersectingActorParams ForEachIntersectingActorParams;
-		ForEachIntersectingActorParams.bIncludeSpatiallyLoadedActors = bIncludeSpatiallyLoadedActors;
-		ForEachIntersectingActorParams.bIncludeNonSpatiallyLoadedActors = bIncludeNonSpatiallyLoadedActors;
+		HandleIntersectingContainer(WorldPartition, InOperation);
+	}
+}
+
+void ILoaderAdapterSpatial::HandleIntersectingContainer(UWorldPartition* InWorldPartition, TFunctionRef<void(const FWorldPartitionHandle&)> InOperation) const
+{
+	const FTransform InstanceTransform = InWorldPartition->GetInstanceTransform();
+	const FBox LocalBoundingBox = GetBoundingBox()->InverseTransformBy(InstanceTransform);
 
-		WorldPartition->EditorHash->ForEachIntersectingActor(*GetBoundingBox(), [this, WorldPartition, &InOperation](FWorldPartitionActorDesc* ActorDesc)
+	UWorldPartitionEditorHash::FForEachIntersectingActorParams ForEachIntersectingActorParams = UWorldPartitionEditorHash::FForEachIntersectingActorParams()
+		.SetIncludeSpatiallyLoadedActors(bIncludeSpatiallyLoadedActors)
+		.SetIncludeNonSpatiallyLoadedActors(bIncludeNonSpatiallyLoadedActors);
+
+	InWorldPartition->EditorHash->ForEachIntersectingActor(LocalBoundingBox, [this, InWorldPartition, &InstanceTransform, &InOperation](FWorldPartitionActorDesc* ActorDesc)
+	{
+		const FBox WorldActorEditorBox = ActorDesc->GetEditorBounds().TransformBy(InstanceTransform);
+		if (Intersect(WorldActorEditorBox))
 		{
-			if (Intersect(ActorDesc->GetEditorBounds()))
+			FWorldPartitionHandle ActorHandle(InWorldPartition, ActorDesc->GetGuid());
+			InOperation(ActorHandle);
+
+			if (ActorHandle->GetIsSpatiallyLoaded() && ActorHandle->IsContainerInstance())
 			{
-				FWorldPartitionHandle ActorHandle(WorldPartition, ActorDesc->GetGuid());
-				InOperation(ActorHandle);
+				FWorldPartitionActorDesc::FContainerInstance ContainerInstance;
+				if (ActorHandle->GetContainerInstance(FWorldPartitionActorDesc::FGetContainerInstanceParams(), ContainerInstance))
+				{
+					if (ContainerInstance.bSupportsPartialEditorLoading)
+					{
+						if (UWorldPartition* ContainerWorldPartition = ContainerInstance.LoadedLevel ? ContainerInstance.LoadedLevel->GetWorldPartition() : nullptr)
+						{
+							HandleIntersectingContainer(ContainerWorldPartition, InOperation);
+						}
+					}
+				}
 			}
-		}, ForEachIntersectingActorParams);
-	}
+		}
+	}, ForEachIntersectingActorParams);
 }
 #endif
diff --git a/Engine/Source/Runtime/Engine/Private/WorldPartition/WorldPartition.cpp b/Engine/Source/Runtime/Engine/Private/WorldPartition/WorldPartition.cpp
--- a/Engine/Source/Runtime/Engine/Private/WorldPartition/WorldPartition.cpp
+++ b/Engine/Source/Runtime/Engine/Private/WorldPartition/WorldPartition.cpp
@@ -274,6 +274,7 @@ UWorldPartition::UWorldPartition(const FObjectInitializer& ObjectInitializer)
 	, bEnablingStreamingJustified(false)
 	, bIsPIE(false)
 	, NumUserCreatedLoadedRegions(0)
+	, bForceEnableStreamingInEditor(false)
 #endif
 	, InitState(EWorldPartitionInitState::Uninitialized)
 	, bStreamingInEnabled(true)
@@ -489,10 +490,14 @@ void UWorldPartition::Initialize(UWorld* InWorld, const FTransform& InTransform)
 		EditorHash->Initialize();
 
 		AlwaysLoadedActors = new FLoaderAdapterAlwaysLoadedActors(OuterWorld);
-		PinnedActors = new FLoaderAdapterPinnedActors(OuterWorld);
+
+		if (IsMainWorldPartition())
+		{			
+			PinnedActors = new FLoaderAdapterPinnedActors(OuterWorld);
 		
-		IWorldPartitionEditorModule& WorldPartitionEditorModule = FModuleManager::LoadModuleChecked<IWorldPartitionEditorModule>("WorldPartitionEditor");
-		ForceLoadedActors = WorldPartitionEditorModule.GetDisableLoadingInEditor() && IsMainWorldPartition() ? new FLoaderAdapterActorList(OuterWorld) : nullptr;
+			IWorldPartitionEditorModule& WorldPartitionEditorModule = FModuleManager::LoadModuleChecked<IWorldPartitionEditorModule>("WorldPartitionEditor");
+			ForceLoadedActors = WorldPartitionEditorModule.GetDisableLoadingInEditor() ? new FLoaderAdapterActorList(OuterWorld) : nullptr;
+		}
 	}
 
 	check(RuntimeHash);
@@ -507,11 +512,15 @@ void UWorldPartition::Initialize(UWorld* InWorld, const FTransform& InTransform)
 		const FName PackageName = LevelPackage->GetLoadedPath().GetPackageFName().IsNone() ? LevelPackage->GetFName() : LevelPackage->GetLoadedPath().GetPackageFName();
 
 		// Currently known Instancing use cases:
-		// - World Partition map template (New Level)
-		// - PIE World Travel
+		//	- World Partition map template (New Level)
+		//	- PIE World Travel
 		FString SourceWorldPath, RemappedWorldPath;
 		const bool bIsInstanced = OuterWorld->GetSoftObjectPathMapping(SourceWorldPath, RemappedWorldPath);
-	
+
+		// Follow the world's streaming enabled value most of the times, except:
+		//	- World is instanced and from a Level Instance that supports partial loading
+		const bool bIsStreamingEnabled = bForceEnableStreamingInEditor || IsStreamingEnabled();
+
 		if (bIsInstanced)
 		{
 			InstancingContext.AddPackageMapping(PackageName, LevelPackage->GetFName());
@@ -540,7 +549,7 @@ void UWorldPartition::Initialize(UWorld* InWorld, const FTransform& InTransform)
 					ActorDescIterator->TransformInstance(SourceWorldPath, RemappedWorldPath);
 				}
 
-				ActorDescIterator->bIsForcedNonSpatiallyLoaded = !IsStreamingEnabled();
+				ActorDescIterator->bIsForcedNonSpatiallyLoaded = !bIsStreamingEnabled;
 
 				if (ForceLoadedActors)
 				{
@@ -592,7 +601,8 @@ void UWorldPartition::Initialize(UWorld* InWorld, const FTransform& InTransform)
 
 	if (bIsEditor && !bIsCooking)
 	{
-		// Load the always loaded cell, don't call LoadCells to avoid creating a transaction
+		// Load the always loaded cell
+		if (AlwaysLoadedActors)
 		{
 			TRACE_CPUPROFILER_EVENT_SCOPE(LoadAlwaysLoaded);
 			AlwaysLoadedActors->Load();
diff --git a/Engine/Source/Runtime/Engine/Private/WorldPartition/WorldPartitionActorLoaderInterface.cpp b/Engine/Source/Runtime/Engine/Private/WorldPartition/WorldPartitionActorLoaderInterface.cpp
--- a/Engine/Source/Runtime/Engine/Private/WorldPartition/WorldPartitionActorLoaderInterface.cpp
+++ b/Engine/Source/Runtime/Engine/Private/WorldPartition/WorldPartitionActorLoaderInterface.cpp
@@ -72,7 +72,7 @@ void IWorldPartitionActorLoaderInterface::ILoaderAdapter::Unload()
 			int32 NumUnloads = 0;
 			{
 				FWorldPartitionLoadingContext::FDeferred LoadingContext;
-				ActorReferences.Empty();
+				ContainerActorReferences.Empty();
 				NumUnloads = LoadingContext.GetNumUnregistrations();
 			}
 
@@ -135,14 +135,16 @@ void IWorldPartitionActorLoaderInterface::ILoaderAdapter::RefreshLoadedState()
 			TArray<FWorldPartitionHandle> ActorsToUnload;			
 			ForEachActor([this, &ActorsToLoad, &ActorsToUnload](const FWorldPartitionHandle& Actor)
 			{
+				const FActorReferenceMap* ActorReferences = GetContainerReferencesConst(Actor->GetContainer());
+
 				if (ShouldActorBeLoaded(Actor))
 				{
-					if (!ActorReferences.Contains(Actor->GetGuid()))
+					if (!ActorReferences || !ActorReferences->Contains(Actor->GetGuid()))
 					{
 						ActorsToLoad.Add(Actor);
 					}
 				}
-				else if (ActorReferences.Contains(Actor->GetGuid()))
+				else if (ActorReferences && ActorReferences->Contains(Actor->GetGuid()))
 				{
 					ActorsToUnload.Add(Actor);
 				}
@@ -162,8 +164,7 @@ void IWorldPartitionActorLoaderInterface::ILoaderAdapter::RefreshLoadedState()
 
 						for (FWorldPartitionHandle& ActorToLoad : ActorsToLoad)
 						{
-							FWorldPartitionHandle ActorHandle(WorldPartition, ActorToLoad->GetGuid());
-							AddReferenceToActor(ActorHandle);
+							AddReferenceToActor(ActorToLoad);
 							SlowTask.EnterProgressFrame(1);
 						}
 					}
@@ -191,6 +192,15 @@ void IWorldPartitionActorLoaderInterface::ILoaderAdapter::RefreshLoadedState()
 			}
 		}
 	}
+
+	// Remove invalid containers (happens when reloading level instances, etc)
+	for (FContainerReferenceMap::TIterator It(ContainerActorReferences); It; ++It)
+	{
+		if (!It->Key.IsValid())
+		{
+			It.RemoveCurrent();
+		}
+	}
 }
 
 void IWorldPartitionActorLoaderInterface::ILoaderAdapter::OnActorDescContainerInitialize(UActorDescContainer* Container)
@@ -205,7 +215,9 @@ void IWorldPartitionActorLoaderInterface::ILoaderAdapter::OnActorDescContainerUn
 	TArray<FWorldPartitionHandle> ActorsToUnload;
 	for(UActorDescContainer::TConstIterator<> It(Container); It; ++It)
 	{
-		if( ActorReferences.Find(It->GetGuid()))
+		const FActorReferenceMap* ActorReferences = GetContainerReferencesConst(It->GetContainer());
+
+		if (!ActorReferences || !ActorReferences->Find(It->GetGuid()))
 		{
 			ActorsToUnload.Emplace(Container, It->GetGuid());
 		}
@@ -273,7 +285,7 @@ void IWorldPartitionActorLoaderInterface::ILoaderAdapter::PostLoadedStateChanged
 
 void IWorldPartitionActorLoaderInterface::ILoaderAdapter::AddReferenceToActor(FWorldPartitionHandle& ActorHandle)
 {
-	TFunction<void(const FWorldPartitionHandle&, TMap<FGuid, FWorldPartitionReference>&)> AddReferences = [this, &ActorHandle, &AddReferences](const FWorldPartitionHandle& Handle, TMap<FGuid, FWorldPartitionReference>& ReferenceMap)
+	TFunction<void(const FWorldPartitionHandle&, FReferenceMap&)> AddReferences = [this, &ActorHandle, &AddReferences](const FWorldPartitionHandle& Handle, FReferenceMap& ReferenceMap)
 	{
 		if (!ReferenceMap.Contains(Handle->GetGuid()))
 		{
@@ -291,11 +303,13 @@ void IWorldPartitionActorLoaderInterface::ILoaderAdapter::AddReferenceToActor(FW
 		}
 	};
 
+	FActorReferenceMap& ActorReferences = GetContainerReferences(ActorHandle->GetContainer());
 	AddReferences(ActorHandle, ActorReferences.Emplace(ActorHandle->GetGuid()));
 }
 
 void IWorldPartitionActorLoaderInterface::ILoaderAdapter::RemoveReferenceToActor(FWorldPartitionHandle& ActorHandle)
 {
+	FActorReferenceMap& ActorReferences = GetContainerReferences(ActorHandle->GetContainer());
 	ActorReferences.Remove(ActorHandle->GetGuid());
 }
 
@@ -308,6 +322,16 @@ void IWorldPartitionActorLoaderInterface::RefreshLoadedState(bool bIsFromUserCha
 {
 	ActorLoaderInterfaceRefreshState.Broadcast(bIsFromUserChange);
 }
+
+IWorldPartitionActorLoaderInterface::FActorReferenceMap& IWorldPartitionActorLoaderInterface::ILoaderAdapter::GetContainerReferences(UActorDescContainer* InContainer)
+{
+	return ContainerActorReferences.FindOrAdd(InContainer);
+}
+
+const IWorldPartitionActorLoaderInterface::FActorReferenceMap* IWorldPartitionActorLoaderInterface::ILoaderAdapter::GetContainerReferencesConst(UActorDescContainer* InContainer) const
+{
+	return ContainerActorReferences.Find(InContainer);
+}
 #endif
 
 #undef LOCTEXT_NAMESPACE
diff --git a/Engine/Source/Runtime/Engine/Private/WorldPartition/WorldPartitionEditorHash.cpp b/Engine/Source/Runtime/Engine/Private/WorldPartition/WorldPartitionEditorHash.cpp
--- a/Engine/Source/Runtime/Engine/Private/WorldPartition/WorldPartitionEditorHash.cpp
+++ b/Engine/Source/Runtime/Engine/Private/WorldPartition/WorldPartitionEditorHash.cpp
@@ -20,9 +20,10 @@ UWorldPartitionEditorHash::FForEachIntersectingActorParams::FForEachIntersecting
 
 int32 UWorldPartitionEditorHash::ForEachIntersectingActor(const FBox& Box, TFunctionRef<void(FWorldPartitionActorDesc*)> InOperation, bool bIncludeSpatiallyLoadedActors, bool bIncludeNonSpatiallyLoadedActors)
 {
-	UWorldPartitionEditorHash::FForEachIntersectingActorParams ForEachIntersectingActorParams;
-	ForEachIntersectingActorParams.bIncludeSpatiallyLoadedActors = bIncludeSpatiallyLoadedActors;
-	ForEachIntersectingActorParams.bIncludeNonSpatiallyLoadedActors = bIncludeNonSpatiallyLoadedActors;
+	UWorldPartitionEditorHash::FForEachIntersectingActorParams ForEachIntersectingActorParams = UWorldPartitionEditorHash::FForEachIntersectingActorParams()
+		.SetIncludeSpatiallyLoadedActors(bIncludeSpatiallyLoadedActors)
+		.SetIncludeNonSpatiallyLoadedActors(bIncludeNonSpatiallyLoadedActors);
+
 	return ForEachIntersectingActor(Box, InOperation, ForEachIntersectingActorParams);
 }
 #endif
diff --git a/Engine/Source/Runtime/Engine/Public/LevelInstance/LevelInstanceInterface.h b/Engine/Source/Runtime/Engine/Public/LevelInstance/LevelInstanceInterface.h
--- a/Engine/Source/Runtime/Engine/Public/LevelInstance/LevelInstanceInterface.h
+++ b/Engine/Source/Runtime/Engine/Public/LevelInstance/LevelInstanceInterface.h
@@ -57,6 +57,8 @@ class ENGINE_API ILevelInstanceInterface
 	virtual ELevelInstanceRuntimeBehavior GetDefaultRuntimeBehavior() const = 0;
 
 	virtual TSubclassOf<AActor> GetEditorPivotClass() const { return nullptr; }
+
+	virtual bool SupportsPartialEditorLoading() const;
 #endif
 	// Pure Interface End
 		
diff --git a/Engine/Source/Runtime/Engine/Public/WorldPartition/LoaderAdapter/LoaderAdapterList.h b/Engine/Source/Runtime/Engine/Public/WorldPartition/LoaderAdapter/LoaderAdapterList.h
--- a/Engine/Source/Runtime/Engine/Public/WorldPartition/LoaderAdapter/LoaderAdapterList.h
+++ b/Engine/Source/Runtime/Engine/Public/WorldPartition/LoaderAdapter/LoaderAdapterList.h
@@ -17,6 +17,11 @@ protected:
 	virtual void ForEachActor(TFunctionRef<void(const FWorldPartitionHandle&)> InOperation) const override;
 	//~ End ILoaderAdapter interface
 
+private:
+	void HandleActorContainer(const FWorldPartitionHandle& InActor, TFunctionRef<void(const FWorldPartitionHandle&)> InOperation) const;
+
+protected:
+
 	TSet<FWorldPartitionHandle> Actors;
 };
 #endif
\ No newline at end of file
diff --git a/Engine/Source/Runtime/Engine/Public/WorldPartition/LoaderAdapter/LoaderAdapterSpatial.h b/Engine/Source/Runtime/Engine/Public/WorldPartition/LoaderAdapter/LoaderAdapterSpatial.h
--- a/Engine/Source/Runtime/Engine/Public/WorldPartition/LoaderAdapter/LoaderAdapterSpatial.h
+++ b/Engine/Source/Runtime/Engine/Public/WorldPartition/LoaderAdapter/LoaderAdapterSpatial.h
@@ -5,6 +5,8 @@
 #include "WorldPartition/WorldPartitionActorLoaderInterface.h"
 
 #if WITH_EDITOR
+class UWorldPartition;
+
 /** Base class for actor loaders that requires spatial queries */
 class ENGINE_API ILoaderAdapterSpatial : public IWorldPartitionActorLoaderInterface::ILoaderAdapter
 {
@@ -20,6 +22,10 @@ protected:
 	// Private interface
 	virtual bool Intersect(const FBox& Box) const =0;
 
+private:
+	void HandleIntersectingContainer(UWorldPartition* InWorldPartition, TFunctionRef<void(const FWorldPartitionHandle&)> InOperation) const;
+
+protected:
 	bool bIncludeSpatiallyLoadedActors;
 	bool bIncludeNonSpatiallyLoadedActors;
 };
diff --git a/Engine/Source/Runtime/Engine/Public/WorldPartition/WorldPartition.h b/Engine/Source/Runtime/Engine/Public/WorldPartition/WorldPartition.h
--- a/Engine/Source/Runtime/Engine/Public/WorldPartition/WorldPartition.h
+++ b/Engine/Source/Runtime/Engine/Public/WorldPartition/WorldPartition.h
@@ -396,6 +396,11 @@ private:
 	TMap<FWorldPartitionReference, AActor*> DirtyActors;
 
 	TSet<FString> GeneratedStreamingPackageNames;
+
+public:
+	bool bForceEnableStreamingInEditor;
+
+private:
 #endif
 
 	EWorldPartitionInitState InitState;
diff --git a/Engine/Source/Runtime/Engine/Public/WorldPartition/WorldPartitionActorDesc.h b/Engine/Source/Runtime/Engine/Public/WorldPartition/WorldPartitionActorDesc.h
--- a/Engine/Source/Runtime/Engine/Public/WorldPartition/WorldPartitionActorDesc.h
+++ b/Engine/Source/Runtime/Engine/Public/WorldPartition/WorldPartitionActorDesc.h
@@ -82,8 +82,10 @@ public:
 	struct FContainerInstance
 	{
 		const UActorDescContainer* Container = nullptr;
+		ULevel* LoadedLevel = nullptr;
+		bool bSupportsPartialEditorLoading = false;
 		FTransform Transform = FTransform::Identity;
-		EContainerClusterMode ClusterMode;
+		EContainerClusterMode ClusterMode;		
 		TMap<FActorContainerID, TSet<FGuid>> FilteredActors;
 	};
 
diff --git a/Engine/Source/Runtime/Engine/Public/WorldPartition/WorldPartitionActorLoaderInterface.h b/Engine/Source/Runtime/Engine/Public/WorldPartition/WorldPartitionActorLoaderInterface.h
--- a/Engine/Source/Runtime/Engine/Public/WorldPartition/WorldPartitionActorLoaderInterface.h
+++ b/Engine/Source/Runtime/Engine/Public/WorldPartition/WorldPartitionActorLoaderInterface.h
@@ -19,6 +19,10 @@ class ENGINE_API IWorldPartitionActorLoaderInterface
 	GENERATED_IINTERFACE_BODY()
 
 #if WITH_EDITOR
+	using FReferenceMap = TMap<FGuid, FWorldPartitionReference>;
+	using FActorReferenceMap = TMap<FGuid, FReferenceMap>;
+	using FContainerReferenceMap = TMap<TWeakObjectPtr<UActorDescContainer>, FActorReferenceMap>;
+
 public:
 	/** Base class for actor loaders */
 	class ENGINE_API ILoaderAdapter
@@ -60,7 +64,11 @@ public:
 		void PostLoadedStateChanged(int32 NumLoads, int32 NumUnloads);
 		void AddReferenceToActor(FWorldPartitionHandle& Actor);
 		void RemoveReferenceToActor(FWorldPartitionHandle& Actor);
-		void OnRefreshLoadedState(bool bFromUserOperation);		
+		void OnRefreshLoadedState(bool bFromUserOperation);
+
+		// Helpers
+		FActorReferenceMap& GetContainerReferences(UActorDescContainer* InContainer);
+		const FActorReferenceMap* GetContainerReferencesConst(UActorDescContainer* InContainer) const;
 
 	private:
 		UWorld* World;
@@ -68,7 +76,7 @@ public:
 		uint8 bLoaded : 1;
 		uint8 bUserCreated : 1;
 
-		TMap<FGuid, TMap<FGuid, FWorldPartitionReference>> ActorReferences;
+		FContainerReferenceMap ContainerActorReferences;
 	};
 
 	virtual ILoaderAdapter* GetLoaderAdapter() =0;
diff --git a/Engine/Source/Runtime/Engine/Public/WorldPartition/WorldPartitionEditorHash.h b/Engine/Source/Runtime/Engine/Public/WorldPartition/WorldPartitionEditorHash.h
--- a/Engine/Source/Runtime/Engine/Public/WorldPartition/WorldPartitionEditorHash.h
+++ b/Engine/Source/Runtime/Engine/Public/WorldPartition/WorldPartitionEditorHash.h
@@ -42,6 +42,10 @@ public:
 
 		/** Optional minimum box to stop searching for actors */
 		TOptional<FBox> MinimumBox;
+
+		FForEachIntersectingActorParams& SetIncludeSpatiallyLoadedActors(bool bInIncludeSpatiallyLoadedActors) { bIncludeSpatiallyLoadedActors = bInIncludeSpatiallyLoadedActors; return *this; };
+		FForEachIntersectingActorParams& SetIncludeNonSpatiallyLoadedActors(bool bInIncludeNonSpatiallyLoadedActors) { bIncludeNonSpatiallyLoadedActors = bInIncludeNonSpatiallyLoadedActors; return *this; };
+		FForEachIntersectingActorParams& SetMinimumBox(const FBox& InMinimumBox) { MinimumBox = InMinimumBox; return *this; };
 	};
 
 	virtual int32 ForEachIntersectingActor(const FBox& Box, TFunctionRef<void(FWorldPartitionActorDesc*)> InOperation, const FForEachIntersectingActorParams& Params = FForEachIntersectingActorParams()) PURE_VIRTUAL(UWorldPartitionEditorHash::ForEachIntersectingActor, return 0;);
