Don't assume that GetWorldRootFolder will always return a valid folder for the persistent level.

diff --git a/Engine/Source/Runtime/Engine/Private/Folder.cpp b/Engine/Source/Runtime/Engine/Private/Folder.cpp
--- a/Engine/Source/Runtime/Engine/Private/Folder.cpp
+++ b/Engine/Source/Runtime/Engine/Private/Folder.cpp
@@ -77,7 +77,7 @@ FFolder FFolder::GetWorldRootFolder(UWorld* InWorld)
 {
 	if (InWorld)
 	{
-		return FFolder::GetOptionalFolderRootObject(InWorld->PersistentLevel).GetValue();
+		return FFolder::GetOptionalFolderRootObject(InWorld->PersistentLevel).Get(FFolder::GetInvalidRootObject());
 	}
 	return FFolder::GetInvalidFolder();
 }
