Fixed WorldPartition uninitialization to properly support streaming-in/out/in again the same instanced partitioned world.

diff --git a/Engine/Source/Runtime/Engine/Private/WorldPartition/WorldPartition.cpp b/Engine/Source/Runtime/Engine/Private/WorldPartition/WorldPartition.cpp
--- a/Engine/Source/Runtime/Engine/Private/WorldPartition/WorldPartition.cpp
+++ b/Engine/Source/Runtime/Engine/Private/WorldPartition/WorldPartition.cpp
@@ -714,6 +714,7 @@ void UWorldPartition::Uninitialize()
 		ActorDescContainer = nullptr;
 
 		EditorHash = nullptr;
+		bIsPIE = false;
 #endif		
 
 		if (DataLayerManager)
diff --git a/Engine/Source/Runtime/Engine/Private/WorldPartition/WorldPartitionRuntimeHash.cpp b/Engine/Source/Runtime/Engine/Private/WorldPartition/WorldPartitionRuntimeHash.cpp
--- a/Engine/Source/Runtime/Engine/Private/WorldPartition/WorldPartitionRuntimeHash.cpp
+++ b/Engine/Source/Runtime/Engine/Private/WorldPartition/WorldPartitionRuntimeHash.cpp
@@ -116,7 +116,6 @@ UWorldPartitionRuntimeCell* UWorldPartitionRuntimeHash::CreateRuntimeCell(UClass
 	};
 
 	const FString CellObjectName = GetCellObjectName(CellName);
-	check(!FindObject<UWorldPartitionRuntimeCell>(this, *CellObjectName));
 	UWorldPartitionRuntimeCell* RuntimeCell = NewObject<UWorldPartitionRuntimeCell>(this, CellClass, *CellObjectName);
 	RuntimeCell->RuntimeCellData = NewObject<UWorldPartitionRuntimeCellData>(RuntimeCell, CellDataClass);
 	return RuntimeCell;
diff --git a/Engine/Source/Runtime/Engine/Private/WorldPartition/WorldPartitionRuntimeLevelStreamingCell.cpp b/Engine/Source/Runtime/Engine/Private/WorldPartition/WorldPartitionRuntimeLevelStreamingCell.cpp
--- a/Engine/Source/Runtime/Engine/Private/WorldPartition/WorldPartitionRuntimeLevelStreamingCell.cpp
+++ b/Engine/Source/Runtime/Engine/Private/WorldPartition/WorldPartitionRuntimeLevelStreamingCell.cpp
@@ -115,7 +115,6 @@ UWorldPartitionLevelStreamingDynamic* UWorldPartitionRuntimeLevelStreamingCell::
 		// When called by Commandlet (PopulateGeneratedPackageForCook), LevelStreaming's outer is set to Cell/WorldPartition's outer to prevent warnings when saving Cell Levels (Warning: Obj in another map). 
 		// At runtime, LevelStreaming's outer will be properly set to the main world (see UWorldPartitionRuntimeLevelStreamingCell::Activate).
 		UWorld* LevelStreamingOuterWorld = IsRunningCommandlet() ? OuterWorld : OwningWorld;
-		check(!FindObject<UWorldPartitionLevelStreamingDynamic>(LevelStreamingOuterWorld, *LevelStreamingName.ToString()));
 		UWorldPartitionLevelStreamingDynamic* NewLevelStreaming = NewObject<UWorldPartitionLevelStreamingDynamic>(LevelStreamingOuterWorld, UWorldPartitionLevelStreamingDynamic::StaticClass(), LevelStreamingName, RF_NoFlags, NULL);
 
 		FName WorldName = OuterWorld->GetFName();
