Improve error report on corrupted ShaderCompileWorker output file.

#rnx

diff --git a/Engine/Source/Runtime/Engine/Private/ShaderCompiler/ShaderCompiler.cpp b/Engine/Source/Runtime/Engine/Private/ShaderCompiler/ShaderCompiler.cpp
--- a/Engine/Source/Runtime/Engine/Private/ShaderCompiler/ShaderCompiler.cpp
+++ b/Engine/Source/Runtime/Engine/Private/ShaderCompiler/ShaderCompiler.cpp
@@ -877,7 +877,7 @@ static constexpr int32 GSingleThreadedRunsDisabled = -2;
 static constexpr int32 GSingleThreadedRunsIncreaseFactor = 8;
 static constexpr int32 GSingleThreadedRunsMaxCount = (1 << 24);
 
-static void ModalErrorOrLog(const FString& Text, int64 CurrentFilePos = 0, int64 ExpectedFileSize = 0)
+static void ModalErrorOrLog(const FString& Title, const FString& Text, int64 CurrentFilePos = 0, int64 ExpectedFileSize = 0)
 {
 	static FThreadSafeBool bModalReported;
 
@@ -885,7 +885,7 @@ static void ModalErrorOrLog(const FString& Text, int64 CurrentFilePos = 0, int64
 	if (CurrentFilePos > ExpectedFileSize)
 	{
 		// Corrupt file
-		BadFile = FString::Printf(TEXT("(Truncated or corrupt output file! Current file pos %lld, file size %lld)"), CurrentFilePos, ExpectedFileSize);
+		BadFile = FString::Printf(TEXT(" (Truncated or corrupt output file! Current file pos %lld, file size %lld)"), CurrentFilePos, ExpectedFileSize);
 	}
 
 	if (FPlatformProperties::SupportsWindowedMode())
@@ -894,7 +894,8 @@ static void ModalErrorOrLog(const FString& Text, int64 CurrentFilePos = 0, int64
 		if (!bModalReported.AtomicSet(true))
 		{
 			// Show dialog box with error message and request exit
-			FMessageDialog::Open(EAppMsgType::Ok, FText::FromString(Text));
+			FText TitleText = FText::FromString(Title);
+			FMessageDialog::Open(EAppMsgType::Ok, FText::FromString(Text), &TitleText);
 			FPlatformMisc::RequestExit(false, TEXT("ShaderCompiler.ModalErrorOrLog"));
 		}
 		else
@@ -905,7 +906,7 @@ static void ModalErrorOrLog(const FString& Text, int64 CurrentFilePos = 0, int64
 	}
 	else
 	{
-		UE_LOG(LogShaderCompilers, Fatal, TEXT("%s%s"), *Text, *BadFile);
+		UE_LOG(LogShaderCompilers, Fatal, TEXT("%s\n%s%s"), *Title, *Text, *BadFile);
 	}
 }
 
@@ -1264,62 +1265,62 @@ namespace ShaderCompileWorkerError
 
 	void HandleBadShaderFormatVersion(const TCHAR* Data)
 	{
-		ModalErrorOrLog(FString::Printf(TEXT("ShaderCompileWorker failed:\n%s\n"), Data));
+		ModalErrorOrLog(TEXT("ShaderCompileWorker failed"), FString::Printf(TEXT("%s\n"), Data));
 	}
 
 	void HandleBadInputVersion(const TCHAR* Data)
 	{
-		ModalErrorOrLog(FString::Printf(TEXT("ShaderCompileWorker failed:\n%s\n"), Data));
+		ModalErrorOrLog(TEXT("ShaderCompileWorker failed"), FString::Printf(TEXT("%s\n"), Data));
 	}
 
 	void HandleBadSingleJobHeader(const TCHAR* Data)
 	{
-		ModalErrorOrLog(FString::Printf(TEXT("ShaderCompileWorker failed:\n%s\n"), Data));
+		ModalErrorOrLog(TEXT("ShaderCompileWorker failed"), FString::Printf(TEXT("%s\n"), Data));
 	}
 
 	void HandleBadPipelineJobHeader(const TCHAR* Data)
 	{
-		ModalErrorOrLog(FString::Printf(TEXT("ShaderCompileWorker failed:\n%s\n"), Data));
+		ModalErrorOrLog(TEXT("ShaderCompileWorker failed"), FString::Printf(TEXT("%s\n"), Data));
 	}
 
 	void HandleCantDeleteInputFile(const TCHAR* Data)
 	{
-		ModalErrorOrLog(FString::Printf(TEXT("ShaderCompileWorker failed:\n%s\n"), Data));
+		ModalErrorOrLog(TEXT("ShaderCompileWorker failed"), FString::Printf(TEXT("%s\n"), Data));
 	}
 
 	void HandleCantSaveOutputFile(const TCHAR* Data)
 	{
-		ModalErrorOrLog(FString::Printf(TEXT("ShaderCompileWorker failed:\n%s\n"), Data));
+		ModalErrorOrLog(TEXT("ShaderCompileWorker failed"), FString::Printf(TEXT("%s\n"), Data));
 	}
 
 	void HandleNoTargetShaderFormatsFound(const TCHAR* Data)
 	{
-		ModalErrorOrLog(FString::Printf(TEXT("ShaderCompileWorker failed:\n%s\n"), Data));
+		ModalErrorOrLog(TEXT("ShaderCompileWorker failed"), FString::Printf(TEXT("%s\n"), Data));
 	}
 
 	void HandleCantCompileForSpecificFormat(const TCHAR* Data)
 	{
-		ModalErrorOrLog(FString::Printf(TEXT("ShaderCompileWorker failed:\n%s\n"), Data));
+		ModalErrorOrLog(TEXT("ShaderCompileWorker failed"), FString::Printf(TEXT("%s\n"), Data));
 	}
 
 	void HandleOutputFileEmpty(const TCHAR* Filename)
 	{
-		ModalErrorOrLog(FString::Printf(TEXT("Output file %s size is 0. Are you out of disk space?"), Filename));
+		ModalErrorOrLog(TEXT("ShaderCompileWorker failed"), FString::Printf(TEXT("Output file %s size is 0. Are you out of disk space?"), Filename));
 	}
 
 	void HandleOutputFileCorrupted(const TCHAR* Filename, int64 ExpectedSize, int64 ActualSize)
 	{
-		ModalErrorOrLog(FString::Printf(TEXT("Output file corrupted (expected %I64d bytes, but only got %I64d): %s"), ExpectedSize, ActualSize, Filename));
+		ModalErrorOrLog(TEXT("ShaderCompileWorker failed"), FString::Printf(TEXT("Output file corrupted (expected %I64d bytes, but only got %I64d): %s"), ExpectedSize, ActualSize, Filename));
 	}
 
 	void HandleCrashInsidePlatformCompiler(const TCHAR* Data)
 	{
-		ModalErrorOrLog(FString::Printf(TEXT("Crash inside the platform compiler!\n%s"), Data));
+		ModalErrorOrLog(TEXT("ShaderCompileWorker failed"), FString::Printf(TEXT("Crash inside the platform compiler:\n%s"), Data));
 	}
 
 	void HandleBadInputFile(const TCHAR* Data)
 	{
-		ModalErrorOrLog(FString::Printf(TEXT("ShaderCompileWorker failed with bad-input-file exception:\n%s\n"), Data));
+		ModalErrorOrLog(TEXT("ShaderCompileWorker failed"), FString::Printf(TEXT("Bad-input-file exception:\n%s\n"), Data));
 	}
 
 	bool HandleOutOfMemory(const TCHAR* ExceptionInfo, const TCHAR* Hostname, uint64 AvailableMemory, uint64 UsedMemory, const TArray<FShaderCommonCompileJobPtr>& QueuedJobs)
@@ -1336,7 +1337,7 @@ namespace ShaderCompileWorkerError
 		}
 		else
 		{
-			ModalErrorOrLog(ErrorReport);
+			ModalErrorOrLog(TEXT("ShaderCompileWorker failed"), ErrorReport);
 			return false;
 		}
 	}
@@ -1982,9 +1983,39 @@ static bool HandleWorkerCrash(const TArray<FShaderCommonCompileJobPtr>& QueuedJo
 }
 UE_ENABLE_OPTIMIZATION_SHIP
 
+// Helper struct to provide consistent error report with detailed information about corrupted ShaderCompileWorker output file.
+struct FSCWOutputFileContext
+{
+	FArchive& OutputFile;
+	int64 FileSize = 0;
+
+	FSCWOutputFileContext(FArchive& OutputFile) :
+		OutputFile(OutputFile)
+	{
+	}
+
+	template <typename FmtType, typename... Types>
+	void ModalErrorOrLog(const FmtType& Format, Types&&... Args)
+	{
+		FString Text = FString::Printf(Format, Args...);
+		Text = FString::Printf(TEXT("File path: \"%s\"\n%s\nForgot to build ShaderCompileWorker or delete invalidated DerivedDataCache?"), *OutputFile.GetArchiveName(), *Text);
+		const TCHAR* Title = TEXT("Corrupted ShaderCompileWorker output file");
+		if (FileSize > 0)
+		{
+			::ModalErrorOrLog(Title, Text, OutputFile.Tell(), FileSize);
+		}
+		else
+		{
+			::ModalErrorOrLog(Title, Text, 0, 0);
+		}
+	}
+};
+
 // Process results from Worker Process
 void FShaderCompileUtilities::DoReadTaskResults(const TArray<FShaderCommonCompileJobPtr>& QueuedJobs, FArchive& OutputFile)
 {
+	FSCWOutputFileContext OutputFileContext(OutputFile);
+
 	if (OutputFile.TotalSize() == 0)
 	{
 		ShaderCompileWorkerError::HandleOutputFileEmpty(*OutputFile.GetArchiveName());
@@ -1995,17 +2026,15 @@ void FShaderCompileUtilities::DoReadTaskResults(const TArray<FShaderCommonCompil
 
 	if (ShaderCompileWorkerOutputVersion != OutputVersion)
 	{
-		FString Text = FString::Printf(TEXT("Expecting ShaderCompileWorker output version %d, got %d instead! Forgot to build ShaderCompileWorker?"), ShaderCompileWorkerOutputVersion, OutputVersion);
-		ModalErrorOrLog(Text);
+		OutputFileContext.ModalErrorOrLog(TEXT("Expecting output version %d, got %d instead!"), ShaderCompileWorkerOutputVersion, OutputVersion);
 	}
 
-	int64 FileSize = 0;
-	OutputFile << FileSize;
+	OutputFile << OutputFileContext.FileSize;
 
 	// Check for corrupted output file
-	if (FileSize > OutputFile.TotalSize())
+	if (OutputFileContext.FileSize > OutputFile.TotalSize())
 	{
-		ShaderCompileWorkerError::HandleOutputFileCorrupted(*OutputFile.GetArchiveName(), FileSize, OutputFile.TotalSize());
+		ShaderCompileWorkerError::HandleOutputFileCorrupted(*OutputFile.GetArchiveName(), OutputFileContext.FileSize, OutputFile.TotalSize());
 	}
 
 	int32 ErrorCode = 0;
@@ -2026,7 +2055,7 @@ void FShaderCompileUtilities::DoReadTaskResults(const TArray<FShaderCommonCompil
 	if (ErrorCode != FSCWErrorCode::Success)
 	{
 		// If worker crashed in a way we were able to recover from, return and expect the compile jobs to be reissued already
-		if (HandleWorkerCrash(QueuedJobs, OutputFile, OutputVersion, FileSize, (FSCWErrorCode::ECode)ErrorCode, NumProcessedJobs, CallstackLength, ExceptionInfoLength, HostnameLength))
+		if (HandleWorkerCrash(QueuedJobs, OutputFile, OutputVersion, OutputFileContext.FileSize, (FSCWErrorCode::ECode)ErrorCode, NumProcessedJobs, CallstackLength, ExceptionInfoLength, HostnameLength))
 		{
 			FSCWErrorCode::Reset();
 			return;
@@ -2044,16 +2073,14 @@ void FShaderCompileUtilities::DoReadTaskResults(const TArray<FShaderCommonCompil
 		OutputFile << SingleJobHeader;
 		if (SingleJobHeader != ShaderCompileWorkerSingleJobHeader)
 		{
-			FString Text = FString::Printf(TEXT("Expecting ShaderCompileWorker Single Jobs %d, got %d instead! Forgot to build ShaderCompileWorker?"), ShaderCompileWorkerSingleJobHeader, SingleJobHeader);
-			ModalErrorOrLog(Text, OutputFile.Tell(), FileSize);
+			OutputFileContext.ModalErrorOrLog(TEXT("Expecting single job header ID 0x%08X, got 0x%08X instead!"), ShaderCompileWorkerSingleJobHeader, SingleJobHeader);
 		}
 
 		int32 NumJobs;
 		OutputFile << NumJobs;
 		if (NumJobs != QueuedSingleJobs.Num())
 		{
-			FString Text = FString::Printf(TEXT("ShaderCompileWorker returned %u single jobs, %u expected"), NumJobs, QueuedSingleJobs.Num());
-			ModalErrorOrLog(Text, OutputFile.Tell(), FileSize);
+			OutputFileContext.ModalErrorOrLog(TEXT("Expecting %d single %s, got %d instead!"), QueuedSingleJobs.Num(), (QueuedSingleJobs.Num() == 1 ? TEXT("job") : TEXT("jobs")), NumJobs);
 		}
 		else
 		{
@@ -2074,16 +2101,14 @@ void FShaderCompileUtilities::DoReadTaskResults(const TArray<FShaderCommonCompil
 		OutputFile << PipelineJobHeader;
 		if (PipelineJobHeader != ShaderCompileWorkerPipelineJobHeader)
 		{
-			FString Text = FString::Printf(TEXT("Expecting ShaderCompileWorker Pipeline Jobs %d, got %d instead! Forgot to build ShaderCompileWorker?"), ShaderCompileWorkerPipelineJobHeader, PipelineJobHeader);
-			ModalErrorOrLog(Text, OutputFile.Tell(), FileSize);
+			OutputFileContext.ModalErrorOrLog(TEXT("Expecting pipeline jobs header ID 0x%08X, got 0x%08X instead!"), ShaderCompileWorkerPipelineJobHeader, PipelineJobHeader);
 		}
 
 		int32 NumJobs;
 		OutputFile << NumJobs;
 		if (NumJobs != QueuedPipelineJobs.Num())
 		{
-			FString Text = FString::Printf(TEXT("Worker returned %u pipeline jobs, %u expected"), NumJobs, QueuedPipelineJobs.Num());
-			ModalErrorOrLog(Text, OutputFile.Tell(), FileSize);
+			OutputFileContext.ModalErrorOrLog(TEXT("Expecting %d pipeline %s, got %d instead!"), QueuedPipelineJobs.Num(), (QueuedPipelineJobs.Num() == 1 ? TEXT("job") : TEXT("jobs")), NumJobs);
 		}
 		else
 		{
@@ -2099,8 +2124,7 @@ void FShaderCompileUtilities::DoReadTaskResults(const TArray<FShaderCommonCompil
 				OutputFile << CurrentJob->bFailedRemovingUnused;
 				if (PipelineName != CurrentJob->Key.ShaderPipeline->GetName())
 				{
-					FString Text = FString::Printf(TEXT("Worker returned Pipeline %s, expected %s!"), *PipelineName, CurrentJob->Key.ShaderPipeline->GetName());
-					ModalErrorOrLog(Text, OutputFile.Tell(), FileSize);
+					OutputFileContext.ModalErrorOrLog(TEXT("Expecting pipeline job \"%s\", got \"%s\" instead!"), CurrentJob->Key.ShaderPipeline->GetName(), *PipelineName);
 				}
 
 				check(!CurrentJob->bFinalized);
@@ -2111,8 +2135,7 @@ void FShaderCompileUtilities::DoReadTaskResults(const TArray<FShaderCommonCompil
 
 				if (NumStageJobs != CurrentJob->StageJobs.Num())
 				{
-					FString Text = FString::Printf(TEXT("Worker returned %u stage pipeline jobs, %u expected"), NumStageJobs, CurrentJob->StageJobs.Num());
-					ModalErrorOrLog(Text, OutputFile.Tell(), FileSize);
+					OutputFileContext.ModalErrorOrLog(TEXT("Expecting %d stage pipeline %s, got %d instead!"), CurrentJob->StageJobs.Num(), (CurrentJob->StageJobs.Num() == 1 ? TEXT("job") : TEXT("jobs")), NumStageJobs);
 				}
 				else
 				{
@@ -2793,16 +2816,13 @@ int32 FShaderCompileThreadRunnable::ReadAvailableResults()
 			{
 				TRACE_CPUPROFILER_EVENT_SCOPE(FShaderCompileThreadRunnable::ProcessOutputFile);
 
-				FArchive* OutputFilePtr = IFileManager::Get().CreateFileReader(*OutputFileNameAndPath, FILEREAD_Silent);
-
-				if (OutputFilePtr)
+				if (TUniquePtr<FArchive> OutputFile = TUniquePtr<FArchive>(IFileManager::Get().CreateFileReader(*OutputFileNameAndPath, FILEREAD_Silent)))
 				{
-					FArchive& OutputFile = *OutputFilePtr;
 					check(!CurrentWorkerInfo.bComplete);
-					FShaderCompileUtilities::DoReadTaskResults(CurrentWorkerInfo.QueuedJobs, OutputFile);
+					FShaderCompileUtilities::DoReadTaskResults(CurrentWorkerInfo.QueuedJobs, *OutputFile);
 
 					// Close the output file.
-					delete OutputFilePtr;
+					OutputFile.Reset();
 
 					// Delete the output file now that we have consumed it, to avoid reading stale data on the next compile loop.
 					bool bDeletedOutput = IFileManager::Get().Delete(*OutputFileNameAndPath, true, true);
