Fixed dirty actors loosing their folders in PIE in partitioned maps.

diff --git a/Engine/Source/Runtime/Engine/Classes/GameFramework/Actor.h b/Engine/Source/Runtime/Engine/Classes/GameFramework/Actor.h
--- a/Engine/Source/Runtime/Engine/Classes/GameFramework/Actor.h
+++ b/Engine/Source/Runtime/Engine/Classes/GameFramework/Actor.h
@@ -3969,6 +3969,7 @@ private:
 	friend struct FSetActorContentBundleGuid;
 	friend struct FAssignActorDataLayer;
 	friend struct FSetActorSelectable;
+	friend struct FSetActorFolderPath;
 #endif
 
 	// Static helpers for accessing functions on SceneComponent.
@@ -4235,6 +4236,16 @@ private:
 	friend class FWorldPartitionLevelHelper;
 };
 
+struct FSetActorFolderPath
+{
+private:
+	FSetActorFolderPath(AActor* InActor, const FName InFolderPath, bool bInBroadcastChange = true)
+	{
+		InActor->SetFolderPathInternal(InFolderPath, bInBroadcastChange);
+	}
+	friend class UWorldPartitionRuntimeCell;
+};
+
 struct FSetActorContentBundleGuid
 {
 private:
diff --git a/Engine/Source/Runtime/Engine/Private/WorldPartition/WorldPartitionRuntimeCell.cpp b/Engine/Source/Runtime/Engine/Private/WorldPartition/WorldPartitionRuntimeCell.cpp
--- a/Engine/Source/Runtime/Engine/Private/WorldPartition/WorldPartitionRuntimeCell.cpp
+++ b/Engine/Source/Runtime/Engine/Private/WorldPartition/WorldPartitionRuntimeCell.cpp
@@ -56,8 +56,14 @@ void UWorldPartitionRuntimeCell::PostDuplicate(bool bDuplicateForPIE)
 		{
 			if (ensure(Actor))
 			{
+				const FName FolderPath = Actor->GetFolderPath();
+				
 				// Don't use AActor::Rename here since the actor is not part of the world, it's only a duplication template.	
 				Actor->UObject::Rename(nullptr, UnsavedActorsContainer);
+				
+				// Mimic what AActor::Rename does under the hood to convert ActorFolders when changing the outer
+				const bool bBroadcastChange = false;
+				FSetActorFolderPath SetFolderPath(Actor, FolderPath, bBroadcastChange);
 			}
 		}
 	}
