Undo //UE5/Main/Engine/Source/Runtime/Engine/Private/Animation/... changelist 24764354

diff --git a/Engine/Source/Runtime/Engine/Private/Animation/AnimCurveCompressionSettings.cpp b/Engine/Source/Runtime/Engine/Private/Animation/AnimCurveCompressionSettings.cpp
--- a/Engine/Source/Runtime/Engine/Private/Animation/AnimCurveCompressionSettings.cpp
+++ b/Engine/Source/Runtime/Engine/Private/Animation/AnimCurveCompressionSettings.cpp
@@ -11,7 +11,6 @@ UAnimCurveCompressionSettings::UAnimCurveCompressionSettings(const FObjectInitia
 {
 	Codec = CreateDefaultSubobject<UAnimCurveCompressionCodec_CompressedRichCurve>(TEXT("CurveCompressionCodec"));
 	Codec->SetFlags(RF_Transactional);
-	Codec->ClearFlags(RF_Public);
 }
 
 UAnimCurveCompressionCodec* UAnimCurveCompressionSettings::GetCodec(const FString& Path)
