Fix access to render LOD0 in editor

diff --git a/Engine/Plugins/Runtime/MeshModelingToolset/Source/ModelingComponents/Private/Scene/MeshSceneAdapter.cpp b/Engine/Plugins/Runtime/MeshModelingToolset/Source/ModelingComponents/Private/Scene/MeshSceneAdapter.cpp
--- a/Engine/Plugins/Runtime/MeshModelingToolset/Source/ModelingComponents/Private/Scene/MeshSceneAdapter.cpp
+++ b/Engine/Plugins/Runtime/MeshModelingToolset/Source/ModelingComponents/Private/Scene/MeshSceneAdapter.cpp
@@ -1150,7 +1150,7 @@ class FStaticMeshRenderDataSpatialWrapper : public TStaticMeshSpatialWrapperBase
 public:
 	virtual TUniquePtr<FStaticMeshLODResourcesMeshSurfaceAdapter> MakeMeshAdapter(const FMeshSceneAdapterBuildOptions& BuildOptions) override
 	{
-		return MakeUnique<FStaticMeshLODResourcesMeshSurfaceAdapter>(&StaticMesh->GetRenderData()->LODResources[0], StaticMesh, BuildOptions.bOnlySurfaceMaterials);
+		return MakeUnique<FStaticMeshLODResourcesMeshSurfaceAdapter>(StaticMesh->GetRenderData()->GetCurrentFirstLOD(0), StaticMesh, BuildOptions.bOnlySurfaceMaterials);
 	}
 };
 
diff --git a/Engine/Plugins/Runtime/MeshModelingToolset/Source/ModelingComponents/Private/Snapping/ModelingSceneSnappingManager.cpp b/Engine/Plugins/Runtime/MeshModelingToolset/Source/ModelingComponents/Private/Snapping/ModelingSceneSnappingManager.cpp
--- a/Engine/Plugins/Runtime/MeshModelingToolset/Source/ModelingComponents/Private/Snapping/ModelingSceneSnappingManager.cpp
+++ b/Engine/Plugins/Runtime/MeshModelingToolset/Source/ModelingComponents/Private/Snapping/ModelingSceneSnappingManager.cpp
@@ -370,12 +370,22 @@ bool GetComponentHitTriangle_Internal(FHitResult HitResult, VectorType* TriVerti
 	}
 	
 	// physics collision data is created from StaticMesh RenderData
-	// so use HitResult.FaceIndex to extract triangle from the LOD0 mesh
+	// so use HitResult.FaceIndex to extract triangle from the mesh
 	// (note: this may be incorrect if there are multiple sections...in that case I think we have to
 	//  first find section whose accumulated index range would contain .FaceIndexX)
 	UStaticMesh* StaticMesh = Component->GetStaticMesh();
-	FStaticMeshLODResources& LOD = StaticMesh->GetRenderData()->LODResources[0];
-	FIndexArrayView Indices = LOD.IndexBuffer.GetArrayView();
+
+	// Get the lowest available render LOD.
+	// If the Minimum LOD index is not zero, their might be a lower LOD available, but it will not be used for rendering in the viewport.
+	const int32 MinLODIdx = StaticMesh->GetMinLODIdx();
+	const FStaticMeshLODResources* LOD = StaticMesh->GetRenderData()->GetCurrentFirstLOD(MinLODIdx);
+	if (!LOD)
+	{
+		return false;
+	}
+
+	const FIndexArrayView Indices = LOD->IndexBuffer.GetArrayView();
+	const FPositionVertexBuffer& Vertices = LOD->VertexBuffers.PositionVertexBuffer;
 
 	const int32 TriIdx = 3 * HitResult.FaceIndex;
 	if (TriIdx + 2 >= Indices.Num())
@@ -384,16 +394,16 @@ bool GetComponentHitTriangle_Internal(FHitResult HitResult, VectorType* TriVerti
 	}
 
 	const uint32 Idx[3] = { Indices[TriIdx], Indices[TriIdx + 1], Indices[TriIdx + 2] };
-	const uint32 NumVertices = LOD.VertexBuffers.PositionVertexBuffer.GetNumVertices();
+	const uint32 NumVertices = Vertices.GetNumVertices();
 	if (Idx[0] >= NumVertices || Idx[1] >= NumVertices || Idx[2] >= NumVertices)
 	{
 		return false;
 	}
 
 	const FVector Positions[3] = {
-		static_cast<FVector>(LOD.VertexBuffers.PositionVertexBuffer.VertexPosition(Idx[0])),
-		static_cast<FVector>(LOD.VertexBuffers.PositionVertexBuffer.VertexPosition(Idx[1])),
-		static_cast<FVector>(LOD.VertexBuffers.PositionVertexBuffer.VertexPosition(Idx[2]))
+		static_cast<FVector>(Vertices.VertexPosition(Idx[0])),
+		static_cast<FVector>(Vertices.VertexPosition(Idx[1])),
+		static_cast<FVector>(Vertices.VertexPosition(Idx[2]))
 	};
 
 	// transform to world space
diff --git a/Engine/Source/Editor/UnrealEd/Private/VertexSnapping.cpp b/Engine/Source/Editor/UnrealEd/Private/VertexSnapping.cpp
--- a/Engine/Source/Editor/UnrealEd/Private/VertexSnapping.cpp
+++ b/Engine/Source/Editor/UnrealEd/Private/VertexSnapping.cpp
@@ -78,8 +78,8 @@ public:
 	FStaticMeshVertexIterator( UStaticMeshComponent* SMC )
 		: ComponentToWorldIT( SMC->GetComponentTransform().ToInverseMatrixWithScale().GetTransposed() )
 		, StaticMeshComponent( SMC )
-		, PositionBuffer( SMC->GetStaticMesh()->GetRenderData()->LODResources[0].VertexBuffers.PositionVertexBuffer )
-		, VertexBuffer( SMC->GetStaticMesh()->GetRenderData()->LODResources[0].VertexBuffers.StaticMeshVertexBuffer )
+		, PositionBuffer( SMC->GetStaticMesh()->GetRenderData()->GetCurrentFirstLOD(SMC->GetStaticMesh()->GetMinLODIdx())->VertexBuffers.PositionVertexBuffer )
+		, VertexBuffer( SMC->GetStaticMesh()->GetRenderData()->GetCurrentFirstLOD(SMC->GetStaticMesh()->GetMinLODIdx())->VertexBuffers.StaticMeshVertexBuffer )
 		, CurrentVertexIndex( 0 )
 	{
 
@@ -110,11 +110,11 @@ private:
 	/** Component To World Inverse Transpose matrix */
 	FMatrix ComponentToWorldIT;
 	/** Component containing the mesh that we are getting vertices from */
-	UStaticMeshComponent* StaticMeshComponent;
+	const UStaticMeshComponent* StaticMeshComponent;
 	/** The static meshes position vertex buffer */
-	FPositionVertexBuffer& PositionBuffer;
+	const FPositionVertexBuffer& PositionBuffer;
 	/** The static meshes vertex buffer for normals */
-	FStaticMeshVertexBuffer& VertexBuffer;
+	const FStaticMeshVertexBuffer& VertexBuffer;
 	/** Current vertex index */
 	uint32 CurrentVertexIndex;
 };
diff --git a/Engine/Source/Runtime/Engine/Private/Components/StaticMeshComponent.cpp b/Engine/Source/Runtime/Engine/Private/Components/StaticMeshComponent.cpp
--- a/Engine/Source/Runtime/Engine/Private/Components/StaticMeshComponent.cpp
+++ b/Engine/Source/Runtime/Engine/Private/Components/StaticMeshComponent.cpp
@@ -3027,12 +3027,12 @@ void UStaticMeshComponent::PostStaticMeshCompilation()
 namespace ComponentIsTouchingSelectionHelpers
 {
 
-const FStaticMeshLODResources& GetMinLODModel(const UStaticMesh& StaticMesh)
+const FStaticMeshLODResources* GetRenderLOD(const UStaticMesh& StaticMesh)
 {
+	// Get the lowest available render LOD.
+	// If the Minimum LOD index is not zero, their might be a lower LOD available, but it will not be used for rendering in the viewport.
 	const int32 MinLODIdx = StaticMesh.GetMinLODIdx();
-	const FStaticMeshRenderData* RenderData = StaticMesh.GetRenderData();
-	const int32 LODIdx = RenderData->GetCurrentFirstLODIdx(MinLODIdx);
-	return RenderData->LODResources[LODIdx];
+	return StaticMesh.GetRenderData()->GetCurrentFirstLOD(MinLODIdx);
 }
 
 enum class ECheckSectionBoundsResult
@@ -3082,12 +3082,16 @@ bool UStaticMeshComponent::ComponentIsTouchingSelectionBox(const FBox& InSelBBox
 			TArray<FVector> TriVertices;
 			TriVertices.SetNumUninitialized(3);
 
-			const FStaticMeshLODResources& LODModel = ComponentIsTouchingSelectionHelpers::GetMinLODModel(*GetStaticMesh());
+			const FStaticMeshLODResources* LOD = ComponentIsTouchingSelectionHelpers::GetRenderLOD(*GetStaticMesh());
+			if (!LOD)
+			{
+				return false;
+			}
 
-			const FIndexArrayView Indices = LODModel.IndexBuffer.GetArrayView();
-			const FPositionVertexBuffer& Vertices = LODModel.VertexBuffers.PositionVertexBuffer;
+			const FIndexArrayView Indices = LOD->IndexBuffer.GetArrayView();
+			const FPositionVertexBuffer& Vertices = LOD->VertexBuffers.PositionVertexBuffer;
 
-			for (const FStaticMeshSection& Section : LODModel.Sections)
+			for (const FStaticMeshSection& Section : LOD->Sections)
 			{
 				switch (ComponentIsTouchingSelectionHelpers::CheckSectionBounds(Section, Indices, Vertices, bMustEncompassEntireComponent))
 				{
@@ -3150,13 +3154,17 @@ bool UStaticMeshComponent::ComponentIsTouchingSelectionFrustum(const FConvexVolu
 				return true;
 			}
 
-			const FStaticMeshLODResources& LODModel = ComponentIsTouchingSelectionHelpers::GetMinLODModel(*GetStaticMesh());
+			const FStaticMeshLODResources* LOD = ComponentIsTouchingSelectionHelpers::GetRenderLOD(*GetStaticMesh());
+			if (!LOD)
+			{
+				return false;
+			}
 
-			const FIndexArrayView Indices = LODModel.IndexBuffer.GetArrayView();
-			const FPositionVertexBuffer& Vertices = LODModel.VertexBuffers.PositionVertexBuffer;
+			const FIndexArrayView Indices = LOD->IndexBuffer.GetArrayView();
+			const FPositionVertexBuffer& Vertices = LOD->VertexBuffers.PositionVertexBuffer;
 
 			const FTransform& TransformToWorld = GetComponentTransform();
-			for (const FStaticMeshSection& Section : LODModel.Sections)
+			for (const FStaticMeshSection& Section : LOD->Sections)
 			{
 				switch (ComponentIsTouchingSelectionHelpers::CheckSectionBounds(Section, Indices, Vertices, bMustEncompassEntireComponent))
 				{
diff --git a/Engine/Source/Runtime/Engine/Private/InstancedStaticMesh.cpp b/Engine/Source/Runtime/Engine/Private/InstancedStaticMesh.cpp
--- a/Engine/Source/Runtime/Engine/Private/InstancedStaticMesh.cpp
+++ b/Engine/Source/Runtime/Engine/Private/InstancedStaticMesh.cpp
@@ -4487,10 +4487,17 @@ bool UInstancedStaticMeshComponent::IsInstanceTouchingSelectionBox(int32 Instanc
 			{
 				TArray<FVector> Vertex;
 
-				FStaticMeshLODResources& LODModel = GetStaticMesh()->GetRenderData()->LODResources[0];
-				FIndexArrayView Indices = LODModel.IndexBuffer.GetArrayView();
+				const int32 MinLODIdx = GetStaticMesh()->GetMinLODIdx();
+				const FStaticMeshLODResources* LODModel = GetStaticMesh()->GetRenderData()->GetCurrentFirstLOD(MinLODIdx);
+				if (!LODModel)
+				{
+					return false;
+				}
+
+				const FIndexArrayView Indices = LODModel->IndexBuffer.GetArrayView();
+				const FPositionVertexBuffer& Vertices = LODModel->VertexBuffers.PositionVertexBuffer;
 
-				for (const FStaticMeshSection& Section : LODModel.Sections)
+				for (const FStaticMeshSection& Section : LODModel->Sections)
 				{
 					// Iterate over each triangle.
 					for (int32 TriangleIndex = 0; TriangleIndex < (int32)Section.NumTriangles; TriangleIndex++)
@@ -4501,7 +4508,7 @@ bool UInstancedStaticMeshComponent::IsInstanceTouchingSelectionBox(int32 Instanc
 						for (int32 i = 0; i < 3; i++)
 						{
 							int32 VertexIndex = Indices[FirstIndex + i];
-							FVector LocalPosition(LODModel.VertexBuffers.PositionVertexBuffer.VertexPosition(VertexIndex));
+							FVector LocalPosition(Vertices.VertexPosition(VertexIndex));
 							Vertex.Emplace(InstanceTransfrom.TransformPosition(LocalPosition));
 						}
 
@@ -4556,19 +4563,26 @@ bool UInstancedStaticMeshComponent::IsInstanceTouchingSelectionFrustum(int32 Ins
 					return true;
 				}
 
-				FStaticMeshLODResources& LODModel = GetStaticMesh()->GetRenderData()->LODResources[0];
-				FIndexArrayView Indices = LODModel.IndexBuffer.GetArrayView();
+				const int32 MinLODIdx = GetStaticMesh()->GetMinLODIdx();
+				const FStaticMeshLODResources* LODModel = GetStaticMesh()->GetRenderData()->GetCurrentFirstLOD(MinLODIdx);
+				if (!LODModel)
+				{
+					return false;
+				}
+
+				const FIndexArrayView Indices = LODModel->IndexBuffer.GetArrayView();
+				const FPositionVertexBuffer& Vertices = LODModel->VertexBuffers.PositionVertexBuffer;
 
-				for (const FStaticMeshSection& Section : LODModel.Sections)
+				for (const FStaticMeshSection& Section : LODModel->Sections)
 				{
 					// Iterate over each triangle.
 					for (int32 TriangleIndex = 0; TriangleIndex < (int32)Section.NumTriangles; TriangleIndex++)
 					{
 						int32 Index = TriangleIndex * 3 + Section.FirstIndex;
 
-						FVector PointA(LODModel.VertexBuffers.PositionVertexBuffer.VertexPosition(Indices[Index]));
-						FVector PointB(LODModel.VertexBuffers.PositionVertexBuffer.VertexPosition(Indices[Index + 1]));
-						FVector PointC(LODModel.VertexBuffers.PositionVertexBuffer.VertexPosition(Indices[Index + 2]));
+						FVector PointA(Vertices.VertexPosition(Indices[Index]));
+						FVector PointB(Vertices.VertexPosition(Indices[Index + 1]));
+						FVector PointC(Vertices.VertexPosition(Indices[Index + 2]));
 
 						PointA = InstanceTransfrom.TransformPosition(PointA);
 						PointB = InstanceTransfrom.TransformPosition(PointB);
