Skeletal Mesh - Fixed the sequencer ability to change the SkeletalMesh asset from a Skeletal Mesh Component in cooked builds.

diff --git a/Engine/Source/Runtime/Engine/Classes/Components/SkeletalMeshComponent.h b/Engine/Source/Runtime/Engine/Classes/Components/SkeletalMeshComponent.h
--- a/Engine/Source/Runtime/Engine/Classes/Components/SkeletalMeshComponent.h
+++ b/Engine/Source/Runtime/Engine/Classes/Components/SkeletalMeshComponent.h
@@ -316,26 +316,26 @@ class ENGINE_API USkeletalMeshComponent : public USkinnedMeshComponent, public I
 	friend struct FLinkedInstancesAdapter;
 	friend struct FLinkedAnimLayerClassData;
 	
-#if WITH_EDITORONLY_DATA
+//#if WITH_EDITORONLY_DATA  // TODO: Re-add these guards once the MovieScene getters/setters are working, so that we can get rid of this redundant pointer in all cooked builds
 private:
 	/** The skeletal mesh used by this component. */
 	UE_DEPRECATED(5.1, "This property isn't deprecated, but getter and setter must be used at all times to preserve correct operations.")
-	UPROPERTY(EditAnywhere, Transient, BlueprintSetter = SetSkeletalMeshAsset, BlueprintGetter = GetSkeletalMeshAsset, Category = Mesh)
-	TObjectPtr<class USkeletalMesh> SkeletalMeshAsset;
-#endif
+	UPROPERTY(EditAnywhere, Transient, Setter = SetSkeletalMeshAsset, BlueprintSetter = SetSkeletalMeshAsset, Getter = GetSkeletalMeshAsset, BlueprintGetter = GetSkeletalMeshAsset, Category = Mesh)
+	TObjectPtr<USkeletalMesh> SkeletalMeshAsset;
+//#endif
 
 public:
 	/**
 	 * Set the SkeletalMesh rendered for this mesh.
 	 */
 	UFUNCTION(BlueprintCallable, Category = "Components|SkeletalMesh")
-	void SetSkeletalMeshAsset(class USkeletalMesh* NewMesh) { SetSkeletalMesh(NewMesh, false); }
+	void SetSkeletalMeshAsset(USkeletalMesh* NewMesh) { SetSkeletalMesh(NewMesh, false); }
 
 	/**
 	 * Get the SkeletalMesh rendered for this mesh.
 	 */
 	UFUNCTION(BlueprintPure, Category = "Components|SkeletalMesh")
-	class USkeletalMesh* GetSkeletalMeshAsset() const;
+	USkeletalMesh* GetSkeletalMeshAsset() const;
 
 #if WITH_EDITORONLY_DATA
 	/** The blueprint for creating an AnimationScript. */
diff --git a/Engine/Source/Runtime/Engine/Private/Components/SkeletalMeshComponent.cpp b/Engine/Source/Runtime/Engine/Private/Components/SkeletalMeshComponent.cpp
--- a/Engine/Source/Runtime/Engine/Private/Components/SkeletalMeshComponent.cpp
+++ b/Engine/Source/Runtime/Engine/Private/Components/SkeletalMeshComponent.cpp
@@ -382,12 +382,12 @@ void USkeletalMeshComponent::PostLoad()
 		bOverrideMinLod = true;
 	}
 
-#if WITH_EDITORONLY_DATA
+//#if WITH_EDITORONLY_DATA  // TODO: Re-add these guards once the MovieScene getters/setters are working, so that we can get rid of this redundant pointer in all cooked builds
 	// Update property alias
 	PRAGMA_DISABLE_DEPRECATION_WARNINGS
 	SkeletalMeshAsset = GetSkeletalMeshAsset();
 	PRAGMA_ENABLE_DEPRECATION_WARNINGS
-#endif
+//#endif
 }
 
 void USkeletalMeshComponent::PostInitProperties()
@@ -3005,11 +3005,11 @@ void USkeletalMeshComponent::SetSkeletalMesh(USkeletalMesh* InSkelMesh, bool bRe
 	}
 
 	// Update property alias
-#if WITH_EDITORONLY_DATA
+//#if WITH_EDITORONLY_DATA  // TODO: Re-add these guards once the MovieScene getters/setters are working, so that we can get rid of this redundant pointer in all cooked builds
 	PRAGMA_DISABLE_DEPRECATION_WARNINGS
 	SkeletalMeshAsset = InSkelMesh;
 	PRAGMA_ENABLE_DEPRECATION_WARNINGS
-#endif
+//#endif
 
 	// We may be doing parallel evaluation on the current anim instance
 	// Calling this here with true will block this init till that thread completes
