Resubmit 24403180 with resolving fix

Ground work for Private DataLayerAssets:
Change TObjectPtr<UDataLayerAsset> to TSoftObjectPtr<UDataLAyerAsset>

diff --git a/Engine/Source/Editor/UnrealEd/Private/Commandlets/WorldPartitionDataLayerToAssetCommandlet.cpp b/Engine/Source/Editor/UnrealEd/Private/Commandlets/WorldPartitionDataLayerToAssetCommandlet.cpp
--- a/Engine/Source/Editor/UnrealEd/Private/Commandlets/WorldPartitionDataLayerToAssetCommandlet.cpp
+++ b/Engine/Source/Editor/UnrealEd/Private/Commandlets/WorldPartitionDataLayerToAssetCommandlet.cpp
@@ -573,7 +573,7 @@ uint32 UDataLayerToAssetCommandlet::RemapDataLayersAssetsFromPreviousConversions
 {
 	uint32 ErrorCount = 0;
 
-	const TArray<TObjectPtr<const UDataLayerAsset>>& ActorDataLayerAssets = Actor->GetDataLayerAssets();
+	TArray<const UDataLayerAsset*> ActorDataLayerAssets = Actor->GetDataLayerAssets();
 	for (int32 i = ActorDataLayerAssets.Num() - 1; i >= 0; --i)
 	{
 		const TObjectPtr<const UDataLayerAsset>& ActorDataLayerAsset = ActorDataLayerAssets[i];
diff --git a/Engine/Source/Runtime/Engine/Classes/GameFramework/Actor.h b/Engine/Source/Runtime/Engine/Classes/GameFramework/Actor.h
--- a/Engine/Source/Runtime/Engine/Classes/GameFramework/Actor.h
+++ b/Engine/Source/Runtime/Engine/Classes/GameFramework/Actor.h
@@ -910,9 +910,9 @@ protected:
 	TArray<FActorDataLayer> DataLayers;
 
 	UPROPERTY(EditAnywhere, AdvancedDisplay, Category = DataLayers)
-	TArray<TObjectPtr<const UDataLayerAsset>> DataLayerAssets;
+	TArray<TSoftObjectPtr<UDataLayerAsset>> DataLayerAssets;
 
-	TArray<TObjectPtr<const UDataLayerAsset>> PreEditChangeDataLayers;
+	TArray<TSoftObjectPtr<UDataLayerAsset>> PreEditChangeDataLayers;
 
 public:
 	/** The copy/paste id used to remap actors during copy operations */
@@ -1251,6 +1251,7 @@ public:
 private:
 	virtual bool ActorTypeSupportsDataLayer() const { return true; }
 	virtual bool ActorTypeShouldSkipFromLevelInstance() const { return false; }
+	TArray<const UDataLayerAsset*> ResolveDataLayerAssets(const TArray<TSoftObjectPtr<UDataLayerAsset>>& InDataLayerAssets) const;
 public:
 	bool AddDataLayer(const UDataLayerInstance* DataLayerInstance);
 	bool RemoveDataLayer(const UDataLayerInstance* DataLayerInstance);
@@ -1264,7 +1265,7 @@ public:
 	static const FName GetDataLayerAssetsPropertyName() { return GET_MEMBER_NAME_CHECKED(AActor, DataLayerAssets); }
 	static const FName GetDataLayerPropertyName() { return GET_MEMBER_NAME_CHECKED(AActor, DataLayers); }
 
-	const TArray<TObjectPtr<const UDataLayerAsset>>& GetDataLayerAssets() const { return DataLayerAssets; }
+	TArray<const UDataLayerAsset*> GetDataLayerAssets() const;
 
 	//~ Begin Deprecated
 
diff --git a/Engine/Source/Runtime/Engine/Private/Actor.cpp b/Engine/Source/Runtime/Engine/Private/Actor.cpp
--- a/Engine/Source/Runtime/Engine/Private/Actor.cpp
+++ b/Engine/Source/Runtime/Engine/Private/Actor.cpp
@@ -466,6 +466,7 @@ void AActor::PreSave(FObjectPreSaveContext ObjectSaveContext)
 {
 	Super::PreSave(ObjectSaveContext);
 #if WITH_EDITOR
+	// Always call FixupDataLayers() on PreSave so that DataLayerAssets SoftObjectPtrs internal WeakPtrs get resolved before GIsSavingPackage gets set to true preventing resolving.
 	FixupDataLayers();
 #endif
 }
@@ -6161,7 +6162,7 @@ TArray<const UDataLayerInstance*> AActor::GetDataLayerInstancesInternal(bool bUs
 		TArray<const UDataLayerInstance*> DataLayerInstances;
 		if (UDataLayerManager* DataLayerManager = bUseLevelContext ? UDataLayerManager::GetDataLayerManager(this) : UDataLayerManager::GetDataLayerManager(GetWorld()))
 		{
-			DataLayerInstances += DataLayerManager->GetDataLayerInstances(DataLayerAssets);
+			DataLayerInstances += DataLayerManager->GetDataLayerInstances(GetDataLayerAssets());
 
 			PRAGMA_DISABLE_DEPRECATION_WARNINGS
 			DataLayerInstances += DataLayerManager->GetDataLayerInstances(DataLayers);
diff --git a/Engine/Source/Runtime/Engine/Private/ActorEditor.cpp b/Engine/Source/Runtime/Engine/Private/ActorEditor.cpp
--- a/Engine/Source/Runtime/Engine/Private/ActorEditor.cpp
+++ b/Engine/Source/Runtime/Engine/Private/ActorEditor.cpp
@@ -32,6 +32,9 @@
 #include "Modules/ModuleManager.h"
 #include "Engine/SimpleConstructionScript.h"
 #include "StaticMeshResources.h"
+#include "WorldPartition/DataLayer/DataLayerAsset.h"
+#include "UObject/SoftObjectPtr.h"
+#include "Algo/Transform.h"
 
 #define LOCTEXT_NAMESPACE "ErrorChecking"
 
@@ -1621,6 +1624,26 @@ bool AActor::RemoveAllDataLayers()
 	return false;
 }
 
+TArray<const UDataLayerAsset*> AActor::ResolveDataLayerAssets(const TArray<TSoftObjectPtr<UDataLayerAsset>>& InDataLayerAssets) const
+{
+	TArray<const UDataLayerAsset*> ResolvedAssets;
+	ResolvedAssets.Reserve(InDataLayerAssets.Num());
+	Algo::TransformIf(InDataLayerAssets, ResolvedAssets, [](const TSoftObjectPtr<UDataLayerAsset>& DataLayerAsset) { return DataLayerAsset.IsValid(); }, [](const TSoftObjectPtr<UDataLayerAsset>& DataLayerAsset) { return DataLayerAsset.Get(); });
+	return ResolvedAssets;
+}
+
+TArray<const UDataLayerAsset*> AActor::GetDataLayerAssets() const
+{
+	TArray<const UDataLayerAsset*> Resolved = ResolveDataLayerAssets(DataLayerAssets);
+
+	TGuardValue<bool> GIsSavingPackageGuard(GIsSavingPackage, false);
+	TArray<const UDataLayerAsset*> ResolvedNoSaving = ResolveDataLayerAssets(DataLayerAssets);
+
+	check(Resolved == ResolvedNoSaving);
+
+	return Resolved;
+}
+
 TArray<FName> AActor::GetDataLayerInstanceNames() const
 {
 	TArray<FName> DataLayerInstanceNames;
@@ -1641,6 +1664,9 @@ TArray<const UDataLayerInstance*> AActor::GetDataLayerInstancesForLevel() const
 
 void AActor::FixupDataLayers(bool bRevertChangesOnLockedDataLayer /*= false*/)
 {
+	// Always call here because this gets called in AActor::Presave and makes sure the SoftObjectPtrs are resolved before saving.
+	TArray<const UDataLayerAsset*> ResolvedAssets = ResolveDataLayerAssets(DataLayerAssets);
+
 	if (IsTemplate())
 	{
 		return;
@@ -1676,9 +1702,10 @@ void AActor::FixupDataLayers(bool bRevertChangesOnLockedDataLayer /*= false*/)
 	
 	if (bRevertChangesOnLockedDataLayer)
 	{
+		TArray<const UDataLayerAsset*> ResolvedPreEditChangeAssets = ResolveDataLayerAssets(PreEditChangeDataLayers);
 		// Since it's not possible to prevent changes of particular elements of an array, rollback change on locked DataLayers.
-		TSet<const UDataLayerAsset*> PreEdit(PreEditChangeDataLayers);
-		TSet<const UDataLayerAsset*> PostEdit(DataLayerAssets);
+		TSet<const UDataLayerAsset*> PreEdit(ResolvedPreEditChangeAssets);
+		TSet<const UDataLayerAsset*> PostEdit(ResolvedAssets);
 
 		auto DifferenceContainsLockedDataLayers = [DataLayerManager](const TSet<const UDataLayerAsset*>& A, const TSet<const UDataLayerAsset*>& B)
 		{
@@ -1725,7 +1752,11 @@ void AActor::FixupDataLayers(bool bRevertChangesOnLockedDataLayer /*= false*/)
 	// Only invalidate DataLayerAssets on cook. In Editor we want to be able to re-resolve if the asset gets readded to the WorldDataLayers actor
 	if (IsRunningCookCommandlet())
 	{
-		CleanupDataLayers(DataLayerAssets);
+		// Get a new resolved array in case DataLayerAssets changed
+		TArray<const UDataLayerAsset*> CleanedUpDataLayers = ResolveDataLayerAssets(DataLayerAssets);
+		CleanupDataLayers(CleanedUpDataLayers);
+		DataLayerAssets.Empty(CleanedUpDataLayers.Num());
+		Algo::Transform(CleanedUpDataLayers, DataLayerAssets, [](const UDataLayerAsset* DataLayerAsset) { return TSoftObjectPtr<UDataLayerAsset>(DataLayerAsset); });
 	}
 	CleanupDataLayers(DataLayers);
 	PRAGMA_ENABLE_DEPRECATION_WARNINGS
diff --git a/Engine/Source/Runtime/Engine/Private/LevelInstance/LevelInstanceSubsystem.cpp b/Engine/Source/Runtime/Engine/Private/LevelInstance/LevelInstanceSubsystem.cpp
--- a/Engine/Source/Runtime/Engine/Private/LevelInstance/LevelInstanceSubsystem.cpp
+++ b/Engine/Source/Runtime/Engine/Private/LevelInstance/LevelInstanceSubsystem.cpp
@@ -1051,7 +1051,7 @@ ILevelInstanceInterface* ULevelInstanceSubsystem::CreateLevelInstanceFrom(const
 					if (UDataLayerManager* DataLayerManager = UDataLayerManager::GetDataLayerManager(ActorToMove))
 					{
 						// Use the raw asset list as we don't want parent DataLayers
-						for (const TObjectPtr<const UDataLayerAsset>& DataLayerAsset : ActorToMove->GetDataLayerAssets())
+						for (const UDataLayerAsset* DataLayerAsset : ActorToMove->GetDataLayerAssets())
 						{
 							// Validate that there's a valid Data Layer Instance for this asset in the source level
 							if (DataLayerManager->GetDataLayerInstanceFromAsset(DataLayerAsset))
diff --git a/Engine/Source/Runtime/Engine/Private/WorldPartition/WorldPartitionActorDesc.cpp b/Engine/Source/Runtime/Engine/Private/WorldPartition/WorldPartitionActorDesc.cpp
--- a/Engine/Source/Runtime/Engine/Private/WorldPartition/WorldPartitionActorDesc.cpp
+++ b/Engine/Source/Runtime/Engine/Private/WorldPartition/WorldPartitionActorDesc.cpp
@@ -100,8 +100,9 @@ void FWorldPartitionActorDesc::Init(const AActor* InActor)
 
 		if (UWorldPartition* ActorWorldPartition = FWorldPartitionHelpers::GetWorldPartition(InActor))
 		{
-			LocalDataLayerAssetPaths.Reserve(InActor->GetDataLayerAssets().Num());
-			for (const TObjectPtr<const UDataLayerAsset>& DataLayerAsset : InActor->GetDataLayerAssets())
+			TArray<const UDataLayerAsset*> DataLayerAssets = InActor->GetDataLayerAssets();
+			LocalDataLayerAssetPaths.Reserve(DataLayerAssets.Num());
+			for (const UDataLayerAsset* DataLayerAsset : DataLayerAssets)
 			{
 				if (DataLayerAsset)
 				{
