Remove submitted debug code

#rnx

diff --git a/Engine/Source/Runtime/Engine/Private/ActorEditor.cpp b/Engine/Source/Runtime/Engine/Private/ActorEditor.cpp
--- a/Engine/Source/Runtime/Engine/Private/ActorEditor.cpp
+++ b/Engine/Source/Runtime/Engine/Private/ActorEditor.cpp
@@ -1634,14 +1634,7 @@ TArray<const UDataLayerAsset*> AActor::ResolveDataLayerAssets(const TArray<TSoft
 
 TArray<const UDataLayerAsset*> AActor::GetDataLayerAssets() const
 {
-	TArray<const UDataLayerAsset*> Resolved = ResolveDataLayerAssets(DataLayerAssets);
-
-	TGuardValue<bool> GIsSavingPackageGuard(GIsSavingPackage, false);
-	TArray<const UDataLayerAsset*> ResolvedNoSaving = ResolveDataLayerAssets(DataLayerAssets);
-
-	check(Resolved == ResolvedNoSaving);
-
-	return Resolved;
+	return ResolveDataLayerAssets(DataLayerAssets);
 }
 
 TArray<FName> AActor::GetDataLayerInstanceNames() const
