Enable easy command line override of shared linear encoding

diff --git a/Engine/Source/Runtime/Engine/Private/TextureEncodingSettings.cpp b/Engine/Source/Runtime/Engine/Private/TextureEncodingSettings.cpp
--- a/Engine/Source/Runtime/Engine/Private/TextureEncodingSettings.cpp
+++ b/Engine/Source/Runtime/Engine/Private/TextureEncodingSettings.cpp
@@ -68,6 +68,14 @@ static void ConstructResolvedSettings(ReadyMask InReadyMask, UDeveloperSettings*
 	OutResolvedSettings->User.ForceEncodeSpeed = UserSettings->ForceEncodeSpeed;
 
 	OutResolvedSettings->Project.bSharedLinearTextureEncoding = ProjectSettings->bSharedLinearTextureEncoding;
+	{
+		bool CmdLineSLE;
+		if (FParse::Bool(FCommandLine::Get(), TEXT("-ForceSharedLinearTextureEncoding="), CmdLineSLE))
+		{
+			OutResolvedSettings->Project.bSharedLinearTextureEncoding = CmdLineSLE;
+			UE_LOG(LogTextureEncodingSettings, Display, TEXT("Shared linear texture encoding forced to %s via command line."), *LexToString(CmdLineSLE));
+		}
+	}
 	OutResolvedSettings->Project.bFinalUsesRDO = ProjectSettings->bFinalUsesRDO;
 	OutResolvedSettings->Project.bFastUsesRDO = ProjectSettings->bFastUsesRDO;
 	OutResolvedSettings->Project.FinalRDOLambda = ProjectSettings->FinalRDOLambda;
