Speculative fix for hang in modulation destinations
#rnx

diff --git a/Engine/Source/Runtime/Engine/Private/SoundModulationDestination.cpp b/Engine/Source/Runtime/Engine/Private/SoundModulationDestination.cpp
--- a/Engine/Source/Runtime/Engine/Private/SoundModulationDestination.cpp
+++ b/Engine/Source/Runtime/Engine/Private/SoundModulationDestination.cpp
@@ -81,10 +81,15 @@ namespace Audio
 		bValueNormalized = InModulationDestination.bValueNormalized;
 		OutputBuffer = InModulationDestination.OutputBuffer;
 
+		TSet<FModulatorHandle> NewHandles;
 		{
 			FScopeLock OtherLock(&InModulationDestination.HandleCritSection);
+			NewHandles = InModulationDestination.Handles;
+		}
+
+		{
 			FScopeLock Lock(&HandleCritSection);
-			Handles = InModulationDestination.Handles;
+			Handles = MoveTemp(NewHandles);
 		}
 
 		Parameter = InModulationDestination.Parameter;
@@ -100,10 +105,15 @@ namespace Audio
 		bValueNormalized = MoveTemp(InModulationDestination.bValueNormalized);
 		bHasProcessed = MoveTemp(InModulationDestination.bHasProcessed);
 		OutputBuffer = MoveTemp(InModulationDestination.OutputBuffer);
+
+		TSet<FModulatorHandle> NewHandles;
 		{
 			FScopeLock OtherLock(&InModulationDestination.HandleCritSection);
+			NewHandles = MoveTemp(InModulationDestination.Handles);
+		}
+		{
 			FScopeLock Lock(&HandleCritSection);
-			Handles = MoveTemp(InModulationDestination.Handles);
+			Handles = MoveTemp(NewHandles);
 		}
 
 		Parameter = MoveTemp(InModulationDestination.Parameter);
