WorldPartitionActorFilter: Fix regression since partially loaded Level Instances are loaded by main world loader. Find world through Actor handle instead.

#rnx

diff --git a/Engine/Source/Runtime/Engine/Private/LevelInstance/LevelInstanceSubsystem.cpp b/Engine/Source/Runtime/Engine/Private/LevelInstance/LevelInstanceSubsystem.cpp
--- a/Engine/Source/Runtime/Engine/Private/LevelInstance/LevelInstanceSubsystem.cpp
+++ b/Engine/Source/Runtime/Engine/Private/LevelInstance/LevelInstanceSubsystem.cpp
@@ -2409,7 +2409,9 @@ bool ULevelInstanceSubsystem::CanUseWorldAsset(const ILevelInstanceInterface* Le
 
 bool ULevelInstanceSubsystem::PassLevelInstanceFilter(UWorld* World, const FWorldPartitionHandle& Actor) const
 {
-	if (const ILevelInstanceInterface* TopAncestor = GetOwningLevelInstance(World->PersistentLevel))
+	UWorld* ContainerOuterWorld = Actor->GetContainer()->GetWorldPartition()->GetTypedOuter<UWorld>();
+	check(ContainerOuterWorld);
+	if (const ILevelInstanceInterface* TopAncestor = GetOwningLevelInstance(ContainerOuterWorld->PersistentLevel))
 	{
 		FActorContainerID ContainerID = TopAncestor->GetLevelInstanceID().GetContainerID();
 		ForEachLevelInstanceAncestors(Cast<AActor>(TopAncestor), [&TopAncestor](const ILevelInstanceInterface* Ancestor)
