Actor Instance Guids: actor instance guids are now always available in AActor::RegisterAllComponents.

diff --git a/Engine/Source/Editor/DetailCustomizations/Private/ActorDetails.cpp b/Engine/Source/Editor/DetailCustomizations/Private/ActorDetails.cpp
--- a/Engine/Source/Editor/DetailCustomizations/Private/ActorDetails.cpp
+++ b/Engine/Source/Editor/DetailCustomizations/Private/ActorDetails.cpp
@@ -614,6 +614,26 @@ void FActorDetails::AddActorCategory( IDetailLayoutBuilder& DetailBuilder, const
 							.Font(IDetailLayoutBuilder::GetDetailFont())
 							.IsEnabled(false)
 					];
+
+				if (Actor->GetActorInstanceGuid() != Actor->GetActorGuid())
+				{
+					const FText ActorInstanceGuidText = FText::FromString(Actor->GetActorInstanceGuid().ToString());
+					ActorCategory.AddCustomRow( LOCTEXT("ActorInstanceGuid", "ActorInstanceGuid") )
+						.NameContent()
+						[
+							SNew(STextBlock)
+							.Text(LOCTEXT("ActorInstanceGuid2", "Actor Instance Guid"))
+							.ToolTipText(LOCTEXT("ActorInstanceGuid_ToolTip", "Actor Instance Guid"))
+							.Font(IDetailLayoutBuilder::GetDetailFont())
+						]
+						.ValueContent()
+						[
+							SNew(STextBlock)
+								.Text(ActorInstanceGuidText)
+								.Font(IDetailLayoutBuilder::GetDetailFont())
+								.IsEnabled(false)
+						];
+				}
 			}
 
 			if (Actor->GetContentBundleGuid().IsValid())
diff --git a/Engine/Source/Editor/UnrealEd/Private/EditorEngine.cpp b/Engine/Source/Editor/UnrealEd/Private/EditorEngine.cpp
--- a/Engine/Source/Editor/UnrealEd/Private/EditorEngine.cpp
+++ b/Engine/Source/Editor/UnrealEd/Private/EditorEngine.cpp
@@ -279,80 +279,6 @@ static USelection* GActorSelection = nullptr;
 static USelection* GComponentSelection = nullptr;
 static USelection* GObjectSelection = nullptr;
 
-class FLevelEvents
-{
-public:
-	FLevelEvents()
-		: bSelectionChanged(false)
-		, ElementSelectionSet(nullptr)
-	{
-		ULevel::OnLoadedActorRemovedFromLevelPreEvent.AddRaw(this, &FLevelEvents::OnLoadedActorRemovedToLevelPreEvent);
-		ULevel::OnLoadedActorRemovedFromLevelPostEvent.AddRaw(this, &FLevelEvents::OnLoadedActorRemovedToLevelPostEvent);
-	}
-
-	~FLevelEvents()
-	{
-		ULevel::OnLoadedActorRemovedFromLevelPreEvent.RemoveAll(this);
-		ULevel::OnLoadedActorRemovedFromLevelPostEvent.RemoveAll(this);
-	}
-
-private:
-	void OnLoadedActorRemovedToLevelPreEvent(ULevel* InLevel)
-	{
-		check(!bSelectionChanged);
-		check(!ElementSelectionSet);
-		check(!LegacySyncBatch);
-
-		ElementSelectionSet = (GEditor && GEditor->GetSelectedActors()) ? GEditor->GetSelectedActors()->GetElementSelectionSet() : nullptr;
-		LegacySyncBatch = ElementSelectionSet ? MakeUnique<FTypedElementList::FLegacySyncScopedBatch>(*ElementSelectionSet->GetElementList()) : nullptr;
-
-		InLevel->OnLoadedActorRemovedFromLevelEvent.AddRaw(this, &FLevelEvents::OnLoadedActorRemoved);
-	}
-
-	void OnLoadedActorRemoved(AActor& InActor)
-	{
-		if (ElementSelectionSet)
-		{
-			if (!InActor.GetRootSelectionParent())
-			{
-				if (FTypedElementHandle ActorHandle = UEngineElementsLibrary::AcquireEditorActorElementHandle(&InActor, /*bAllowCreate*/false))
-				{
-					static const FTypedElementSelectionOptions SelectionOptions = FTypedElementSelectionOptions()
-						.SetAllowHidden(true)
-						.SetAllowGroups(false)
-						.SetWarnIfLocked(false)
-						.SetChildElementInclusionMethod(ETypedElementChildInclusionMethod::Recursive);
-
-					bSelectionChanged |= ElementSelectionSet->DeselectElement(ActorHandle, SelectionOptions);
-				}
-			}
-		}
-	}
-
-	void OnLoadedActorRemovedToLevelPostEvent(ULevel* InLevel)
-	{
-		if (ElementSelectionSet)
-		{
-			if (bSelectionChanged)
-			{
-				GEditor->NoteSelectionChange();
-				bSelectionChanged = false;
-			}
-
-			LegacySyncBatch.Reset();
-			ElementSelectionSet = nullptr;
-		}
-
-		InLevel->OnLoadedActorRemovedFromLevelEvent.RemoveAll(this);
-	}
-
-	bool bSelectionChanged;
-	UTypedElementSelectionSet* ElementSelectionSet;
-	TUniquePtr<FTypedElementList::FLegacySyncScopedBatch> LegacySyncBatch;
-};
-
-static FLevelEvents* GLevelEvents = nullptr;
-
 void InitSelectionSets()
 {
 	// Note: The actor and component typed element selection set is set and owned by the level editor, so it is deliberately left null here
@@ -380,8 +306,6 @@ void InitSelectionSets()
 	{
 		return GObjectSelection->IsSelected(InObject);
 	};
-
-	GLevelEvents = new FLevelEvents;
 }
 
 void DestroySelectionSets()
@@ -401,8 +325,6 @@ void DestroySelectionSets()
 				ObjectSelectionSet->ClearSelection(FTypedElementSelectionOptions());
 			}
 		}
-
-		delete GLevelEvents;
 	}
 
 	GIsActorSelectedInEditor = nullptr;
@@ -412,7 +334,6 @@ void DestroySelectionSets()
 	GActorSelection = nullptr;
 	GComponentSelection = nullptr;
 	GObjectSelection = nullptr;
-	GLevelEvents = nullptr;
 }
 
 } // namespace PrivateEditorSelection
diff --git a/Engine/Source/Editor/UnrealEd/Private/Selection.cpp b/Engine/Source/Editor/UnrealEd/Private/Selection.cpp
--- a/Engine/Source/Editor/UnrealEd/Private/Selection.cpp
+++ b/Engine/Source/Editor/UnrealEd/Private/Selection.cpp
@@ -12,6 +12,7 @@
 #include "GameFramework/Actor.h"
 #include "Logging/LogMacros.h"
 #include "Templates/UnrealTemplate.h"
+#include "Editor.h"
 
 DEFINE_LOG_CATEGORY_STATIC(LogSelection, Log, All);
 
@@ -356,3 +357,37 @@ bool USelection::Modify(bool bAlwaysMarkDirty)
 	return ElementSelectionSet 
 		&& ElementSelectionSet->Modify(bAlwaysMarkDirty);
 }
+
+FDeselectedActorsEvent::~FDeselectedActorsEvent()
+{
+	if (GEditor && GEditor->GetSelectedActors())
+	{
+		if (UTypedElementSelectionSet* ElementSelectionSet = GEditor->GetSelectedActors()->GetElementSelectionSet())
+		{
+			bool bSelectionChanged = false;
+			FTypedElementList::FLegacySyncScopedBatch LegacySyncBatch(*ElementSelectionSet->GetElementList());
+
+			for (AActor* Actor : DeselectedActors)
+			{
+				if (!Actor->GetRootSelectionParent())
+				{
+					if (FTypedElementHandle ActorHandle = UEngineElementsLibrary::AcquireEditorActorElementHandle(Actor, /*bAllowCreate*/false))
+					{
+						static const FTypedElementSelectionOptions SelectionOptions = FTypedElementSelectionOptions()
+							.SetAllowHidden(true)
+							.SetAllowGroups(false)
+							.SetWarnIfLocked(false)
+							.SetChildElementInclusionMethod(ETypedElementChildInclusionMethod::Recursive);
+
+						bSelectionChanged |= ElementSelectionSet->DeselectElement(ActorHandle, SelectionOptions);
+					}
+				}
+			}
+
+			if (bSelectionChanged)
+			{
+				GEditor->NoteSelectionChange();
+			}
+		}
+	}
+}
\ No newline at end of file
diff --git a/Engine/Source/Editor/UnrealEd/Public/Selection.h b/Engine/Source/Editor/UnrealEd/Public/Selection.h
--- a/Engine/Source/Editor/UnrealEd/Public/Selection.h
+++ b/Engine/Source/Editor/UnrealEd/Public/Selection.h
@@ -494,3 +494,16 @@ public:
 		: TSelectionIterator<FSelectedEditableComponentFilter>(InSelection)
 	{}
 };
+
+class UNREALED_API FDeselectedActorsEvent
+{
+public:
+	FDeselectedActorsEvent(const TArray<AActor*>& InDeselectedActors)
+		: DeselectedActors(InDeselectedActors)
+	{}
+
+	~FDeselectedActorsEvent();
+
+private:
+	const TArray<AActor*>& DeselectedActors;
+};
\ No newline at end of file
diff --git a/Engine/Source/Runtime/Engine/Classes/Engine/Level.h b/Engine/Source/Runtime/Engine/Classes/Engine/Level.h
--- a/Engine/Source/Runtime/Engine/Classes/Engine/Level.h
+++ b/Engine/Source/Runtime/Engine/Classes/Engine/Level.h
@@ -949,28 +949,28 @@ public:
 	ENGINE_API void RemoveLoadedActors(const TArray<AActor*>& ActorList, const FTransform* TransformToRemove = nullptr);
 
 	/** Called when dynamically loaded actors are being added to this level */
-	DECLARE_MULTICAST_DELEGATE_OneParam(FLoadedActorAddedToLevelPreEvent, ULevel*);
-	static ENGINE_API FLoadedActorAddedToLevelPreEvent OnLoadedActorAddedToLevelPreEvent;
+	DECLARE_EVENT_OneParam(ULevel, FLoadedActorAddedToLevelPreEvent, const TArray<AActor*>&);
+	FLoadedActorAddedToLevelPreEvent OnLoadedActorAddedToLevelPreEvent;
 
 	/** Called when dynamically loaded actor is added to this level */
 	DECLARE_EVENT_OneParam(ULevel, FLoadedActorAddedToLevelEvent, AActor&);
 	FLoadedActorAddedToLevelEvent OnLoadedActorAddedToLevelEvent;
 
 	/** Called when dynamically loaded actors were added to this level */
-	DECLARE_MULTICAST_DELEGATE_OneParam(FLoadedActorAddedToLevelPostEvent, ULevel*);
-	static ENGINE_API FLoadedActorAddedToLevelPostEvent OnLoadedActorAddedToLevelPostEvent;
+	DECLARE_EVENT_OneParam(ULevel, FLoadedActorAddedToLevelPostEvent, const TArray<AActor*>&);
+	FLoadedActorAddedToLevelPostEvent OnLoadedActorAddedToLevelPostEvent;
 
 	/** Called when dynamically loaded actors are being removed from this level */
-	DECLARE_MULTICAST_DELEGATE_OneParam(FLoadedActorRemovedFromLevelPreEvent, ULevel*);
-	static ENGINE_API FLoadedActorRemovedFromLevelPreEvent OnLoadedActorRemovedFromLevelPreEvent;
+	DECLARE_EVENT_OneParam(ULevel, FLoadedActorRemovedFromLevelPreEvent, const TArray<AActor*>&);
+	FLoadedActorRemovedFromLevelPreEvent OnLoadedActorRemovedFromLevelPreEvent;
 
 	/** Called when dynamically loaded actor is removed from this level */
 	DECLARE_EVENT_OneParam(ULevel, FLoadedActorRemovedFromLevelEvent, AActor&);
 	FLoadedActorRemovedFromLevelEvent OnLoadedActorRemovedFromLevelEvent;
 
 	/** Called when dynamically loaded actors were removed from this level */
-	DECLARE_MULTICAST_DELEGATE_OneParam(FLoadedActorRemovedFromLevelPostEvent, ULevel*);
-	static ENGINE_API FLoadedActorRemovedFromLevelPostEvent OnLoadedActorRemovedFromLevelPostEvent;
+	DECLARE_EVENT_OneParam(ULevel, FLoadedActorRemovedFromLevelPostEvent, const TArray<AActor*>&);
+	FLoadedActorRemovedFromLevelPostEvent OnLoadedActorRemovedFromLevelPostEvent;
 #endif
 
 	/* Called when level is loaded. */
diff --git a/Engine/Source/Runtime/Engine/Classes/Engine/LevelStreamingDynamic.h b/Engine/Source/Runtime/Engine/Classes/Engine/LevelStreamingDynamic.h
--- a/Engine/Source/Runtime/Engine/Classes/Engine/LevelStreamingDynamic.h
+++ b/Engine/Source/Runtime/Engine/Classes/Engine/LevelStreamingDynamic.h
@@ -86,6 +86,8 @@ class ENGINE_API ULevelStreamingDynamic : public ULevelStreaming
 
 	static ULevelStreamingDynamic* LoadLevelInstance(const FLoadLevelInstanceParams& Params, bool& bOutSuccess);
 
+	static FString GetLevelInstancePackageName(const FLoadLevelInstanceParams& Params);
+
 	//~ Begin UObject Interface
 	virtual void PostLoad() override;
 	//~ End UObject Interface
diff --git a/Engine/Source/Runtime/Engine/Classes/GameFramework/Actor.h b/Engine/Source/Runtime/Engine/Classes/GameFramework/Actor.h
--- a/Engine/Source/Runtime/Engine/Classes/GameFramework/Actor.h
+++ b/Engine/Source/Runtime/Engine/Classes/GameFramework/Actor.h
@@ -137,6 +137,22 @@ DECLARE_EVENT_TwoParams(AActor, FActorOnPackagingModeChanged, AActor*, bool /* b
 DECLARE_DELEGATE_RetVal_ThreeParams(bool, FOnProcessEvent, AActor*, UFunction*, void*);
 #endif
 
+#if WITH_EDITOR
+class FActorInstanceGuidMapper
+{
+	using FGuidMapper = TFunction<FGuid(const FGuid&)>;
+
+public:
+	void RegisterGuidMapper(FName InPackageName, const FGuidMapper& InGuidMapper);
+	void UnregisterGuidMapper(FName InPackageName);
+
+	FGuid MapGuid(FName InPackageName, const FGuid& InGuid);
+
+private:
+	TMap<FName, FGuidMapper> GuidMappers;
+};
+#endif
+
 /**
  * TInlineComponentArray is simply a TArray that reserves a fixed amount of space on the stack
  * to try to avoid heap allocation when there are fewer than a specified number of elements expected in the result.
@@ -883,7 +899,7 @@ private:
 protected:
 	/**
 	 * The GUID for this actor; this guid will be the same for actors from instanced streaming levels.
-	 * @see		ActorInstanceGuid
+	 * @see		ActorInstanceGuid, FActorInstanceGuidMapper
 	 * @note	Don't use VisibleAnywhere here to avoid getting the CPF_Edit flag and get this property reset when resetting to defaults.
 	 *			See FActorDetails::AddActorCategory and EditorUtilities::CopySingleProperty for details.
 	 */
@@ -893,8 +909,7 @@ protected:
 	/**
 	 * The instance GUID for this actor; this guid will be unique for actors from instanced streaming levels.
 	 * @see		ActorGuid
-	 * @note	Don't use VisibleAnywhere here to avoid getting the CPF_Edit flag and get this property reset when resetting to defaults.
-	 *			See FActorDetails::AddActorCategory and EditorUtilities::CopySingleProperty for details.
+	 * @note	This is not guaranteed to be valid during PostLoad in all situations, but safe to access from RegisterAllComponents.
 	 */
 	UPROPERTY(BluePrintReadOnly, AdvancedDisplay, Category=Actor, Transient, NonTransactional)
 	FGuid ActorInstanceGuid;
diff --git a/Engine/Source/Runtime/Engine/Private/Actor.cpp b/Engine/Source/Runtime/Engine/Private/Actor.cpp
--- a/Engine/Source/Runtime/Engine/Private/Actor.cpp
+++ b/Engine/Source/Runtime/Engine/Private/Actor.cpp
@@ -60,8 +60,10 @@
 #include "FoliageHelper.h"
 #include "LevelInstance/LevelInstanceSubsystem.h"
 #include "LevelInstance/LevelInstanceInterface.h"
+#include "Misc/LazySingleton.h"
 #endif
 
+#include "WorldPartition/WorldPartitionLog.h"
 #include "WorldPartition/DataLayer/WorldDataLayers.h"
 #include "WorldPartition/WorldPartitionActorDescUtils.h"
 #include "WorldPartition/DataLayer/DataLayerAsset.h"
@@ -100,6 +102,29 @@ FCriticalSection CSVActorClassNameToCountMapLock;
 
 #endif // (CSV_PROFILER && !UE_BUILD_SHIPPING)
 
+#if WITH_EDITOR
+void FActorInstanceGuidMapper::RegisterGuidMapper(FName InPackageName, const FGuidMapper& InGuidMapper)
+{
+	check(!GuidMappers.Contains(InPackageName));
+	GuidMappers.Add(InPackageName, InGuidMapper);
+}
+
+void FActorInstanceGuidMapper::UnregisterGuidMapper(FName InPackageName)
+{
+	check(GuidMappers.Contains(InPackageName));
+	GuidMappers.Remove(InPackageName);
+}
+
+FGuid FActorInstanceGuidMapper::MapGuid(FName InPackageName, const FGuid& InGuid)
+{
+	if (FGuidMapper* GuidMapper = GuidMappers.Find(InPackageName))
+	{
+		return (*GuidMapper)(InGuid);
+	}
+	return InGuid;
+}
+#endif
+
 uint32 AActor::BeginPlayCallDepth = 0;
 
 AActor::AActor()
@@ -861,6 +886,13 @@ void AActor::Serialize(FArchive& Ar)
 			ActorGuid = FGuid::NewGuid();
 		}
 
+		check(!ActorInstanceGuid.IsValid());
+		const FGuid NewActorInstanceGuid = TLazySingleton<FActorInstanceGuidMapper>::Get().MapGuid(GetOuter()->GetPackage()->GetFName(), ActorGuid);
+		if (NewActorInstanceGuid != ActorGuid)
+		{
+			ActorInstanceGuid = NewActorInstanceGuid;
+		}
+
 		if (!CanChangeIsSpatiallyLoadedFlag() && (Ar.CustomVer(FUE5ReleaseStreamObjectVersion::GUID) < FUE5ReleaseStreamObjectVersion::ActorGridPlacementDeprecateDefaultValueFixup))
 		{
 			bIsSpatiallyLoaded = GetClass()->GetDefaultObject<AActor>()->bIsSpatiallyLoaded;
@@ -879,7 +911,6 @@ void AActor::Serialize(FArchive& Ar)
 				ContentBundleGuid = FixupContentBundleGuid;
 			}
 		}
-
 	}
 #endif
 }
diff --git a/Engine/Source/Runtime/Engine/Private/Level.cpp b/Engine/Source/Runtime/Engine/Private/Level.cpp
--- a/Engine/Source/Runtime/Engine/Private/Level.cpp
+++ b/Engine/Source/Runtime/Engine/Private/Level.cpp
@@ -60,6 +60,7 @@ Level.cpp: Level-related functions
 #include "WorldPartition/WorldPartitionHelpers.h"
 #include "HAL/PlatformFileManager.h"
 #include "Misc/PathViews.h"
+#include "Selection.h"
 #endif
 #include "WorldPartition/WorldPartition.h"
 #include "WorldPartition/WorldPartitionLog.h"
@@ -194,13 +195,6 @@ FPendingAutoReceiveInputActor::~FPendingAutoReceiveInputActor() = default;
 ULevel implementation.
 -----------------------------------------------------------------------------*/
 
-#if WITH_EDITOR
-ULevel::FLoadedActorAddedToLevelPreEvent ULevel::OnLoadedActorAddedToLevelPreEvent;
-ULevel::FLoadedActorAddedToLevelPostEvent ULevel::OnLoadedActorAddedToLevelPostEvent;
-ULevel::FLoadedActorRemovedFromLevelPreEvent ULevel::OnLoadedActorRemovedFromLevelPreEvent;
-ULevel::FLoadedActorRemovedFromLevelPostEvent ULevel::OnLoadedActorRemovedFromLevelPostEvent;
-#endif
-
 /** Called when a level package has been dirtied. */
 FSimpleMulticastDelegate ULevel::LevelDirtiedEvent;
 
@@ -836,7 +830,7 @@ void ULevel::AddLoadedActors(const TArray<AActor*>& ActorList, const FTransform*
 	FScopedSlowTask SlowTask(ActorsQueue.Num() * 3, LOCTEXT("RegisteringActors", "Registering actors..."));
 	SlowTask.MakeDialogDelayed(1.0f);
 
-	OnLoadedActorAddedToLevelPreEvent.Broadcast(this);
+	OnLoadedActorAddedToLevelPreEvent.Broadcast(ActorsQueue);
 
 	// Register all components
 	for (AActor* Actor : ActorsQueue)
@@ -905,7 +899,11 @@ void ULevel::AddLoadedActors(const TArray<AActor*>& ActorList, const FTransform*
 		SlowTask.EnterProgressFrame(1);
 	}
 
-	OnLoadedActorAddedToLevelPostEvent.Broadcast(this);
+	TArray<AActor*> ValidActorsQueue;
+	ValidActorsQueue.Reserve(ActorsQueue.Num());
+	Algo::CopyIf(ActorsQueue, ValidActorsQueue, [&](AActor* Actor) { return IsValid(Actor); });
+
+	OnLoadedActorAddedToLevelPostEvent.Broadcast(ValidActorsQueue);
 }
 
 void ULevel::RemoveLoadedActor(AActor* Actor, const FTransform* TransformToRemove)
@@ -951,8 +949,8 @@ void ULevel::RemoveLoadedActors(const TArray<AActor*>& ActorList, const FTransfo
 
 	FScopedSlowTask SlowTask(ActorsQueue.Num(), LOCTEXT("UnregisteringActors", "Unregistering actors..."));
 	SlowTask.MakeDialogDelayed(1.0f);
-			
-	OnLoadedActorRemovedFromLevelPreEvent.Broadcast(this);
+
+	OnLoadedActorRemovedFromLevelPreEvent.Broadcast(ActorsQueue);
 
 	for (AActor* Actor : ActorsQueue)
 	{
@@ -978,7 +976,12 @@ void ULevel::RemoveLoadedActors(const TArray<AActor*>& ActorList, const FTransfo
 		SlowTask.EnterProgressFrame(1);
 	}
 
-	OnLoadedActorRemovedFromLevelPostEvent.Broadcast(this);
+	TArray<AActor*> ValidActorsQueue;
+	ValidActorsQueue.Reserve(ActorsQueue.Num());
+	Algo::CopyIf(ActorsQueue, ValidActorsQueue, [&](AActor* Actor) { return IsValid(Actor); });
+
+	FDeselectedActorsEvent DeselectedActorsEvent(ValidActorsQueue);
+	OnLoadedActorRemovedFromLevelPostEvent.Broadcast(ValidActorsQueue);
 }
 #endif
 
diff --git a/Engine/Source/Runtime/Engine/Private/LevelInstance/LevelInstanceLevelStreaming.cpp b/Engine/Source/Runtime/Engine/Private/LevelInstance/LevelInstanceLevelStreaming.cpp
--- a/Engine/Source/Runtime/Engine/Private/LevelInstance/LevelInstanceLevelStreaming.cpp
+++ b/Engine/Source/Runtime/Engine/Private/LevelInstance/LevelInstanceLevelStreaming.cpp
@@ -23,6 +23,8 @@
 #include "LevelInstance/LevelInstanceEditorInstanceActor.h"
 #include "LevelUtils.h"
 #include "ActorFolder.h"
+#include "Misc/LazySingleton.h"
+#include "Misc/ScopeExit.h"
 #include "UObject/LinkerLoad.h"
 
 static bool GDisableLevelInstanceEditorPartialLoading = false;
@@ -89,6 +91,20 @@ FBox ULevelStreamingLevelInstance::GetBounds() const
 	return CachedBounds;
 }
 
+void ULevelStreamingLevelInstance::OnLoadedActorPreAddedToLevel(const TArray<AActor*>& InActors)
+{
+	check(LevelInstanceEditorInstanceActor.IsValid());
+	if (ILevelInstanceInterface* LevelInstance = GetLevelInstance())
+	{
+		for (AActor* Actor : InActors)
+		{
+			// Must happen before the actors are registered with the world, which is the case for this delegate.
+			const FActorContainerID& ContainerID = LevelInstance->GetLevelInstanceID().GetContainerID();
+			FSetActorInstanceGuid SetActorInstanceGuid(Actor, ContainerID.GetActorGuid(Actor->GetActorGuid()));
+		}
+	}
+}
+
 void ULevelStreamingLevelInstance::OnLoadedActorAddedToLevel(AActor& InActor)
 {
 	check(LevelInstanceEditorInstanceActor.IsValid());
@@ -159,9 +175,6 @@ void ULevelStreamingLevelInstance::PrepareLevelInstanceLoadedActor(AActor& InAct
 	if (InLevelInstance)
 	{
 		InActor.PushLevelInstanceEditingStateToProxies(CastChecked<AActor>(InLevelInstance)->IsInEditingLevelInstance());
-
-		const FActorContainerID& ContainerID = InLevelInstance->GetLevelInstanceID().GetContainerID();
-		FSetActorInstanceGuid SetActorInstanceGuid(&InActor, ContainerID.GetActorGuid(InActor.GetActorGuid()));
 	}
 
 	if (LevelInstanceEditorInstanceActor.IsValid())
@@ -221,7 +234,15 @@ ULevelStreamingLevelInstance* ULevelStreamingLevelInstance::LoadInstance(ILevelI
 	{
 		Params.bInitiallyVisible = LevelInstance->IsInitiallyVisible();
 	}
-	
+
+	const FString LevelInstancePackageName = ULevelStreamingDynamic::GetLevelInstancePackageName(Params);
+
+#if WITH_EDITOR
+	FActorInstanceGuidMapper& ActorInstanceGuidMapper = TLazySingleton<FActorInstanceGuidMapper>::Get();
+	ActorInstanceGuidMapper.RegisterGuidMapper(*LevelInstancePackageName, [LevelInstance](const FGuid& InActorGuid) { return LevelInstance->GetLevelInstanceID().GetContainerID().GetActorGuid(InActorGuid); });
+	ON_SCOPE_EXIT { ActorInstanceGuidMapper.UnregisterGuidMapper(*LevelInstancePackageName); };
+#endif
+
 	ULevelStreamingLevelInstance* LevelStreaming = Cast<ULevelStreamingLevelInstance>(ULevelStreamingDynamic::LoadLevelInstance(Params, bOutSuccess));
 	if (bOutSuccess)
 	{
@@ -240,6 +261,7 @@ ULevelStreamingLevelInstance* ULevelStreamingLevelInstance::LoadInstance(ILevelI
 			{
 				check(LevelStreaming->GetLevelStreamingState() == ELevelStreamingState::LoadedVisible);
 
+				Level->OnLoadedActorAddedToLevelPreEvent.AddUObject(LevelStreaming, &ULevelStreamingLevelInstance::OnLoadedActorPreAddedToLevel);
 				Level->OnLoadedActorAddedToLevelEvent.AddUObject(LevelStreaming, &ULevelStreamingLevelInstance::OnLoadedActorAddedToLevel);
 				Level->OnLoadedActorRemovedFromLevelEvent.AddUObject(LevelStreaming, &ULevelStreamingLevelInstance::OnLoadedActorRemovedFromLevel);
 
@@ -301,6 +323,7 @@ void ULevelStreamingLevelInstance::UnloadInstance(ULevelStreamingLevelInstance*
 	else
 	{
 		ULevel* LoadedLevel = LevelStreaming->GetLoadedLevel();
+		LoadedLevel->OnLoadedActorAddedToLevelPreEvent.RemoveAll(LevelStreaming);
 		LoadedLevel->OnLoadedActorAddedToLevelEvent.RemoveAll(LevelStreaming);
 		LoadedLevel->OnLoadedActorRemovedFromLevelEvent.RemoveAll(LevelStreaming);
 		LevelStreaming->LevelInstanceEditorInstanceActor.Reset();
diff --git a/Engine/Source/Runtime/Engine/Private/LevelStreaming.cpp b/Engine/Source/Runtime/Engine/Private/LevelStreaming.cpp
--- a/Engine/Source/Runtime/Engine/Private/LevelStreaming.cpp
+++ b/Engine/Source/Runtime/Engine/Private/LevelStreaming.cpp
@@ -2440,7 +2440,6 @@ ULevelStreamingDynamic* ULevelStreamingDynamic::LoadLevelInstance(const FLoadLev
 	return LoadLevelInstance_Internal(Params, bOutSuccess);
 }
 
-
 ULevelStreamingDynamic* ULevelStreamingDynamic::LoadLevelInstance_Internal(const FLoadLevelInstanceParams& Params, bool& bOutSuccess)
 {
 	const FString PackagePath = FPackageName::GetLongPackagePath(Params.LongPackageName);
@@ -2456,29 +2455,10 @@ ULevelStreamingDynamic* ULevelStreamingDynamic::LoadLevelInstance_Internal(const
 	const FString OnDiskPackageName = PackagePath + TEXT("/") + ShortPackageName;
 
 	// Determine loaded package name
-	TStringBuilder<512> LevelPackageNameStrBuilder;
-	if (Params.bLoadAsTempPackage)
-	{
-		LevelPackageNameStrBuilder.Append(TEXT("/Temp"));
-	}
-	LevelPackageNameStrBuilder.Append(PackagePath);
-	LevelPackageNameStrBuilder.Append(TEXT("/"));
-		
-	bool bNeedsUniqueTest = false;
-	if (Params.OptionalLevelNameOverride)
-	{
-		// Use the supplied suffix, which is expected to result in a unique package name but we have to check if it is not.
-		LevelPackageNameStrBuilder.Append(*Params.OptionalLevelNameOverride);
-		bNeedsUniqueTest = true;
-	}
-	else
-	{
-		LevelPackageNameStrBuilder.Append(ShortPackageName);
-		LevelPackageNameStrBuilder.Append(TEXT("_LevelInstance_"));
-		LevelPackageNameStrBuilder.Append(FString::FromInt(++UniqueLevelInstanceId));
-	}
+	const FString LevelPackageNameStr(GetLevelInstancePackageName(Params));
+
+	const bool bNeedsUniqueTest = Params.OptionalLevelNameOverride != nullptr;
 
-	const FString LevelPackageNameStr(LevelPackageNameStrBuilder.ToString());
 	const FName UnmodifiedLevelPackageName = FName(*LevelPackageNameStr);
 #if WITH_EDITOR
 	const bool bIsPlayInEditor = Params.World->IsPlayInEditor();
@@ -2536,6 +2516,43 @@ ULevelStreamingDynamic* ULevelStreamingDynamic::LoadLevelInstance_Internal(const
     return StreamingLevel;
 }	
 
+FString ULevelStreamingDynamic::GetLevelInstancePackageName(const FLoadLevelInstanceParams& Params)
+{
+	const FString PackagePath = FPackageName::GetLongPackagePath(Params.LongPackageName);
+	FString ShortPackageName = FPackageName::GetShortName(Params.LongPackageName);
+
+	if (ShortPackageName.StartsWith(Params.World->StreamingLevelsPrefix))
+	{
+		ShortPackageName.RightChopInline(Params.World->StreamingLevelsPrefix.Len(), false);
+	}
+
+	// Remove PIE prefix if it's there before we actually load the level
+	const FString OnDiskPackageName = PackagePath + TEXT("/") + ShortPackageName;
+
+	// Determine loaded package name
+	TStringBuilder<512> LevelPackageNameStrBuilder;
+	if (Params.bLoadAsTempPackage)
+	{
+		LevelPackageNameStrBuilder.Append(TEXT("/Temp"));
+	}
+	LevelPackageNameStrBuilder.Append(PackagePath);
+	LevelPackageNameStrBuilder.Append(TEXT("/"));
+		
+	if (Params.OptionalLevelNameOverride)
+	{
+		// Use the supplied suffix, which is expected to result in a unique package name but we have to check if it is not.
+		LevelPackageNameStrBuilder.Append(*Params.OptionalLevelNameOverride);
+	}
+	else
+	{
+		LevelPackageNameStrBuilder.Append(ShortPackageName);
+		LevelPackageNameStrBuilder.Append(TEXT("_LevelInstance_"));
+		LevelPackageNameStrBuilder.Append(FString::FromInt(++UniqueLevelInstanceId));
+	}
+
+	return LevelPackageNameStrBuilder.ToString();
+}
+
 /*-----------------------------------------------------------------------------
 	ULevelStreamingAlwaysLoaded implementation.
 -----------------------------------------------------------------------------*/
diff --git a/Engine/Source/Runtime/Engine/Public/LevelInstance/LevelInstanceLevelStreaming.h b/Engine/Source/Runtime/Engine/Public/LevelInstance/LevelInstanceLevelStreaming.h
--- a/Engine/Source/Runtime/Engine/Public/LevelInstance/LevelInstanceLevelStreaming.h
+++ b/Engine/Source/Runtime/Engine/Public/LevelInstance/LevelInstanceLevelStreaming.h
@@ -39,6 +39,7 @@ private:
 #if WITH_EDITOR
 	void ResetLevelInstanceLoaders();
 	void PrepareLevelInstanceLoadedActor(AActor& InActor, ILevelInstanceInterface* InLevelInstance, bool bResetLoaders);
+	void OnLoadedActorPreAddedToLevel(const TArray<AActor*>& InActors);
 	void OnLoadedActorAddedToLevel(AActor& InActor);
 	void OnLoadedActorRemovedFromLevel(AActor& InActor);
 
