Exec commands: audit usage through analytics. Keep track of parsed commands and exec commands being called when not allowed.

diff --git a/Engine/Source/Runtime/Core/Private/Misc/CoreDelegates.cpp b/Engine/Source/Runtime/Core/Private/Misc/CoreDelegates.cpp
--- a/Engine/Source/Runtime/Core/Private/Misc/CoreDelegates.cpp
+++ b/Engine/Source/Runtime/Core/Private/Misc/CoreDelegates.cpp
@@ -356,4 +356,6 @@ CORE_API void RegisterEncryptionKeyCallback(TEncryptionKeyFunc InCallback)
 FSimpleMulticastDelegate FCoreDelegates::OnParentBeginFork;
 FSimpleMulticastDelegate FCoreDelegates::OnParentPreFork;
 TMulticastDelegate<void(EForkProcessRole /* ProcessRole */)> FCoreDelegates::OnPostFork;
-FSimpleMulticastDelegate FCoreDelegates::OnChildEndFramePostFork;
\ No newline at end of file
+FSimpleMulticastDelegate FCoreDelegates::OnChildEndFramePostFork;
+FCoreDelegates::FOnDisallowedExecCommandCalled FCoreDelegates::OnDisallowedExecCommandCalled;
+FCoreDelegates::FOnNamedCommandParsed FCoreDelegates::OnNamedCommandParsed;
\ No newline at end of file
diff --git a/Engine/Source/Runtime/Core/Private/Misc/Exec.cpp b/Engine/Source/Runtime/Core/Private/Misc/Exec.cpp
--- a/Engine/Source/Runtime/Core/Private/Misc/Exec.cpp
+++ b/Engine/Source/Runtime/Core/Private/Misc/Exec.cpp
@@ -2,6 +2,9 @@
 
 #include "Misc/Exec.h"
 
+#include "Misc/AssertionMacros.h"
+#include "Misc/CoreDelegates.h"
+
 FExec::~FExec()
 {
 }
@@ -23,4 +26,17 @@ bool FExec::Exec( UWorld* InWorld, const TCHAR* Cmd, FOutputDevice& Ar )
 
 	return bExecSuccess;
 }
-#endif // UE_ALLOW_EXEC_COMMANDS
\ No newline at end of file
+#else // UE_ALLOW_EXEC_COMMANDS
+bool FExec::Exec( UWorld* InWorld, const TCHAR* Cmd, FOutputDevice& Ar )
+{
+	return Exec( InWorld, Cmd );
+}
+
+bool FExec::Exec( UWorld* InWorld, const TCHAR* Cmd )
+{
+	ensure(false);
+	FCoreDelegates::OnDisallowedExecCommandCalled.Broadcast( Cmd );
+
+	return false;
+}
+#endif // !UE_ALLOW_EXEC_COMMANDS
\ No newline at end of file
diff --git a/Engine/Source/Runtime/Core/Private/Misc/Parse.cpp b/Engine/Source/Runtime/Core/Private/Misc/Parse.cpp
--- a/Engine/Source/Runtime/Core/Private/Misc/Parse.cpp
+++ b/Engine/Source/Runtime/Core/Private/Misc/Parse.cpp
@@ -10,6 +10,7 @@
 #include "Containers/Set.h"
 #include "Internationalization/Text.h"
 #include "Misc/AsciiSet.h"
+#include "Misc/CoreDelegates.h"
 #include "Misc/Guid.h"
 #include "Misc/OutputDeviceNull.h"
 #include "Misc/StringBuilder.h"
@@ -738,6 +739,8 @@ bool FParse::Command( const TCHAR** Stream, const TCHAR* Match, bool bParseMight
 				(*Stream)++;
 			}
 
+			FCoreDelegates::OnNamedCommandParsed.Broadcast(Match);
+
 			return true; // Success.
 		}
 		else
diff --git a/Engine/Source/Runtime/Core/Public/Misc/CoreDelegates.h b/Engine/Source/Runtime/Core/Public/Misc/CoreDelegates.h
--- a/Engine/Source/Runtime/Core/Public/Misc/CoreDelegates.h
+++ b/Engine/Source/Runtime/Core/Public/Misc/CoreDelegates.h
@@ -572,6 +572,14 @@ public:
 	// Called at the end of the frame where the process spawned a fork
 	static FSimpleMulticastDelegate OnChildEndFramePostFork;
 
+	// Called when FParse::Command succeeds. Used to audit named commands usage.
+	DECLARE_MULTICAST_DELEGATE_OneParam(FOnNamedCommandParsed, const TCHAR* /*Cmd*/);
+	static FOnNamedCommandParsed OnNamedCommandParsed;
+
+	// Called when FExec::Exec is called in a context where it's disallowed (UE_ALLOW_EXEC_COMMANDS = 0).
+	DECLARE_MULTICAST_DELEGATE_OneParam(FOnDisallowedExecCommandCalled, const TCHAR* /*Cmd*/);
+	static FOnDisallowedExecCommandCalled OnDisallowedExecCommandCalled;
+
 private:
 
 	// Callbacks for hotfixes
diff --git a/Engine/Source/Runtime/Core/Public/Misc/Exec.h b/Engine/Source/Runtime/Core/Public/Misc/Exec.h
--- a/Engine/Source/Runtime/Core/Public/Misc/Exec.h
+++ b/Engine/Source/Runtime/Core/Public/Misc/Exec.h
@@ -44,22 +44,21 @@ public:
 	virtual bool Exec( UWorld* InWorld, const TCHAR* Cmd, FOutputDevice& Ar );
 #else // UE_ALLOW_EXEC_COMMANDS
 
-public:
 	/**
 	 * final override of Exec that asserts if called
 	 */
-	virtual bool Exec( UWorld* InWorld, const TCHAR* Cmd, FOutputDevice& Ar ) final { check(false); return false; }
+	virtual bool Exec( UWorld* InWorld, const TCHAR* Cmd, FOutputDevice& Ar ) final;
 	
 	/**
 	 * final override of Exec to replace overrides where a default value for Ar is provided
 	 */
-	virtual bool Exec( UWorld* InWorld, const TCHAR* Cmd ) final { check(false); return false; }
+	virtual bool Exec( UWorld* InWorld, const TCHAR* Cmd ) final;
 
 #endif // !UE_ALLOW_EXEC_COMMANDS
 
 protected:
 	/** Implementation of Exec that is called on all targets where UE_ALLOW_EXEC_COMMANDS is true */
-	virtual bool Exec_Runtime(UWorld* InWorld, const TCHAR* Cmd, FOutputDevice& Ar) { return false; }
+	virtual bool Exec_Runtime( UWorld* InWorld, const TCHAR* Cmd, FOutputDevice& Ar ) { return false; }
 
 	/** Implementation of Exec that is only called in non-shipping targets */
 	virtual bool Exec_Dev( UWorld* InWorld, const TCHAR* Cmd, FOutputDevice& Ar ) { return false; }
diff --git a/Engine/Source/Runtime/Engine/Private/Analytics/EngineAnalyticsSessionSummary.cpp b/Engine/Source/Runtime/Engine/Private/Analytics/EngineAnalyticsSessionSummary.cpp
--- a/Engine/Source/Runtime/Engine/Private/Analytics/EngineAnalyticsSessionSummary.cpp
+++ b/Engine/Source/Runtime/Engine/Private/Analytics/EngineAnalyticsSessionSummary.cpp
@@ -45,6 +45,8 @@ namespace EngineAnalyticsProperties
 	static const TAnalyticsProperty<bool> RunningOnBattery       = TEXT("RunningOnBattery");
 	static const TAnalyticsProperty<bool> IsUserLoggingOut       = FAnalyticsSessionSummaryManager::IsUserLoggingOutProperty;
 	static const TAnalyticsProperty<int32> ShutdownType          = FAnalyticsSessionSummaryManager::ShutdownTypeCodeProperty;
+	static const TAnalyticsProperty<FString> CalledDisallowedExecCommands = TEXT("CalledDisallowedExecCommands");
+	static const TAnalyticsProperty<FString> ParsedNamedCommands = TEXT("ParsedNamedCommands");
 
 	static FString UnknownUserActivity(TEXT("Unknown"));
 
@@ -119,6 +121,22 @@ namespace EngineAnalyticsProperties
 		return true; // Shutdown|Terminate|Abnormal can overwrite each other.
 	}
 
+	bool UpdateCommandValue(const TCHAR* Cmd, FString& Value)
+	{
+		constexpr int32 MaxAllowedStringSize = 128;
+		if (Value.Len() >= MaxAllowedStringSize || Value.Contains(Cmd))
+		{
+			return false;
+		}
+
+		if (!Value.IsEmpty())
+		{
+			Value += TEXT(",");
+		}
+
+		Value += Cmd;
+		return true;
+	}
 } // namespace EngineAnalyticsProperties
 
 
@@ -209,6 +227,10 @@ FEngineAnalyticsSessionSummary::FEngineAnalyticsSessionSummary(TSharedPtr<IAnaly
 	EngineAnalyticsProperties::ShutdownType.Set(GetStore(), (int32)EAnalyticsSessionShutdownType::Abnormal);
 	EngineAnalyticsProperties::RunningOnBattery.Set(GetStore(), FPlatformMisc::IsRunningOnBattery());
 
+	const uint32 CharCountCapacityHint = 128;
+	EngineAnalyticsProperties::CalledDisallowedExecCommands.Set(GetStore(), TEXT(""), CharCountCapacityHint);
+	EngineAnalyticsProperties::ParsedNamedCommands.Set(GetStore(), TEXT(""), CharCountCapacityHint);
+
 	Store->Flush();
 
 	// Listen to interesting events.
@@ -216,6 +238,8 @@ FEngineAnalyticsSessionSummary::FEngineAnalyticsSessionSummary(TSharedPtr<IAnaly
 	FCoreDelegates::ApplicationWillTerminateDelegate.AddRaw(this, &FEngineAnalyticsSessionSummary::OnTerminate); // WARNING: Don't assume this function is only called from game thread.
 	FCoreDelegates::IsVanillaProductChanged.AddRaw(this, &FEngineAnalyticsSessionSummary::OnVanillaStateChanged);
 	FCoreDelegates::OnUserLoginChangedEvent.AddRaw(this, &FEngineAnalyticsSessionSummary::OnUserLoginChanged);
+	FCoreDelegates::OnDisallowedExecCommandCalled.AddRaw(this, &FEngineAnalyticsSessionSummary::OnDisallowedExecCommandCalled);
+	FCoreDelegates::OnNamedCommandParsed.AddRaw(this, &FEngineAnalyticsSessionSummary::OnNamedCommandParsed);
 	FUserActivityTracking::OnActivityChanged.AddRaw(this, &FEngineAnalyticsSessionSummary::OnUserActivity);
 }
 
@@ -230,6 +254,8 @@ void FEngineAnalyticsSessionSummary::Shutdown()
 		FCoreDelegates::ApplicationWillTerminateDelegate.RemoveAll(this);
 		FCoreDelegates::IsVanillaProductChanged.RemoveAll(this);
 		FCoreDelegates::OnUserLoginChangedEvent.RemoveAll(this);
+		FCoreDelegates::OnDisallowedExecCommandCalled.RemoveAll(this);
+		FCoreDelegates::OnNamedCommandParsed.RemoveAll(this);
 		FUserActivityTracking::OnActivityChanged.RemoveAll(this);
 
 		EngineAnalyticsProperties::WasShutdown.Set(GetStore(), true);
@@ -301,6 +327,24 @@ void FEngineAnalyticsSessionSummary::OnVanillaStateChanged(bool bIsVanilla)
 	UpdateSessionProgress();
 }
 
+void FEngineAnalyticsSessionSummary::OnDisallowedExecCommandCalled(const TCHAR* Cmd)
+{
+	EngineAnalyticsProperties::CalledDisallowedExecCommands.Update(GetStore(),
+		[Cmd](FString& Value)
+		{
+			return EngineAnalyticsProperties::UpdateCommandValue(Cmd, Value);
+		});
+}
+
+void FEngineAnalyticsSessionSummary::OnNamedCommandParsed(const TCHAR* Cmd)
+{
+	EngineAnalyticsProperties::ParsedNamedCommands.Update(GetStore(),
+		[Cmd](FString& Value)
+		{
+			return EngineAnalyticsProperties::UpdateCommandValue(Cmd, Value);
+		});
+}
+
 void FEngineAnalyticsSessionSummary::LowDriveSpaceDetected()
 {
 	EngineAnalyticsProperties::IsLowDriveSpace.Set(GetStore(), true);
diff --git a/Engine/Source/Runtime/Engine/Public/Analytics/EngineAnalyticsSessionSummary.h b/Engine/Source/Runtime/Engine/Public/Analytics/EngineAnalyticsSessionSummary.h
--- a/Engine/Source/Runtime/Engine/Public/Analytics/EngineAnalyticsSessionSummary.h
+++ b/Engine/Source/Runtime/Engine/Public/Analytics/EngineAnalyticsSessionSummary.h
@@ -71,6 +71,12 @@ private:
 	/** Invoked if the engine detects some local changes.. */
 	void OnVanillaStateChanged(bool bIsVanilla);
 
+	/** Invoked on disallowed usage of FExec::Exec */
+	void OnDisallowedExecCommandCalled(const TCHAR* Cmd);
+
+	/** Invoked on FParse::Command success. Used to audit named commands usage. */
+	void OnNamedCommandParsed(const TCHAR* Cmd);
+
 	/** Check on CRC and return true if its running state state changed (so that the change is persisted on disk). */
 	bool UpdateExternalProcessReporterState(bool bQuickCheck);
 
