LevelInstance: Disable partial loading until we can fix sub Level Instance issue. (also fix failed check when disabling partial loading)

#rnx

diff --git a/Engine/Source/Runtime/Engine/Private/LevelInstance/LevelInstanceLevelStreaming.cpp b/Engine/Source/Runtime/Engine/Private/LevelInstance/LevelInstanceLevelStreaming.cpp
--- a/Engine/Source/Runtime/Engine/Private/LevelInstance/LevelInstanceLevelStreaming.cpp
+++ b/Engine/Source/Runtime/Engine/Private/LevelInstance/LevelInstanceLevelStreaming.cpp
@@ -27,7 +27,7 @@
 #include "Misc/ScopeExit.h"
 #include "UObject/LinkerLoad.h"
 
-static bool GDisableLevelInstanceEditorPartialLoading = false;
+static bool GDisableLevelInstanceEditorPartialLoading = true;
 FAutoConsoleVariableRef CVarDisableLevelInstanceEditorPartialLoading(
 	TEXT("wp.Editor.DisableLevelInstanceEditorPartialLoading"),
 	GDisableLevelInstanceEditorPartialLoading,
@@ -361,12 +361,11 @@ void ULevelStreamingLevelInstance::OnLevelLoadedChanged(ULevel* InLevel)
 		}
 
 #if WITH_EDITOR
-		if (!GDisableLevelInstanceEditorPartialLoading)
+		if (ULevelInstanceSubsystem* LevelInstanceSubsystem = GetWorld()->GetSubsystem<ULevelInstanceSubsystem>())
 		{
-			if (ULevelInstanceSubsystem* LevelInstanceSubsystem = GetWorld()->GetSubsystem<ULevelInstanceSubsystem>())
+			LevelInstanceSubsystem->RegisterLoadedLevelStreamingLevelInstance(this);
+			if (!GDisableLevelInstanceEditorPartialLoading)
 			{
-				LevelInstanceSubsystem->RegisterLoadedLevelStreamingLevelInstance(this);
-
 				if (UWorldPartition* WorldPartition = NewLoadedLevel->GetWorldPartition())
 				{
 					check(!WorldPartition->IsInitialized());
