Exposed CachedMips metric for detection of texture pool issues.

#tests PIE
#rnx

diff --git a/Engine/Source/Runtime/Engine/Private/Streaming/StreamingManagerTexture.h b/Engine/Source/Runtime/Engine/Private/Streaming/StreamingManagerTexture.h
--- a/Engine/Source/Runtime/Engine/Private/Streaming/StreamingManagerTexture.h
+++ b/Engine/Source/Runtime/Engine/Private/Streaming/StreamingManagerTexture.h
@@ -99,6 +99,8 @@ struct FRenderAssetStreamingManager final : public IRenderAssetStreamingManager
 
 	virtual int64 GetMaxEverRequired() const override { return MaxEverRequired; }
 
+	virtual float GetCachedMips() const override { return DisplayedStats.CachedMips; }
+
 	virtual void ResetMaxEverRequired() override { MaxEverRequired = 0; }
 
 	/**
diff --git a/Engine/Source/Runtime/Engine/Public/ContentStreaming.h b/Engine/Source/Runtime/Engine/Public/ContentStreaming.h
--- a/Engine/Source/Runtime/Engine/Public/ContentStreaming.h
+++ b/Engine/Source/Runtime/Engine/Public/ContentStreaming.h
@@ -472,6 +472,9 @@ struct IRenderAssetStreamingManager : public IStreamingManager
 	/** Max required textures/meshes ever seen in bytes. */
 	virtual int64 GetMaxEverRequired() const = 0;
 
+	/** Amount of memory cached in pool */
+	virtual float GetCachedMips() const = 0;
+
 	/** Resets the max ever required textures/meshes.  For possibly when changing resolutions or screen pct. */
 	virtual void ResetMaxEverRequired() = 0;
 
