PR #10277: Fix GetTanAndInvTanHalfFOV (Contributed by [REDACTED_PERSON])

diff --git a/Engine/Source/Runtime/Engine/Public/SceneView.h b/Engine/Source/Runtime/Engine/Public/SceneView.h
--- a/Engine/Source/Runtime/Engine/Public/SceneView.h
+++ b/Engine/Source/Runtime/Engine/Public/SceneView.h
@@ -481,10 +481,10 @@ public:
 		//No concept of FOV for orthographic projection so only return perspective related values or 1.0f
 		if (IsPerspectiveProjection())
 		{
-			FVector4f(static_cast<float>(InvProjectionMatrix.M[0][0]), //ClipToView[0][0] - X axis
-					  static_cast<float>(InvProjectionMatrix.M[1][1]), //ClipToView[1][1] - Y axis
-					  static_cast<float>(ProjectionMatrix.M[0][0]), //ViewToClip[0][0] - X axis
-					  static_cast<float>(ProjectionMatrix.M[1][1])); //ViewToClip[1][1] - Y axis
+			return FVector4f(static_cast<float>(InvProjectionMatrix.M[0][0]), //ClipToView[0][0] - X axis
+							 static_cast<float>(InvProjectionMatrix.M[1][1]), //ClipToView[1][1] - Y axis
+							 static_cast<float>(ProjectionMatrix.M[0][0]), //ViewToClip[0][0] - X axis
+							 static_cast<float>(ProjectionMatrix.M[1][1])); //ViewToClip[1][1] - Y axis
 		}
 		return FVector4f(1.0f, 1.0f, 1.0f, 1.0f);
 	}
