World Partition: don't allow Level Instance partial loading if any parent is in edit mode.

diff --git a/Engine/Source/Runtime/Engine/Private/LevelInstance/LevelInstanceInterface.cpp b/Engine/Source/Runtime/Engine/Private/LevelInstance/LevelInstanceInterface.cpp
--- a/Engine/Source/Runtime/Engine/Private/LevelInstance/LevelInstanceInterface.cpp
+++ b/Engine/Source/Runtime/Engine/Private/LevelInstance/LevelInstanceInterface.cpp
@@ -21,7 +21,7 @@ ULevelInstanceInterface::ULevelInstanceInterface(const FObjectInitializer& Objec
 #if WITH_EDITOR
 bool ILevelInstanceInterface::SupportsPartialEditorLoading() const
 {
-	return CastChecked<AActor>(this)->GetIsSpatiallyLoaded() && !IsEditing();
+	return CastChecked<AActor>(this)->GetIsSpatiallyLoaded() && !IsEditing() && !HasParentEdit();
 }
 #endif
 
@@ -142,6 +142,19 @@ bool ILevelInstanceInterface::HasChildEdit() const
 	return false;
 }
 
+bool ILevelInstanceInterface::HasParentEdit() const
+{
+	if (HasValidLevelInstanceID())
+	{
+		if (ULevelInstanceSubsystem* LevelInstanceSubsystem = GetLevelInstanceSubsystem())
+		{
+			return LevelInstanceSubsystem->HasParentEdit(this);
+		}
+	}
+
+	return false;
+}
+
 bool ILevelInstanceInterface::IsDirty() const
 {
 	if (HasValidLevelInstanceID())
diff --git a/Engine/Source/Runtime/Engine/Private/LevelInstance/LevelInstanceSubsystem.cpp b/Engine/Source/Runtime/Engine/Private/LevelInstance/LevelInstanceSubsystem.cpp
--- a/Engine/Source/Runtime/Engine/Private/LevelInstance/LevelInstanceSubsystem.cpp
+++ b/Engine/Source/Runtime/Engine/Private/LevelInstance/LevelInstanceSubsystem.cpp
@@ -2221,6 +2221,20 @@ bool ULevelInstanceSubsystem::HasChildEdit(const ILevelInstanceInterface* LevelI
 	return ChildEditCountPtr && *ChildEditCountPtr;
 }
 
+bool ULevelInstanceSubsystem::HasParentEdit(const ILevelInstanceInterface* LevelInstance) const
+{
+	bool bResult = false;
+	
+	const AActor* LevelInstanceActor = CastChecked<AActor>(LevelInstance);
+	ForEachLevelInstanceAncestors(LevelInstanceActor, [LevelInstance, &bResult](const ILevelInstanceInterface* Ancestor)
+	{
+		bResult = Ancestor->IsEditing();
+		return !bResult;
+	});
+
+	return bResult;
+}
+
 void ULevelInstanceSubsystem::OnCommitChild(const FLevelInstanceID& LevelInstanceID, bool bChildChanged)
 {
 	int32& ChildEditCount = ChildEdits.FindChecked(LevelInstanceID);
diff --git a/Engine/Source/Runtime/Engine/Public/LevelInstance/LevelInstanceInterface.h b/Engine/Source/Runtime/Engine/Public/LevelInstance/LevelInstanceInterface.h
--- a/Engine/Source/Runtime/Engine/Public/LevelInstance/LevelInstanceInterface.h
+++ b/Engine/Source/Runtime/Engine/Public/LevelInstance/LevelInstanceInterface.h
@@ -102,6 +102,8 @@ class ENGINE_API ILevelInstanceInterface
 	virtual bool IsEditing() const;
 	
 	virtual bool HasChildEdit() const;
+
+	virtual bool HasParentEdit() const;
 	
 	virtual bool IsDirty() const;
 	
diff --git a/Engine/Source/Runtime/Engine/Public/LevelInstance/LevelInstanceSubsystem.h b/Engine/Source/Runtime/Engine/Public/LevelInstance/LevelInstanceSubsystem.h
--- a/Engine/Source/Runtime/Engine/Public/LevelInstance/LevelInstanceSubsystem.h
+++ b/Engine/Source/Runtime/Engine/Public/LevelInstance/LevelInstanceSubsystem.h
@@ -105,6 +105,7 @@ public:
 	void BlockUnloadLevelInstance(ILevelInstanceInterface* LevelInstance);
 		
 	bool HasChildEdit(const ILevelInstanceInterface* LevelInstance) const;
+	bool HasParentEdit(const ILevelInstanceInterface* LevelInstance) const;
 
 	TArray<ILevelInstanceInterface*> GetLevelInstances(const FString& WorldAssetPackage);
 	
