World Partition: don't clear the transaction buffer when unloading actors from containers that aren't the main world.

diff --git a/Engine/Source/Runtime/Engine/Private/WorldPartition/WorldPartitionActorLoaderInterface.cpp b/Engine/Source/Runtime/Engine/Private/WorldPartition/WorldPartitionActorLoaderInterface.cpp
--- a/Engine/Source/Runtime/Engine/Private/WorldPartition/WorldPartitionActorLoaderInterface.cpp
+++ b/Engine/Source/Runtime/Engine/Private/WorldPartition/WorldPartitionActorLoaderInterface.cpp
@@ -87,7 +87,7 @@ void IWorldPartitionActorLoaderInterface::ILoaderAdapter::Unload()
 
 			if (NumUnloads)
 			{
-				PostLoadedStateChanged(0, NumUnloads);
+				PostLoadedStateChanged(0, NumUnloads, true);
 			}
 		}
 	}
@@ -154,6 +154,7 @@ void IWorldPartitionActorLoaderInterface::ILoaderAdapter::RefreshLoadedState()
 			{
 				int32 NumLoads = 0;
 				int32 NumUnloads = 0;
+				bool bClearTransactions = false;
 				{
 					FWorldPartitionLoadingContext::FDeferred LoadingContext;
 
@@ -183,11 +184,12 @@ void IWorldPartitionActorLoaderInterface::ILoaderAdapter::RefreshLoadedState()
 
 					NumLoads = LoadingContext.GetNumRegistrations();
 					NumUnloads = LoadingContext.GetNumUnregistrations();
+					bClearTransactions = LoadingContext.GetNeedsClearTransactions();
 				}
 
 				if (NumLoads || NumUnloads)
 				{
-					PostLoadedStateChanged(NumLoads, NumUnloads);
+					PostLoadedStateChanged(NumLoads, NumUnloads, bClearTransactions);
 				}
 			}
 		}
@@ -236,7 +238,7 @@ void IWorldPartitionActorLoaderInterface::ILoaderAdapter::OnActorDescContainerUn
 			SlowTask.EnterProgressFrame(1);
 		}
 	
-		PostLoadedStateChanged(0, ActorsToUnload.Num());
+		PostLoadedStateChanged(0, ActorsToUnload.Num(), true);
 	}
 }
 
@@ -263,7 +265,7 @@ bool IWorldPartitionActorLoaderInterface::ILoaderAdapter::ShouldActorBeLoaded(co
 	return true;
 };
 
-void IWorldPartitionActorLoaderInterface::ILoaderAdapter::PostLoadedStateChanged(int32 NumLoads, int32 NumUnloads)
+void IWorldPartitionActorLoaderInterface::ILoaderAdapter::PostLoadedStateChanged(int32 NumLoads, int32 NumUnloads, bool bClearTransactions)
 {
 	check(NumLoads || NumUnloads);
 
@@ -273,7 +275,7 @@ void IWorldPartitionActorLoaderInterface::ILoaderAdapter::PostLoadedStateChanged
 
 		if (NumUnloads)
 		{
-			if (!GIsTransacting && World->HasAnyFlags(RF_Transactional))
+			if (!GIsTransacting && bClearTransactions)
 			{
 				GEditor->ResetTransaction(LOCTEXT("UnloadingEditorActorResetTrans", "Editor Actors Unloaded"));
 			}
diff --git a/Engine/Source/Runtime/Engine/Private/WorldPartition/WorldPartitionHandle.cpp b/Engine/Source/Runtime/Engine/Private/WorldPartition/WorldPartitionHandle.cpp
--- a/Engine/Source/Runtime/Engine/Private/WorldPartition/WorldPartitionHandle.cpp
+++ b/Engine/Source/Runtime/Engine/Private/WorldPartition/WorldPartitionHandle.cpp
@@ -1,9 +1,9 @@
 // Copyright Epic Games, Inc. All Rights Reserved.
 
 #include "WorldPartition/WorldPartitionHandle.h"
-#include "Engine/Level.h"
 #include "WorldPartition/ActorDescContainerCollection.h"
-
+#include "Engine/Level.h"
+#include "Engine/World.h"
 
 #if WITH_EDITOR
 /**
@@ -214,6 +214,8 @@ void FWorldPartitionLoadingContext::FDeferred::UnregisterActor(FWorldPartitionAc
 		NumUnregistrations++;
 
 		check(!ContainerOps.FindChecked(Container).Registrations.Contains(ActorDesc));
+
+		bNeedsClearTransactions |= Actor->GetTypedOuter<UWorld>()->HasAnyFlags(RF_Transactional);
 	}
 }
 
diff --git a/Engine/Source/Runtime/Engine/Public/WorldPartition/WorldPartitionActorLoaderInterface.h b/Engine/Source/Runtime/Engine/Public/WorldPartition/WorldPartitionActorLoaderInterface.h
--- a/Engine/Source/Runtime/Engine/Public/WorldPartition/WorldPartitionActorLoaderInterface.h
+++ b/Engine/Source/Runtime/Engine/Public/WorldPartition/WorldPartitionActorLoaderInterface.h
@@ -61,7 +61,7 @@ public:
 		virtual bool PassActorDescFilter(const FWorldPartitionHandle& Actor) const;
 		void RefreshLoadedState();
 
-		void PostLoadedStateChanged(int32 NumLoads, int32 NumUnloads);
+		void PostLoadedStateChanged(int32 NumLoads, int32 NumUnloads, bool bClearTransactions);
 		void AddReferenceToActor(FWorldPartitionHandle& Actor);
 		void RemoveReferenceToActor(FWorldPartitionHandle& Actor);
 		void OnRefreshLoadedState(bool bFromUserOperation);
diff --git a/Engine/Source/Runtime/Engine/Public/WorldPartition/WorldPartitionHandle.h b/Engine/Source/Runtime/Engine/Public/WorldPartition/WorldPartitionHandle.h
--- a/Engine/Source/Runtime/Engine/Public/WorldPartition/WorldPartitionHandle.h
+++ b/Engine/Source/Runtime/Engine/Public/WorldPartition/WorldPartitionHandle.h
@@ -54,6 +54,7 @@ public:
 
 		int32 GetNumRegistrations() const { return NumRegistrations; }
 		int32 GetNumUnregistrations() const { return NumUnregistrations; }
+		bool GetNeedsClearTransactions() const { return bNeedsClearTransactions; }
 
 	private:
 		virtual void RegisterActor(FWorldPartitionActorDesc* ActorDesc) override;
@@ -69,6 +70,7 @@ public:
 
 		int32 NumRegistrations = 0;
 		int32 NumUnregistrations = 0;
+		bool bNeedsClearTransactions = false;
 	};
 
 	/**
