Fix incorrect "stat fps" output when r.GTSyncType 2 is used on some platforms
 - Not all platforms implement type 2 syncing, so RHIGetFrameTime() returns zero.
 - Workaround the bug by ignoring zero from this function.

diff --git a/Engine/Source/Runtime/Engine/Private/UnrealEngine.cpp b/Engine/Source/Runtime/Engine/Private/UnrealEngine.cpp
--- a/Engine/Source/Runtime/Engine/Private/UnrealEngine.cpp
+++ b/Engine/Source/Runtime/Engine/Private/UnrealEngine.cpp
@@ -641,7 +641,11 @@ void CalculateFPSTimings()
 	static auto CVarVsync = IConsoleManager::Get().FindConsoleVariable(TEXT("r.VSync"));
 	if (CVarGTSyncType->GetInt() == 2 && CVarVsync->GetInt() != 0)
 	{
-		FrameTimeMS = RHIGetFrameTime();
+		float RHIFrameTime = RHIGetFrameTime();
+		if (RHIFrameTime != 0)
+		{
+			FrameTimeMS = RHIFrameTime;
+		}
 	}
 	
 	// A 3/4, 1/4 split gets close to a simple 10 frame moving average
