Improve removal cost for chaos proxy objects.

diff --git a/Engine/Source/Runtime/Engine/Private/PhysicsEngine/Experimental/PhysScene_Chaos.cpp b/Engine/Source/Runtime/Engine/Private/PhysicsEngine/Experimental/PhysScene_Chaos.cpp
--- a/Engine/Source/Runtime/Engine/Private/PhysicsEngine/Experimental/PhysScene_Chaos.cpp
+++ b/Engine/Source/Runtime/Engine/Private/PhysicsEngine/Experimental/PhysScene_Chaos.cpp
@@ -1027,11 +1027,11 @@ void FPhysScene_Chaos::RemoveFromComponentMaps(IPhysicsProxyBase* InObject)
 		TArray<IPhysicsProxyBase*>* ProxyArray = ComponentToPhysicsProxyMap.Find(*Component);
 		if (ProxyArray)
 		{
-			ProxyArray->Remove(InObject);
+			ProxyArray->RemoveSingleSwap(InObject, false);
 			if (ProxyArray->Num() == 0)
 			{
 				ComponentToPhysicsProxyMap.Remove(*Component);
-			}		
+			}
 		}
 	}
 
