LevelInstance: reenable partial loading following 24801297.

diff --git a/Engine/Source/Runtime/Engine/Private/LevelInstance/LevelInstanceLevelStreaming.cpp b/Engine/Source/Runtime/Engine/Private/LevelInstance/LevelInstanceLevelStreaming.cpp
--- a/Engine/Source/Runtime/Engine/Private/LevelInstance/LevelInstanceLevelStreaming.cpp
+++ b/Engine/Source/Runtime/Engine/Private/LevelInstance/LevelInstanceLevelStreaming.cpp
@@ -27,7 +27,7 @@
 #include "Misc/ScopeExit.h"
 #include "UObject/LinkerLoad.h"
 
-static bool GDisableLevelInstanceEditorPartialLoading = true;
+static bool GDisableLevelInstanceEditorPartialLoading = false;
 FAutoConsoleVariableRef CVarDisableLevelInstanceEditorPartialLoading(
 	TEXT("wp.Editor.DisableLevelInstanceEditorPartialLoading"),
 	GDisableLevelInstanceEditorPartialLoading,
