Partial undo of CL 24796093, Removed scope-based evaluation lock

diff --git a/Engine/Plugins/Animation/AnimationData/Source/AnimationData/Private/AnimSequencerController.cpp b/Engine/Plugins/Animation/AnimationData/Source/AnimationData/Private/AnimSequencerController.cpp
--- a/Engine/Plugins/Animation/AnimationData/Source/AnimationData/Private/AnimSequencerController.cpp
+++ b/Engine/Plugins/Animation/AnimationData/Source/AnimationData/Private/AnimSequencerController.cpp
@@ -53,7 +53,7 @@ void UAnimSequencerController::OpenBracket(const FText& InTitle, bool bShouldTra
 	{
 		FBracketPayload Payload;
 		Payload.Description = InTitle.ToString();
-		BracketLock = MakeUnique<IAnimationDataModel::FEvaluationAndModificationLock>(*ModelInterface);
+
 		Model->GetNotifier().Notify(EAnimDataModelNotifyType::BracketOpened, Payload);
 	}
 
@@ -83,7 +83,6 @@ void UAnimSequencerController::CloseBracket(bool bShouldTransact /*= true*/)
 			ChangeTransactor.CloseTransaction();
 		}
 		
-		BracketLock.Reset();
 		Model->GetNotifier().Notify(EAnimDataModelNotifyType::BracketClosed);
 	}
 }
diff --git a/Engine/Plugins/Animation/AnimationData/Source/AnimationData/Private/AnimSequencerController.h b/Engine/Plugins/Animation/AnimationData/Source/AnimationData/Private/AnimSequencerController.h
--- a/Engine/Plugins/Animation/AnimationData/Source/AnimationData/Private/AnimSequencerController.h
+++ b/Engine/Plugins/Animation/AnimationData/Source/AnimationData/Private/AnimSequencerController.h
@@ -129,8 +129,6 @@ private:
 	UPROPERTY(transient)
 	TScriptInterface<IAnimationDataModel> ModelInterface;
 
-	TUniquePtr<IAnimationDataModel::FEvaluationAndModificationLock> BracketLock;
-
 	friend class FAnimDataControllerTestBase;
 	friend UE::Anim::FOpenBracketAction;
 	friend UE::Anim::FCloseBracketAction;
diff --git a/Engine/Source/Runtime/Engine/Private/Animation/AnimSequence.cpp b/Engine/Source/Runtime/Engine/Private/Animation/AnimSequence.cpp
--- a/Engine/Source/Runtime/Engine/Private/Animation/AnimSequence.cpp
+++ b/Engine/Source/Runtime/Engine/Private/Animation/AnimSequence.cpp
@@ -4662,14 +4662,17 @@ void UAnimSequence::OnModelModified(const EAnimDataModelNotifyType& NotifyType,
 		UpdateDependentStreamingAnimations();
 	};
 
-	auto HandleTrackDataChanged = [this, UpdateRawDataGuid]()
+	auto HandleTrackDataChanged = [this, UpdateRawDataGuid](bool bWasReset = false)
 	{
 		if (NotifyCollector.IsNotWithinBracket())
 		{
 			CalculateNumberOfSampledKeys();
 			UpdateRawDataGuid(RegenerateGUID);
 			ClearAllCachedCookedPlatformData();
-			BeginCacheDerivedDataForCurrentPlatform();
+			if(!bWasReset)
+			{
+				BeginCacheDerivedDataForCurrentPlatform();
+			}
 		}
 	};
 
@@ -4798,7 +4801,7 @@ void UAnimSequence::OnModelModified(const EAnimDataModelNotifyType& NotifyType,
 					const bool bWasModelReset = NotifyCollector.Contains(EAnimDataModelNotifyType::Reset);
 					UpdateRawDataGuid(bWasModelReset ? GenerateNewGUID : RegenerateGUID);
 					ClearCompressedCurveData();
-					HandleTrackDataChanged();
+					HandleTrackDataChanged(bWasModelReset);
 				}
 			}
 
