Cleaning up object references during replay scrub

diff --git a/Engine/Source/Runtime/Engine/Private/World.cpp b/Engine/Source/Runtime/Engine/Private/World.cpp
--- a/Engine/Source/Runtime/Engine/Private/World.cpp
+++ b/Engine/Source/Runtime/Engine/Private/World.cpp
@@ -4333,6 +4333,22 @@ void UWorld::HandleTimelineScrubbed()
 			FXSystemComponent->DeactivateImmediate();
 		}
 	}
+
+	// Clearing game state references
+	if (GameState && !IsValid(GameState))
+	{
+		GameState = nullptr;
+	}
+
+	// Clear collections separately as the game state could differ (instant replay)
+	for (FLevelCollection& Collection : LevelCollections)
+	{
+		const AGameStateBase* CollectionGameState = Collection.GetGameState();
+		if (CollectionGameState && !IsValid(CollectionGameState))
+		{
+			Collection.SetGameState(nullptr);
+		}
+	}
 }
 
 bool UWorld::HandleDemoScrubCommand(const TCHAR* Cmd, FOutputDevice& Ar, UWorld* InWorld)
