Addition of a console variable to toggle navmesh resolution.

diff --git a/Engine/Source/Runtime/Engine/Private/AI/Navigation/NavigationModifier.cpp b/Engine/Source/Runtime/Engine/Private/AI/Navigation/NavigationModifier.cpp
--- a/Engine/Source/Runtime/Engine/Private/AI/Navigation/NavigationModifier.cpp
+++ b/Engine/Source/Runtime/Engine/Private/AI/Navigation/NavigationModifier.cpp
@@ -18,6 +18,13 @@
 // should be less than voxel size (recast navmesh)
 static const FVector::FReal CONVEX_HULL_POINTS_MIN_DISTANCE_SQ = 4.0f * 4.0f;
 
+bool FCompositeNavModifier::bEnableNavMeshResolutions = true;
+
+namespace UE::Navigation::Private
+{
+	FAutoConsoleVariableRef CVarEnableNavMeshResolutions(TEXT("ai.nav.EnableNavMeshResolutions"), FCompositeNavModifier::bEnableNavMeshResolutions, TEXT("When set to false, navmesh resoutions will be ignored."), ECVF_Default);
+}
+
 //----------------------------------------------------------------------//
 // FNavigationLinkBase
 //----------------------------------------------------------------------//
diff --git a/Engine/Source/Runtime/Engine/Public/AI/NavigationModifier.h b/Engine/Source/Runtime/Engine/Public/AI/NavigationModifier.h
--- a/Engine/Source/Runtime/Engine/Public/AI/NavigationModifier.h
+++ b/Engine/Source/Runtime/Engine/Public/AI/NavigationModifier.h
@@ -332,7 +332,7 @@ struct ENGINE_API FCompositeNavModifier : public FNavigationModifier
 	FORCEINLINE void SetFillCollisionUnderneathForNavmesh(bool bValue) { bFillCollisionUnderneathForNavmesh = bValue; }
 	FORCEINLINE bool GetMaskFillCollisionUnderneathForNavmesh() const { return bMaskFillCollisionUnderneathForNavmesh; }
 	FORCEINLINE void SetMaskFillCollisionUnderneathForNavmesh(bool bValue) { bMaskFillCollisionUnderneathForNavmesh = bValue; }
-	FORCEINLINE ENavigationDataResolution GetNavMeshResolution() const { return NavMeshResolution; }
+	FORCEINLINE ENavigationDataResolution GetNavMeshResolution() const { return bEnableNavMeshResolutions ? NavMeshResolution : ENavigationDataResolution::Invalid; }
 	FORCEINLINE void SetNavMeshResolution(ENavigationDataResolution Resolution) { NavMeshResolution = Resolution; }
 	FORCEINLINE void ReserveForAdditionalAreas(int32 AdditionalElementsCount) { Areas.Reserve(Areas.Num() + AdditionalElementsCount); }
 
@@ -356,6 +356,8 @@ public:
 	// Gathers per instance data for navigation area modifiers in a specified area box
 	FNavDataPerInstanceTransformDelegate NavDataPerInstanceTransformDelegate;
 
+	static bool bEnableNavMeshResolutions;
+	
 private:
 	TArray<FAreaNavModifier> Areas;
 	TArray<FSimpleLinkNavModifier> SimpleLinks;
