World Partition: delay loading state refresh when in a transaction in all cases.

diff --git a/Engine/Source/Runtime/Engine/Private/WorldPartition/DataLayer/DataLayerInstance.cpp b/Engine/Source/Runtime/Engine/Private/WorldPartition/DataLayer/DataLayerInstance.cpp
--- a/Engine/Source/Runtime/Engine/Private/WorldPartition/DataLayer/DataLayerInstance.cpp
+++ b/Engine/Source/Runtime/Engine/Private/WorldPartition/DataLayer/DataLayerInstance.cpp
@@ -5,13 +5,13 @@
 #include "Internationalization/Text.h"
 #include "UObject/UnrealType.h"
 #include "WorldPartition/WorldPartitionLog.h"
+#include "WorldPartition/WorldPartitionActorLoaderInterface.h"
 #include "WorldPartition/DataLayer/DataLayerUtils.h"
 #include "WorldPartition/DataLayer/WorldDataLayers.h"
 #include "WorldPartition/ErrorHandling/WorldPartitionStreamingGenerationErrorHandler.h"
 
 #include UE_INLINE_GENERATED_CPP_BY_NAME(DataLayerInstance)
 
-
 #define LOCTEXT_NAMESPACE "DataLayer"
 
 UDataLayerInstance::UDataLayerInstance(const FObjectInitializer& ObjectInitializer)
@@ -54,6 +54,21 @@ AWorldDataLayers* UDataLayerInstance::GetOuterWorldDataLayers() const
 }
 
 #if WITH_EDITOR
+void UDataLayerInstance::PreEditUndo()
+{
+	Super::PreEditUndo();
+	bUndoIsLoadedInEditor = bIsLoadedInEditor;
+}
+
+void UDataLayerInstance::PostEditUndo()
+{
+	Super::PostEditUndo();
+	if (bIsLoadedInEditor != bUndoIsLoadedInEditor)
+	{
+		IWorldPartitionActorLoaderInterface::RefreshLoadedState(true);
+	}
+}
+
 void UDataLayerInstance::SetVisible(bool bInIsVisible)
 {
 	if (bIsVisible != bInIsVisible)
diff --git a/Engine/Source/Runtime/Engine/Private/WorldPartition/Filter/WorldPartitionActorFilter.cpp b/Engine/Source/Runtime/Engine/Private/WorldPartition/Filter/WorldPartitionActorFilter.cpp
--- a/Engine/Source/Runtime/Engine/Private/WorldPartition/Filter/WorldPartitionActorFilter.cpp
+++ b/Engine/Source/Runtime/Engine/Private/WorldPartition/Filter/WorldPartitionActorFilter.cpp
@@ -1,7 +1,6 @@
 // Copyright Epic Games, Inc. All Rights Reserved.
 
 #include "WorldPartition/Filter/WorldPartitionActorFilter.h"
-#include "Containers/Ticker.h"
 #include "WorldPartition/WorldPartitionActorLoaderInterface.h"
 
 #if WITH_EDITORONLY_DATA
@@ -58,24 +57,7 @@ FWorldPartitionActorFilter& FWorldPartitionActorFilter::operator=(FWorldPartitio
 #if WITH_EDITORONLY_DATA
 void FWorldPartitionActorFilter::RequestFilterRefresh(bool bIsFromUserChange)
 {
-	if (GUndo || GIsTransacting)
-	{
-		FTSTicker::GetCoreTicker().AddTicker(FTickerDelegate::CreateLambda([=](float DeltaTime)
-		{
-			if (GUndo || GIsTransacting)
-			{
-				// Tick again we are still inside a transaction...
-				return true;
-			}
-			
-			IWorldPartitionActorLoaderInterface::RefreshLoadedState(bIsFromUserChange);
-			return false;
-		}));
-	}
-	else
-	{
-		IWorldPartitionActorLoaderInterface::RefreshLoadedState(bIsFromUserChange);
-	}
+	IWorldPartitionActorLoaderInterface::RefreshLoadedState(bIsFromUserChange);
 }
 
 void FWorldPartitionActorFilter::AddChildFilter(const FGuid& InGuid, FWorldPartitionActorFilter* InChildFilter)
diff --git a/Engine/Source/Runtime/Engine/Private/WorldPartition/WorldPartitionActorLoaderInterface.cpp b/Engine/Source/Runtime/Engine/Private/WorldPartition/WorldPartitionActorLoaderInterface.cpp
--- a/Engine/Source/Runtime/Engine/Private/WorldPartition/WorldPartitionActorLoaderInterface.cpp
+++ b/Engine/Source/Runtime/Engine/Private/WorldPartition/WorldPartitionActorLoaderInterface.cpp
@@ -7,6 +7,7 @@
 #if WITH_EDITOR
 #include "Editor.h"
 #include "Engine/Level.h"
+#include "Containers/Ticker.h"
 #include "Misc/ScopedSlowTask.h"
 #include "WorldPartition/WorldPartition.h"
 #endif
@@ -322,7 +323,22 @@ void IWorldPartitionActorLoaderInterface::ILoaderAdapter::OnRefreshLoadedState(b
 
 void IWorldPartitionActorLoaderInterface::RefreshLoadedState(bool bIsFromUserChange)
 {
-	ActorLoaderInterfaceRefreshState.Broadcast(bIsFromUserChange);
+	if (!GUndo && !GIsTransacting)
+	{
+		ActorLoaderInterfaceRefreshState.Broadcast(bIsFromUserChange);
+	}
+	else
+	{
+		FTSTicker::GetCoreTicker().AddTicker(FTickerDelegate::CreateLambda([=](float DeltaTime)
+		{
+			if (!GUndo && !GIsTransacting)
+			{
+				ActorLoaderInterfaceRefreshState.Broadcast(bIsFromUserChange);
+				return false;
+			}
+			return true;
+		}));
+	}
 }
 
 IWorldPartitionActorLoaderInterface::FActorReferenceMap& IWorldPartitionActorLoaderInterface::ILoaderAdapter::GetContainerReferences(UActorDescContainer* InContainer)
diff --git a/Engine/Source/Runtime/Engine/Public/WorldPartition/DataLayer/DataLayerInstance.h b/Engine/Source/Runtime/Engine/Public/WorldPartition/DataLayer/DataLayerInstance.h
--- a/Engine/Source/Runtime/Engine/Public/WorldPartition/DataLayer/DataLayerInstance.h
+++ b/Engine/Source/Runtime/Engine/Public/WorldPartition/DataLayer/DataLayerInstance.h
@@ -66,6 +66,9 @@ public:
 	AWorldDataLayers* GetOuterWorldDataLayers() const;
 
 #if WITH_EDITOR
+	virtual void PreEditUndo() override;
+	virtual void PostEditUndo() override;
+
 	void SetVisible(bool bIsVisible);
 	void SetIsInitiallyVisible(bool bIsInitiallyVisible);
 	void SetIsLoadedInEditor(bool bIsLoadedInEditor, bool bFromUserChange);
@@ -199,6 +202,7 @@ protected:
 	/** Wheter the data layer is loaded in editor (user setting) */
 	UPROPERTY(Transient)
 	uint32 bIsLoadedInEditor : 1;
+	uint32 bUndoIsLoadedInEditor : 1;
 
 	/** Whether this data layer editor visibility was changed by a user operation */
 	UPROPERTY(Transient)
@@ -218,5 +222,4 @@ private:
 
 	UPROPERTY(Transient)
 	TArray<TObjectPtr<UDataLayerInstance>> Children;
-
 };
\ No newline at end of file
