No PCH build fix

#rnx

diff --git a/Engine/Source/Runtime/Engine/Private/Animation/AnimCompositeBase.cpp b/Engine/Source/Runtime/Engine/Private/Animation/AnimCompositeBase.cpp
--- a/Engine/Source/Runtime/Engine/Private/Animation/AnimCompositeBase.cpp
+++ b/Engine/Source/Runtime/Engine/Private/Animation/AnimCompositeBase.cpp
@@ -13,6 +13,7 @@
 #include "Animation/AnimationPoseData.h"
 #include "Animation/AttributesRuntime.h"
 #include "EngineLogs.h"
+#include "UObject/LinkerLoad.h"
 
 #include UE_INLINE_GENERATED_CPP_BY_NAME(AnimCompositeBase)
 
