FWorldPartitionActorDesc::IsDefaultActorDesc

diff --git a/Engine/Source/Runtime/Engine/Public/WorldPartition/WorldPartitionActorDesc.h b/Engine/Source/Runtime/Engine/Public/WorldPartition/WorldPartitionActorDesc.h
--- a/Engine/Source/Runtime/Engine/Public/WorldPartition/WorldPartitionActorDesc.h
+++ b/Engine/Source/Runtime/Engine/Public/WorldPartition/WorldPartitionActorDesc.h
@@ -150,6 +150,8 @@ PRAGMA_ENABLE_DEPRECATION_WARNINGS
 	FName GetActorLabelOrName() const;
 	FName GetDisplayClassName() const;
 
+	inline bool IsDefaultActorDesc() const { return bIsDefaultActorDesc; }
+
 	virtual bool IsContainerInstance() const { return false; }
 	virtual FName GetLevelPackage() const { return NAME_None; }
 	virtual const FWorldPartitionActorFilter* GetContainerFilter() const { return nullptr; }
