Implement GetResourceSizeEx for Animation assets used by "obj list" and "ListAnims" commands.

#rnx
#tests replay,PIE

diff --git a/Engine/Source/Runtime/Engine/Classes/Animation/AnimComposite.h b/Engine/Source/Runtime/Engine/Classes/Animation/AnimComposite.h
--- a/Engine/Source/Runtime/Engine/Classes/Animation/AnimComposite.h
+++ b/Engine/Source/Runtime/Engine/Classes/Animation/AnimComposite.h
@@ -35,6 +35,10 @@ public:
 	TObjectPtr<UAnimSequence> PreviewBasePose;
 #endif // WITH_EDITORONLY_DATA
 
+	//~ Begin UObject Interface
+	virtual void GetResourceSizeEx(FResourceSizeEx& CumulativeResourceSize) override;
+	//~ End UObject Interface
+
 	//~ Begin UAnimSequenceBase Interface
 	ENGINE_API virtual void HandleAssetPlayerTickedInternal(FAnimAssetTickContext &Context, const float PreviousTime, const float MoveDelta, const FAnimTickRecord &Instance, struct FAnimNotifyQueue& NotifyQueue) const override;
 
diff --git a/Engine/Source/Runtime/Engine/Classes/Animation/AnimCompositeBase.h b/Engine/Source/Runtime/Engine/Classes/Animation/AnimCompositeBase.h
--- a/Engine/Source/Runtime/Engine/Classes/Animation/AnimCompositeBase.h
+++ b/Engine/Source/Runtime/Engine/Classes/Animation/AnimCompositeBase.h
@@ -338,6 +338,8 @@ struct FAnimTrack
 
 	/** return true if anim notify is available */
 	bool IsNotifyAvailable() const;
+
+	ENGINE_API int32 GetTotalBytesUsed() const;
 };
 
 UCLASS(abstract, MinimalAPI)
diff --git a/Engine/Source/Runtime/Engine/Classes/Animation/AnimMontage.h b/Engine/Source/Runtime/Engine/Classes/Animation/AnimMontage.h
--- a/Engine/Source/Runtime/Engine/Classes/Animation/AnimMontage.h
+++ b/Engine/Source/Runtime/Engine/Classes/Animation/AnimMontage.h
@@ -734,6 +734,7 @@ public:
 #if WITH_EDITOR
 	virtual void PostEditChangeProperty(FPropertyChangedEvent& PropertyChangedEvent) override;
 #endif // WITH_EDITOR
+	virtual void GetResourceSizeEx(FResourceSizeEx& CumulativeResourceSize) override;
 	//~ End UObject Interface
 
 	//~ Begin AnimSequenceBase Interface
diff --git a/Engine/Source/Runtime/Engine/Classes/Animation/AnimSequenceBase.h b/Engine/Source/Runtime/Engine/Classes/Animation/AnimSequenceBase.h
--- a/Engine/Source/Runtime/Engine/Classes/Animation/AnimSequenceBase.h
+++ b/Engine/Source/Runtime/Engine/Classes/Animation/AnimSequenceBase.h
@@ -79,6 +79,7 @@ public:
 #if WITH_EDITORONLY_DATA
 	static void DeclareConstructClasses(TArray<FTopLevelAssetPath>& OutConstructClasses, const UClass* SpecificSubclass);
 #endif
+	virtual void GetResourceSizeEx(FResourceSizeEx& CumulativeResourceSize) override;
 	//~ End UObject Interface
 
 	/** Returns the total play length of the montage, if played back with a speed of 1.0. */
diff --git a/Engine/Source/Runtime/Engine/Classes/Animation/BlendSpace.h b/Engine/Source/Runtime/Engine/Classes/Animation/BlendSpace.h
--- a/Engine/Source/Runtime/Engine/Classes/Animation/BlendSpace.h
+++ b/Engine/Source/Runtime/Engine/Classes/Animation/BlendSpace.h
@@ -464,6 +464,7 @@ public:
 	virtual void PreEditChange(FProperty* PropertyAboutToChange) override;
 	virtual void PostEditChangeProperty( struct FPropertyChangedEvent& PropertyChangedEvent) override;
 #endif // WITH_EDITOR
+	virtual void GetResourceSizeEx(FResourceSizeEx& CumulativeResourceSize) override;
 	//~ End UObject Interface
 
 	//~ Begin UAnimationAsset Interface
diff --git a/Engine/Source/Runtime/Engine/Private/Animation/AnimComposite.cpp b/Engine/Source/Runtime/Engine/Private/Animation/AnimComposite.cpp
--- a/Engine/Source/Runtime/Engine/Private/Animation/AnimComposite.cpp
+++ b/Engine/Source/Runtime/Engine/Private/Animation/AnimComposite.cpp
@@ -66,6 +66,13 @@ void UAnimComposite::GetAnimNotifiesFromDeltaPositions(const float& PreviousPosi
 	}
 }
 
+void UAnimComposite::GetResourceSizeEx(FResourceSizeEx& CumulativeResourceSize)
+{
+	Super::GetResourceSizeEx(CumulativeResourceSize);
+
+	CumulativeResourceSize.AddDedicatedSystemMemoryBytes(AnimationTrack.GetTotalBytesUsed());
+}
+
 void UAnimComposite::HandleAssetPlayerTickedInternal(FAnimAssetTickContext &Context, const float PreviousTime, const float MoveDelta, const FAnimTickRecord &Instance, struct FAnimNotifyQueue& NotifyQueue) const
 {
 	Super::HandleAssetPlayerTickedInternal(Context, PreviousTime, MoveDelta, Instance, NotifyQueue);
diff --git a/Engine/Source/Runtime/Engine/Private/Animation/AnimCompositeBase.cpp b/Engine/Source/Runtime/Engine/Private/Animation/AnimCompositeBase.cpp
--- a/Engine/Source/Runtime/Engine/Private/Animation/AnimCompositeBase.cpp
+++ b/Engine/Source/Runtime/Engine/Private/Animation/AnimCompositeBase.cpp
@@ -706,6 +706,11 @@ bool FAnimTrack::IsNotifyAvailable() const
 	return false;
 }
 
+int32 FAnimTrack::GetTotalBytesUsed() const
+{
+	return AnimSegments.GetAllocatedSize();
+}
+
 bool FAnimTrack::IsValidToAdd(const UAnimSequenceBase* SequenceBase) const
 {
 	bool bValid = false;
diff --git a/Engine/Source/Runtime/Engine/Private/Animation/AnimMontage.cpp b/Engine/Source/Runtime/Engine/Private/Animation/AnimMontage.cpp
--- a/Engine/Source/Runtime/Engine/Private/Animation/AnimMontage.cpp
+++ b/Engine/Source/Runtime/Engine/Private/Animation/AnimMontage.cpp
@@ -759,6 +759,20 @@ void UAnimMontage::PropagateChanges()
 }
 #endif // WITH_EDITOR
 
+void UAnimMontage::GetResourceSizeEx(FResourceSizeEx& CumulativeResourceSize)
+{
+	Super::GetResourceSizeEx(CumulativeResourceSize);
+
+	CumulativeResourceSize.AddDedicatedSystemMemoryBytes(CompositeSections.GetAllocatedSize());
+	CumulativeResourceSize.AddDedicatedSystemMemoryBytes(SlotAnimTracks.GetAllocatedSize());
+	for (FSlotAnimationTrack& Slot : SlotAnimTracks)
+	{
+		CumulativeResourceSize.AddDedicatedSystemMemoryBytes(Slot.AnimTrack.GetTotalBytesUsed());
+	}
+	CumulativeResourceSize.AddDedicatedSystemMemoryBytes(BranchingPointMarkers.GetAllocatedSize());
+	CumulativeResourceSize.AddDedicatedSystemMemoryBytes(BranchingPointStateNotifyIndices.GetAllocatedSize());
+}
+
 bool UAnimMontage::IsValidAdditive() const
 {
 	// if first one is additive, this is additive
diff --git a/Engine/Source/Runtime/Engine/Private/Animation/AnimSequenceBase.cpp b/Engine/Source/Runtime/Engine/Private/Animation/AnimSequenceBase.cpp
--- a/Engine/Source/Runtime/Engine/Private/Animation/AnimSequenceBase.cpp
+++ b/Engine/Source/Runtime/Engine/Private/Animation/AnimSequenceBase.cpp
@@ -83,6 +83,17 @@ void UAnimSequenceBase::DeclareConstructClasses(TArray<FTopLevelAssetPath>& OutC
 		}
 	}
 }
+
+void UAnimSequenceBase::GetResourceSizeEx(FResourceSizeEx& CumulativeResourceSize)
+{
+	Super::GetResourceSizeEx(CumulativeResourceSize);
+
+	const FRawCurveTracks& CurveData = GetCurveData();
+	const UStruct* Struct = CurveData.StaticStruct();
+	ensure(Struct);
+	CumulativeResourceSize.AddDedicatedSystemMemoryBytes(Struct->GetStructureSize());
+	CumulativeResourceSize.AddDedicatedSystemMemoryBytes(CurveData.FloatCurves.GetAllocatedSize());
+}
 #endif
 
 void UAnimSequenceBase::PostLoad()
diff --git a/Engine/Source/Runtime/Engine/Private/Animation/BlendSpace.cpp b/Engine/Source/Runtime/Engine/Private/Animation/BlendSpace.cpp
--- a/Engine/Source/Runtime/Engine/Private/Animation/BlendSpace.cpp
+++ b/Engine/Source/Runtime/Engine/Private/Animation/BlendSpace.cpp
@@ -199,8 +199,19 @@ void UBlendSpace::PostEditChangeProperty(struct FPropertyChangedEvent& PropertyC
 	}
 	Super::PostEditChangeProperty(PropertyChangedEvent);
 }
+
 #endif // WITH_EDITOR
 
+void UBlendSpace::GetResourceSizeEx(FResourceSizeEx& CumulativeResourceSize)
+{
+	Super::GetResourceSizeEx(CumulativeResourceSize);
+	
+	CumulativeResourceSize.AddDedicatedSystemMemoryBytes(this->SampleData.GetAllocatedSize());
+	CumulativeResourceSize.AddDedicatedSystemMemoryBytes(this->GridSamples.GetAllocatedSize());
+	CumulativeResourceSize.AddDedicatedSystemMemoryBytes(this->BlendSpaceData.StaticStruct()->GetStructureSize());
+	CumulativeResourceSize.AddDedicatedSystemMemoryBytes(this->BlendSpaceData.Segments.GetAllocatedSize());
+	CumulativeResourceSize.AddDedicatedSystemMemoryBytes(this->BlendSpaceData.Triangles.GetAllocatedSize());
+}
 
 bool UBlendSpace::UpdateBlendSamples_Internal(
 	const FVector&            InBlendSpacePosition, 
