Added DetailMode Epic for more scalability options

- Mirrored to ParticleDetailMode (Cascade)
- Removed DetailMode_DEPRECATED (Cascade)

diff --git a/Engine/Source/Editor/Cascade/Private/Cascade.cpp b/Engine/Source/Editor/Cascade/Private/Cascade.cpp
--- a/Engine/Source/Editor/Cascade/Private/Cascade.cpp
+++ b/Engine/Source/Editor/Cascade/Private/Cascade.cpp
@@ -2386,6 +2386,12 @@ void FCascade::BindCommands()
 		FCanExecuteAction(),
 		FIsActionChecked::CreateSP(this, &FCascade::IsDetailModeChecked, DM_High));
 
+	ToolkitCommands->MapAction(
+		Commands.DetailMode_Epic,
+		FExecuteAction::CreateSP(this, &FCascade::OnDetailMode, DM_Epic),
+		FCanExecuteAction(),
+		FIsActionChecked::CreateSP(this, &FCascade::IsDetailModeChecked, DM_Epic));
+
 	ToolkitCommands->MapAction(
 		Commands.Significance_Critical,
 		FExecuteAction::CreateSP(this, &FCascade::OnSignificance, EParticleSignificanceLevel::Critical),
diff --git a/Engine/Source/Editor/Cascade/Private/CascadeActions.cpp b/Engine/Source/Editor/Cascade/Private/CascadeActions.cpp
--- a/Engine/Source/Editor/Cascade/Private/CascadeActions.cpp
+++ b/Engine/Source/Editor/Cascade/Private/CascadeActions.cpp
@@ -51,6 +51,7 @@ void FCascadeCommands::RegisterCommands()
 	UI_COMMAND(DetailMode_Low, "Low", "Select Low Detail Mode", EUserInterfaceActionType::RadioButton, FInputChord());
 	UI_COMMAND(DetailMode_Medium, "Medium", "Select Medium Detail Mode", EUserInterfaceActionType::RadioButton, FInputChord());
 	UI_COMMAND(DetailMode_High, "High", "Select High Detail Mode", EUserInterfaceActionType::RadioButton, FInputChord());
+	UI_COMMAND(DetailMode_Epic, "Epic", "Select Epic Detail Mode", EUserInterfaceActionType::RadioButton, FInputChord());
 	UI_COMMAND(Significance_Critical, "Critical", "Require >= Critical Significance", EUserInterfaceActionType::RadioButton, FInputChord());
 	UI_COMMAND(Significance_High, "High", "Require High >= Significance", EUserInterfaceActionType::RadioButton, FInputChord());
 	UI_COMMAND(Significance_Medium, "Medium", "Require >= Medium Significance", EUserInterfaceActionType::RadioButton, FInputChord());
diff --git a/Engine/Source/Editor/Cascade/Private/CascadeActions.h b/Engine/Source/Editor/Cascade/Private/CascadeActions.h
--- a/Engine/Source/Editor/Cascade/Private/CascadeActions.h
+++ b/Engine/Source/Editor/Cascade/Private/CascadeActions.h
@@ -65,6 +65,7 @@ public:
 	TSharedPtr<FUICommandInfo> DetailMode_Low;
 	TSharedPtr<FUICommandInfo> DetailMode_Medium;
 	TSharedPtr<FUICommandInfo> DetailMode_High;
+	TSharedPtr<FUICommandInfo> DetailMode_Epic;
 	TSharedPtr<FUICommandInfo> Significance_Critical;
 	TSharedPtr<FUICommandInfo> Significance_High;
 	TSharedPtr<FUICommandInfo> Significance_Medium;
diff --git a/Engine/Source/Editor/Cascade/Private/CascadePreviewViewportClient.cpp b/Engine/Source/Editor/Cascade/Private/CascadePreviewViewportClient.cpp
--- a/Engine/Source/Editor/Cascade/Private/CascadePreviewViewportClient.cpp
+++ b/Engine/Source/Editor/Cascade/Private/CascadePreviewViewportClient.cpp
@@ -467,9 +467,9 @@ void FCascadeEdPreviewViewportClient::Draw(FViewport* InViewport, FCanvas* Canva
 
 	int32 DetailMode = CascadePtr.Pin()->GetDetailMode();
 	
-	if (DetailMode != DM_High)
+	if (DetailMode != DM_Epic)
 	{
-		FString DetailModeOutput = FString::Printf(TEXT("DETAIL MODE: %s"), (DetailMode == DM_Medium)? TEXT("MEDIUM"): TEXT("LOW"));
+		FString DetailModeOutput = FString::Printf(TEXT("DETAIL MODE: %s"), (DetailMode == DM_High) ? TEXT("HIGH") : (DetailMode == DM_Medium) ? TEXT("MEDIUM") : TEXT("LOW"));
 		TextItem.SetColor( FLinearColor::Red );
 		TextItem.Text = FText::FromString( DetailModeOutput );
 		Canvas->DrawItem(TextItem, 5.0f, MiscMessageY);
diff --git a/Engine/Source/Editor/Cascade/Private/SCascadePreviewToolbar.cpp b/Engine/Source/Editor/Cascade/Private/SCascadePreviewToolbar.cpp
--- a/Engine/Source/Editor/Cascade/Private/SCascadePreviewToolbar.cpp
+++ b/Engine/Source/Editor/Cascade/Private/SCascadePreviewToolbar.cpp
@@ -109,6 +109,7 @@ TSharedRef<SWidget> SCascadePreviewViewportToolBar::GenerateViewMenu() const
 				Menu.AddMenuEntry(FCascadeCommands::Get().DetailMode_Low);
 				Menu.AddMenuEntry(FCascadeCommands::Get().DetailMode_Medium);
 				Menu.AddMenuEntry(FCascadeCommands::Get().DetailMode_High);
+				Menu.AddMenuEntry(FCascadeCommands::Get().DetailMode_Epic);
 			}
 			Menu.EndSection();
 		}
diff --git a/Engine/Source/Editor/UnrealEd/Private/EditorServer.cpp b/Engine/Source/Editor/UnrealEd/Private/EditorServer.cpp
--- a/Engine/Source/Editor/UnrealEd/Private/EditorServer.cpp
+++ b/Engine/Source/Editor/UnrealEd/Private/EditorServer.cpp
@@ -6630,7 +6630,7 @@ bool UEditorEngine::HandleSetDetailModeCommand( const TCHAR* Str, FOutputDevice&
 {
 	TArray<AActor*> ActorsToDeselect;
 
-	uint8 ParsedDetailMode = DM_High;
+	uint8 ParsedDetailMode = DM_Epic;
 	if ( FParse::Value( Str, TEXT("MODE="), ParsedDetailMode ) )
 	{
 		for ( FSelectionIterator It( GetSelectedActorIterator() ) ; It ; ++It )
@@ -6674,7 +6674,7 @@ bool UEditorEngine::HandleSetDetailModeCommand( const TCHAR* Str, FOutputDevice&
 
 bool UEditorEngine::HandleSetDetailModeViewCommand( const TCHAR* Str, FOutputDevice& Ar, UWorld* InWorld )
 {
-	uint8 DM = DM_High;
+	uint8 DM = DM_Epic;
 	if ( FParse::Value( Str, TEXT("MODE="), DM ) )
 	{
 		DetailMode = (EDetailMode)DM;
diff --git a/Engine/Source/Runtime/Engine/Classes/Components/SceneComponent.h b/Engine/Source/Runtime/Engine/Classes/Components/SceneComponent.h
--- a/Engine/Source/Runtime/Engine/Classes/Components/SceneComponent.h
+++ b/Engine/Source/Runtime/Engine/Classes/Components/SceneComponent.h
@@ -37,6 +37,7 @@ enum EDetailMode : int
 	DM_Low UMETA(DisplayName="Low"),
 	DM_Medium UMETA(DisplayName="Medium"),
 	DM_High UMETA(DisplayName="High"),
+	DM_Epic UMETA(DisplayName="Epic"),
 	DM_MAX,
 };
 
diff --git a/Engine/Source/Runtime/Engine/Classes/Particles/ParticleEmitter.h b/Engine/Source/Runtime/Engine/Classes/Particles/ParticleEmitter.h
--- a/Engine/Source/Runtime/Engine/Classes/Particles/ParticleEmitter.h
+++ b/Engine/Source/Runtime/Engine/Classes/Particles/ParticleEmitter.h
@@ -92,11 +92,6 @@ DECLARE_STATS_GROUP(TEXT("Emitters"), STATGROUP_EmittersRT, STATCAT_Advanced);
 DECLARE_DWORD_COUNTER_STAT_EXTERN(TEXT("STAT_EmittersStatGroupTester"), STAT_EmittersStatGroupTester, STATGROUP_Emitters, ENGINE_API);
 DECLARE_DWORD_COUNTER_STAT_EXTERN(TEXT("STAT_EmittersRTStatGroupTester"), STAT_EmittersRTStatGroupTester, STATGROUP_EmittersRT, ENGINE_API);
 
-/* detail modes for emitters are now flags instead of a single enum
-	an emitter is shown if it is flagged for the current system scalability level
- */
-#define NUM_DETAILMODE_FLAGS 3
-
 UCLASS(hidecategories=Object, editinlinenew, abstract, MinimalAPI)
 class UParticleEmitter : public UObject
 {
@@ -171,10 +166,6 @@ class UParticleEmitter : public UObject
 	UPROPERTY(EditAnywhere, Category=Cascade)
 	uint8 bCollapsed:1;
 
-	/** If detail mode is >= system detail mode, primitive won't be rendered. */
-	UPROPERTY()
-	TEnumAsByte<EDetailMode> DetailMode_DEPRECATED;
-
 	/**
 	 *	The color of the emitter in the curve editor and debug rendering modes.
 	 */
@@ -262,7 +253,8 @@ class UParticleEmitter : public UObject
 		DetailModeDisplay = "";
 		DetailModeDisplay += DetailModeBitmask & (1 << EParticleDetailMode::PDM_Low) ? "Low, " : "";
 		DetailModeDisplay += DetailModeBitmask & (1 << EParticleDetailMode::PDM_Medium) ? "Medium, " : "";
-		DetailModeDisplay += DetailModeBitmask & (1 << EParticleDetailMode::PDM_High) ? "High" : "";
+		DetailModeDisplay += DetailModeBitmask & (1 << EParticleDetailMode::PDM_High) ? "High, " : "";
+		DetailModeDisplay += DetailModeBitmask & (1 << EParticleDetailMode::PDM_Epic) ? "Epic" : "";
 	}
 #endif // WITH_EDITORONLY_DATA
 	virtual void Serialize(FArchive& Ar)override;
diff --git a/Engine/Source/Runtime/Engine/Private/Particles/ParticleComponents.cpp b/Engine/Source/Runtime/Engine/Private/Particles/ParticleComponents.cpp
--- a/Engine/Source/Runtime/Engine/Private/Particles/ParticleComponents.cpp
+++ b/Engine/Source/Runtime/Engine/Private/Particles/ParticleComponents.cpp
@@ -124,6 +124,8 @@ const FGuid FParticleSystemCustomVersion::GUID(0x4A56EB40, 0x10F511DC, 0x92D3347
 // Register the custom version with core
 FCustomVersionRegistration GRegisterParticleSystemCustomVersion(FParticleSystemCustomVersion::GUID, FParticleSystemCustomVersion::LatestVersion, TEXT("ParticleSystemVer"));
 
+static_assert(EDetailMode::DM_MAX == EParticleDetailMode::PDM_MAX);
+
 //////////////////////////////////////////////////////////////////////////
 
 void UFXSystemAsset::PostInitProperties()
@@ -197,18 +199,6 @@ FAutoConsoleVariableRef CVarFXLWCTileRecache(
 	TEXT("Setting this value to 0 will remove this behavior but could introduce rendering & simulation artifacts.\n"),
 	ECVF_Default);
 
-int32 OldDetailModeToBitmask(int32 OldMode)
-{
-	// low = L+M+H, Medium = M+H, High = H
-	uint32 AllDetailModes = /*(1<<EParticleDetailMode::PDM_VeryLow) |*/ (1 << EParticleDetailMode::PDM_Low) | (1 << EParticleDetailMode::PDM_Medium) | (1 << EParticleDetailMode::PDM_High);
-	uint32 DetailModeBitmask = OldMode == EDetailMode::DM_Low ? AllDetailModes
-		: OldMode == EDetailMode::DM_Medium ? ((1 << EParticleDetailMode::PDM_Medium) | (1 << EParticleDetailMode::PDM_High))
-		: OldMode == EDetailMode::DM_High ? (1 << EParticleDetailMode::PDM_High) : AllDetailModes;
-
-	return DetailModeBitmask;
-}
-
-
 /** Whether to allow particle systems to perform work. */
 ENGINE_API bool GIsAllowingParticles = true;
 
@@ -1025,7 +1015,16 @@ void UParticleEmitter::PostLoad()
 	if (PSysVer < FParticleSystemCustomVersion::FixLegacySpawningBugs)
 	{
 		bUseLegacySpawningBehavior = true;
-	}	
+	}
+
+	if (PSysVer < FParticleSystemCustomVersion::AddEpicDetailMode)
+	{
+		// Init epic detail mode to enabled if high is set
+		if (DetailModeBitmask & (1 << EParticleDetailMode::PDM_High))
+		{
+			DetailModeBitmask |= (1 << EParticleDetailMode::PDM_Epic);
+		}
+	}
 
 	for (int32 LODIndex = 0; LODIndex < LODLevels.Num(); LODIndex++)
 	{
@@ -1052,13 +1051,7 @@ void UParticleEmitter::PostLoad()
 		}
 	}
 
-#if	WITH_EDITORONLY_DATA
-	// set up DetailModeFlags from deprecated DetailMode if needed
-	if (DetailModeBitmask == PDM_DefaultValue)
-	{
-		DetailModeBitmask = OldDetailModeToBitmask(DetailMode_DEPRECATED);
-	}
-
+#if WITH_EDITORONLY_DATA
 	UpdateDetailModeDisplayString();
 #endif
 
@@ -5866,7 +5859,7 @@ void UParticleSystemComponent::InitParticles()
 			if (Emitter)
 			{
 				FParticleEmitterInstance* Instance = NumInstances == 0 ? NULL : EmitterInstances[Idx];
-				check(GlobalDetailMode < NUM_DETAILMODE_FLAGS);
+				check(GlobalDetailMode < EParticleDetailMode::PDM_MAX);
 				const bool bDetailModeAllowsRendering = DetailMode <= GlobalDetailMode && (Emitter->DetailModeBitmask & (1 << GlobalDetailMode));
 				const bool bShouldCreateAndOrInit = bDetailModeAllowsRendering && Emitter->HasAnyEnabledLODs() && bCanEverRender;
 
diff --git a/Engine/Source/Runtime/Engine/Private/UnrealEngine.cpp b/Engine/Source/Runtime/Engine/Private/UnrealEngine.cpp
--- a/Engine/Source/Runtime/Engine/Private/UnrealEngine.cpp
+++ b/Engine/Source/Runtime/Engine/Private/UnrealEngine.cpp
@@ -9563,9 +9563,9 @@ bool UEngine::HandleTrackParticleRenderingStatsCommand( const TCHAR* Cmd, FOutpu
 	GTrackParticleRenderingStats = !GTrackParticleRenderingStats;
 	if (GTrackParticleRenderingStats)
 	{
-		if (GetCachedScalabilityCVars().DetailMode == DM_High)
+		if (GetCachedScalabilityCVars().DetailMode >= DM_High)
 		{
-			Ar.Logf(TEXT("Currently in high detail mode, note that particle stats will only be captured in medium or low detail modes (eg splitscreen)."));
+			Ar.Logf(TEXT("Currently in at least high detail mode, note that particle stats will only be captured in medium or low detail modes (eg splitscreen)."));
 		}
 		Ar.Logf(TEXT("Enabled particle render stat tracking with %.1fs between captures, min tracked time of %.4fs, use DUMPPARTICLERENDERINGSTATS to save results."),
 			GTimeBetweenParticleRenderStatCaptures, GMinParticleDrawTimeToTrack);
diff --git a/Engine/Source/Runtime/Engine/Public/ParticleHelper.h b/Engine/Source/Runtime/Engine/Public/ParticleHelper.h
--- a/Engine/Source/Runtime/Engine/Public/ParticleHelper.h
+++ b/Engine/Source/Runtime/Engine/Public/ParticleHelper.h
@@ -50,6 +50,7 @@ enum EParticleDetailMode : int
 	PDM_Low UMETA(DisplayName = "Low"),
 	PDM_Medium UMETA(DisplayName = "Medium"),
 	PDM_High UMETA(DisplayName = "High"),
+	PDM_Epic UMETA(DisplayName = "Epic"),
 	PDM_MAX UMETA(Hidden),
 };
 inline const int32 PDM_DefaultValue = 0xFFFF;
@@ -2703,6 +2704,7 @@ struct FParticleSystemCustomVersion
 		BeforeCustomVersionWasAdded = 0,
 		SkipCookingEmittersBasedOnDetailMode,	// skip emitter cooking if their detail mode doesn't match predefined
 		FixLegacySpawningBugs,					// fixing some spawning bugs but must keep old behavior around for existing systems.
+		AddEpicDetailMode,						// adding another bitmask entry to EParticleDetailMode
 
 		// -----<new versions can be added above this line>-------------------------------------------------
 		VersionPlusOne,
