Level Instance: don't set an instance guid on blueprint class default objects.

diff --git a/Engine/Source/Runtime/Engine/Private/Actor.cpp b/Engine/Source/Runtime/Engine/Private/Actor.cpp
--- a/Engine/Source/Runtime/Engine/Private/Actor.cpp
+++ b/Engine/Source/Runtime/Engine/Private/Actor.cpp
@@ -886,11 +886,14 @@ void AActor::Serialize(FArchive& Ar)
 			ActorGuid = FGuid::NewGuid();
 		}
 
-		check(!ActorInstanceGuid.IsValid());
-		const FGuid NewActorInstanceGuid = TLazySingleton<FActorInstanceGuidMapper>::Get().MapGuid(GetOuter()->GetPackage()->GetFName(), ActorGuid);
-		if (NewActorInstanceGuid != ActorGuid)
+		if (!IsTemplate())
 		{
-			ActorInstanceGuid = NewActorInstanceGuid;
+			check(!ActorInstanceGuid.IsValid());
+			const FGuid NewActorInstanceGuid = TLazySingleton<FActorInstanceGuidMapper>::Get().MapGuid(GetOuter()->GetPackage()->GetFName(), ActorGuid);
+			if (NewActorInstanceGuid != ActorGuid)
+			{
+				ActorInstanceGuid = NewActorInstanceGuid;
+			}
 		}
 
 		if (!CanChangeIsSpatiallyLoadedFlag() && (Ar.CustomVer(FUE5ReleaseStreamObjectVersion::GUID) < FUE5ReleaseStreamObjectVersion::ActorGridPlacementDeprecateDefaultValueFixup))
