Chaos - fix joint axis locked/limited axis selection in FPhysInterface_Chaos::SetAngularMotionLimitType_AssumesLocked

This was a knockon from a previous fix for a different issue in CL 23422440 (See FConeConstraint::UpdateConeLimit_AssumesLocked).

diff --git a/Engine/Source/Runtime/Engine/Private/PhysicsEngine/ConstraintInstance.cpp b/Engine/Source/Runtime/Engine/Private/PhysicsEngine/ConstraintInstance.cpp
--- a/Engine/Source/Runtime/Engine/Private/PhysicsEngine/ConstraintInstance.cpp
+++ b/Engine/Source/Runtime/Engine/Private/PhysicsEngine/ConstraintInstance.cpp
@@ -1079,15 +1079,14 @@ void FConstraintInstance::SetAngularDOFLimitScale(float InSwing1LimitScale, floa
 	{
 		if ( ProfileInstance.ConeLimit.Swing1Motion == ACM_Limited || ProfileInstance.ConeLimit.Swing2Motion == ACM_Limited )
 		{
-			// PhysX swing directions are different from Unreal's - so change here.
 			if (ProfileInstance.ConeLimit.Swing1Motion == ACM_Limited)
 			{
-				FPhysicsInterface::SetAngularMotionLimitType_AssumesLocked(InUnbrokenConstraint, ELimitAxis::Swing2, ACM_Limited);
+				FPhysicsInterface::SetAngularMotionLimitType_AssumesLocked(InUnbrokenConstraint, ELimitAxis::Swing1, ACM_Limited);
 			}
 
 			if (ProfileInstance.ConeLimit.Swing2Motion == ACM_Limited)
 			{
-				FPhysicsInterface::SetAngularMotionLimitType_AssumesLocked(InUnbrokenConstraint, ELimitAxis::Swing1, ACM_Limited);
+				FPhysicsInterface::SetAngularMotionLimitType_AssumesLocked(InUnbrokenConstraint, ELimitAxis::Swing2, ACM_Limited);
 			}
 		
 			//The limit values need to be clamped so it will be valid in PhysX
@@ -1100,12 +1099,12 @@ void FConstraintInstance::SetAngularDOFLimitScale(float InSwing1LimitScale, floa
 
 		if ( ProfileInstance.ConeLimit.Swing1Motion  == ACM_Locked )
 		{
-			FPhysicsInterface::SetAngularMotionLimitType_AssumesLocked(InUnbrokenConstraint, ELimitAxis::Swing2, ACM_Locked);
+			FPhysicsInterface::SetAngularMotionLimitType_AssumesLocked(InUnbrokenConstraint, ELimitAxis::Swing1, ACM_Locked);
 		}
 
 		if ( ProfileInstance.ConeLimit.Swing2Motion  == ACM_Locked )
 		{
-			FPhysicsInterface::SetAngularMotionLimitType_AssumesLocked(InUnbrokenConstraint, ELimitAxis::Swing1, ACM_Locked);
+			FPhysicsInterface::SetAngularMotionLimitType_AssumesLocked(InUnbrokenConstraint, ELimitAxis::Swing2, ACM_Locked);
 		}
 
 		if ( ProfileInstance.TwistLimit.TwistMotion == ACM_Limited )
diff --git a/Engine/Source/Runtime/Engine/Private/PhysicsEngine/Experimental/PhysInterface_Chaos.cpp b/Engine/Source/Runtime/Engine/Private/PhysicsEngine/Experimental/PhysInterface_Chaos.cpp
--- a/Engine/Source/Runtime/Engine/Private/PhysicsEngine/Experimental/PhysInterface_Chaos.cpp
+++ b/Engine/Source/Runtime/Engine/Private/PhysicsEngine/Experimental/PhysInterface_Chaos.cpp
@@ -262,6 +262,11 @@ Chaos::EJointMotionType ConvertMotionType(EAngularConstraintMotion InEngineType)
 
 void FPhysInterface_Chaos::SetAngularMotionLimitType_AssumesLocked(const FPhysicsConstraintHandle& InConstraintRef, PhysicsInterfaceTypes::ELimitAxis InAxis, EAngularConstraintMotion InMotion)
 {
+	// Twist is X component, Swing1 is Z component, and Swing2 is Y component in Chaos (see EJointAngularConstraintIndex)
+	static_assert(((int32)Chaos::EJointAngularConstraintIndex::Twist == 0), "EJointAngularConstraintIndex has changed");
+	static_assert(((int32)Chaos::EJointAngularConstraintIndex::Swing1 == 2), "EJointAngularConstraintIndex has changed");
+	static_assert(((int32)Chaos::EJointAngularConstraintIndex::Swing2 == 1), "EJointAngularConstraintIndex has changed");
+
 	if (InConstraintRef.IsValid() && InConstraintRef.Constraint->IsType(Chaos::EConstraintType::JointConstraintType))
 	{
 		if (Chaos::FJointConstraint* Constraint = static_cast<Chaos::FJointConstraint*>(InConstraintRef.Constraint))
@@ -273,11 +278,11 @@ void FPhysInterface_Chaos::SetAngularMotionLimitType_AssumesLocked(const FPhysic
 				break;
 
 			case PhysicsInterfaceTypes::ELimitAxis::Swing1:
-				Constraint->SetAngularMotionTypesY(ConvertMotionType(InMotion));
+				Constraint->SetAngularMotionTypesZ(ConvertMotionType(InMotion));
 				break;
 
 			case PhysicsInterfaceTypes::ELimitAxis::Swing2:
-				Constraint->SetAngularMotionTypesZ(ConvertMotionType(InMotion));
+				Constraint->SetAngularMotionTypesY(ConvertMotionType(InMotion));
 				break;
 			default:
 				ensure(false);
