Fix CIS static analyser to accept non null ensure.

#rnx
#tests compile

diff --git a/Engine/Source/Runtime/Engine/Private/Animation/AnimSequenceBase.cpp b/Engine/Source/Runtime/Engine/Private/Animation/AnimSequenceBase.cpp
--- a/Engine/Source/Runtime/Engine/Private/Animation/AnimSequenceBase.cpp
+++ b/Engine/Source/Runtime/Engine/Private/Animation/AnimSequenceBase.cpp
@@ -91,9 +91,11 @@ void UAnimSequenceBase::GetResourceSizeEx(FResourceSizeEx& CumulativeResourceSiz
 
 	const FRawCurveTracks& CurveData = GetCurveData();
 	const UStruct* Struct = CurveData.StaticStruct();
-	ensure(Struct);
-	CumulativeResourceSize.AddDedicatedSystemMemoryBytes(Struct->GetStructureSize());
-	CumulativeResourceSize.AddDedicatedSystemMemoryBytes(CurveData.FloatCurves.GetAllocatedSize());
+	if (ensure(Struct))
+	{
+		CumulativeResourceSize.AddDedicatedSystemMemoryBytes(Struct->GetStructureSize());
+		CumulativeResourceSize.AddDedicatedSystemMemoryBytes(CurveData.FloatCurves.GetAllocatedSize());
+	}
 }
 
 void UAnimSequenceBase::PostLoad()
