Add scope locks to the LWI subsystem in the appropriate places to avoid crashes with multithreaded code. General hardening and clean up.

diff --git a/Engine/Source/Runtime/Engine/Classes/GameFramework/LightWeightInstanceSubsystem.h b/Engine/Source/Runtime/Engine/Classes/GameFramework/LightWeightInstanceSubsystem.h
--- a/Engine/Source/Runtime/Engine/Classes/GameFramework/LightWeightInstanceSubsystem.h
+++ b/Engine/Source/Runtime/Engine/Classes/GameFramework/LightWeightInstanceSubsystem.h
@@ -14,9 +14,6 @@ class FAutoConsoleVariableRef;
 
 struct ENGINE_API FLightWeightInstanceSubsystem
 {
-	friend struct FActorInstanceHandle;
-	friend class ALightWeightInstanceManager;
-
 	static FCriticalSection GetFunctionCS;
 
 	static FLightWeightInstanceSubsystem& Get()
@@ -105,6 +102,12 @@ struct ENGINE_API FLightWeightInstanceSubsystem
 	// Helper that converts a position (world space) into a coordinate for the LWI grid.
 	static FInt32Vector3 ConvertPositionToCoord(const FVector & InPosition);
 
+	// Add a manager to the subsystem, thread safe.
+	bool AddManager(ALightWeightInstanceManager* Manager);
+
+	// Remove a manager from the subsystem, thread safe.
+	bool RemoveManager(ALightWeightInstanceManager* Manager);
+
 protected:
 	// Returns the class of the instance manager best suited to support instances of type ActorClass
 	UClass* FindBestInstanceManagerClass(const UClass* ActorClass);
@@ -127,6 +130,9 @@ private:
 	static int32 LWIGridSize;
 	static FAutoConsoleVariableRef CVarLWIGridSize;
 
+	// Mutex to make sure we don't change the LWInstanceManagers array while reading/writing it.
+	mutable FRWLock LWIManagersRWLock;
+
 #ifdef WITH_EDITOR
 private:
 	FDelegateHandle OnLevelActorAddedHandle;
diff --git a/Engine/Source/Runtime/Engine/Private/LightWeightInstanceManager.cpp b/Engine/Source/Runtime/Engine/Private/LightWeightInstanceManager.cpp
--- a/Engine/Source/Runtime/Engine/Private/LightWeightInstanceManager.cpp
+++ b/Engine/Source/Runtime/Engine/Private/LightWeightInstanceManager.cpp
@@ -21,13 +21,13 @@ ALightWeightInstanceManager::ALightWeightInstanceManager(const FObjectInitialize
 	AcceptedClass = AActor::StaticClass();
 	if (!HasAnyFlags(RF_ClassDefaultObject))
 	{
-		FLightWeightInstanceSubsystem::Get().LWInstanceManagers.Add(this);
+		FLightWeightInstanceSubsystem::Get().AddManager(this);
 	}
 }
 
 ALightWeightInstanceManager::~ALightWeightInstanceManager()
 {
-	FLightWeightInstanceSubsystem::Get().LWInstanceManagers.Remove(this);
+	FLightWeightInstanceSubsystem::Get().RemoveManager(this);
 }
 
 void ALightWeightInstanceManager::SetRepresentedClass(UClass* ActorClass)
@@ -269,22 +269,22 @@ FString ALightWeightInstanceManager::GetName(const FActorInstanceHandle& Handle)
 
 bool ALightWeightInstanceManager::DoesRepresentClass(const UClass* OtherClass) const
 {
-	return OtherClass ? OtherClass->IsChildOf(RepresentedClass) : false;
+	return OtherClass ? OtherClass->IsChildOf(GetRepresentedClass()) : false;
 }
 
 bool ALightWeightInstanceManager::DoesAcceptClass(const UClass* OtherClass) const
 {
-	return OtherClass ? OtherClass->IsChildOf(AcceptedClass) : false;
+	return OtherClass ? OtherClass->IsChildOf(GetAcceptedClass()) : false;
 }
 
 UClass* ALightWeightInstanceManager::GetRepresentedClass() const
 {
-	return RepresentedClass;
+	return IsValid(RepresentedClass) ? RepresentedClass : nullptr;
 }
 
 UClass* ALightWeightInstanceManager::GetAcceptedClass() const
 {
-	return AcceptedClass;
+	return IsValid(AcceptedClass) ? AcceptedClass : nullptr;
 }
 
 int32 ALightWeightInstanceManager::AddNewInstance(FLWIData* InitData)
diff --git a/Engine/Source/Runtime/Engine/Private/LightWeightInstanceSubsystem.cpp b/Engine/Source/Runtime/Engine/Private/LightWeightInstanceSubsystem.cpp
--- a/Engine/Source/Runtime/Engine/Private/LightWeightInstanceSubsystem.cpp
+++ b/Engine/Source/Runtime/Engine/Private/LightWeightInstanceSubsystem.cpp
@@ -29,14 +29,14 @@ FLightWeightInstanceSubsystem::FLightWeightInstanceSubsystem()
 			{
 				if (ALightWeightInstanceManager* LWIManager = Cast<ALightWeightInstanceManager>(Actor))
 				{
-					FLightWeightInstanceSubsystem::Get().LWInstanceManagers.AddUnique(LWIManager);
+					Get().AddManager(LWIManager);
 				}
 			});
 		OnLevelActorDeletedHandle = GEngine->OnLevelActorDeleted().AddLambda([this](AActor* Actor)
 			{
 				if (ALightWeightInstanceManager* LWIManager = Cast<ALightWeightInstanceManager>(Actor))
 				{
-					FLightWeightInstanceSubsystem::Get().LWInstanceManagers.Remove(LWIManager);
+					Get().RemoveManager(LWIManager);
 				}
 			});
 	}
@@ -56,6 +56,7 @@ FLightWeightInstanceSubsystem::~FLightWeightInstanceSubsystem()
 
 int32 FLightWeightInstanceSubsystem::GetManagerIndex(const ALightWeightInstanceManager* Manager) const
 {
+	FReadScopeLock Lock(LWIManagersRWLock);
 	for (int32 Idx = 0; Idx < LWInstanceManagers.Num(); ++Idx)
 	{
 		if (LWInstanceManagers[Idx] == Manager)
@@ -69,9 +70,22 @@ int32 FLightWeightInstanceSubsystem::GetManagerIndex(const ALightWeightInstanceM
 
 const ALightWeightInstanceManager* FLightWeightInstanceSubsystem::GetManagerAt(int32 Index) const
 {
+	FReadScopeLock Lock(LWIManagersRWLock);
 	return LWInstanceManagers.IsValidIndex(Index) ? LWInstanceManagers[Index] : nullptr;
 }
 
+bool FLightWeightInstanceSubsystem::AddManager(ALightWeightInstanceManager* Manager)
+{
+	FWriteScopeLock Lock(LWIManagersRWLock);
+	return Get().LWInstanceManagers.Add(Manager) ? true : false;
+}
+
+bool FLightWeightInstanceSubsystem::RemoveManager(ALightWeightInstanceManager* Manager)
+{
+	FWriteScopeLock Lock(LWIManagersRWLock);
+	return Get().LWInstanceManagers.Remove(Manager) ? true : false;
+}
+
 FInt32Vector3 FLightWeightInstanceSubsystem::ConvertPositionToCoord(const FVector & InPosition)
 {
 	const int32 GridSize = LWIGridSize;
@@ -97,7 +111,8 @@ ALightWeightInstanceManager* FLightWeightInstanceSubsystem::FindLightWeightInsta
 	if (Handle.Actor.IsValid())
 	{
 		const FInt32Vector3 GridCoord = ConvertPositionToCoord(Handle.GetLocation());
-		
+
+		FReadScopeLock Lock(LWIManagersRWLock);
 		// see if we already have a match
 		for (ALightWeightInstanceManager* LWInstance : LWInstanceManagers)
 		{
@@ -141,18 +156,22 @@ ALightWeightInstanceManager* FLightWeightInstanceSubsystem::FindOrAddLightWeight
 ALightWeightInstanceManager* FLightWeightInstanceSubsystem::FindLightWeightInstanceManager(UClass& ActorClass, UWorld& World, const FVector& InPos, const UDataLayerInstance* DataLayer) const
 {
 	const FInt32Vector3 GridCoord = ConvertPositionToCoord(InPos);
+	FReadScopeLock Lock(LWIManagersRWLock);
 
 	for (ALightWeightInstanceManager* InstanceManager : LWInstanceManagers)
 	{
-		const FInt32Vector3 ManagerGridCoord = ConvertPositionToCoord(InstanceManager->GetActorLocation());
-		
-		if (InstanceManager->GetRepresentedClass() == &ActorClass && ManagerGridCoord == GridCoord)
+		if (IsValid(InstanceManager))
 		{
-#if WITH_EDITOR
-			if (!DataLayer || (InstanceManager->SupportsDataLayer() && InstanceManager->ContainsDataLayer(DataLayer)))
-#endif // WITH_EDITOR
+			const FInt32Vector3 ManagerGridCoord = ConvertPositionToCoord(InstanceManager->GetActorLocation());
+			
+			if (InstanceManager->GetRepresentedClass() == &ActorClass && ManagerGridCoord == GridCoord)
 			{
-				return InstanceManager;
+	#if WITH_EDITOR
+				if (!DataLayer || (InstanceManager->SupportsDataLayer() && InstanceManager->ContainsDataLayer(DataLayer)))
+	#endif // WITH_EDITOR
+				{
+					return InstanceManager;
+				}
 			}
 		}
 	}
@@ -166,7 +185,7 @@ ALightWeightInstanceManager* FLightWeightInstanceSubsystem::FindOrAddLightWeight
 	{
 		return FoundManager;
 	}
-	
+
 	// we didn't find a match so we should add one.
 	// Find the best base class to start from
 	UClass* BestMatchingClass = FindBestInstanceManagerClass(&ActorClass);
