Fixed replay playback not resetting runtime data layer states to their initial state.

diff --git a/Engine/Source/Runtime/Engine/Private/WorldPartition/DataLayer/WorldDataLayers.cpp b/Engine/Source/Runtime/Engine/Private/WorldPartition/DataLayer/WorldDataLayers.cpp
--- a/Engine/Source/Runtime/Engine/Private/WorldPartition/DataLayer/WorldDataLayers.cpp
+++ b/Engine/Source/Runtime/Engine/Private/WorldPartition/DataLayer/WorldDataLayers.cpp
@@ -76,11 +76,12 @@ void AWorldDataLayers::BeginPlay()
 {
 	Super::BeginPlay();
 
-	// When running a Replay we want to reset our state to CDO (empty) and rely on the Replay/Replication.
+	// When running a Replay we want to reset our state to the initial state and rely on the Replay/Replication.
 	// Unfortunately this can't be tested in the PostLoad as the World doesn't have a demo driver yet.
 	if (GetWorld()->IsPlayingReplay())
 	{
 		ResetDataLayerRuntimeStates();
+		InitializeDataLayerRuntimeStates();
 	}
 }
 
@@ -88,8 +89,9 @@ void AWorldDataLayers::RewindForReplay()
 {
 	Super::RewindForReplay();
 
-	// Same as BeginPlay when rewinding we want to reset our state to CDO (empty) and rely on Replay/Replication.
+	// Same as BeginPlay when rewinding we want to reset our state to the initial state and rely on the Replay/Replication.
 	ResetDataLayerRuntimeStates();
+	InitializeDataLayerRuntimeStates();
 }
 
 void AWorldDataLayers::InitializeDataLayerRuntimeStates()
