Fixed check when running resave commandlet on external actor packages of partitioned worlds

diff --git a/Engine/Source/Runtime/Engine/Private/WorldPartition/WorldPartitionActorDesc.cpp b/Engine/Source/Runtime/Engine/Private/WorldPartition/WorldPartitionActorDesc.cpp
--- a/Engine/Source/Runtime/Engine/Private/WorldPartition/WorldPartitionActorDesc.cpp
+++ b/Engine/Source/Runtime/Engine/Private/WorldPartition/WorldPartitionActorDesc.cpp
@@ -133,8 +133,10 @@ void FWorldPartitionActorDesc::Init(const AActor* InActor)
 			// In Editor, always use owning world to resolve data layers.
 			// In game (PIE), always use actor's world partition DataLayerManager to resolve data layers.
 			const UDataLayerManager* ResolvingDataLayerManager = !ActorOwningWorld->IsGameWorld() ? ActorOwningWorld->GetDataLayerManager() : ActorWorldPartition->GetDataLayerManager();
-			// The only case where it's notmal not to have a valid ResolvingDataLayerManager is if OwningWorld is not game and not partitioned
-			check(ResolvingDataLayerManager || (!ActorOwningWorld->IsGameWorld() && !ActorOwningWorld->IsPartitionedWorld()));
+			// The only case where it's normal not to have a valid ResolvingDataLayerManager is if OwningWorld is not game and not partitioned
+			// Also, consider that we don't need the DataLayerManager for any running commandlet except for cook
+			const bool bCanSkipDataLayerManagerValidation = (IsRunningCommandlet() && !IsRunningCookCommandlet()) || (!ActorOwningWorld->IsGameWorld() && !ActorOwningWorld->IsPartitionedWorld());
+			check(ResolvingDataLayerManager || bCanSkipDataLayerManagerValidation);
 			if (ResolvingDataLayerManager)
 			{
 				if (ensure(ResolvingDataLayerManager->CanResolveDataLayers()))
