DataLayer: Private DataLayer

- Editor only DataLayer type that can be used for filters without having to create DataLayer asset

#rnx

diff --git a/Engine/Source/Editor/DataLayerEditor/Private/DataLayer/DataLayerEditorSubsystem.cpp b/Engine/Source/Editor/DataLayerEditor/Private/DataLayer/DataLayerEditorSubsystem.cpp
--- a/Engine/Source/Editor/DataLayerEditor/Private/DataLayer/DataLayerEditorSubsystem.cpp
+++ b/Engine/Source/Editor/DataLayerEditor/Private/DataLayer/DataLayerEditorSubsystem.cpp
@@ -45,6 +45,7 @@
 #include "Widgets/Text/STextBlock.h"
 #include "WorldPartition/DataLayer/DataLayerAsset.h"
 #include "WorldPartition/DataLayer/DataLayerInstanceWithAsset.h"
+#include "WorldPartition/DataLayer/DataLayerInstancePrivate.h"
 #include "WorldPartition/DataLayer/DataLayerManager.h"
 #include "WorldPartition/DataLayer/DeprecatedDataLayerInstance.h"
 #include "WorldPartition/DataLayer/WorldDataLayers.h"
@@ -60,7 +61,8 @@ DEFINE_LOG_CATEGORY_STATIC(LogDataLayerEditorSubsystem, All, All);
 
 FDataLayerCreationParameters::FDataLayerCreationParameters()
 	:DataLayerAsset(nullptr),
-	WorldDataLayers(nullptr)
+	WorldDataLayers(nullptr),
+	bIsPrivate(false)
 {
 
 }
@@ -1227,7 +1229,14 @@ UDataLayerInstance* UDataLayerEditorSubsystem::CreateDataLayerInstance(const FDa
 	{
 		if (!WorldDataLayers->HasDeprecatedDataLayers())
 		{
-			NewDataLayer = CreateDataLayerInstance<UDataLayerInstanceWithAsset>(WorldDataLayers, Parameters.DataLayerAsset);
+			if (Parameters.bIsPrivate)
+			{
+				NewDataLayer = CreateDataLayerInstance<UDataLayerInstancePrivate>(WorldDataLayers);
+			}
+			else
+			{
+				NewDataLayer = CreateDataLayerInstance<UDataLayerInstanceWithAsset>(WorldDataLayers, Parameters.DataLayerAsset);
+			}
 		}
 		else
 		{
diff --git a/Engine/Source/Editor/DataLayerEditor/Private/DataLayer/DataLayerMode.cpp b/Engine/Source/Editor/DataLayerEditor/Private/DataLayer/DataLayerMode.cpp
--- a/Engine/Source/Editor/DataLayerEditor/Private/DataLayer/DataLayerMode.cpp
+++ b/Engine/Source/Editor/DataLayerEditor/Private/DataLayer/DataLayerMode.cpp
@@ -1332,7 +1332,7 @@ void FDataLayerMode::RegisterContextMenu()
 					return BestCandidate;
 				};
 
-				auto CreateNewDataLayerInternal = [&GetBestCandidateWorldDataLayersFromSelection, SceneOutliner, Mode](UDataLayerInstance* InParentDataLayer, const UDataLayerAsset* InDataLayerAsset = nullptr) -> UDataLayerInstance*
+				auto CreateNewDataLayerInternal = [&GetBestCandidateWorldDataLayersFromSelection, SceneOutliner, Mode](UDataLayerInstance* InParentDataLayer, const UDataLayerAsset* InDataLayerAsset = nullptr, bool bInIsPrivate = false) -> UDataLayerInstance*
 				{
 					const FScopedTransaction Transaction(LOCTEXT("CreateNewDataLayer", "Create New Data Layer"));
 					AWorldDataLayers* TargetWorldDataLayers = GetBestCandidateWorldDataLayersFromSelection(SceneOutliner, Mode);
@@ -1342,6 +1342,7 @@ void FDataLayerMode::RegisterContextMenu()
 					FDataLayerCreationParameters CreationParams;
 					CreationParams.DataLayerAsset = const_cast<UDataLayerAsset*>(InDataLayerAsset);
 					CreationParams.WorldDataLayers = InParentDataLayer ? InParentDataLayer->GetOuterWorldDataLayers() : TargetWorldDataLayers;
+					CreationParams.bIsPrivate = bInIsPrivate;
 					if (UDataLayerInstance* NewDataLayerInstance = UDataLayerEditorSubsystem::Get()->CreateDataLayerInstance(CreationParams))
 					{
 						UDataLayerEditorSubsystem::Get()->SetParentDataLayer(NewDataLayerInstance, InParentDataLayer);
@@ -1356,26 +1357,33 @@ void FDataLayerMode::RegisterContextMenu()
 					return nullptr;
 				};
 
-				auto CreateNewDataLayer = [SceneOutliner, Mode, CreateNewDataLayerInternal](UDataLayerInstance* InParentDataLayer = nullptr, const UDataLayerAsset* InDataLayerAsset = nullptr) -> UDataLayerInstance*
+				auto CreateNewDataLayer = [SceneOutliner, Mode, CreateNewDataLayerInternal](UDataLayerInstance* InParentDataLayer = nullptr, const UDataLayerAsset* InDataLayerAsset = nullptr, bool bInIsPrivate = false) -> UDataLayerInstance*
 				{
 					if (UDataLayerEditorSubsystem::Get()->HasDeprecatedDataLayers())
 					{
 						check(InDataLayerAsset == nullptr);
-						return CreateNewDataLayerInternal(InParentDataLayer);
+						return CreateNewDataLayerInternal(InParentDataLayer, nullptr, false);
 					}
 					else
 					{
-						return CreateNewDataLayerInternal(InParentDataLayer, InDataLayerAsset);
+						return CreateNewDataLayerInternal(InParentDataLayer, InDataLayerAsset, bInIsPrivate);
 					}
 				};
 
-				auto CreateNewEmptyDataLayer = [SceneOutliner, Mode, CreateNewDataLayerInternal](UDataLayerInstance* InParentDataLayer = nullptr) -> UDataLayerInstance*
+				auto CreateNewEmptyDataLayer = [SceneOutliner, Mode, CreateNewDataLayerInternal](UDataLayerInstance* InParentDataLayer = nullptr, bool bInIsPrivate = false) -> UDataLayerInstance*
 				{
-					return CreateNewDataLayerInternal(InParentDataLayer);
+					return CreateNewDataLayerInternal(InParentDataLayer, nullptr, bInIsPrivate);
 				};
 				
 				Section.AddMenuEntry("CreateNewDataLayer", LOCTEXT("CreateNewDataLayer", "Create New Data Layer"), FText(), FSlateIcon(),
 					FUIAction(FExecuteAction::CreateLambda([CreateNewEmptyDataLayer]() { CreateNewEmptyDataLayer(); })));
+				
+				Section.AddMenuEntry("CreateNewDataLayerPrivate", LOCTEXT("CreateNewDataLayerPrivate", "Create New Private Data Layer"), LOCTEXT("CreateNewDataLayerPrivateToolTip", "Creates an Editor Only Data Layer that cannot be used by other Worlds."), FSlateIcon(),
+					FUIAction(
+						FExecuteAction::CreateLambda([CreateNewEmptyDataLayer]() { CreateNewEmptyDataLayer(nullptr, true); }),
+						FCanExecuteAction(),
+						FIsActionChecked(),
+						FIsActionButtonVisible::CreateLambda([]() { return !UDataLayerEditorSubsystem::Get()->HasDeprecatedDataLayers(); })));
 
 				const bool bAllSelectedDataLayersCanAddActors = !!SelectedDataLayers.FindByPredicate([](const UDataLayerInstance* DataLayerInstance) { return !DataLayerInstance->CanUserAddActors(); });
 				Section.AddMenuEntry("AddSelectedActorsToSelectedDataLayers", LOCTEXT("AddSelectedActorsToSelectedDataLayersMenu", "Add Selected Actors to Selected Data Layers"), FText(), FSlateIcon(),
diff --git a/Engine/Source/Editor/DataLayerEditor/Public/DataLayer/DataLayerEditorSubsystem.h b/Engine/Source/Editor/DataLayerEditor/Public/DataLayer/DataLayerEditorSubsystem.h
--- a/Engine/Source/Editor/DataLayerEditor/Public/DataLayer/DataLayerEditorSubsystem.h
+++ b/Engine/Source/Editor/DataLayerEditor/Public/DataLayer/DataLayerEditorSubsystem.h
@@ -59,6 +59,9 @@ struct DATALAYEREDITOR_API FDataLayerCreationParameters
 	// Optional. Will default at the level WorldDataLayers if unset.
 	UPROPERTY(EditAnywhere, BlueprintReadWrite, Category = "Data Layer")
 	TObjectPtr<AWorldDataLayers> WorldDataLayers;
+
+	UPROPERTY(EditAnywhere, BlueprintReadWrite, Category = "Data Layer")
+	bool bIsPrivate;
 };
 
 UCLASS()
diff --git a/Engine/Source/Runtime/Core/Public/UObject/FortniteMainBranchObjectVersion.h b/Engine/Source/Runtime/Core/Public/UObject/FortniteMainBranchObjectVersion.h
--- a/Engine/Source/Runtime/Core/Public/UObject/FortniteMainBranchObjectVersion.h
+++ b/Engine/Source/Runtime/Core/Public/UObject/FortniteMainBranchObjectVersion.h
@@ -283,6 +283,9 @@ struct CORE_API FFortniteMainBranchObjectVersion
 		// FActorContainerID is now an FGuid instead of a uint64
 		WorldPartitionFActorContainerIDu64ToGuid,
 
+		// FDataLayerInstanceDesc support for private data layers
+		WorldPartitionPrivateDataLayers,
+
 		// -----<new versions can be added above this line>-------------------------------------------------
 		VersionPlusOne,
 		LatestVersion = VersionPlusOne - 1
diff --git a/Engine/Source/Runtime/Engine/Classes/GameFramework/Actor.h b/Engine/Source/Runtime/Engine/Classes/GameFramework/Actor.h
--- a/Engine/Source/Runtime/Engine/Classes/GameFramework/Actor.h
+++ b/Engine/Source/Runtime/Engine/Classes/GameFramework/Actor.h
@@ -4279,6 +4279,7 @@ private:
 	static bool RemoveDataLayerAsset(AActor* InActor, const UDataLayerAsset* InDataLayerAsset);
 
 	friend class UDataLayerInstanceWithAsset;
+	friend class UDataLayerInstancePrivate;
 	friend class ULevelInstanceSubsystem;
 };
 #endif
diff --git a/Engine/Source/Runtime/Engine/Private/WorldPartition/DataLayer/DataLayerAsset.cpp b/Engine/Source/Runtime/Engine/Private/WorldPartition/DataLayer/DataLayerAsset.cpp
--- a/Engine/Source/Runtime/Engine/Private/WorldPartition/DataLayer/DataLayerAsset.cpp
+++ b/Engine/Source/Runtime/Engine/Private/WorldPartition/DataLayer/DataLayerAsset.cpp
@@ -1,6 +1,7 @@
 // Copyright Epic Games, Inc. All Rights Reserved.
 
 #include "WorldPartition/DataLayer/DataLayerAsset.h"
+#include "WorldPartition/DataLayer/DataLayerInstance.h"
 
 #include UE_INLINE_GENERATED_CPP_BY_NAME(DataLayerAsset)
 
@@ -10,6 +11,11 @@ UDataLayerAsset::UDataLayerAsset(const FObjectInitializer& ObjectInitializer)
 	, DebugColor(FColor::Black)
 {}
 
+bool UDataLayerAsset::IsPrivate() const
+{
+	return !!GetTypedOuter<UDataLayerInstance>();
+}
+
 #if WITH_EDITOR
 void UDataLayerAsset::PostLoad()
 {
@@ -20,4 +26,23 @@ void UDataLayerAsset::PostLoad()
 
 	Super::PostLoad();
 }
+
+bool UDataLayerAsset::CanEditChange(const FProperty* InProperty) const
+{
+	if (!Super::CanEditChange(InProperty))
+	{
+		return false;
+	}
+
+	if (InProperty)
+	{
+		if (InProperty->GetFName() == GET_MEMBER_NAME_CHECKED(UDataLayerAsset, DataLayerType))
+		{
+			// UDataLayerAsset outered to a UDataLayerInstance do not support Runtime type
+			return !IsPrivate();
+		}
+	}
+
+	return true;
+}
 #endif
diff --git a/Engine/Source/Runtime/Engine/Private/WorldPartition/DataLayer/DataLayerInstancePrivate.cpp b/Engine/Source/Runtime/Engine/Private/WorldPartition/DataLayer/DataLayerInstancePrivate.cpp
new file mode 100644
--- /dev/null
+++ b/Engine/Source/Runtime/Engine/Private/WorldPartition/DataLayer/DataLayerInstancePrivate.cpp
@@ -0,0 +1,85 @@
+// Copyright Epic Games, Inc. All Rights Reserved.
+
+#include "WorldPartition/DataLayer/DataLayerInstancePrivate.h"
+#include "WorldPartition/DataLayer/WorldDataLayers.h"
+#include "WorldPartition/DataLayer/DataLayerManager.h"
+#include "WorldPartition/DataLayer/DataLayerUtils.h"
+#include "Engine/Level.h"
+#include "UObject/UnrealType.h"
+
+#include UE_INLINE_GENERATED_CPP_BY_NAME(DataLayerInstancePrivate)
+
+UDataLayerInstancePrivate::UDataLayerInstancePrivate(const FObjectInitializer& ObjectInitializer)
+	: Super(ObjectInitializer)
+#if WITH_EDITORONLY_DATA
+	, bIsIncludedInActorFilterDefault(true)
+#endif
+{
+	DataLayerAsset = CreateDefaultSubobject<UDataLayerAsset>(TEXT("DataLayerAsset"));
+#if WITH_EDITOR
+	DataLayerAsset->SetType(EDataLayerType::Editor);
+#endif
+}
+
+#if WITH_EDITOR
+FName UDataLayerInstancePrivate::MakeName()
+{
+	return FName(FString::Format(TEXT("DataLayer_{0}"), { FGuid::NewGuid().ToString() }));
+}
+
+void UDataLayerInstancePrivate::OnCreated()
+{
+	check(!GetOuterWorldDataLayers()->HasDeprecatedDataLayers() || IsRunningCommandlet());
+
+	Modify(/*bAlwaysMarkDirty*/false);
+	FDataLayerUtils::SetDataLayerShortName(this, TEXT("DataLayer"));
+		
+	SetVisible(true);
+}
+
+bool UDataLayerInstancePrivate::IsReadOnly() const
+{
+	if (Super::IsReadOnly())
+	{
+		return true;
+	}
+	return GetOuterWorldDataLayers()->IsReadOnly();
+}
+
+bool UDataLayerInstancePrivate::IsLocked() const
+{
+	if (Super::IsLocked())
+	{
+		return true;
+	}
+	return IsReadOnly();
+}
+
+bool UDataLayerInstancePrivate::PerformAddActor(AActor* InActor) const
+{
+	check(GetOuterWorldDataLayers() == InActor->GetLevel()->GetWorldDataLayers()); // Make sure the instance is part of the same WorldDataLayers as the actor's level WorldDataLayer.
+	check(UDataLayerManager::GetDataLayerManager(InActor)->GetDataLayerInstance(DataLayerAsset) != nullptr); // Make sure the DataLayerInstance exists for this level
+	return FAssignActorDataLayer::AddDataLayerAsset(InActor, DataLayerAsset);
+}
+
+bool UDataLayerInstancePrivate::PerformRemoveActor(AActor* InActor) const
+{
+	return FAssignActorDataLayer::RemoveDataLayerAsset(InActor, DataLayerAsset);
+}
+
+bool UDataLayerInstancePrivate::CanEditChange(const FProperty* InProperty) const
+{
+	if (!Super::CanEditChange(InProperty))
+	{
+		return false;
+	}
+
+	if (InProperty && (InProperty->GetFName() == GET_MEMBER_NAME_CHECKED(UDataLayerInstancePrivate, bIsIncludedInActorFilterDefault)))
+	{
+		return SupportsActorFilters();
+	}
+
+	return true;
+}
+
+#endif
diff --git a/Engine/Source/Runtime/Engine/Private/WorldPartition/DataLayer/WorldDataLayers.cpp b/Engine/Source/Runtime/Engine/Private/WorldPartition/DataLayer/WorldDataLayers.cpp
--- a/Engine/Source/Runtime/Engine/Private/WorldPartition/DataLayer/WorldDataLayers.cpp
+++ b/Engine/Source/Runtime/Engine/Private/WorldPartition/DataLayer/WorldDataLayers.cpp
@@ -704,13 +704,9 @@ const UDataLayerInstance* AWorldDataLayers::GetDataLayerInstance(const UDataLaye
 #if WITH_EDITOR	
 	for (const UDataLayerInstance* DataLayerInstance : DataLayerInstances)
 	{
-		static_assert(DATALAYER_TO_INSTANCE_RUNTIME_CONVERSION_ENABLED, "Remove unnecessary cast. All DataLayerInstance now have assets");
-		if (const UDataLayerInstanceWithAsset* DataLayerInstanceWithAsset = Cast<UDataLayerInstanceWithAsset>(DataLayerInstance))
+		if (DataLayerInstance->GetAsset() == InDataLayerAsset)
 		{
-			if (DataLayerInstanceWithAsset->GetAsset() == InDataLayerAsset)
-			{
-				return DataLayerInstance;
-			}
+			return DataLayerInstance;
 		}
 	}
 #else
diff --git a/Engine/Source/Runtime/Engine/Private/WorldPartition/DataLayer/WorldDataLayersActorDesc.cpp b/Engine/Source/Runtime/Engine/Private/WorldPartition/DataLayer/WorldDataLayersActorDesc.cpp
--- a/Engine/Source/Runtime/Engine/Private/WorldPartition/DataLayer/WorldDataLayersActorDesc.cpp
+++ b/Engine/Source/Runtime/Engine/Private/WorldPartition/DataLayer/WorldDataLayersActorDesc.cpp
@@ -17,10 +17,10 @@ FArchive& operator<<(FArchive& Ar, FDataLayerInstanceDesc& Desc)
 
 	if (Ar.CustomVer(FFortniteNCBranchObjectVersion::GUID) < FFortniteNCBranchObjectVersion::FixedDataLayerInstanceDesc)
 	{
-		Ar << Desc.ShortName;
+		Ar << Desc.PrivateShortName;
 		if (Desc.bIsUsingAsset)
 		{
-			Ar << Desc.AssetPath << Desc.bIsRuntime;
+			Ar << Desc.AssetPath << Desc.bDeprecatedIsRuntime;
 		}
 	}
 	else
@@ -31,7 +31,7 @@ FArchive& operator<<(FArchive& Ar, FDataLayerInstanceDesc& Desc)
 		}
 		else
 		{
-			Ar << Desc.bIsRuntime << Desc.ShortName;
+			Ar << Desc.bDeprecatedIsRuntime << Desc.PrivateShortName;
 		}
 	}
 
@@ -46,6 +46,17 @@ FArchive& operator<<(FArchive& Ar, FDataLayerInstanceDesc& Desc)
 		Ar << Desc.bIsIncludedInActorFilterDefault;
 	}
 
+	if (Ar.CustomVer(FFortniteMainBranchObjectVersion::GUID) >= FFortniteMainBranchObjectVersion::WorldPartitionPrivateDataLayers)
+	{
+		Ar << Desc.bIsPrivate;
+
+		if (Desc.bIsPrivate)
+		{
+			Ar << Desc.PrivateShortName;
+			Ar << Desc.bPrivateDataLayerSupportsActorFilter;
+		}
+	}
+
 	return Ar;
 }
 
@@ -55,9 +66,11 @@ bool operator == (const FDataLayerInstanceDesc& Lhs, const FDataLayerInstanceDes
 		(Lhs.ParentName == Rhs.ParentName) &&
 		(Lhs.bIsUsingAsset == Rhs.bIsUsingAsset) &&
 		(Lhs.AssetPath == Rhs.AssetPath) &&
-		(Lhs.bIsRuntime == Rhs.bIsRuntime) &&
-		(Lhs.ShortName == Rhs.ShortName) &&
-		(Lhs.bIsIncludedInActorFilterDefault == Rhs.bIsIncludedInActorFilterDefault);
+		(Lhs.bDeprecatedIsRuntime == Rhs.bDeprecatedIsRuntime) &&
+		(Lhs.PrivateShortName == Rhs.PrivateShortName) &&
+		(Lhs.bIsIncludedInActorFilterDefault == Rhs.bIsIncludedInActorFilterDefault) &&
+		(Lhs.bIsPrivate == Rhs.bIsPrivate) &&
+		(Lhs.bPrivateDataLayerSupportsActorFilter == Rhs.bPrivateDataLayerSupportsActorFilter);
 }
 
 bool operator < (const FDataLayerInstanceDesc& Lhs, const FDataLayerInstanceDesc& Rhs)
@@ -70,15 +83,25 @@ bool operator < (const FDataLayerInstanceDesc& Lhs, const FDataLayerInstanceDesc
 			{
 				if (Lhs.AssetPath == Rhs.AssetPath)
 				{
-					if (Lhs.bIsRuntime == Rhs.bIsRuntime)
+					if (Lhs.bDeprecatedIsRuntime == Rhs.bDeprecatedIsRuntime)
 					{
-						if (Lhs.ShortName == Rhs.ShortName)
+						if (Lhs.PrivateShortName == Rhs.PrivateShortName)
 						{
+							if (Lhs.bIsIncludedInActorFilterDefault == Rhs.bIsIncludedInActorFilterDefault)
+							{
+								if (Lhs.bIsPrivate == Rhs.bIsPrivate)
+								{
+									return (int)Lhs.bPrivateDataLayerSupportsActorFilter == (int)Rhs.bPrivateDataLayerSupportsActorFilter;
+								}
+
+								return (int)Lhs.bIsPrivate < (int)Rhs.bIsPrivate;
+							}
+							
 							return (int)Lhs.bIsIncludedInActorFilterDefault < (int)Rhs.bIsIncludedInActorFilterDefault;
 						}
-						return Lhs.ShortName < Rhs.ShortName;
+						return Lhs.PrivateShortName < Rhs.PrivateShortName;
 					}
-					return (int)Lhs.bIsRuntime < (int)Rhs.bIsRuntime;
+					return (int)Lhs.bDeprecatedIsRuntime < (int)Rhs.bDeprecatedIsRuntime;
 				}
 				return Lhs.AssetPath.LexicalLess(Rhs.AssetPath);
 			}
@@ -91,7 +114,7 @@ bool operator < (const FDataLayerInstanceDesc& Lhs, const FDataLayerInstanceDesc
 
 UDataLayerAsset* FDataLayerInstanceDesc::GetAsset() const
 {
-	if (IsUsingAsset() && !AssetPath.IsNone())
+	if (!bIsPrivate && IsUsingAsset() && !AssetPath.IsNone())
 	{
 		UDataLayerAsset* Asset = LoadObject<UDataLayerAsset>(nullptr, *AssetPath.ToString());
 		return Asset;
@@ -101,6 +124,11 @@ UDataLayerAsset* FDataLayerInstanceDesc::GetAsset() const
 
 EDataLayerType FDataLayerInstanceDesc::GetDataLayerType() const
 {
+	if (bIsPrivate)
+	{
+		return EDataLayerType::Editor;
+	}
+
 	if (IsUsingAsset())
 	{
 		if (UDataLayerAsset* DataLayerAsset = GetAsset())
@@ -109,24 +137,31 @@ EDataLayerType FDataLayerInstanceDesc::GetDataLayerType() const
 		}
 		return EDataLayerType::Unknown;
 	}
-	return bIsRuntime ? EDataLayerType::Runtime : EDataLayerType::Editor;
+	return bDeprecatedIsRuntime ? EDataLayerType::Runtime : EDataLayerType::Editor;
 }
 
 FString FDataLayerInstanceDesc::GetShortName() const
 {
-	if (IsUsingAsset())
+	if (!bIsPrivate && IsUsingAsset())
 	{
 		static FString UnknownString(TEXT("Unknown"));
 		UDataLayerAsset* DataLayerAsset = GetAsset();
 		return DataLayerAsset ? DataLayerAsset->GetName() : UnknownString;
 	}
-	return ShortName;
+	return PrivateShortName;
+}
+
+bool FDataLayerInstanceDesc::SupportsActorFilters() const
+{
+	return bIsPrivate ? bPrivateDataLayerSupportsActorFilter : GetAsset()->SupportsActorFilters();
 }
 
 FDataLayerInstanceDesc::FDataLayerInstanceDesc()
 : bIsUsingAsset(false)
-, bIsRuntime(false)
 , bIsIncludedInActorFilterDefault(true)
+, bIsPrivate(false)
+, bPrivateDataLayerSupportsActorFilter(false)
+, bDeprecatedIsRuntime(false)
 {
 }
 
@@ -135,18 +170,18 @@ void FDataLayerInstanceDesc::Init(UDataLayerInstance* InDataLayerInstance)
 	Name = InDataLayerInstance->GetDataLayerFName();
 	ParentName = InDataLayerInstance->GetParent() ? InDataLayerInstance->GetParent()->GetDataLayerFName() : NAME_None;
 	const UDataLayerInstanceWithAsset* DataLayerWithAsset = Cast<UDataLayerInstanceWithAsset>(InDataLayerInstance);
-	bIsUsingAsset = (DataLayerWithAsset != nullptr);
-	if (bIsUsingAsset)
+	const UDataLayerAsset* DataLayerAsset = InDataLayerInstance->GetAsset();
+	bIsUsingAsset = InDataLayerInstance->Implements<UDataLayerInstanceWithAsset>() || (DataLayerAsset != nullptr);
+	if (bIsUsingAsset && DataLayerAsset)
 	{
-		if (const UDataLayerAsset* DataLayerAsset = DataLayerWithAsset->GetAsset())
-		{
-			AssetPath = FName(DataLayerAsset->GetPathName());
-			bIsIncludedInActorFilterDefault = DataLayerWithAsset->IsIncludedInActorFilterDefault();
-			
-		}
+		AssetPath = FName(DataLayerAsset->GetPathName());
+		bIsPrivate = DataLayerAsset->IsPrivate();
 	}
-	ShortName = InDataLayerInstance->GetDataLayerShortName();
-	bIsRuntime = InDataLayerInstance->IsRuntime();
+	
+	bIsIncludedInActorFilterDefault = InDataLayerInstance->IsIncludedInActorFilterDefault();
+	bPrivateDataLayerSupportsActorFilter = InDataLayerInstance->SupportsActorFilters();
+	PrivateShortName = InDataLayerInstance->GetDataLayerShortName();
+	bDeprecatedIsRuntime = InDataLayerInstance->IsRuntime();
 }
 
 FWorldDataLayersActorDesc::FWorldDataLayersActorDesc()
diff --git a/Engine/Source/Runtime/Engine/Private/WorldPartition/WorldPartitionSubsystem.cpp b/Engine/Source/Runtime/Engine/Private/WorldPartition/WorldPartitionSubsystem.cpp
--- a/Engine/Source/Runtime/Engine/Private/WorldPartition/WorldPartitionSubsystem.cpp
+++ b/Engine/Source/Runtime/Engine/Private/WorldPartition/WorldPartitionSubsystem.cpp
@@ -177,7 +177,7 @@ FWorldPartitionActorFilter UWorldPartitionSubsystem::GetWorldPartitionActorFilte
 		for (const FDataLayerInstanceDesc& DataLayerInstanceDesc : WorldDataLayersActorDesc->GetDataLayerInstances())
 		{
 			// For now consider all DataLayerInstances using Assets as filters that are included by default
-			if (DataLayerInstanceDesc.IsUsingAsset() && DataLayerInstanceDesc.GetAsset() && DataLayerInstanceDesc.GetAsset()->SupportsActorFilters())
+			if (DataLayerInstanceDesc.SupportsActorFilters())
 			{
 				Filter.DataLayerFilters.Add(FSoftObjectPath(DataLayerInstanceDesc.GetAssetPath().ToString()), FWorldPartitionActorFilter::FDataLayerFilter(DataLayerInstanceDesc.GetShortName(), DataLayerInstanceDesc.IsIncludedInActorFilterDefault()));
 			}
diff --git a/Engine/Source/Runtime/Engine/Public/WorldPartition/DataLayer/DataLayerAsset.h b/Engine/Source/Runtime/Engine/Public/WorldPartition/DataLayer/DataLayerAsset.h
--- a/Engine/Source/Runtime/Engine/Public/WorldPartition/DataLayer/DataLayerAsset.h
+++ b/Engine/Source/Runtime/Engine/Public/WorldPartition/DataLayer/DataLayerAsset.h
@@ -21,18 +21,24 @@ class ENGINE_API UDataLayerAsset : public UObject
 #if WITH_EDITOR
 	//~ Begin UObject Interface
 	virtual void PostLoad() override;
+	virtual bool CanEditChange(const FProperty* InProperty) const override;
 	//~ End UObject Interface
 
 public:
-	void SetType(EDataLayerType Type) { DataLayerType = Type; }
+	void SetType(EDataLayerType Type) 
+	{ 
+		check(Type == EDataLayerType::Editor || !IsPrivate());
+		DataLayerType = Type; 
+	}
 	void SetDebugColor(FColor InDebugColor) { DebugColor = InDebugColor; }
 #endif
+	bool IsPrivate() const;
 
 	UFUNCTION(Category = "Data Layer", BlueprintCallable)
 	EDataLayerType GetType() const { return DataLayerType; }
 
 	UFUNCTION(Category = "Data Layer|Runtime", BlueprintCallable)
-	bool IsRuntime() const { return DataLayerType == EDataLayerType::Runtime; }
+	bool IsRuntime() const { return !IsPrivate() && DataLayerType == EDataLayerType::Runtime; }
 
 	UFUNCTION(Category = "Data Layer|Runtime", BlueprintCallable)
 	FColor GetDebugColor() const { return DebugColor; }
@@ -43,9 +49,9 @@ private:
 	UPROPERTY(Category = "Data Layer", EditAnywhere)
 	EDataLayerType DataLayerType;
 
-	UPROPERTY(Category = "Data Layer|Actor Filter", EditAnywhere)
+	UPROPERTY(Category = "Actor Filter", EditAnywhere)
 	bool bSupportsActorFilters;
 		
-	UPROPERTY(Category = "Data Layer|Runtime", EditAnywhere)
+	UPROPERTY(Category = "Runtime", EditAnywhere)
 	FColor DebugColor;
 };
\ No newline at end of file
diff --git a/Engine/Source/Runtime/Engine/Public/WorldPartition/DataLayer/DataLayerInstance.h b/Engine/Source/Runtime/Engine/Public/WorldPartition/DataLayer/DataLayerInstance.h
--- a/Engine/Source/Runtime/Engine/Public/WorldPartition/DataLayer/DataLayerInstance.h
+++ b/Engine/Source/Runtime/Engine/Public/WorldPartition/DataLayer/DataLayerInstance.h
@@ -16,6 +16,7 @@ class AWorldDataLayers;
 class FText;
 class IStreamingGenerationErrorHandler;
 class UWorld;
+class UDataLayerAsset;
 
 UENUM(BlueprintType)
 enum class EDataLayerRuntimeState : uint8
@@ -107,6 +108,9 @@ public:
 
 	virtual bool CanEditDataLayerShortName() const { return false; }
 
+	virtual bool SupportsActorFilters() const { return false; }
+	virtual bool IsIncludedInActorFilterDefault() const { return false; }
+
 	// Whether the DataLayer was created by a user and can be deleted by a user.
 	virtual bool IsUserManaged() const { return true; }
 
@@ -155,6 +159,8 @@ public:
 	EDataLayerRuntimeState GetRuntimeState() const;
 	EDataLayerRuntimeState GetEffectiveRuntimeState() const;
 
+	virtual const UDataLayerAsset* GetAsset() const { return nullptr; }
+
 	const TArray<TObjectPtr<UDataLayerInstance>>& GetChildren() const { return Children; }
 	void ForEachChild(TFunctionRef<bool(const UDataLayerInstance*)> Operation) const;
 
@@ -192,11 +198,11 @@ protected:
 	uint32 bIsVisible : 1;
 
 	/** Whether actors associated with the Data Layer should be initially visible in the viewport when loading the map */
-	UPROPERTY(Category = "Data Layer|Editor", EditAnywhere)
+	UPROPERTY(Category = "Editor", EditAnywhere)
 	uint32 bIsInitiallyVisible : 1;
 
 	/** Determines the default value of the data layer's loaded state in editor if it hasn't been changed in data layer outliner by the user */
-	UPROPERTY(Category = "Data Layer|Editor", EditAnywhere, meta = (DisplayName = "Is Initially Loaded"))
+	UPROPERTY(Category = "Editor", EditAnywhere, meta = (DisplayName = "Is Initially Loaded"))
 	uint32 bIsInitiallyLoadedInEditor : 1;
 
 	/** Wheter the data layer is loaded in editor (user setting) */
@@ -213,7 +219,7 @@ protected:
 	uint32 bIsLocked : 1;
 #endif
 
-	UPROPERTY(Category = "Data Layer|Advanced|Runtime", EditAnywhere)
+	UPROPERTY(Category = "Runtime", EditAnywhere)
 	EDataLayerRuntimeState InitialRuntimeState;
 
 private:
diff --git a/Engine/Source/Runtime/Engine/Public/WorldPartition/DataLayer/DataLayerInstancePrivate.h b/Engine/Source/Runtime/Engine/Public/WorldPartition/DataLayer/DataLayerInstancePrivate.h
new file mode 100644
--- /dev/null
+++ b/Engine/Source/Runtime/Engine/Public/WorldPartition/DataLayer/DataLayerInstancePrivate.h
@@ -0,0 +1,54 @@
+// Copyright Epic Games, Inc. All Rights Reserved.
+#pragma once
+
+#include "WorldPartition/DataLayer/DataLayerInstance.h"
+#include "WorldPartition/DataLayer/DataLayerAsset.h"
+
+#include "DataLayerInstancePrivate.generated.h"
+
+UCLASS(Config = Engine, PerObjectConfig, Within = WorldDataLayers, AutoExpandCategories = ("Data Layer|Editor"))
+class ENGINE_API UDataLayerInstancePrivate : public UDataLayerInstance
+{
+	GENERATED_UCLASS_BODY()
+
+public:
+#if WITH_EDITOR
+	static FName MakeName();
+	void OnCreated();
+
+	virtual bool CanEditChange(const FProperty* InProperty) const;
+	virtual bool IsLocked() const override;
+	virtual bool IsReadOnly() const override;
+		
+	virtual bool SupportsActorFilters() const override { return GetAsset()->SupportsActorFilters(); }
+	virtual bool IsIncludedInActorFilterDefault() const override { return bIsIncludedInActorFilterDefault; }
+	virtual bool CanEditDataLayerShortName() const override { return true; }
+#endif
+
+	virtual const UDataLayerAsset* GetAsset() const override { return DataLayerAsset; }
+
+	virtual EDataLayerType GetType() const override { return DataLayerAsset->GetType(); }
+	virtual bool IsRuntime() const override { return DataLayerAsset->IsRuntime(); }
+	virtual FColor GetDebugColor() const override { return DataLayerAsset->GetDebugColor(); }
+
+	virtual FString GetDataLayerShortName() const override { return ShortName; }
+	virtual FString GetDataLayerFullName() const override { return DataLayerAsset->GetPathName(); }
+protected:
+#if WITH_EDITOR
+	virtual bool PerformAddActor(AActor* InActor) const override;
+	virtual bool PerformRemoveActor(AActor* InActor) const override;
+	virtual void PerformSetDataLayerShortName(const FString& InNewShortName) override { ShortName = InNewShortName; }
+#endif
+
+private:
+	UPROPERTY()
+	FString ShortName;
+
+	UPROPERTY(Category = "Asset", VisibleAnywhere, Instanced)
+	TObjectPtr<UDataLayerAsset> DataLayerAsset;
+
+#if WITH_EDITORONLY_DATA
+	UPROPERTY(Category = "Actor Filter", EditAnywhere, meta = (DisplayName = "Is Included", ToolTip = "Whether actors assigned to this DataLayer are included by default when used in a filter"))
+	bool bIsIncludedInActorFilterDefault;
+#endif
+};
\ No newline at end of file
diff --git a/Engine/Source/Runtime/Engine/Public/WorldPartition/DataLayer/DataLayerInstanceWithAsset.h b/Engine/Source/Runtime/Engine/Public/WorldPartition/DataLayer/DataLayerInstanceWithAsset.h
--- a/Engine/Source/Runtime/Engine/Public/WorldPartition/DataLayer/DataLayerInstanceWithAsset.h
+++ b/Engine/Source/Runtime/Engine/Public/WorldPartition/DataLayer/DataLayerInstanceWithAsset.h
@@ -33,7 +33,7 @@ public:
 	bool IsIncludedInActorFilterDefault() const;
 #endif
 
-	const UDataLayerAsset* GetAsset() const { return DataLayerAsset; }
+	const UDataLayerAsset* GetAsset() const override { return DataLayerAsset; }
 
 	virtual EDataLayerType GetType() const override { return DataLayerAsset != nullptr ? DataLayerAsset->GetType() : EDataLayerType::Unknown; }
 
diff --git a/Engine/Source/Runtime/Engine/Public/WorldPartition/DataLayer/DeprecatedDataLayerInstance.h b/Engine/Source/Runtime/Engine/Public/WorldPartition/DataLayer/DeprecatedDataLayerInstance.h
--- a/Engine/Source/Runtime/Engine/Public/WorldPartition/DataLayer/DeprecatedDataLayerInstance.h
+++ b/Engine/Source/Runtime/Engine/Public/WorldPartition/DataLayer/DeprecatedDataLayerInstance.h
@@ -37,8 +37,6 @@ public:
 	static FName MakeName(const UDEPRECATED_DataLayer* DeprecatedDataLayer);
 	void OnCreated(const UDEPRECATED_DataLayer* DeprecatedDataLayer);
 	PRAGMA_ENABLE_DEPRECATION_WARNINGS
-
-	void SetDataLayerLabel(FName InDataLayerLabel);
 #endif
 	
 	FName GetDataLayerLabel() const { return Label; }
diff --git a/Engine/Source/Runtime/Engine/Public/WorldPartition/DataLayer/WorldDataLayersActorDesc.h b/Engine/Source/Runtime/Engine/Public/WorldPartition/DataLayer/WorldDataLayersActorDesc.h
--- a/Engine/Source/Runtime/Engine/Public/WorldPartition/DataLayer/WorldDataLayersActorDesc.h
+++ b/Engine/Source/Runtime/Engine/Public/WorldPartition/DataLayer/WorldDataLayersActorDesc.h
@@ -20,33 +20,39 @@ public:
 	FName GetParentName() const { return ParentName; }
 	bool IsUsingAsset() const { return bIsUsingAsset; }
 	FName GetAssetPath() const { return AssetPath; }
-	class UDataLayerAsset* GetAsset() const;
 	EDataLayerType GetDataLayerType() const;
 	FString GetShortName() const;
 	bool IsIncludedInActorFilterDefault() const { return bIsIncludedInActorFilterDefault; }
+	bool SupportsActorFilters() const;
 
 private:
+	class UDataLayerAsset* GetAsset() const;
+
 	// DataLayerInstance Name
 	FName Name;
 	// Parent DataLayerInstance Name
 	FName ParentName;
-	// We can't rely on a valid AssetPath to determine if DataLayerInstance is UDataLayerInstanceWithAsset
+	// We can't rely on a valid AssetPath to determine if DataLayerInstance is UDataLayerInstanceWithAsset/UDataLayerInstancePrivate
 	bool bIsUsingAsset;
 
-	//~ Begin UDataLayerInstanceWithAsset
+	//~ Begin UDataLayerInstance
 	// DataLayer Asset Path 
 	FName AssetPath;
-	//~ End UDataLayerInstanceWithAsset
+	// Returns if data layer should be included by default in FWorldPartitionActorFilter.
+	bool bIsIncludedInActorFilterDefault;
+	// DataLayer Asset is Private
+	bool bIsPrivate;
+	//~ End UDataLayerInstance
+
+	// If DataLayer Asset is Private store the SupportsActorFilter() flag
+	bool bPrivateDataLayerSupportsActorFilter;
+	// If DataLayer Asset is Private store the ShortName (also used for UDeprecatedDataLayerInstance)
+	FString PrivateShortName;
 
 	//~ Begin UDeprecatedDataLayerInstance
-	// Runtime or Editor
-	bool bIsRuntime;
-	// Label
-	FString ShortName;
+	bool bDeprecatedIsRuntime;
 	//~ End UDeprecatedDataLayerInstance
 
-	// Returns if data layer should be included by default in FWorldPartitionActorFilter.
-	bool bIsIncludedInActorFilterDefault;
 };
 
 /**
