Add missing function for override that was lost in cherry-pick/merge conflict

diff --git a/Engine/Source/Runtime/Engine/Classes/Components/SkeletalMeshComponent.h b/Engine/Source/Runtime/Engine/Classes/Components/SkeletalMeshComponent.h
--- a/Engine/Source/Runtime/Engine/Classes/Components/SkeletalMeshComponent.h
+++ b/Engine/Source/Runtime/Engine/Classes/Components/SkeletalMeshComponent.h
@@ -2376,6 +2376,8 @@ protected:
 
 private:
 
+	virtual void OnClearAnimScriptInstance() {};
+
 	friend struct FSkeletalMeshComponentEndPhysicsTickFunction;
 
 	/** Update systems after physics sim is done */
