Fix random crash in ULevelStreamingProfilingSubsystem::OnLevelFinishedAddToWorld at the end of a match.

diff --git a/Engine/Source/Runtime/Engine/Private/ProfilingDebugging/LevelStreamingProfilingSubsystem.cpp b/Engine/Source/Runtime/Engine/Private/ProfilingDebugging/LevelStreamingProfilingSubsystem.cpp
--- a/Engine/Source/Runtime/Engine/Private/ProfilingDebugging/LevelStreamingProfilingSubsystem.cpp
+++ b/Engine/Source/Runtime/Engine/Private/ProfilingDebugging/LevelStreamingProfilingSubsystem.cpp
@@ -83,6 +83,8 @@ bool ULevelStreamingProfilingSubsystem::ShouldCreateSubsystem(UObject* Outer) co
 
 void ULevelStreamingProfilingSubsystem::PostInitialize()
 {
+	Super::PostInitialize();
+
 	Handle_OnLevelStreamingTargetStateChanged = FLevelStreamingDelegates::OnLevelStreamingTargetStateChanged.AddUObject(this, &ULevelStreamingProfilingSubsystem::OnLevelStreamingTargetStateChanged);
 	Handle_OnLevelStreamingStateChanged = FLevelStreamingDelegates::OnLevelStreamingStateChanged.AddUObject(this, &ULevelStreamingProfilingSubsystem::OnLevelStreamingStateChanged);
 	Handle_OnLevelBeginAddToWorld = FLevelStreamingDelegates::OnLevelBeginMakingVisible.AddUObject(this, &ULevelStreamingProfilingSubsystem::OnLevelStartedAddToWorld);
@@ -101,6 +103,16 @@ void ULevelStreamingProfilingSubsystem::Deinitialize()
 		StopTrackingAndReport();
 	}
 
+	FLevelStreamingDelegates::OnLevelStreamingTargetStateChanged.Remove(Handle_OnLevelStreamingTargetStateChanged);
+	FLevelStreamingDelegates::OnLevelStreamingStateChanged.Remove(Handle_OnLevelStreamingStateChanged);
+	FLevelStreamingDelegates::OnLevelBeginMakingVisible.Remove(Handle_OnLevelBeginAddToWorld);
+	FLevelStreamingDelegates::OnLevelBeginMakingInvisible.Remove(Handle_OnLevelBeginRemoveFromWorld);
+	
+	Handle_OnLevelStreamingTargetStateChanged.Reset();
+	Handle_OnLevelStreamingStateChanged.Reset();
+	Handle_OnLevelBeginAddToWorld.Reset();
+	Handle_OnLevelBeginRemoveFromWorld.Reset();
+
 	Super::Deinitialize();
 }
 
@@ -698,44 +710,47 @@ void ULevelStreamingProfilingSubsystem::OnLevelFinishedAddToWorld(UWorld* World,
 	LevelStats[Level->StatsIndex].TimeAddingToWorld = Time - Level->StateStartTime;
 	Level->TimeAddedToWorld = Time;
 
-	FBox CellBounds = LevelStats[Level->StatsIndex].CellBounds;
-	if (CellBounds.IsValid)
+	if (World && World->GetWorldPartition())
 	{
-		double MinDistance = MAX_dbl;
-		MinDistance = MAX_dbl;
-		FVector Location;
-		for (const FWorldPartitionStreamingSource& Source : World->GetWorldPartition()->GetStreamingSources())
+		FBox CellBounds = LevelStats[Level->StatsIndex].CellBounds;
+		if (CellBounds.IsValid)
 		{
-			double Distance = CellBounds.ComputeSquaredDistanceToPoint(Source.Location);
-			if (Distance < MinDistance)
-			{ 
-				MinDistance = Distance;
-				Location = Source.Location;
+			double MinDistance = MAX_dbl;
+			MinDistance = MAX_dbl;
+			FVector Location;
+			for (const FWorldPartitionStreamingSource& Source : World->GetWorldPartition()->GetStreamingSources())
+			{
+				double Distance = CellBounds.ComputeSquaredDistanceToPoint(Source.Location);
+				if (Distance < MinDistance)
+				{ 
+					MinDistance = Distance;
+					Location = Source.Location;
+				}
+			}
+			if (MinDistance != MAX_dbl)
+			{
+				UE_LOG(LogLevelStreamingProfiling, Verbose, TEXT("Streaming level %s %s streamed in at distance %.0f from cell bounds"), 
+					*StreamingLevel->GetName(), *StreamingLevel->GetWorldAssetPackageName(), MinDistance);
+				LevelStats[Level->StatsIndex].FinalStreamInDistance_Cell = FMath::Sqrt(MinDistance);
+				LevelStats[Level->StatsIndex].FinalStreamInLocation = Location;
 			}
 		}
-		if (MinDistance != MAX_dbl)
-		{
-			UE_LOG(LogLevelStreamingProfiling, Verbose, TEXT("Streaming level %s %s streamed in at distance %.0f from cell bounds"), 
-				*StreamingLevel->GetName(), *StreamingLevel->GetWorldAssetPackageName(), MinDistance);
-			LevelStats[Level->StatsIndex].FinalStreamInDistance_Cell = FMath::Sqrt(MinDistance);
-			LevelStats[Level->StatsIndex].FinalStreamInLocation = Location;
-		}
-	}
 
-	FBox ContentBounds = LevelStats[Level->StatsIndex].ContentBounds;
-	if (ContentBounds.IsValid)
-	{
-		double MinDistance = MAX_dbl;
-		for (const FWorldPartitionStreamingSource& Source : World->GetWorldPartition()->GetStreamingSources())
-		{
-			double Distance = ContentBounds.ComputeSquaredDistanceToPoint(Source.Location);
-			MinDistance = FMath::Min(Distance, MinDistance);
-		}
-		if (MinDistance != MAX_dbl)
+		FBox ContentBounds = LevelStats[Level->StatsIndex].ContentBounds;
+		if (ContentBounds.IsValid)
 		{
-			UE_LOG(LogLevelStreamingProfiling, Verbose, TEXT("Streaming level %s %s streamed in at distance %.0f from content bounds"), 
-				*StreamingLevel->GetName(), *StreamingLevel->GetWorldAssetPackageName(), MinDistance);
-			LevelStats[Level->StatsIndex].FinalStreamInDistance_Content = FMath::Sqrt(MinDistance);
+			double MinDistance = MAX_dbl;
+			for (const FWorldPartitionStreamingSource& Source : World->GetWorldPartition()->GetStreamingSources())
+			{
+				double Distance = ContentBounds.ComputeSquaredDistanceToPoint(Source.Location);
+				MinDistance = FMath::Min(Distance, MinDistance);
+			}
+			if (MinDistance != MAX_dbl)
+			{
+				UE_LOG(LogLevelStreamingProfiling, Verbose, TEXT("Streaming level %s %s streamed in at distance %.0f from content bounds"), 
+					*StreamingLevel->GetName(), *StreamingLevel->GetWorldAssetPackageName(), MinDistance);
+				LevelStats[Level->StatsIndex].FinalStreamInDistance_Content = FMath::Sqrt(MinDistance);
+			}
 		}
 	}
 	UpdateLevelState(*Level, StreamingLevel, ELevelState::AddedToWorld, Time);
