Level Instances: fix creation and duplication not loading actors when partial loading is enabled.

diff --git a/Engine/Source/Runtime/Engine/Private/LevelInstance/LevelInstanceLevelStreaming.cpp b/Engine/Source/Runtime/Engine/Private/LevelInstance/LevelInstanceLevelStreaming.cpp
--- a/Engine/Source/Runtime/Engine/Private/LevelInstance/LevelInstanceLevelStreaming.cpp
+++ b/Engine/Source/Runtime/Engine/Private/LevelInstance/LevelInstanceLevelStreaming.cpp
@@ -364,6 +364,7 @@ void ULevelStreamingLevelInstance::OnLevelLoadedChanged(ULevel* InLevel)
 		if (ULevelInstanceSubsystem* LevelInstanceSubsystem = GetWorld()->GetSubsystem<ULevelInstanceSubsystem>())
 		{
 			LevelInstanceSubsystem->RegisterLoadedLevelStreamingLevelInstance(this);
+
 			if (!GDisableLevelInstanceEditorPartialLoading)
 			{
 				if (UWorldPartition* WorldPartition = NewLoadedLevel->GetWorldPartition())
@@ -371,7 +372,13 @@ void ULevelStreamingLevelInstance::OnLevelLoadedChanged(ULevel* InLevel)
 					check(!WorldPartition->IsInitialized());
 					if (ILevelInstanceInterface* LevelInstance = LevelInstanceSubsystem->GetLevelInstance(LevelInstanceID))
 					{
-						WorldPartition->bForceEnableStreamingInEditor = LevelInstance->SupportsPartialEditorLoading();
+						if (AActor* Actor = CastChecked<AActor>(LevelInstance))
+						{
+							if (!Actor->GetPackage()->HasAnyPackageFlags(PKG_NewlyCreated))
+							{
+								WorldPartition->bForceEnableStreamingInEditor = LevelInstance->SupportsPartialEditorLoading();
+							}
+						}
 					}
 				}
 			}
