Added functionality to IXRLoadingScreen to check if XR loading movie is playing, and to skip updating viewports in LoadMap if so. This matches existing logic for non-VR loading movies. Fixes swapchain crash when loading a new level with an OpenXR splash movie playing on packaged Windows builds.

diff --git a/Engine/Plugins/Runtime/Oculus/OculusVR/Source/OculusHMD/Private/OculusHMD_Splash.cpp b/Engine/Plugins/Runtime/Oculus/OculusVR/Source/OculusHMD/Private/OculusHMD_Splash.cpp
--- a/Engine/Plugins/Runtime/Oculus/OculusVR/Source/OculusHMD/Private/OculusHMD_Splash.cpp
+++ b/Engine/Plugins/Runtime/Oculus/OculusVR/Source/OculusHMD/Private/OculusHMD_Splash.cpp
@@ -670,6 +670,24 @@ void FSplash::UnloadTexture(FSplashLayer& InSplashLayer)
 	InSplashLayer.Layer.Reset();
 }
 
+bool FSplash::IsPlayingLoadingMovie() const
+{
+	if (!bIsShown)
+	{
+		return false;
+	}
+
+	for (FSplashLayer Splash : SplashLayers)
+	{
+		if (Splash.Desc.bIsDynamic)
+		{
+			return true;
+		}
+	}
+
+	return false;
+}
+
 
 } // namespace OculusHMD
 
diff --git a/Engine/Plugins/Runtime/Oculus/OculusVR/Source/OculusHMD/Private/OculusHMD_Splash.h b/Engine/Plugins/Runtime/Oculus/OculusVR/Source/OculusHMD/Private/OculusHMD_Splash.h
--- a/Engine/Plugins/Runtime/Oculus/OculusVR/Source/OculusHMD/Private/OculusHMD_Splash.h
+++ b/Engine/Plugins/Runtime/Oculus/OculusVR/Source/OculusHMD/Private/OculusHMD_Splash.h
@@ -87,6 +87,7 @@ public:
 	virtual void ClearSplashes() override;
 	virtual void AddSplash(const FSplashDesc& Splash) override;
 	virtual bool IsShown() const override { return bIsShown; }
+	virtual bool IsPlayingLoadingMovie() const override;
 
 protected:
 	void DoShow();
diff --git a/Engine/Source/Runtime/Engine/Private/UnrealEngine.cpp b/Engine/Source/Runtime/Engine/Private/UnrealEngine.cpp
--- a/Engine/Source/Runtime/Engine/Private/UnrealEngine.cpp
+++ b/Engine/Source/Runtime/Engine/Private/UnrealEngine.cpp
@@ -99,6 +99,7 @@ UnrealEngine.cpp: Implements the UEngine class and helpers.
 #include "IHeadMountedDisplayModule.h"
 #include "IHeadMountedDisplay.h"
 #include "IXRTrackingSystem.h"
+#include "IXRLoadingScreen.h"
 #include "Stats/StatsFile.h"
 #include "Audio/AudioDebug.h"
 #include "AudioDevice.h"
@@ -14765,8 +14766,18 @@ bool UEngine::LoadMap( FWorldContext& WorldContext, FURL URL, class UPendingNetG
 
 			// Display loading screen.		
 			// Check if a loading movie is playing.  If so it is not safe to redraw the viewport due to potential race conditions with font rendering
-			bool bIsLoadingMovieCurrentlyPlaying = FCoreDelegates::IsLoadingMovieCurrentlyPlaying.IsBound() ? FCoreDelegates::IsLoadingMovieCurrentlyPlaying.Execute() : false;
-			if(!bIsLoadingMovieCurrentlyPlaying)
+			bool bIsDesktopLoadingMovieCurrentlyPlaying = FCoreDelegates::IsLoadingMovieCurrentlyPlaying.IsBound() ? FCoreDelegates::IsLoadingMovieCurrentlyPlaying.Execute() : false;
+
+			// Check if a loading movie is playing on an XR device as well. Redrawing viewports with a dynamic OpenXR stereo layer active
+			// causes a race condition with CopyTexture_RenderThread that results in acquiring swapchain images twice without releasing
+			bool bIsXRLoadingMovieCurrentlyPlaying = false;
+			if (GEngine && GEngine->XRSystem.IsValid())
+			{
+				IXRLoadingScreen* XRLoadingScreen = GEngine->XRSystem->GetLoadingScreen();
+				bIsXRLoadingMovieCurrentlyPlaying = XRLoadingScreen ? XRLoadingScreen->IsPlayingLoadingMovie() : false;
+			}
+
+			if(!bIsDesktopLoadingMovieCurrentlyPlaying && !bIsXRLoadingMovieCurrentlyPlaying)
 			{
 				LoadMapRedrawViewports();
 			}
diff --git a/Engine/Source/Runtime/HeadMountedDisplay/Private/DefaultXRLoadingScreen.cpp b/Engine/Source/Runtime/HeadMountedDisplay/Private/DefaultXRLoadingScreen.cpp
--- a/Engine/Source/Runtime/HeadMountedDisplay/Private/DefaultXRLoadingScreen.cpp
+++ b/Engine/Source/Runtime/HeadMountedDisplay/Private/DefaultXRLoadingScreen.cpp
@@ -60,6 +60,24 @@ void FDefaultXRLoadingScreen::HideLoadingScreen()
 	}
 }
 
+bool FDefaultXRLoadingScreen::IsPlayingLoadingMovie() const
+{
+	if (!bShowing)
+	{
+		return false;
+	}
+
+	for (FSplashData Splash : Splashes)
+	{
+		if (Splash.Desc.bIsDynamic)
+		{
+			return true;
+		}
+	}
+
+	return false;
+}
+
 void FDefaultXRLoadingScreen::DoHideSplash(FSplashData& Splash)
 {
 	// Pop layer state will destroy all created layers, so we clear the layer id stored locally
diff --git a/Engine/Source/Runtime/HeadMountedDisplay/Public/DefaultXRLoadingScreen.h b/Engine/Source/Runtime/HeadMountedDisplay/Public/DefaultXRLoadingScreen.h
--- a/Engine/Source/Runtime/HeadMountedDisplay/Public/DefaultXRLoadingScreen.h
+++ b/Engine/Source/Runtime/HeadMountedDisplay/Public/DefaultXRLoadingScreen.h
@@ -31,6 +31,7 @@ public:
 	/* IXRLoadingScreen interface */
 	virtual void ShowLoadingScreen() override;
 	virtual void HideLoadingScreen() override;
+	virtual bool IsPlayingLoadingMovie() const override;
 
 private:
 
diff --git a/Engine/Source/Runtime/HeadMountedDisplay/Public/IXRLoadingScreen.h b/Engine/Source/Runtime/HeadMountedDisplay/Public/IXRLoadingScreen.h
--- a/Engine/Source/Runtime/HeadMountedDisplay/Public/IXRLoadingScreen.h
+++ b/Engine/Source/Runtime/HeadMountedDisplay/Public/IXRLoadingScreen.h
@@ -82,6 +82,11 @@ public:
 	 */
 	virtual bool IsShown() const = 0;
 
+	/**
+	 * Returns whether a loading movie is playing or not.
+	 */
+	virtual bool IsPlayingLoadingMovie() const = 0;
+
 	/**
 	 * Internal utility method for implementing backwards compatibility with IStereoLayers::Show/HideSplashScreen.
 	 * Should be called from implementations overriding IStereoLayer::UpdateSplashScreen()
diff --git a/Engine/Source/Runtime/HeadMountedDisplay/Public/XRLoadingScreenBase.h b/Engine/Source/Runtime/HeadMountedDisplay/Public/XRLoadingScreenBase.h
--- a/Engine/Source/Runtime/HeadMountedDisplay/Public/XRLoadingScreenBase.h
+++ b/Engine/Source/Runtime/HeadMountedDisplay/Public/XRLoadingScreenBase.h
@@ -96,6 +96,7 @@ public:
 	}
 
 	virtual bool IsShown() const override { return bShowing; }
+	virtual bool IsPlayingLoadingMovie() const = 0;
 
 
 protected:
