WorldPartitionActorFilter: Use 'or' logic when applying filters

#rnx

diff --git a/Engine/Source/Runtime/Engine/Private/WorldPartition/WorldPartitionSubsystem.cpp b/Engine/Source/Runtime/Engine/Private/WorldPartition/WorldPartitionSubsystem.cpp
--- a/Engine/Source/Runtime/Engine/Private/WorldPartition/WorldPartitionSubsystem.cpp
+++ b/Engine/Source/Runtime/Engine/Private/WorldPartition/WorldPartitionSubsystem.cpp
@@ -212,18 +212,15 @@ TMap<FActorContainerID, TSet<FGuid>> UWorldPartitionSubsystem::GetFilteredActors
 	ContainerFilter.Override(InActorFilter);
 
 	// Flatten Filter to FActorContainerID map
-	TMap<FActorContainerID, TSet<FSoftObjectPath>> FilteredOutDataLayersPerContainer;
-	TFunction<void(const FActorContainerID&, const FWorldPartitionActorFilter&)> ProcessFilter = [&FilteredOutDataLayersPerContainer, &ProcessFilter](const FActorContainerID& InContainerID, const FWorldPartitionActorFilter& InContainerFilter)
+	TMap<FActorContainerID, TMap<FSoftObjectPath, FWorldPartitionActorFilter::FDataLayerFilter>> DataLayerFiltersPerContainer;
+	TFunction<void(const FActorContainerID&, const FWorldPartitionActorFilter&)> ProcessFilter = [&DataLayerFiltersPerContainer, &ProcessFilter](const FActorContainerID& InContainerID, const FWorldPartitionActorFilter& InContainerFilter)
 	{
-		check(!FilteredOutDataLayersPerContainer.Contains(InContainerID));
-		TSet<FSoftObjectPath>& Filtered = FilteredOutDataLayersPerContainer.Add(InContainerID);
-
+		check(!DataLayerFiltersPerContainer.Contains(InContainerID));
+		TMap<FSoftObjectPath,FWorldPartitionActorFilter::FDataLayerFilter>& DataLayerFilters = DataLayerFiltersPerContainer.Add(InContainerID);
+		
 		for (auto& [AssetPath, DataLayerFilter] : InContainerFilter.DataLayerFilters)
 		{
-			if (!DataLayerFilter.bIncluded)
-			{
-				Filtered.Add(AssetPath);
-			}
+			DataLayerFilters.Add(AssetPath, DataLayerFilter);
 		}
 
 		for (auto& [ActorGuid, WorldPartitionActorFilter] : InContainerFilter.GetChildFilters())
@@ -234,22 +231,35 @@ TMap<FActorContainerID, TSet<FGuid>> UWorldPartitionSubsystem::GetFilteredActors
 
 	ProcessFilter(InContainerID, ContainerFilter);
 
-	TFunction<void(const FActorContainerID&, const UActorDescContainer*)> ProcessContainers = [&FilteredOutDataLayersPerContainer, &FilteredActors, &ProcessContainers](const FActorContainerID& InContainerID, const UActorDescContainer* InContainer)
+	TFunction<void(const FActorContainerID&, const UActorDescContainer*)> ProcessContainers = [&DataLayerFiltersPerContainer, &FilteredActors, &ProcessContainers](const FActorContainerID& InContainerID, const UActorDescContainer* InContainer)
 	{
-		const TSet<FSoftObjectPath>& Filtered = FilteredOutDataLayersPerContainer.FindChecked(InContainerID);
+		const TMap<FSoftObjectPath, FWorldPartitionActorFilter::FDataLayerFilter>& DataLayerFilters = DataLayerFiltersPerContainer.FindChecked(InContainerID);
 		for (FActorDescList::TConstIterator<> ActorDescIt(InContainer); ActorDescIt; ++ActorDescIt)
 		{
 			if (ActorDescIt->GetDataLayers().Num() > 0 && ActorDescIt->IsUsingDataLayerAsset())
 			{
+				bool bExcluded = false;
 				for (FName DataLayerName : ActorDescIt->GetDataLayers())
 				{
 					FSoftObjectPath DataLayerAsset(DataLayerName.ToString());
-					if (Filtered.Contains(DataLayerAsset))
+					if (const FWorldPartitionActorFilter::FDataLayerFilter* DataLayerFilter = DataLayerFilters.Find(DataLayerAsset))
 					{
-						FilteredActors.FindOrAdd(InContainerID).Add(ActorDescIt->GetGuid());
-						break;
+						if (DataLayerFilter->bIncluded)
+						{
+							bExcluded = false;
+							break;
+						}
+						else
+						{
+							bExcluded = true;
+						}
 					}
 				}
+
+				if (bExcluded)
+				{
+					FilteredActors.FindOrAdd(InContainerID).Add(ActorDescIt->GetGuid());
+				}
 			}
 
 			if (ActorDescIt->IsContainerInstance())
