Move CheckMaterialUsage() to UHeterogeneousVolumeComponent's PostEditChangeProperty().

diff --git a/Engine/Source/Runtime/Engine/Classes/Components/HeterogeneousVolumeComponent.h b/Engine/Source/Runtime/Engine/Classes/Components/HeterogeneousVolumeComponent.h
--- a/Engine/Source/Runtime/Engine/Classes/Components/HeterogeneousVolumeComponent.h
+++ b/Engine/Source/Runtime/Engine/Classes/Components/HeterogeneousVolumeComponent.h
@@ -37,6 +37,11 @@ public:
 	//~ Begin USceneComponent Interface.
 	virtual FBoxSphereBounds CalcBounds(const FTransform& LocalToWorld) const override;
 	virtual void TickComponent(float DeltaTime, ELevelTick TickType, FActorComponentTickFunction* ThisTickFunction) override;
+
+#if WITH_EDITOR
+	virtual void PostEditChangeProperty(FPropertyChangedEvent& PropertyChangedEvent) override;
+#endif // WITH_EDITOR
+
 	//~ End USceneComponent Interface.
 
 	//~ Begin UPrimitiveComponent Interface.
diff --git a/Engine/Source/Runtime/Engine/Private/Components/HeterogeneousVolumeComponent.cpp b/Engine/Source/Runtime/Engine/Private/Components/HeterogeneousVolumeComponent.cpp
--- a/Engine/Source/Runtime/Engine/Private/Components/HeterogeneousVolumeComponent.cpp
+++ b/Engine/Source/Runtime/Engine/Private/Components/HeterogeneousVolumeComponent.cpp
@@ -80,17 +80,6 @@ FHeterogeneousVolumeSceneProxy::FHeterogeneousVolumeSceneProxy(UHeterogeneousVol
 
 	HeterogeneousVolumeData.LightingDownsampleFactor = InComponent->LightingDownsampleFactor;
 
-	// Update material assignment to include Heterogeneous Volumes
-	if (MaterialInterface)
-	{
-		const UMaterial* Material = MaterialInterface->GetMaterial();
-		if (Material && Material->MaterialDomain == EMaterialDomain::MD_Volume)
-		{
-			Material->GetRenderProxy();
-			MaterialInterface->CheckMaterialUsage(MATUSAGE_HeterogeneousVolumes);
-		}
-	}
-
 	// Initialize vertex buffer data for a quad
 	StaticMeshVertexBuffers.PositionVertexBuffer.Init(4);
 	StaticMeshVertexBuffers.StaticMeshVertexBuffer.Init(4, 1);
@@ -237,6 +226,33 @@ FBoxSphereBounds UHeterogeneousVolumeComponent::CalcBounds(const FTransform& Loc
 	return NewBounds.TransformBy(LocalToWorld);
 }
 
+#if WITH_EDITOR
+void UHeterogeneousVolumeComponent::PostEditChangeProperty(FPropertyChangedEvent& PropertyChangedEvent)
+{
+	Super::PostEditChangeProperty(PropertyChangedEvent);
+
+	FName PropertyName;
+	if (PropertyChangedEvent.Property)
+	{
+		PropertyName = PropertyChangedEvent.Property->GetFName();
+	}
+
+	if (PropertyName == GET_MEMBER_NAME_CHECKED(UHeterogeneousVolumeComponent, OverrideMaterials))
+	{
+		UMaterialInterface* MaterialInterface = GetMaterial(0);
+		if (MaterialInterface)
+		{
+			const UMaterial* Material = MaterialInterface->GetMaterial();
+			if (Material && Material->MaterialDomain == EMaterialDomain::MD_Volume)
+			{
+				Material->GetRenderProxy();
+				MaterialInterface->CheckMaterialUsage(MATUSAGE_HeterogeneousVolumes);
+			}
+		}
+	}
+}
+#endif // WITH_EDITOR
+
 void UHeterogeneousVolumeComponent::TickComponent(float DeltaTime, ELevelTick TickType, FActorComponentTickFunction* ThisTickFunction)
 {
 	Super::TickComponent(DeltaTime, TickType, ThisTickFunction);
@@ -269,18 +285,21 @@ void UHeterogeneousVolumeComponent::TickComponent(float DeltaTime, ELevelTick Ti
 			USparseVolumeTexture* DefaultSVT = nullptr;
 			MaterialInstanceDynamic->GetSparseVolumeTextureParameterDefaultValue(SVTParameterInfo[MaterialIndex], DefaultSVT);
 
-			if (bAnimate)
-			{
-				int32 FrameIndex = int32(Time * Framerate) % DefaultSVT->GetNumFrames();
-				int32 MipLevel = 0;
-				USparseVolumeTextureFrame* SVTFrame = USparseVolumeTextureFrame::CreateFrame(DefaultSVT, FrameIndex, MipLevel);
-				MaterialInstanceDynamic->SetSparseVolumeTextureParameterValue(SVTParameterInfo[MaterialIndex].Name, SVTFrame);
-
-				VolumeResolution = SVTFrame ? SVTFrame->GetVolumeResolution() : FIntVector(1);
-			}
-			else
+			if (DefaultSVT)
 			{
-				VolumeResolution = DefaultSVT->GetVolumeResolution();
+				if (bAnimate)
+				{
+					int32 FrameIndex = int32(Time * Framerate) % DefaultSVT->GetNumFrames();
+					int32 MipLevel = 0;
+					USparseVolumeTextureFrame* SVTFrame = USparseVolumeTextureFrame::CreateFrame(DefaultSVT, FrameIndex, MipLevel);
+					MaterialInstanceDynamic->SetSparseVolumeTextureParameterValue(SVTParameterInfo[MaterialIndex].Name, SVTFrame);
+
+					VolumeResolution = SVTFrame ? SVTFrame->GetVolumeResolution() : FIntVector(1);
+				}
+				else
+				{
+					VolumeResolution = DefaultSVT->GetVolumeResolution();
+				}
 			}
 		}
 	}
