Fix (somewhat speculative) for a crash in UBlendSpace::TickFollowerSamples because of hardcoded bLooping

diff --git a/Engine/Source/Runtime/Engine/Classes/Animation/BlendSpace.h b/Engine/Source/Runtime/Engine/Classes/Animation/BlendSpace.h
--- a/Engine/Source/Runtime/Engine/Classes/Animation/BlendSpace.h
+++ b/Engine/Source/Runtime/Engine/Classes/Animation/BlendSpace.h
@@ -671,7 +671,10 @@ protected:
 	/** Initialize Per Bone Blend **/
 	void InitializePerBoneBlend();
 
-	void TickFollowerSamples(TArray<FBlendSampleData> &SampleDataList, const int32 HighestWeightIndex, FAnimAssetTickContext &Context, bool bResetMarkerDataOnFollowers, const UMirrorDataTable* MirrorDataTable = nullptr) const;
+	/** Ticks the samples in SampleDataList apart from the HighestWeightIndex one. */
+	void TickFollowerSamples(
+		TArray<FBlendSampleData> &SampleDataList, const int32 HighestWeightIndex, FAnimAssetTickContext &Context, 
+		bool bResetMarkerDataOnFollowers, bool bLooping, const UMirrorDataTable* MirrorDataTable = nullptr) const;
 
 	/** Returns the blend input clamped to the valid range, unless that axis has been set to wrap in which case no clamping is done **/
 	FVector GetClampedBlendInput(const FVector& BlendInput) const;
diff --git a/Engine/Source/Runtime/Engine/Private/Animation/AnimSequence.cpp b/Engine/Source/Runtime/Engine/Private/Animation/AnimSequence.cpp
--- a/Engine/Source/Runtime/Engine/Private/Animation/AnimSequence.cpp
+++ b/Engine/Source/Runtime/Engine/Private/Animation/AnimSequence.cpp
@@ -4245,7 +4245,9 @@ float UAnimSequence::GetCurrentTimeFromMarkers(FMarkerPair& PrevMarker, FMarkerP
 	return CurrentTime;
 }
 
-void UAnimSequence::GetMarkerIndicesForPosition(const FMarkerSyncAnimPosition& SyncPosition, bool bLooping, FMarkerPair& OutPrevMarker, FMarkerPair& OutNextMarker, float& OutCurrentTime, const UMirrorDataTable* MirrorTable) const
+void UAnimSequence::GetMarkerIndicesForPosition(
+	const FMarkerSyncAnimPosition& SyncPosition, bool bLooping, FMarkerPair& OutPrevMarker, 
+	FMarkerPair& OutNextMarker, float& OutCurrentTime, const UMirrorDataTable* MirrorTable) const
 {
 	auto GetMarkerName = [MirrorTable](const FAnimSyncMarker& SyncMarker)->FName
 	{
@@ -4253,7 +4255,8 @@ void UAnimSequence::GetMarkerIndicesForPosition(const FMarkerSyncAnimPosition& S
 	};
 
 	// If we're not looping, assume we're playing a transition and we need to stay where we are.
-	if (!bLooping)
+	// Also do this if we have no usable SyncPosition
+	if (!bLooping || (SyncPosition.PreviousMarkerName == NAME_None && SyncPosition.NextMarkerName == NAME_None))
 	{
 		OutPrevMarker.MarkerIndex = INDEX_NONE;
 		OutNextMarker.MarkerIndex = INDEX_NONE;
diff --git a/Engine/Source/Runtime/Engine/Private/Animation/BlendSpace.cpp b/Engine/Source/Runtime/Engine/Private/Animation/BlendSpace.cpp
--- a/Engine/Source/Runtime/Engine/Private/Animation/BlendSpace.cpp
+++ b/Engine/Source/Runtime/Engine/Private/Animation/BlendSpace.cpp
@@ -493,7 +493,7 @@ void UBlendSpace::TickAssetPlayer(FAnimTickRecord& Instance, struct FAnimNotifyQ
 							Context.SetLeaderDelta(NewDeltaTime);
 							Sample.Animation->TickByMarkerAsLeader(SampleDataItem.MarkerTickRecord, Context.MarkerTickContext, SampleDataItem.Time, SampleDataItem.PreviousTime, NewDeltaTime, Instance.bLooping, Instance.MirrorDataTable);
 							check(!Instance.bLooping || Context.MarkerTickContext.IsMarkerSyncStartValid());
-							TickFollowerSamples(SampleDataList, HighestMarkerSyncWeightIndex, Context, bResetMarkerDataOnFollowers, Instance.MirrorDataTable);
+							TickFollowerSamples(SampleDataList, HighestMarkerSyncWeightIndex, Context, bResetMarkerDataOnFollowers, Instance.bLooping, Instance.MirrorDataTable);
 						}
 						NormalizedCurrentTime = SampleDataItem.Time / Sample.Animation->GetPlayLength();
 						*Instance.MarkerTickRecord = SampleDataItem.MarkerTickRecord;
@@ -534,7 +534,7 @@ void UBlendSpace::TickAssetPlayer(FAnimTickRecord& Instance, struct FAnimNotifyQ
 								SampleDataItem.Time = NormalizedCurrentTime * Sample.Animation->GetPlayLength();
 							}
 
-							TickFollowerSamples(SampleDataList, -1, Context, false, Instance.MirrorDataTable);
+							TickFollowerSamples(SampleDataList, -1, Context, false, Instance.bLooping, Instance.MirrorDataTable);
 						}
 						*Instance.MarkerTickRecord = SampleDataItem.MarkerTickRecord;
 						NormalizedCurrentTime = SampleDataItem.Time / Sample.Animation->GetPlayLength();
@@ -1595,7 +1595,9 @@ void UBlendSpace::InitializePerBoneBlend()
 	}
 }
 
-void UBlendSpace::TickFollowerSamples(TArray<FBlendSampleData>& SampleDataList, const int32 HighestWeightIndex, FAnimAssetTickContext& Context, bool bResetMarkerDataOnFollowers, const UMirrorDataTable* MirrorDataTable) const
+void UBlendSpace::TickFollowerSamples(
+	TArray<FBlendSampleData>& SampleDataList, const int32 HighestWeightIndex, FAnimAssetTickContext& Context, 
+	bool bResetMarkerDataOnFollowers, bool bLooping, const UMirrorDataTable* MirrorDataTable) const
 {
 	for (int32 SampleIndex = 0; SampleIndex < SampleDataList.Num(); ++SampleIndex)
 	{
@@ -1608,9 +1610,12 @@ void UBlendSpace::TickFollowerSamples(TArray<FBlendSampleData>& SampleDataList,
 				SampleDataItem.MarkerTickRecord.Reset();
 			}
 
-			if (Sample.Animation->AuthoredSyncMarkers.Num() > 0) // Update followers who can do marker sync, others will be handled later in TickAssetPlayer
+			// Update followers who can do marker sync, others will be handled later in TickAssetPlayer
+			if (Sample.Animation->AuthoredSyncMarkers.Num() > 0) 
 			{
-				Sample.Animation->TickByMarkerAsFollower(SampleDataItem.MarkerTickRecord, Context.MarkerTickContext, SampleDataItem.Time, SampleDataItem.PreviousTime, Context.GetLeaderDelta(), true, MirrorDataTable);
+				Sample.Animation->TickByMarkerAsFollower(
+					SampleDataItem.MarkerTickRecord, Context.MarkerTickContext, SampleDataItem.Time, 
+					SampleDataItem.PreviousTime, Context.GetLeaderDelta(), bLooping, MirrorDataTable);
 			}
 		}
 	}
