Improved FSparseVolumeTextureRuntime::SetTileMappings() performance. Replaced unnecessary Read/WriteVoxel() calls with a single memcpy() call to copy entire rows of source data at once. Also used ParallelFor to copy tiles in parallel.

#rnx

diff --git a/Engine/Source/Runtime/Engine/Private/SparseVolumeTexture/SparseVolumeTextureSceneProxy.cpp b/Engine/Source/Runtime/Engine/Private/SparseVolumeTexture/SparseVolumeTextureSceneProxy.cpp
--- a/Engine/Source/Runtime/Engine/Private/SparseVolumeTexture/SparseVolumeTextureSceneProxy.cpp
+++ b/Engine/Source/Runtime/Engine/Private/SparseVolumeTexture/SparseVolumeTextureSceneProxy.cpp
@@ -3,6 +3,7 @@
 #include "SparseVolumeTexture/SparseVolumeTextureSceneProxy.h"
 #include "SparseVolumeTexture/SparseVolumeTextureUtility.h"
 #include "SparseVolumeTexture/SparseVolumeTextureData.h"
+#include "Async/ParallelFor.h"
 
 DEFINE_LOG_CATEGORY_STATIC(LogSparseVolumeTextureProxy, Log, All);
 
@@ -44,6 +45,16 @@ namespace Private
 
 		return TileDataVolumeResolution * SPARSE_VOLUME_TILE_RES_PADDED;
 	}
+
+	static FIntVector3 UnflattenCoord(int32 Index, const FIntVector3& VolumeSize)
+	{
+		FIntVector3 Result;
+		Result.X = Index % VolumeSize.X;
+		Result.Y = (Index % (VolumeSize.Y * VolumeSize.X)) / VolumeSize.X;
+		Result.Z = Index / (VolumeSize.Y * VolumeSize.X);
+
+		return Result;
+	}
 } // Private
 } // SVT
 } // UE
@@ -198,10 +209,11 @@ bool FSparseVolumeTextureRuntime::SetTileMappings(const TArrayView<const FSparse
 	PhysicalTileDataB.SetNum(Header.TileDataVolumeResolution.X * Header.TileDataVolumeResolution.Y * Header.TileDataVolumeResolution.Z * FormatSize[1]);
 	uint8* PhysicalTileDataPtrs[] = { PhysicalTileDataA.GetData(), PhysicalTileDataB.GetData() };
 
-	int32 NumWrittenTiles = 0;
-	FIntVector3 DstTileCoord = FIntVector3::ZeroValue;
-
 	// Write null tile
+	uint8 NullTileDataBytesA[sizeof(float) * 4];
+	uint8 NullTileDataBytesB[sizeof(float) * 4];
+	WriteVoxel(0, NullTileDataBytesA, Header.AttributesFormats[0], Header.NullTileValues[0]);
+	WriteVoxel(0, NullTileDataBytesB, Header.AttributesFormats[1], Header.NullTileValues[1]);
 	for (int32 Z = 0; Z < SPARSE_VOLUME_TILE_RES_PADDED; ++Z)
 	{
 		for (int32 Y = 0; Y < SPARSE_VOLUME_TILE_RES_PADDED; ++Y)
@@ -209,18 +221,69 @@ bool FSparseVolumeTextureRuntime::SetTileMappings(const TArrayView<const FSparse
 			for (int32 X = 0; X < SPARSE_VOLUME_TILE_RES_PADDED; ++X)
 			{
 				const int32 VoxelIndex = Z * (SPARSE_VOLUME_TILE_RES_PADDED * SPARSE_VOLUME_TILE_RES_PADDED) + Y * SPARSE_VOLUME_TILE_RES_PADDED + X;
-				WriteVoxel(VoxelIndex, PhysicalTileDataPtrs[0], Header.AttributesFormats[0], Header.NullTileValues[0]);
-				WriteVoxel(VoxelIndex, PhysicalTileDataPtrs[1], Header.AttributesFormats[1], Header.NullTileValues[1]);
+				if (FormatSize[0])
+				{
+					uint8* DstVoxelA = PhysicalTileDataPtrs[0] + FormatSize[0] * VoxelIndex;
+					FMemory::Memcpy(DstVoxelA, NullTileDataBytesA, FormatSize[0]);
+				}
+				if (FormatSize[1])
+				{
+					uint8* DstVoxelB = PhysicalTileDataPtrs[1] + FormatSize[1] * VoxelIndex;
+					FMemory::Memcpy(DstVoxelB, NullTileDataBytesB, FormatSize[1]);
+				}
 			}
 		}
 	}
-	++NumWrittenTiles;
-	DstTileCoord = AdvanceTileCoord(DstTileCoord, TileCoordSpace);
+
+	// Write to tile data
+	ParallelFor(NumTiles - 1, [&](int32 TileIndex)
+		{
+			TileIndex += 1; // skipped null tile
+
+			int32 TileIndexOffset = 1; // makes TileIndex relative to the range of source tiles of a mip
+			const uint8* SrcTileData[2] = {};
+			for (const FSparseVolumeTextureTileMapping& Mapping : Mappings)
+			{
+				const int32 NumPhysicalTiles = Mapping.NumPhysicalTiles;
+				if (TileIndex < (TileIndexOffset + NumPhysicalTiles) && Mapping.TileIndices)
+				{
+					SrcTileData[0] = Mapping.TileDataA;
+					SrcTileData[1] = Mapping.TileDataB;
+					break;
+				}
+				TileIndexOffset += NumPhysicalTiles;
+			}
+			const int32 SrcTileIndex = TileIndex - TileIndexOffset;
+			const FIntVector3 DstTileCoord = UnflattenCoord(TileIndex, TileCoordSpace);
+
+			for (int32 Z = 0; Z < SPARSE_VOLUME_TILE_RES_PADDED; ++Z)
+			{
+				for (int32 Y = 0; Y < SPARSE_VOLUME_TILE_RES_PADDED; ++Y)
+				{
+					const int32 SrcRowStartVoxelIndex = SrcTileIndex * SVTNumVoxelsPerPaddedTile + Z * (SPARSE_VOLUME_TILE_RES_PADDED * SPARSE_VOLUME_TILE_RES_PADDED) + Y * SPARSE_VOLUME_TILE_RES_PADDED;
+					const FIntVector3 RowStartVoxelCoord = DstTileCoord * SPARSE_VOLUME_TILE_RES_PADDED + FIntVector3(0, Y, Z);
+					const int32 DstRowStartVoxelIndex = RowStartVoxelCoord.Z * (Header.TileDataVolumeResolution.Y * Header.TileDataVolumeResolution.X) + RowStartVoxelCoord.Y * Header.TileDataVolumeResolution.X + RowStartVoxelCoord.X;
+					if (FormatSize[0])
+					{
+						uint8* Dst = PhysicalTileDataPtrs[0] + FormatSize[0] * DstRowStartVoxelIndex;
+						const uint8* Src = SrcTileData[0] + FormatSize[0] * SrcRowStartVoxelIndex;
+						FMemory::Memcpy(Dst, Src, FormatSize[0] * SPARSE_VOLUME_TILE_RES_PADDED);
+					}
+					if (FormatSize[1])
+					{
+						uint8* Dst = PhysicalTileDataPtrs[1] + FormatSize[1] * DstRowStartVoxelIndex;
+						const uint8* Src = SrcTileData[1] + FormatSize[1] * SrcRowStartVoxelIndex;
+						FMemory::Memcpy(Dst, Src, FormatSize[1] * SPARSE_VOLUME_TILE_RES_PADDED);
+					}
+				}
+			}
+		});
 
 	// Write page table and physical tiles
 	Header.HighestResidentLevel = INT32_MIN;
 	Header.LowestResidentLevel = INT32_MAX;
 	FIntVector3 PageTableRes = Header.PageTableVolumeResolution * 2;
+	int32 MipTileOffset = 1;
 	for (int32 MipLevel = 0; MipLevel < NumMipLevels; ++MipLevel)
 	{
 		PageTableRes = FIntVector3(FMath::Max(1, PageTableRes.X / 2), FMath::Max(1, PageTableRes.Y / 2), FMath::Max(1, PageTableRes.Z / 2));
@@ -237,7 +300,6 @@ bool FSparseVolumeTextureRuntime::SetTileMappings(const TArrayView<const FSparse
 		// Write page table
 		PageTable[MipLevel].SetNum(PageTableRes.X * PageTableRes.Y * PageTableRes.Z);
 		uint32* PageTablePtr = PageTable[MipLevel].GetData();
-		const int32 MipTileOffset = NumWrittenTiles;
 		for (int32 PageZ = 0; PageZ < PageTableRes.Z; ++PageZ)
 		{
 			for (int32 PageY = 0; PageY < PageTableRes.Y; ++PageY)
@@ -254,11 +316,7 @@ bool FSparseVolumeTextureRuntime::SetTileMappings(const TArrayView<const FSparse
 					{
 						// Make index relative to the start index of this mip level
 						const uint32 TileIndex = MipTileOffset + MipLocalTileIndex - 1; // -1 because MipLocalTileIndex uses 0 to refer to the null tile, but we only want a single null tile for the entire resource.
-
-						FIntVector3 TileCoord;
-						TileCoord.X = TileIndex % TileCoordSpace.X;
-						TileCoord.Z = TileIndex / (TileCoordSpace.Y * TileCoordSpace.X);
-						TileCoord.Y = (TileIndex - (TileCoord.Z * (TileCoordSpace.Y * TileCoordSpace.X))) / TileCoordSpace.X;
+						const FIntVector3 TileCoord = UnflattenCoord(TileIndex, TileCoordSpace);
 
 						// Write to page table
 						PageTablePtr[PageIndex] = PackPageTableEntry(TileCoord);
@@ -266,30 +324,7 @@ bool FSparseVolumeTextureRuntime::SetTileMappings(const TArrayView<const FSparse
 				}
 			}
 		}
-
-		// Write to tile data
-		for (int32 PhysicalTileIndex = 0; PhysicalTileIndex < Mapping.NumPhysicalTiles; ++PhysicalTileIndex)
-		{
-			for (int32 Z = 0; Z < SPARSE_VOLUME_TILE_RES_PADDED; ++Z)
-			{
-				for (int32 Y = 0; Y < SPARSE_VOLUME_TILE_RES_PADDED; ++Y)
-				{
-					for (int32 X = 0; X < SPARSE_VOLUME_TILE_RES_PADDED; ++X)
-					{
-						const int32 SrcVoxelIndex = PhysicalTileIndex * SVTNumVoxelsPerPaddedTile + Z * (SPARSE_VOLUME_TILE_RES_PADDED * SPARSE_VOLUME_TILE_RES_PADDED) + Y * SPARSE_VOLUME_TILE_RES_PADDED + X;
-						const FVector4f ValueA = ReadVoxel(SrcVoxelIndex, Mapping.TileDataA, Header.AttributesFormats[0]);
-						const FVector4f ValueB = ReadVoxel(SrcVoxelIndex, Mapping.TileDataB, Header.AttributesFormats[1]);
-
-						const FIntVector3 VoxelCoord = DstTileCoord * SPARSE_VOLUME_TILE_RES_PADDED + FIntVector3(X, Y, Z);
-						const int32 DstVoxelIndex = VoxelCoord.Z * (Header.TileDataVolumeResolution.Y * Header.TileDataVolumeResolution.X) + VoxelCoord.Y * Header.TileDataVolumeResolution.X + VoxelCoord.X;
-						WriteVoxel(DstVoxelIndex, PhysicalTileDataPtrs[0], Header.AttributesFormats[0], ValueA);
-						WriteVoxel(DstVoxelIndex, PhysicalTileDataPtrs[1], Header.AttributesFormats[1], ValueB);
-					}
-				}
-			}
-			DstTileCoord = AdvanceTileCoord(DstTileCoord, TileCoordSpace);
-		}
-		NumWrittenTiles += Mapping.NumPhysicalTiles;
+		MipTileOffset += Mapping.NumPhysicalTiles;
 	}
 	return true;
 }
