AnimSandbox - Assertion failed: bUseRawDataOnly in AnimSequence.cpp when copy/pasting anim curves from one animation sequence to another
#fix additional change for different repro steps -> do not transact bUseRawDataOnly

diff --git a/Engine/Source/Runtime/Engine/Private/Animation/AnimSequence.cpp b/Engine/Source/Runtime/Engine/Private/Animation/AnimSequence.cpp
--- a/Engine/Source/Runtime/Engine/Private/Animation/AnimSequence.cpp
+++ b/Engine/Source/Runtime/Engine/Private/Animation/AnimSequence.cpp
@@ -687,7 +687,10 @@ void UAnimSequence::Serialize(FArchive& Ar)
 		{
 			PRAGMA_DISABLE_DEPRECATION_WARNINGS
 			SerializeCompressedData(Ar,false);
-			Ar << bUseRawDataOnly;
+			if (!bIsTransacting)
+			{
+				Ar << bUseRawDataOnly;
+			}
 			PRAGMA_ENABLE_DEPRECATION_WARNINGS
 		}
 	}
