Add UIMax slider value rect light and clamp value for very large value.

diff --git a/Engine/Source/Runtime/Engine/Classes/Components/RectLightComponent.h b/Engine/Source/Runtime/Engine/Classes/Components/RectLightComponent.h
--- a/Engine/Source/Runtime/Engine/Classes/Components/RectLightComponent.h
+++ b/Engine/Source/Runtime/Engine/Classes/Components/RectLightComponent.h
@@ -25,14 +25,14 @@ class ENGINE_API URectLightComponent : public ULocalLightComponent
 	 * Width of light source rect.
 	 * Note that light sources shapes which intersect shadow casting geometry can cause shadowing artifacts.
 	 */
-	UPROPERTY(EditAnywhere, BlueprintReadOnly, interp, Category=Light)
+	UPROPERTY(EditAnywhere, BlueprintReadOnly, interp, Category=Light, meta=(UIMin = "0.0", UIMax = "1000.0", ClampMax = "100000"))
 	float SourceWidth;
 
 	/** 
 	 * Height of light source rect.
 	 * Note that light sources shapes which intersect shadow casting geometry can cause shadowing artifacts.
 	 */
-	UPROPERTY(EditAnywhere, BlueprintReadOnly, interp, Category=Light)
+	UPROPERTY(EditAnywhere, BlueprintReadOnly, interp, Category=Light, meta=(UIMin = "0.0", UIMax = "1000.0", ClampMax = "100000"))
 	float SourceHeight;
 
 	/**
