UnrealbuildTool: Stop disabling msvc warning w4463, fix remaining instances

diff --git a/Engine/Plugins/Experimental/Water/Source/Runtime/Private/BuoyancyComponent.cpp b/Engine/Plugins/Experimental/Water/Source/Runtime/Private/BuoyancyComponent.cpp
--- a/Engine/Plugins/Experimental/Water/Source/Runtime/Private/BuoyancyComponent.cpp
+++ b/Engine/Plugins/Experimental/Water/Source/Runtime/Private/BuoyancyComponent.cpp
@@ -294,7 +294,7 @@ void UBuoyancyComponent::EnteredWaterBody(UWaterBodyComponent* WaterBodyComponen
 	}
 	if (bIsFirstBody)
 	{
-		bIsOverlappingWaterBody = true;
+		bIsOverlappingWaterBody = -1;
 	}
 }
 
diff --git a/Engine/Source/Programs/UnrealBuildTool/Platform/Windows/VCToolChain.cs b/Engine/Source/Programs/UnrealBuildTool/Platform/Windows/VCToolChain.cs
--- a/Engine/Source/Programs/UnrealBuildTool/Platform/Windows/VCToolChain.cs
+++ b/Engine/Source/Programs/UnrealBuildTool/Platform/Windows/VCToolChain.cs
@@ -808,8 +808,6 @@ namespace UnrealBuildTool
 					Arguments.Add("/we4458"); // 4458 - declaration of 'parameter' hides class member
 					Arguments.Add("/we4459"); // 4459 - declaration of 'LocalVariable' hides global declaration
 				}
-
-				Arguments.Add("/wd4463"); // 4463 - overflow; assigning 1 to bit-field that can only hold values from -1 to 0
 			}
 
 			if (CompileEnvironment.bEnableUndefinedIdentifierWarnings && !CompileEnvironment.bPreprocessOnly)
diff --git a/Engine/Source/Runtime/Engine/Private/AI/Navigation/NavigationModifier.cpp b/Engine/Source/Runtime/Engine/Private/AI/Navigation/NavigationModifier.cpp
--- a/Engine/Source/Runtime/Engine/Private/AI/Navigation/NavigationModifier.cpp
+++ b/Engine/Source/Runtime/Engine/Private/AI/Navigation/NavigationModifier.cpp
@@ -827,7 +827,7 @@ FCompositeNavModifier FCompositeNavModifier::GetInstantiatedMetaModifier(const F
 			{
 				Area->SetAreaClass(UNavAreaBase::PickAreaClassForAgent(Area->GetAreaClass(), *ActorOwner, *NavAgent));
 				Area->SetAreaClassToReplace(UNavAreaBase::PickAreaClassForAgent(Area->GetAreaClassToReplace(), *ActorOwner, *NavAgent));
-				Result.bHasMetaAreas = true;
+				Result.bHasMetaAreas = -1;
 			}
 		}
 	}
diff --git a/Engine/Source/Runtime/Engine/Private/Materials/MaterialRenderProxy.cpp b/Engine/Source/Runtime/Engine/Private/Materials/MaterialRenderProxy.cpp
--- a/Engine/Source/Runtime/Engine/Private/Materials/MaterialRenderProxy.cpp
+++ b/Engine/Source/Runtime/Engine/Private/Materials/MaterialRenderProxy.cpp
@@ -269,7 +269,7 @@ IAllocatedVirtualTexture* FMaterialRenderProxy::GetPreallocatedVTStack(const FMa
 	}
 
 	GetRendererModule().AddVirtualTextureProducerDestroyedCallback(Texture->GetProducerHandle(), &OnVirtualTextureDestroyedCB, const_cast<FMaterialRenderProxy*>(this));
-	HasVirtualTextureCallbacks = true;
+	HasVirtualTextureCallbacks = -1;
 
 	return Texture->GetAllocatedVirtualTexture();
 }
@@ -346,7 +346,7 @@ IAllocatedVirtualTexture* FMaterialRenderProxy::AllocateVTStack(const FMaterialR
 
 	if (bFoundValidLayer)
 	{
-		HasVirtualTextureCallbacks = true;
+		HasVirtualTextureCallbacks = -1;
 		return GetRendererModule().AllocateVirtualTexture(VTDesc);
 	}
 	return nullptr;
@@ -598,7 +598,7 @@ void FMaterialRenderProxy::ReleaseDynamicRHI()
 
 void FMaterialRenderProxy::ReleaseResource()
 {
-	ReleaseResourceFlag = true;
+	ReleaseResourceFlag = -1;
 	FRenderResource::ReleaseResource();
 	if (HasVirtualTextureCallbacks)
 	{
diff --git a/Engine/Source/Runtime/RuntimeAssetCache/Private/RuntimeAssetCacheBuilders.cpp b/Engine/Source/Runtime/RuntimeAssetCache/Private/RuntimeAssetCacheBuilders.cpp
--- a/Engine/Source/Runtime/RuntimeAssetCache/Private/RuntimeAssetCacheBuilders.cpp
+++ b/Engine/Source/Runtime/RuntimeAssetCache/Private/RuntimeAssetCacheBuilders.cpp
@@ -90,7 +90,7 @@ void URuntimeAssetCacheBuilder_ObjectBase::GetFromCacheAsyncComplete(int32 Handl
 		// Once complete, call GetFromCacheAsync() again and it will loop back to this function, but should succeed.
 		if (!bProcessedCacheMiss)
 		{
-			bProcessedCacheMiss = true;
+			bProcessedCacheMiss = -1;
 			OnAssetCacheMiss();
 		}
 		else
