Include time spent in FLevelStreamingDelegates OnLevelBeginMakingVisible/OnLevelBeginMakingInvisible and ULevelStreaming::BroadcastLevelVisibleStatus in their respective time budget.

diff --git a/Engine/Source/Runtime/Engine/Private/World.cpp b/Engine/Source/Runtime/Engine/Private/World.cpp
--- a/Engine/Source/Runtime/Engine/Private/World.cpp
+++ b/Engine/Source/Runtime/Engine/Private/World.cpp
@@ -2838,6 +2838,27 @@ void UWorld::AddToWorld( ULevel* Level, const FTransform& LevelTransform, bool b
 		}
 	}
 
+	// Don't consider the time limit if the match hasn't started as we need to ensure that the levels are fully loaded
+	bConsiderTimeLimit &= bMatchStarted && bIsGameWorld;
+	double TimeLimit = 0.0;
+
+	if (bExecuteNextStep && bConsiderTimeLimit)
+	{
+		TimeLimit = GLevelStreamingActorsUpdateTimeLimit;
+
+		// Give the actor initialization code more time if we're performing a high priority load or are in seamless travel
+		if (AWorldSettings* WorldSettings = GetWorldSettings(false, false))
+		{
+			if (WorldSettings->bHighPriorityLoading || WorldSettings->bHighPriorityLoadingLocal || IsInSeamlessTravel())
+			{
+				TimeLimit += GPriorityLevelStreamingActorsUpdateExtraTime;
+			}
+		}
+
+		// Remove cumulated time since UpdateLevelStreaming
+		TimeLimit = FMath::Max<double>(0.0, TimeLimit - GAddToWorldTimeCumul);
+	}
+
 	// Don't make this level visible if it's currently being made invisible
 	if( bExecuteNextStep && CurrentLevelPendingVisibility == NULL && CurrentLevelPendingInvisibility != Level )
 	{
@@ -2866,27 +2887,8 @@ void UWorld::AddToWorld( ULevel* Level, const FTransform& LevelTransform, bool b
 		SortActorListTime = 0.0;
 		PerformLastStepTime = 0.0;
 #endif
-	}
 
-	// Don't consider the time limit if the match hasn't started as we need to ensure that the levels are fully loaded
-	bConsiderTimeLimit &= bMatchStarted && bIsGameWorld;
-	double TimeLimit = 0.0;
-
-	if (bConsiderTimeLimit)
-	{
-		TimeLimit = GLevelStreamingActorsUpdateTimeLimit;
-
-		// Give the actor initialization code more time if we're performing a high priority load or are in seamless travel
-		if (AWorldSettings* WorldSettings = GetWorldSettings(false, false))
-		{
-			if (WorldSettings->bHighPriorityLoading || WorldSettings->bHighPriorityLoadingLocal || IsInSeamlessTravel())
-			{
-				TimeLimit += GPriorityLevelStreamingActorsUpdateExtraTime;
-			}
-		}
-
-		// Remove cumulated time since UpdateLevelStreaming
-		TimeLimit = FMath::Max<double>(0.0, TimeLimit - GAddToWorldTimeCumul);
+		bExecuteNextStep = (!bConsiderTimeLimit || !IsTimeLimitExceeded(TEXT("begin making visible"), StartTime, Level, TimeLimit));
 	}
 
 	if( bExecuteNextStep && !Level->bAlreadyMovedActors )
@@ -3204,8 +3206,10 @@ void UWorld::RemoveFromWorld( ULevel* Level, bool bAllowIncrementalRemoval, FNet
 	// If the level may be removed incrementally then there must also be no level pending visibility
 	if ( ((CurrentLevelPendingVisibility == nullptr) || (!bAllowIncrementalRemoval && (CurrentLevelPendingVisibility != Level))) && (Level->bIsVisible || Level->bIsBeingRemoved) )
 	{
+#if PERF_TRACK_DETAILED_ASYNC_STATS
 		// Keep track of timing.
-		double StartTime = FPlatformTime::Seconds();	
+		double PerfTrackStartTime = FPlatformTime::Seconds();
+#endif
 
 		bool bFinishRemovingLevel = true;
 		if ( bAllowIncrementalRemoval && GLevelStreamingUnregisterComponentsTimeLimit > 0.0 )
@@ -3213,11 +3217,16 @@ void UWorld::RemoveFromWorld( ULevel* Level, bool bAllowIncrementalRemoval, FNet
 			bFinishRemovingLevel = false;
 			if (CurrentLevelPendingInvisibility == nullptr)
 			{
+				double StartTime = FPlatformTime::Seconds();
 				// Mark level as being the one in process of being made invisible. 
 				// This will prevent this level from being unloaded or made visible in the meantime
 				CurrentLevelPendingInvisibility = Level;
 
 				BeginRemoval();
+
+				// Take into consideration time taken by BeginRemoval
+				double DeltaTime = (FPlatformTime::Seconds() - StartTime) * 1000;
+				GRemoveFromWorldUnregisterComponentTimeCumul += DeltaTime;
 			}
 
 			if (CurrentLevelPendingInvisibility == Level)
@@ -3225,6 +3234,7 @@ void UWorld::RemoveFromWorld( ULevel* Level, bool bAllowIncrementalRemoval, FNet
 				double TimeLimit = FMath::Max<double>(0.0, GLevelStreamingUnregisterComponentsTimeLimit - GRemoveFromWorldUnregisterComponentTimeCumul);
 				if (TimeLimit > 0.0)
 				{
+					double StartTime = FPlatformTime::Seconds();
 					// Incrementally unregister actor components. 
 					// This avoids spikes on the renderthread and gamethread when we subsequently call ClearLevelComponents() further down
 					check(IsGameWorld());
@@ -3343,7 +3353,7 @@ void UWorld::RemoveFromWorld( ULevel* Level, bool bAllowIncrementalRemoval, FNet
 		Level->bIsDisassociatingLevel = false;
 
 #if PERF_TRACK_DETAILED_ASYNC_STATS
-		UE_LOG(LogStreaming, Display, TEXT("UWorld::RemoveFromWorld for %s took %5.2f ms"), *Level->GetOutermost()->GetName(), (FPlatformTime::Seconds() - StartTime) * 1000.0);
+		UE_LOG(LogStreaming, Display, TEXT("UWorld::RemoveFromWorld for %s took %5.2f ms"), *Level->GetOutermost()->GetName(), (FPlatformTime::Seconds() - PerfTrackStartTime) * 1000.0);
 #endif // PERF_TRACK_DETAILED_ASYNC_STATS
 	}
 }
