Fixed spline components debug drawing out-of-date spline points. The render proxy wasn't being updated in shipping after the spline had been modified. The reason was that we were not updating the render proxy in shipping builds. We'll now update the render proxy for all builds that have debug draw enabled.

diff --git a/Engine/Source/Runtime/Engine/Private/Components/SplineComponent.cpp b/Engine/Source/Runtime/Engine/Private/Components/SplineComponent.cpp
--- a/Engine/Source/Runtime/Engine/Private/Components/SplineComponent.cpp
+++ b/Engine/Source/Runtime/Engine/Private/Components/SplineComponent.cpp
@@ -268,7 +268,7 @@ void USplineComponent::UpdateSpline()
 	MARK_PROPERTY_DIRTY_FROM_NAME(USplineComponent, bSplineHasBeenEdited, this);
 	MARK_PROPERTY_DIRTY_FROM_NAME(USplineComponent, bInputSplinePointsToConstructionScript, this);
 
-#if !UE_BUILD_SHIPPING
+#if UE_ENABLE_DEBUG_DRAWING
 	if (bDrawDebug)
 	{
 		MarkRenderStateDirty();
