DataLayerManager: Fix cooking issue

- Make sure AWorldDataLayers actor is part of the ULevel::Actors array in UDataLayerManager::Initialize (this happens on some levels that haven't been resaved yet)

diff --git a/Engine/Source/Runtime/Engine/Private/WorldPartition/DataLayer/DataLayerManager.cpp b/Engine/Source/Runtime/Engine/Private/WorldPartition/DataLayer/DataLayerManager.cpp
--- a/Engine/Source/Runtime/Engine/Private/WorldPartition/DataLayer/DataLayerManager.cpp
+++ b/Engine/Source/Runtime/Engine/Private/WorldPartition/DataLayer/DataLayerManager.cpp
@@ -154,8 +154,8 @@ void UDataLayerManager::Initialize()
 	UActorDescContainer* ActorDescContainer = OuterWorldPartition->GetActorDescContainer();
 	// In PIE, main world partition doesn't have an ActorDescContainer and doesn't need it, but instanced world partitions do have one.
 	check(!ActorDescContainer || ActorDescContainer->IsInitialized());
-
-	if (!GetWorldDataLayers())
+	AWorldDataLayers* WorldDataLayers = GetWorldDataLayers();
+	if (!WorldDataLayers)
 	{
 		if (ActorDescContainer)
 		{
@@ -170,15 +170,28 @@ void UDataLayerManager::Initialize()
 			}
 		}
 
-		if (!GetWorldDataLayers())
+		WorldDataLayers = GetWorldDataLayers();
+		if (!WorldDataLayers)
 		{
 			// Create missing AWorldDataLayers actor
-			AWorldDataLayers* WorldDataLayers = AWorldDataLayers::Create(OuterWorld);
-			OuterWorld->SetWorldDataLayers(WorldDataLayers);
-			check(GetWorldDataLayers());
+			AWorldDataLayers* NewWorldDataLayers = AWorldDataLayers::Create(OuterWorld);
+			OuterWorld->SetWorldDataLayers(NewWorldDataLayers);
+			WorldDataLayers = GetWorldDataLayers();
+			check(WorldDataLayers);
 		}
 	}
 
+	// Some levels do not have the WorldDataLayer actor serialized as part of their Actors array. Make sure we add it here if it isn't.
+	// Make sure WorldDataLayers is part of the Actors list so that it gets cooked properly as part of the Persistent Level
+	// This auto-corrects itself when resaving the level.
+	ULevel* WorldDataLayerLevel = WorldDataLayers->GetLevel();
+	int32 ActorIndex;
+	if (!WorldDataLayerLevel->Actors.Find(WorldDataLayers, ActorIndex))
+	{
+		WorldDataLayerLevel->Actors.Add(WorldDataLayers);
+		WorldDataLayerLevel->ActorsForGC.Add(WorldDataLayers);
+	}
+
 	// Partitioned Level Instance's DataLayerManager will never resolve DataLayers (it's up to its owning WorldPartition DataLayerManager to do the job.
 	ULevelInstanceSubsystem* LevelInstanceSubsystem = !GetWorld()->IsGameWorld() ? UWorld::GetSubsystem<ULevelInstanceSubsystem>(GetWorld()) : nullptr;
 	ILevelInstanceInterface* LevelInstance = LevelInstanceSubsystem ? LevelInstanceSubsystem->GetOwningLevelInstance(OuterWorld->PersistentLevel) : nullptr;
