Add an adaptive AddToWorld timeslice. This adjusts the timeslice to take into account the amount of remaining work and the rate at which we're completing it. Work is estimated in terms of "work units" which take into account component registration and RouteActorInit for actors.

Also add LevelStreamingDetail CSV category (off by default) with additional stats to track levels in flight and breakdown addToWorld time

This can be controlled with the following cvars:
s.AdaptiveAddToWorld.Enabled - Enables the adaptive AddToWorld timeslice (replaces s.LevelStreamingActorsUpdateTimeLimit) (default: off)
s.AdaptiveAddToWorld.MinWorkUnitCount - Work unit threshold at which the adaptive timeslice kicks in. If we're below this, we'll just use the min timeslice
s.AdaptiveAddToWorld.AddToWorldTimeSliceMin - Minimum adaptive AddToWorld timeslice
s.AdaptiveAddToWorld.AddToWorldTimeSliceMax - Maximum adaptive AddToWorld timeslice
s.AdaptiveAddToWorld.TargetMaxTimeRemaining - Target max time remaining in seconds. If our estimated completion time is longer than this, the timeslice will increase. Lower values are more aggressive.
s.AdaptiveAddToWorld.TimeSliceMaxIncreasePerSecond - Max rate at which the adptive AddToWorld timeslice will increase. Set to 0 to increase instantly
s.AdaptiveAddToWorld.TimeSliceMaxReducePerSecond - Max rate at which the adptive AddToWorld timeslice will reduce. Set to 0 to reduce instantly

Note: There is a single instance of this class shared by all views, but multiple views (e.g splitscreen) are supported. Each view is simply treated as a separate update for the purposes of tracking/extrapolation and we end up with a single averaged target timeslice (rather than separate timeslices per view)

diff --git a/Engine/Source/Runtime/Engine/Classes/Engine/Level.h b/Engine/Source/Runtime/Engine/Classes/Engine/Level.h
--- a/Engine/Source/Runtime/Engine/Classes/Engine/Level.h
+++ b/Engine/Source/Runtime/Engine/Classes/Engine/Level.h
@@ -1310,6 +1310,9 @@ public:
 	virtual UAssetUserData* GetAssetUserDataOfClass(TSubclassOf<UAssetUserData> InUserDataClass) override;
 	//~ End IInterface_AssetUserData Interface
 
+	/** Estimate the amount of work remaining for AddToWorld for this level. Used by the adaptive level streaming timeslice (see s.AdaptiveAddToWorld.Enabled) */
+	int32 GetEstimatedAddToWorldWorkUnitsRemaining() const;
+
 #if WITH_EDITOR
 	/** meant to be called only from editor, calculating and storing static geometry to be used with off-line and/or on-line navigation building */
 	ENGINE_API void RebuildStaticNavigableGeometry();
diff --git a/Engine/Source/Runtime/Engine/Private/Level.cpp b/Engine/Source/Runtime/Engine/Private/Level.cpp
--- a/Engine/Source/Runtime/Engine/Private/Level.cpp
+++ b/Engine/Source/Runtime/Engine/Private/Level.cpp
@@ -86,6 +86,15 @@ static FAutoConsoleVariableRef CVarActorClusteringEnabled(
 	ECVF_Default
 );
 
+// RouteActorInit for a single actor generally takes about half the time to complete compared to component initialization
+float GRouteActorInitializationWorkUnitWeighting = 0.5f;
+static FAutoConsoleVariableRef CVarRouteActorInitializationWorkUnitWeighting(
+	TEXT("s.RouteActorInitializationWorkUnitWeighting"),
+	GRouteActorInitializationWorkUnitWeighting,
+	TEXT("Weighting to apply to RouteActorInitialization when computing work units (relative to component init)"),
+	ECVF_Default
+);
+
 #if WITH_EDITOR
 void FLevelActorFoldersHelper::SetUseActorFolders(ULevel* InLevel, bool bInEnabled)
 {
@@ -3292,6 +3301,33 @@ void ULevel::ResetRouteActorInitializationState()
 	 RouteActorInitializationIndex = 0;
 }
 
+int32 ULevel::GetEstimatedAddToWorldWorkUnitsRemaining() const
+{
+	// We count "work units" as the number of actors requiring component update, plus the number of actors requiring initialization
+	if (bIsVisible)
+	{
+		return 0;
+	}
+	const int32 ActorCount = Actors.Num();
+	int32 TotalWorkUnits = 0;
+	if (!bAlreadyUpdatedComponents && IncrementalComponentState != EIncrementalComponentState::Finalize)
+	{
+		TotalWorkUnits += ActorCount - CurrentActorIndexForIncrementalUpdate;
+	}
+	if (!IsFinishedRouteActorInitialization())
+	{
+		// Count remaining work items from the current state, plus work items from the states that follow
+		const int32 NumStates = (int32)ERouteActorInitializationState::Finished;
+		int32 ActorInitWorkItemCount = ActorCount - RouteActorInitializationIndex;
+
+		int32 CompleteStatesRemaining = (int32)ERouteActorInitializationState::Finished - (int32)RouteActorInitializationState - 1;
+		ActorInitWorkItemCount += ActorCount * CompleteStatesRemaining;
+
+		TotalWorkUnits += ActorInitWorkItemCount * GRouteActorInitializationWorkUnitWeighting / NumStates;
+	}
+	return TotalWorkUnits;
+}
+
 void ULevel::RouteActorInitialize(int32 NumActorsToProcess)
 {
 	TRACE_OBJECT_EVENT(this, RouteActorInitialize);
diff --git a/Engine/Source/Runtime/Engine/Private/World.cpp b/Engine/Source/Runtime/Engine/Private/World.cpp
--- a/Engine/Source/Runtime/Engine/Private/World.cpp
+++ b/Engine/Source/Runtime/Engine/Private/World.cpp
@@ -133,6 +133,9 @@ DEFINE_LOG_CATEGORY_STATIC(LogWorld, Log, All);
 DEFINE_LOG_CATEGORY(LogSpawn);
 
 CSV_DECLARE_CATEGORY_MODULE_EXTERN(CORE_API, Basic);
+CSV_DEFINE_CATEGORY(LevelStreamingAdaptive, true);
+CSV_DEFINE_CATEGORY(LevelStreamingAdaptiveDetail, false);
+CSV_DEFINE_CATEGORY(LevelStreamingDetail, false);
 
 #define LOCTEXT_NAMESPACE "World"
 
@@ -145,6 +148,180 @@ static TAutoConsoleVariable<int32> CVarPurgeEditorSceneDuringPIE(
 	TEXT("0 to keep editor scene fully initialized during PIE (default)\n")
 	TEXT("1 to purge editor scene from memory during PIE and restore when the session finishes."));
 
+/*-----------------------------------------------------------------------------
+	FAdaptiveAddToWorld implementation.
+-----------------------------------------------------------------------------*/
+static int32 GAdaptiveAddToWorldEnabled = 0;
+FAutoConsoleVariableRef CVarAdaptiveAddToWorldEnabled(TEXT("s.AdaptiveAddToWorld.Enabled"), GAdaptiveAddToWorldEnabled, 
+	TEXT("Enables the adaptive AddToWorld timeslice (replaces s.LevelStreamingActorsUpdateTimeLimit) (default: off)"));
+
+static int32 GAdaptiveAddToWorldMinWorkUnitCount = 100;
+FAutoConsoleVariableRef CVarAdaptiveAddToWorldMinWorkUnitCount(TEXT("s.AdaptiveAddToWorld.MinWorkUnitCount"), GAdaptiveAddToWorldMinWorkUnitCount, 
+	TEXT("Work unit threshold at which the adaptive timeslice kicks in. If we're below this, we'll just use the min timeslice"));
+
+static float GAdaptiveAddToWorldTimeSliceMin = 1.0f;
+FAutoConsoleVariableRef CVarAdaptiveAddToWorldTimeSliceMin(TEXT("s.AdaptiveAddToWorld.AddToWorldTimeSliceMin"), GAdaptiveAddToWorldTimeSliceMin, 
+	TEXT("Minimum adaptive AddToWorld timeslice"));
+
+static float GAdaptiveAddToWorldTimeSliceMax = 6.0f;
+FAutoConsoleVariableRef CVarAdaptiveAddToWorldTimeSliceMax(TEXT("s.AdaptiveAddToWorld.AddToWorldTimeSliceMax"), GAdaptiveAddToWorldTimeSliceMax, 
+	TEXT("Maximum adaptive AddToWorld timeslice"));
+
+static float GAdaptiveAddToWorldTargetMaxTimeRemaining = 6.0f;
+FAutoConsoleVariableRef CVarAdaptiveAddToWorldTargetMaxTimeRemaining(TEXT("s.AdaptiveAddToWorld.TargetMaxTimeRemaining"), GAdaptiveAddToWorldTargetMaxTimeRemaining, 
+	TEXT("Target max time remaining in seconds. If our estimated completion time is longer than this, the timeslice will increase. Lower values are more aggressive"));
+
+static float GAdaptiveAddToWorldTimeSliceMaxIncreasePerSecond = 0.0f;
+FAutoConsoleVariableRef CVarAdaptiveAddToWorldTimeSliceMaxIncreasePerSecond(TEXT("s.AdaptiveAddToWorld.TimeSliceMaxIncreasePerSecond"), GAdaptiveAddToWorldTimeSliceMaxIncreasePerSecond,
+	TEXT("Max rate at which the adptive AddToWorld timeslice will increase. Set to 0 to increase instantly"));
+
+static float GAdaptiveAddToWorldTimeSliceMaxReducePerSecond = 0.0f;
+FAutoConsoleVariableRef CVarAdaptiveAddToWorldTimesliceReductionSpeedPerSecond(TEXT("s.AdaptiveAddToWorld.TimeSliceMaxReducePerSecond"), GAdaptiveAddToWorldTimeSliceMaxReducePerSecond,
+	TEXT("Max rate at which the adptive AddToWorld timeslice will reduce. Set to 0 to reduce instantly"));
+
+class FAdaptiveAddToWorld
+{
+	struct FHistoryEntry
+	{
+		float DeltaT = 0.0f;
+		float TimeSlice = 0.0f;
+		int32 WorkUnitsProcessed = 0;
+	};
+public:
+	FAdaptiveAddToWorld()
+	{
+	}
+
+	bool IsEnabled() const
+	{
+		return bEnabled;
+	} 
+
+	void SetEnabled(bool bInEnabled)
+	{
+		bEnabled = bInEnabled;
+	}
+
+	// Called when level streaming update begins. 
+	// Note: multiple updates per frame are supported, e.g for multiple views/splitscreen. Each update is counted as a separate history entry
+	void BeginUpdate()
+	{
+		if (bEnabled)
+		{
+		    WorkUnitsProcessedThisUpdate = 0;
+		    TotalWorkUnitsRemaining = 0;
+		}
+	}
+
+	// Called when level streaming update ends 
+	void EndUpdate()
+	{
+		if ( !bEnabled )
+		{
+			return;
+		}
+		QUICK_SCOPE_CYCLE_COUNTER(STAT_LevelStreamingAdaptiveUpdate);
+		CSV_SCOPED_TIMING_STAT(LevelStreamingAdaptiveDetail, Update);
+		// Only count frames where we have a reasonable number of work units to process
+		double CurrentTimestamp = FPlatformTime::Seconds();
+		float DeltaT = float(CurrentTimestamp - LastUpdateTimestamp);
+		if (TotalWorkUnitsRemaining >= GAdaptiveAddToWorldMinWorkUnitCount)
+		{
+			FHistoryEntry& HistoryEntry = UpdateHistory[HistoryFrameIndex];
+			HistoryEntry.WorkUnitsProcessed = WorkUnitsProcessedThisUpdate;
+			HistoryEntry.DeltaT = DeltaT;
+			HistoryEntry.TimeSlice = GetTimeSlice();
+			HistoryFrameIndex = (HistoryFrameIndex + 1) % HistorySize;
+
+			// Make sure we have a complete history buffer before 
+			if (ValidFrameCount < HistorySize)
+			{
+				ValidFrameCount++;
+			}
+			else
+			{
+				// Compute the rolling average
+				int32 TotalWorkUnits = 0;
+				float TotalTime = 0.0f;
+				float TotalWorkTimeMs = 0.0f;
+				for (int i = 0; i < HistorySize; i++)
+				{
+					TotalWorkUnits += UpdateHistory[i].WorkUnitsProcessed;
+					TotalTime += UpdateHistory[i].DeltaT;
+					TotalWorkTimeMs += UpdateHistory[i].TimeSlice;
+				}
+
+				float AverageTimeSliceMs = TotalWorkTimeMs / float(HistorySize);
+
+				float AverageWorkUnitsPerSecondWith1MsTimeSlice = (float)TotalWorkUnits / ( TotalTime * AverageTimeSliceMs );
+				CSV_CUSTOM_STAT(LevelStreamingAdaptiveDetail, AverageWorkUnitsPerSecondWith1MsTimeSlice, AverageWorkUnitsPerSecondWith1MsTimeSlice, ECsvCustomStatOp::Set);
+
+				// Work out what the timeslice should be
+				float EstimatedTimeRemainingFor1MsTimeSlice = (float)TotalWorkUnitsRemaining / AverageWorkUnitsPerSecondWith1MsTimeSlice;
+				CSV_CUSTOM_STAT(LevelStreamingAdaptiveDetail, EstimatedTimeRemainingFor1MsTimeSlice, EstimatedTimeRemainingFor1MsTimeSlice, ECsvCustomStatOp::Set);
+
+				float TargetTimeSlice = EstimatedTimeRemainingFor1MsTimeSlice / GAdaptiveAddToWorldTargetMaxTimeRemaining - ExtraTimeSlice;
+				float DesiredBaseTimeSlice = FMath::Clamp(TargetTimeSlice, GAdaptiveAddToWorldTimeSliceMin, GAdaptiveAddToWorldTimeSliceMax);
+
+				// Limit the timeslice change based on MaxIncreasePerSecond/ MaxReducePerSecond
+				float NewBaseTimeSlice = DesiredBaseTimeSlice;
+				if (DesiredBaseTimeSlice > BaseTimeSlice && GAdaptiveAddToWorldTimeSliceMaxIncreasePerSecond > 0.0)
+				{
+					NewBaseTimeSlice = FMath::Min(NewBaseTimeSlice, BaseTimeSlice + GAdaptiveAddToWorldTimeSliceMaxIncreasePerSecond * DeltaT);
+				}
+				if (DesiredBaseTimeSlice < BaseTimeSlice && GAdaptiveAddToWorldTimeSliceMaxReducePerSecond > 0.0)
+				{
+					NewBaseTimeSlice = FMath::Max(NewBaseTimeSlice, BaseTimeSlice - GAdaptiveAddToWorldTimeSliceMaxReducePerSecond * DeltaT);
+				}
+				BaseTimeSlice = NewBaseTimeSlice;
+			}
+		}
+		else
+		{
+			BaseTimeSlice = GAdaptiveAddToWorldTimeSliceMin;
+		}
+
+		CSV_CUSTOM_STAT(LevelStreamingAdaptiveDetail, WorkUnitsProcessedPerSecond, float(WorkUnitsProcessedThisUpdate)/DeltaT, ECsvCustomStatOp::Set);
+		CSV_CUSTOM_STAT(LevelStreamingAdaptiveDetail, WorkUnitsProcessedThisUpdate, WorkUnitsProcessedThisUpdate, ECsvCustomStatOp::Set);
+		CSV_CUSTOM_STAT(LevelStreamingAdaptive, TimeSlice, GetTimeSlice(), ECsvCustomStatOp::Set);
+		CSV_CUSTOM_STAT(LevelStreamingAdaptive, WorkUnitsRemaining, TotalWorkUnitsRemaining, ECsvCustomStatOp::Set);
+		LastUpdateTimestamp = CurrentTimestamp;
+	}
+
+	void RegisterAddToWorldWorkUnits(int32 StartWorkUnits, int32 EndWorkUnits)
+	{
+		WorkUnitsProcessedThisUpdate += StartWorkUnits - EndWorkUnits;
+		TotalWorkUnitsRemaining += EndWorkUnits;
+	}
+
+	float GetTimeSlice()
+	{
+		return ExtraTimeSlice + BaseTimeSlice;
+	}
+
+	void SetExtraTimeSlice(float InExtraTimeSlice)
+	{
+		ExtraTimeSlice = InExtraTimeSlice;
+	}
+
+
+private:
+	static const int32 HistorySize = 16;
+
+	int32 WorkUnitsProcessedThisUpdate = 0;
+	int32 TotalWorkUnitsRemaining = 0;
+	int32 HistoryFrameIndex = 0;
+	FHistoryEntry UpdateHistory[HistorySize];
+	float BaseTimeSlice = 1.0f;
+	float ExtraTimeSlice = 0.0f;
+	double LastUpdateTimestamp = 0.0;
+	int32 ValidFrameCount = 0;
+    bool bEnabled = false;
+};
+
+static FAdaptiveAddToWorld GAdaptiveAddToWorld;
+
+
 template<class Function>
 static void ForEachNetDriver(UEngine* Engine, const UWorld* const World, const Function InFunction)
 {
@@ -2844,16 +3021,24 @@ void UWorld::AddToWorld( ULevel* Level, const FTransform& LevelTransform, bool b
 
 	if (bExecuteNextStep && bConsiderTimeLimit)
 	{
-		TimeLimit = GLevelStreamingActorsUpdateTimeLimit;
-
+		float ExtraTimeLimit = 0.0f;
 		// Give the actor initialization code more time if we're performing a high priority load or are in seamless travel
 		if (AWorldSettings* WorldSettings = GetWorldSettings(false, false))
 		{
 			if (WorldSettings->bHighPriorityLoading || WorldSettings->bHighPriorityLoadingLocal || IsInSeamlessTravel())
 			{
-				TimeLimit += GPriorityLevelStreamingActorsUpdateExtraTime;
+				ExtraTimeLimit = GPriorityLevelStreamingActorsUpdateExtraTime;
 			}
 		}
+		if (GAdaptiveAddToWorld.IsEnabled())
+		{
+			GAdaptiveAddToWorld.SetExtraTimeSlice(ExtraTimeLimit);
+			TimeLimit = GAdaptiveAddToWorld.GetTimeSlice();
+		}
+		else
+		{
+			TimeLimit = GLevelStreamingActorsUpdateTimeLimit + ExtraTimeLimit;
+		}
 
 		// Remove cumulated time since UpdateLevelStreaming
 		TimeLimit = FMath::Max<double>(0.0, TimeLimit - GAddToWorldTimeCumul);
@@ -2891,8 +3076,11 @@ void UWorld::AddToWorld( ULevel* Level, const FTransform& LevelTransform, bool b
 		bExecuteNextStep = (!bConsiderTimeLimit || !IsTimeLimitExceeded(TEXT("begin making visible"), StartTime, Level, TimeLimit));
 	}
 
+	int32 StartWorkUnitsRemaining = Level->GetEstimatedAddToWorldWorkUnitsRemaining();
+
 	if( bExecuteNextStep && !Level->bAlreadyMovedActors )
 	{
+		CSV_SCOPED_TIMING_STAT(LevelStreamingDetail, AddToWorld_MoveActors);
 		QUICK_SCOPE_CYCLE_COUNTER(STAT_AddToWorldTime_MoveActors);
 		SCOPE_TIME_TO_VAR(&MoveActorTime);
 
@@ -2906,6 +3094,7 @@ void UWorld::AddToWorld( ULevel* Level, const FTransform& LevelTransform, bool b
 
 	if( bExecuteNextStep && !Level->bAlreadyShiftedActors )
 	{
+		CSV_SCOPED_TIMING_STAT(LevelStreamingDetail, AddToWorld_ShiftActors);
 		QUICK_SCOPE_CYCLE_COUNTER(STAT_AddToWorldTime_ShiftActors);
 		SCOPE_TIME_TO_VAR(&ShiftActorsTime);
 
@@ -2963,6 +3152,7 @@ void UWorld::AddToWorld( ULevel* Level, const FTransform& LevelTransform, bool b
 	// Updates the level components (Actor components and UModelComponents).
 	if( bExecuteNextStep && !Level->bAlreadyUpdatedComponents )
 	{
+		CSV_SCOPED_TIMING_STAT(LevelStreamingDetail, AddToWorld_UpdateComponents);
 		QUICK_SCOPE_CYCLE_COUNTER(STAT_AddToWorldTime_UpdatingComponents);
 		SCOPE_TIME_TO_VAR(&UpdateComponentsTime);
 
@@ -3015,6 +3205,7 @@ void UWorld::AddToWorld( ULevel* Level, const FTransform& LevelTransform, bool b
 		if (bExecuteNextStep && !(Level->bAlreadyInitializedNetworkActors && Level->bAlreadyClearedActorsSeamlessTravelFlag))
 		{
 			QUICK_SCOPE_CYCLE_COUNTER(STAT_AddToWorldTime_InitializeNetworkActors);
+			CSV_SCOPED_TIMING_STAT(LevelStreamingDetail, AddToWorld_InitNetworkActors);
 			SCOPE_TIME_TO_VAR(&InitActorTime);
 
 			// InitializeNetworkActors only needs to be called the first time a level is loaded,
@@ -3037,6 +3228,7 @@ void UWorld::AddToWorld( ULevel* Level, const FTransform& LevelTransform, bool b
 		// Route various initialization functions and set volumes.
 		if (bExecuteNextStep && !Level->IsFinishedRouteActorInitialization())
 		{
+			CSV_SCOPED_TIMING_STAT(LevelStreamingDetail, AddToWorld_RouteActorInit);
 			QUICK_SCOPE_CYCLE_COUNTER(STAT_AddToWorldTime_RouteActorInitialize);
 			SCOPE_TIME_TO_VAR(&RouteActorInitializeTime);
 			const int32 NumActorsToProcess = (!bConsiderTimeLimit || !IsGameWorld() || IsRunningCommandlet()) ? 0 : GLevelStreamingRouteActorInitializationGranularity;
@@ -3053,6 +3245,7 @@ void UWorld::AddToWorld( ULevel* Level, const FTransform& LevelTransform, bool b
 		// Sort the actor list; can't do this on save as the relevant properties for sorting might have been changed by code
 		if( bExecuteNextStep && !Level->bAlreadySortedActorList )
 		{
+			CSV_SCOPED_TIMING_STAT(LevelStreamingDetail, AddToWorld_SortActorList);
 			SCOPE_TIME_TO_VAR(&SortActorListTime);
 
 			Level->SortActorList();
@@ -3072,6 +3265,7 @@ void UWorld::AddToWorld( ULevel* Level, const FTransform& LevelTransform, bool b
 	if( bPerformedLastStep )
 	{
 		SCOPE_TIME_TO_VAR(&PerformLastStepTime);
+		CSV_CUSTOM_STAT(LevelStreamingDetail, AddToWorldLevelCompleteCount, 1, ECsvCustomStatOp::Accumulate);
 		
 		ResetLevelFlagsOnLevelAddedToWorld(Level);
 
@@ -3123,6 +3317,12 @@ void UWorld::AddToWorld( ULevel* Level, const FTransform& LevelTransform, bool b
 			ULevelStreaming::BroadcastLevelVisibleStatus(this, Level->GetOutermost()->GetFName(), true);
 		}
 	}
+#if CSV_PROFILER
+	else
+	{
+		CSV_CUSTOM_STAT(LevelStreamingDetail, AddToWorldLevelIncompleteCount, 1, ECsvCustomStatOp::Accumulate);
+	}
+#endif
 
 #if PERF_TRACK_DETAILED_ASYNC_STATS
 	if (bPerformedLastStep)
@@ -3154,6 +3354,12 @@ void UWorld::AddToWorld( ULevel* Level, const FTransform& LevelTransform, bool b
 	}
 #endif // PERF_TRACK_DETAILED_ASYNC_STATS
 
+    // Register work completed with adaptive level streaming
+	if (GAdaptiveAddToWorld.IsEnabled())
+	{
+		GAdaptiveAddToWorld.RegisterAddToWorldWorkUnits(StartWorkUnitsRemaining, Level->GetEstimatedAddToWorldWorkUnitsRemaining());
+	}
+
 	// Delta time in ms.
 	double DeltaTime = (FPlatformTime::Seconds() - StartTime) * 1000;
 	GAddToWorldTimeCumul += DeltaTime;
@@ -3909,6 +4115,13 @@ void UWorld::UpdateLevelStreaming()
 		return;
 	}
 
+	GAdaptiveAddToWorld.SetEnabled(GAdaptiveAddToWorldEnabled == 1);
+
+	if ( GAdaptiveAddToWorld.IsEnabled() )
+	{
+		GAdaptiveAddToWorld.BeginUpdate();
+	}
+
 	// Store current number of pending unload levels, it may change in loop bellow
 	const int32 NumLevelsPendingPurge = FLevelStreamingGCHelper::GetNumLevelsPendingPurge();
 
@@ -3972,6 +4185,10 @@ void UWorld::UpdateLevelStreaming()
 			GEngine->ForceGarbageCollection(true); 
 		}
 	}
+	if ( GAdaptiveAddToWorld.IsEnabled() )
+	{
+		GAdaptiveAddToWorld.EndUpdate();
+	}
 }
 
 void UWorld::SetShouldForceUnloadStreamingLevels(const bool bInShouldForceUnloadStreamingLevels)
