Updated engine code to move SetWorldAsset, ForEachActorInLevelInstance and GetLevelInstanceLevel out of the editor only section

diff --git a/Engine/Source/Runtime/Engine/Private/LevelInstance/LevelInstanceSubsystem.cpp b/Engine/Source/Runtime/Engine/Private/LevelInstance/LevelInstanceSubsystem.cpp
--- a/Engine/Source/Runtime/Engine/Private/LevelInstance/LevelInstanceSubsystem.cpp
+++ b/Engine/Source/Runtime/Engine/Private/LevelInstance/LevelInstanceSubsystem.cpp
@@ -308,6 +308,26 @@ void ULevelInstanceSubsystem::RegisterLoadedLevelStreamingLevelInstance(ULevelSt
 	}
 }
 
+ULevel* ULevelInstanceSubsystem::GetLevelInstanceLevel(const ILevelInstanceInterface* LevelInstance) const
+{
+	if (LevelInstance->HasValidLevelInstanceID())
+	{
+#if WITH_EDITOR
+		if (const FLevelInstanceEdit* CurrentEdit = GetLevelInstanceEdit(LevelInstance))
+		{
+			return LevelInstanceEdit->LevelStreaming->GetLoadedLevel();
+		}
+		else 
+#endif // WITH_EDITOR
+			if (const FLevelInstance* LevelInstanceEntry = LoadedLevelInstances.Find(LevelInstance->GetLevelInstanceID()))
+		{
+			return LevelInstanceEntry->LevelStreaming->GetLoadedLevel();
+		}
+	}
+
+	return nullptr;
+}
+
 #if WITH_EDITOR
 
 void ULevelInstanceSubsystem::OnAssetsPreDelete(const TArray<UObject*>& Objects)
@@ -518,6 +538,14 @@ ILevelInstanceInterface* ULevelInstanceSubsystem::GetOwningLevelInstance(const U
 	return nullptr;
 }
 
+void ULevelInstanceSubsystem::ForEachActorInLevelInstance(const ILevelInstanceInterface* LevelInstance, TFunctionRef<bool(AActor* LevelActor)> Operation) const
+{
+	if (ULevel* LevelInstanceLevel = GetLevelInstanceLevel(LevelInstance))
+	{
+		ForEachActorInLevel(LevelInstanceLevel, Operation);
+	}
+}
+
 #if WITH_EDITOR
 
 void ULevelInstanceSubsystem::Tick()
@@ -740,14 +768,6 @@ bool ULevelInstanceSubsystem::GetLevelInstanceBoundsFromPackage(const FTransform
 	return false;
 }
 
-void ULevelInstanceSubsystem::ForEachActorInLevelInstance(const ILevelInstanceInterface* LevelInstance, TFunctionRef<bool(AActor * LevelActor)> Operation) const
-{
-	if (ULevel* LevelInstanceLevel = GetLevelInstanceLevel(LevelInstance))
-	{
-		ForEachActorInLevel(LevelInstanceLevel, Operation);
-	}
-}
-
 void ULevelInstanceSubsystem::ForEachLevelInstanceAncestorsAndSelf(const AActor* Actor, TFunctionRef<bool(const ILevelInstanceInterface*)> Operation) const
 {
 	if (const ILevelInstanceInterface* LevelInstance = Cast<ILevelInstanceInterface>(Actor))
@@ -1422,23 +1442,6 @@ void ULevelInstanceSubsystem::BreakLevelInstance_Impl(ILevelInstanceInterface* L
 	return;
 }
 
-ULevel* ULevelInstanceSubsystem::GetLevelInstanceLevel(const ILevelInstanceInterface* LevelInstance) const
-{
-	if (LevelInstance->HasValidLevelInstanceID())
-	{
-		if (const FLevelInstanceEdit* CurrentEdit = GetLevelInstanceEdit(LevelInstance))
-		{
-			return LevelInstanceEdit->LevelStreaming->GetLoadedLevel();
-		}
-		else if (const FLevelInstance* LevelInstanceEntry = LoadedLevelInstances.Find(LevelInstance->GetLevelInstanceID()))
-		{
-			return LevelInstanceEntry->LevelStreaming->GetLoadedLevel();
-		}
-	}
-
-	return nullptr;
-}
-
 bool ULevelInstanceSubsystem::LevelInstanceHasLevelScriptBlueprint(const ILevelInstanceInterface* LevelInstance) const
 {
 	if (LevelInstance)
