Track non streaming mips in CSV

#rnx
#tests none

diff --git a/Engine/Source/Runtime/Engine/Private/Streaming/StreamingManagerTexture.cpp b/Engine/Source/Runtime/Engine/Private/Streaming/StreamingManagerTexture.cpp
--- a/Engine/Source/Runtime/Engine/Private/Streaming/StreamingManagerTexture.cpp
+++ b/Engine/Source/Runtime/Engine/Private/Streaming/StreamingManagerTexture.cpp
@@ -1750,6 +1750,7 @@ void FRenderAssetStreamingManager::UpdateResourceStreaming( float DeltaTime, boo
 	CSV_CUSTOM_STAT(TextureStreaming, WantedMips, ((float)DisplayedStats.WantedMips) / (1024.0f * 1024.0f), ECsvCustomStatOp::Set);
 	CSV_CUSTOM_STAT(TextureStreaming, ResidentMeshMem, ((float)DisplayedStats.ResidentMeshMem) / (1024.0f * 1024.0f), ECsvCustomStatOp::Set);
 	CSV_CUSTOM_STAT(TextureStreaming, StreamedMeshMem, ((float)DisplayedStats.StreamedMeshMem) / (1024.0f * 1024.0f), ECsvCustomStatOp::Set);
+	CSV_CUSTOM_STAT(TextureStreaming, NonStreamingMips, ((float)DisplayedStats.NonStreamingMips) / (1024.0f * 1024.0f), ECsvCustomStatOp::Set);
 
 	RenderAssetInstanceAsyncWork->EnsureCompletion();
 
