Substrate - added distortion coverage mode to be compatible with legacy distortion.

diff --git a/Engine/Shaders/Private/DistortAccumulatePS.usf b/Engine/Shaders/Private/DistortAccumulatePS.usf
--- a/Engine/Shaders/Private/DistortAccumulatePS.usf
+++ b/Engine/Shaders/Private/DistortAccumulatePS.usf
@@ -128,7 +128,11 @@ void Main(
 	TopF0 = SanitizeF0(TopF0);	// Sanitize, in case a Add node has been used
 
 	const FStrataOperator RootOperator	= StrataTree.Operators[StrataData.OperatorIndex];
+#if DISTORTION_ACCOUNT_FOR_COVERAGE
 	const float RefractionCoverage		= RootOperator.Coverage;														// Coverage of the material
+#else
+	const float RefractionCoverage		= 1.0;																			// Behavior where coverage is always ignored (legacy support)
+#endif
 	const float RefractionDepthMeter	= length(MaterialParameters.WorldPosition_CamRelative) * CENTIMETER_TO_METER;	// Distance for camera in meters
 	const float RefractionLobeVariance	= Sigma;																		// The lobe variance that can be added to in the en recover a single lobe.
 
diff --git a/Engine/Source/Runtime/Engine/Classes/Engine/EngineTypes.h b/Engine/Source/Runtime/Engine/Classes/Engine/EngineTypes.h
--- a/Engine/Source/Runtime/Engine/Classes/Engine/EngineTypes.h
+++ b/Engine/Source/Runtime/Engine/Classes/Engine/EngineTypes.h
@@ -390,6 +390,24 @@ enum ERefractionMode : int
 	RM_None UMETA(DisplayName = "None"),
 };
 
+/** Determines how the refraction account for the coverage with Substrate. It can only be used when Substrate is enabled. */
+UENUM()
+enum ERefractionCoverageMode : int
+{
+	/**
+	 * This is the pre-Substrate behavior: coverage is ignored and always 1.
+	 * When rough refraction is disabled, this is behavior is forced ON.
+	 */
+	RCM_CoverageIgnored UMETA(DisplayName = "Coverage Ignored"),
+
+	/**
+	 * This is a new behavior available with Substrate when rough refraction are enabled: account for roughness, coverage and depth.
+	 * This is a more physically based behavior: the background scene will be visible untouched according to (1-coverage),
+	 * while the blurred version will be visible according to coverage.
+	 */
+	RCM_CoverageAccountedFor UMETA(DisplayName = "Coverage Accounted For"),
+};
+
 /**
  * Enumerates available options for the translucency sort policy.
  */
diff --git a/Engine/Source/Runtime/Engine/Classes/Materials/Material.h b/Engine/Source/Runtime/Engine/Classes/Materials/Material.h
--- a/Engine/Source/Runtime/Engine/Classes/Materials/Material.h
+++ b/Engine/Source/Runtime/Engine/Classes/Materials/Material.h
@@ -985,6 +985,10 @@ public:
 	UPROPERTY(EditAnywhere, Category=Refraction)
 	TEnumAsByte<ERefractionMode> RefractionMethod;
 
+	/** Controls whether refraction takes into account the material surface coverage, or not. */
+	UPROPERTY(EditAnywhere, Category=Refraction)
+	TEnumAsByte<ERefractionCoverageMode> RefractionCoverageMode;
+
 	/** If multiple nodes with the same  type are inserted at the same point, this defined order and if they get combined, only used if domain is PostProcess */
 	UPROPERTY(EditAnywhere, BlueprintReadWrite, Category = PostProcessMaterial, meta = (DisplayName = "Blendable Priority"))
 	int32 BlendablePriority;
diff --git a/Engine/Source/Runtime/Engine/Private/Materials/HLSLMaterialTranslator.cpp b/Engine/Source/Runtime/Engine/Private/Materials/HLSLMaterialTranslator.cpp
--- a/Engine/Source/Runtime/Engine/Private/Materials/HLSLMaterialTranslator.cpp
+++ b/Engine/Source/Runtime/Engine/Private/Materials/HLSLMaterialTranslator.cpp
@@ -1895,6 +1895,7 @@ void FHLSLMaterialTranslator::ValidateShadingModelsForFeatureLevel(const FMateri
 
 void FHLSLMaterialTranslator::GetMaterialEnvironment(EShaderPlatform InPlatform, FShaderCompilerEnvironment& OutEnvironment)
 {
+	const bool bStrataEnabled = Strata::IsStrataEnabled();
 	bool bMaterialRequestsDualSourceBlending = false;
 
 	if (bNeedsParticlePosition || Material->ShouldGenerateSphericalParticleNormals() || bUsesSphericalParticleOpacity)
@@ -1993,7 +1994,8 @@ void FHLSLMaterialTranslator::GetMaterialEnvironment(EShaderPlatform InPlatform,
 	OutEnvironment.SetDefine(TEXT("USES_WORLD_POSITION_OFFSET"), bUsesWorldPositionOffset);
 	OutEnvironment.SetDefine(TEXT("USES_EMISSIVE_COLOR"), bUsesEmissiveColor);
 	// Distortion uses tangent space transform 
-	OutEnvironment.SetDefine(TEXT("USES_DISTORTION"), Material->IsDistorted()); 
+	OutEnvironment.SetDefine(TEXT("USES_DISTORTION"), Material->IsDistorted());
+	OutEnvironment.SetDefine(TEXT("DISTORTION_ACCOUNT_FOR_COVERAGE"), bStrataEnabled && Material->GetRefractionCoverageMode() == RCM_CoverageAccountedFor ? 1 : 0);
 
 	OutEnvironment.SetDefine(TEXT("MATERIAL_ENABLE_TRANSLUCENCY_FOGGING"), Material->ShouldApplyFogging());
 	OutEnvironment.SetDefine(TEXT("MATERIAL_ENABLE_TRANSLUCENCY_CLOUD_FOGGING"), Material->ShouldApplyCloudFogging());
diff --git a/Engine/Source/Runtime/Engine/Private/Materials/Material.cpp b/Engine/Source/Runtime/Engine/Private/Materials/Material.cpp
--- a/Engine/Source/Runtime/Engine/Private/Materials/Material.cpp
+++ b/Engine/Source/Runtime/Engine/Private/Materials/Material.cpp
@@ -1008,6 +1008,7 @@ UMaterial::UMaterial(const FObjectInitializer& ObjectInitializer)
 	bIsMaterialEditorStatsMaterial = false;
 
 	RefractionMethod = RM_None;
+	RefractionCoverageMode = RCM_CoverageAccountedFor;
 
 #if WITH_EDITORONLY_DATA
 	MaterialGraph = NULL;
@@ -3327,6 +3328,7 @@ void UMaterial::ConvertMaterialToStrataMaterial()
 		}
 
 		BlendMode = ConvertLegacyBlendMode(BlendMode, ShadingModels);
+		RefractionCoverageMode = RCM_CoverageIgnored;
 		bInvalidateShader = true;
 	}
 	else if (!bUseMaterialAttributes && !EditorOnly->FrontMaterial.IsConnected() && GetExpressions().IsEmpty())
@@ -3611,6 +3613,7 @@ void UMaterial::ConvertMaterialToStrataMaterial()
 		}
 
 		BlendMode = ConvertLegacyBlendMode(BlendMode, ShadingModels);
+		RefractionCoverageMode = RCM_CoverageIgnored;
 	}
 
 	if (bRelinkCustomOutputNodes)
@@ -4417,9 +4420,9 @@ bool UMaterial::CanEditChange(const FProperty* InProperty) const
 			if (bStrataEnabled)
 			{
 				return ((MaterialDomain != MD_PostProcess && MaterialDomain != MD_LightFunction && MaterialDomain != MD_Volume) || (MaterialDomain == MD_PostProcess && BlendableOutputAlpha));
-		}
+			}
 			else
-		{
+			{
 				return (MaterialDomain == MD_DeferredDecal || MaterialDomain == MD_Surface || MaterialDomain == MD_Volume || MaterialDomain == MD_UI || (MaterialDomain == MD_PostProcess && BlendableOutputAlpha));
 			}
 		}
@@ -4446,6 +4449,11 @@ bool UMaterial::CanEditChange(const FProperty* InProperty) const
 		{
 			return EditorOnly->Refraction.IsConnected();
 		}
+
+		if (PropertyName == GET_MEMBER_NAME_STRING_CHECKED(UMaterial, RefractionCoverageMode))
+		{
+			return bStrataEnabled;
+		}
 	
 		if (PropertyName == GET_MEMBER_NAME_STRING_CHECKED(UMaterial, TranslucencyPass)
 			|| PropertyName == GET_MEMBER_NAME_STRING_CHECKED(UMaterial, bEnableResponsiveAA)
diff --git a/Engine/Source/Runtime/Engine/Private/Materials/MaterialShared.cpp b/Engine/Source/Runtime/Engine/Private/Materials/MaterialShared.cpp
--- a/Engine/Source/Runtime/Engine/Private/Materials/MaterialShared.cpp
+++ b/Engine/Source/Runtime/Engine/Private/Materials/MaterialShared.cpp
@@ -1882,6 +1882,7 @@ bool FMaterialResource::AllowNegativeEmissiveColor() const
 }
 
 bool FMaterialResource::IsDistorted() const { return Material->bUsesDistortion && IsTranslucentBlendMode(GetBlendMode()); }
+ERefractionCoverageMode FMaterialResource::GetRefractionCoverageMode() const { return Material->RefractionCoverageMode; }
 float FMaterialResource::GetTranslucencyDirectionalLightingIntensity() const { return Material->TranslucencyDirectionalLightingIntensity; }
 float FMaterialResource::GetTranslucentShadowDensityScale() const { return Material->TranslucentShadowDensityScale; }
 float FMaterialResource::GetTranslucentSelfShadowDensityScale() const { return Material->TranslucentSelfShadowDensityScale; }
diff --git a/Engine/Source/Runtime/Engine/Public/MaterialShared.h b/Engine/Source/Runtime/Engine/Public/MaterialShared.h
--- a/Engine/Source/Runtime/Engine/Public/MaterialShared.h
+++ b/Engine/Source/Runtime/Engine/Public/MaterialShared.h
@@ -1898,6 +1898,7 @@ public:
 	virtual float GetOpacityMaskClipValue() const = 0;
 	virtual bool GetCastDynamicShadowAsMasked() const = 0;
 	virtual bool IsDistorted() const { return false; };
+	virtual ERefractionCoverageMode GetRefractionCoverageMode() const { return RCM_CoverageIgnored; }
 	virtual float GetTranslucencyDirectionalLightingIntensity() const { return 1.0f; }
 	virtual float GetTranslucentShadowDensityScale() const { return 1.0f; }
 	virtual float GetTranslucentSelfShadowDensityScale() const { return 1.0f; }
@@ -2621,6 +2622,7 @@ public:
 	ENGINE_API virtual float GetOpacityMaskClipValue() const override;
 	ENGINE_API virtual bool GetCastDynamicShadowAsMasked() const override;
 	ENGINE_API virtual bool IsDistorted() const override;
+	ENGINE_API virtual ERefractionCoverageMode GetRefractionCoverageMode() const override;
 	ENGINE_API virtual float GetTranslucencyDirectionalLightingIntensity() const override;
 	ENGINE_API virtual float GetTranslucentShadowDensityScale() const override;
 	ENGINE_API virtual float GetTranslucentSelfShadowDensityScale() const override;
