Latest Clang toolchain fixes for FPostProcessSettings. deprecated-copy-with-user-provided-copy. Also addresses cases where members weren't being copied in the copy constructor.

diff --git a/Engine/Source/Runtime/Engine/Classes/Engine/Scene.h b/Engine/Source/Runtime/Engine/Classes/Engine/Scene.h
--- a/Engine/Source/Runtime/Engine/Classes/Engine/Scene.h
+++ b/Engine/Source/Runtime/Engine/Classes/Engine/Scene.h
@@ -668,11 +668,9 @@ struct FPostProcessSettingsDebugInfo
 // Each property consists of a bool to enable it (by default off),
 // the variable declaration and further down the default value for it.
 // The comment should include the meaning and usable range.
-PRAGMA_DISABLE_DEPRECATION_WARNINGS
 USTRUCT(BlueprintType, meta=(HiddenByDefault))
 struct FPostProcessSettings
 {
-	PRAGMA_ENABLE_DEPRECATION_WARNINGS
 	GENERATED_USTRUCT_BODY()
 
 	// first all bOverride_... as they get grouped together into bitfields
@@ -2335,8 +2333,14 @@ struct FPostProcessSettings
 
 	// good start values for a new volume, by default no value is overriding
 	ENGINE_API FPostProcessSettings();
-
-	ENGINE_API FPostProcessSettings(const FPostProcessSettings&);
+	
+	// Workaround for clang deprecation warnings for any deprecated members in implicitly-defined special member functions
+PRAGMA_DISABLE_DEPRECATION_WARNINGS
+	FPostProcessSettings(const FPostProcessSettings&) = default;
+	FPostProcessSettings(FPostProcessSettings&&) = default;
+	FPostProcessSettings& operator=(const FPostProcessSettings&) = default;
+	FPostProcessSettings& operator=(FPostProcessSettings&&) = default;
+PRAGMA_ENABLE_DEPRECATION_WARNINGS
 
 	/**
 		* Used to define the values before any override happens.
@@ -2359,9 +2363,7 @@ struct FPostProcessSettings
 	void PostSerialize(const FArchive& Ar);
 #endif
 
-	PRAGMA_DISABLE_DEPRECATION_WARNINGS
 }; // struct FPostProcessSettings
-PRAGMA_ENABLE_DEPRECATION_WARNINGS
 
 #if WITH_EDITORONLY_DATA
 template<>
diff --git a/Engine/Source/Runtime/Engine/Private/Scene.cpp b/Engine/Source/Runtime/Engine/Private/Scene.cpp
--- a/Engine/Source/Runtime/Engine/Private/Scene.cpp
+++ b/Engine/Source/Runtime/Engine/Private/Scene.cpp
@@ -658,422 +658,7 @@ FPostProcessSettings::FPostProcessSettings()
 #endif // DO_CHECK
 }
 
-FPostProcessSettings::FPostProcessSettings(const FPostProcessSettings& Settings)
-	: bOverride_TemperatureType(Settings.bOverride_TemperatureType)
-	, bOverride_WhiteTemp(Settings.bOverride_WhiteTemp)
-	, bOverride_WhiteTint(Settings.bOverride_WhiteTint)
-	, bOverride_ColorSaturation(Settings.bOverride_ColorSaturation)
-	, bOverride_ColorContrast(Settings.bOverride_ColorContrast)
-	, bOverride_ColorGamma(Settings.bOverride_ColorGamma)
-	, bOverride_ColorGain(Settings.bOverride_ColorGain)
-	, bOverride_ColorOffset(Settings.bOverride_ColorOffset)
-	, bOverride_ColorSaturationShadows(Settings.bOverride_ColorSaturationShadows)
-	, bOverride_ColorContrastShadows(Settings.bOverride_ColorContrastShadows)
-	, bOverride_ColorGammaShadows(Settings.bOverride_ColorGammaShadows)
-	, bOverride_ColorGainShadows(Settings.bOverride_ColorGainShadows)
-	, bOverride_ColorOffsetShadows(Settings.bOverride_ColorOffsetShadows)
-	, bOverride_ColorSaturationMidtones(Settings.bOverride_ColorSaturationMidtones)
-	, bOverride_ColorContrastMidtones(Settings.bOverride_ColorContrastMidtones)
-	, bOverride_ColorGammaMidtones(Settings.bOverride_ColorGammaMidtones)
-	, bOverride_ColorGainMidtones(Settings.bOverride_ColorGainMidtones)
-	, bOverride_ColorOffsetMidtones(Settings.bOverride_ColorOffsetMidtones)
-	, bOverride_ColorSaturationHighlights(Settings.bOverride_ColorSaturationHighlights)
-	, bOverride_ColorContrastHighlights(Settings.bOverride_ColorContrastHighlights)
-	, bOverride_ColorGammaHighlights(Settings.bOverride_ColorGammaHighlights)
-	, bOverride_ColorGainHighlights(Settings.bOverride_ColorGainHighlights)
-	, bOverride_ColorOffsetHighlights(Settings.bOverride_ColorOffsetHighlights)
-	, bOverride_ColorCorrectionShadowsMax(Settings.bOverride_ColorCorrectionShadowsMax)
-	, bOverride_ColorCorrectionHighlightsMin(Settings.bOverride_ColorCorrectionHighlightsMin)
-	, bOverride_ColorCorrectionHighlightsMax(Settings.bOverride_ColorCorrectionHighlightsMax)
-	, bOverride_BlueCorrection(Settings.bOverride_BlueCorrection)
-	, bOverride_ExpandGamut(Settings.bOverride_ExpandGamut)
-	, bOverride_ToneCurveAmount(Settings.bOverride_ToneCurveAmount)
-	, bOverride_FilmSlope(Settings.bOverride_FilmSlope)
-	, bOverride_FilmToe(Settings.bOverride_FilmToe)
-	, bOverride_FilmShoulder(Settings.bOverride_FilmShoulder)
-	, bOverride_FilmBlackClip(Settings.bOverride_FilmBlackClip)
-	, bOverride_FilmWhiteClip(Settings.bOverride_FilmWhiteClip)
-	, bOverride_SceneColorTint(Settings.bOverride_SceneColorTint)
-	, bOverride_SceneFringeIntensity(Settings.bOverride_SceneFringeIntensity)
-	, bOverride_ChromaticAberrationStartOffset(Settings.bOverride_ChromaticAberrationStartOffset)
-	, bOverride_AmbientCubemapTint(Settings.bOverride_AmbientCubemapTint)
-	, bOverride_AmbientCubemapIntensity(Settings.bOverride_AmbientCubemapIntensity)
-	, bOverride_BloomMethod(Settings.bOverride_BloomMethod)
-	, bOverride_BloomIntensity(Settings.bOverride_BloomIntensity)
-	, bOverride_BloomThreshold(Settings.bOverride_BloomThreshold)
-	, bOverride_Bloom1Tint(Settings.bOverride_Bloom1Tint)
-	, bOverride_Bloom1Size(Settings.bOverride_Bloom1Size)
-	, bOverride_Bloom2Size(Settings.bOverride_Bloom2Size)
-	, bOverride_Bloom2Tint(Settings.bOverride_Bloom2Tint)
-	, bOverride_Bloom3Tint(Settings.bOverride_Bloom3Tint)
-	, bOverride_Bloom3Size(Settings.bOverride_Bloom3Size)
-	, bOverride_Bloom4Tint(Settings.bOverride_Bloom4Tint)
-	, bOverride_Bloom4Size(Settings.bOverride_Bloom4Size)
-	, bOverride_Bloom5Tint(Settings.bOverride_Bloom5Tint)
-	, bOverride_Bloom5Size(Settings.bOverride_Bloom5Size)
-	, bOverride_Bloom6Tint(Settings.bOverride_Bloom6Tint)
-	, bOverride_Bloom6Size(Settings.bOverride_Bloom6Size)
-	, bOverride_BloomSizeScale(Settings.bOverride_BloomSizeScale)
-	, bOverride_BloomConvolutionTexture(Settings.bOverride_BloomConvolutionTexture)
-	, bOverride_BloomConvolutionScatterDispersion(Settings.bOverride_BloomConvolutionScatterDispersion)
-	, bOverride_BloomConvolutionSize(Settings.bOverride_BloomConvolutionSize)
-	, bOverride_BloomConvolutionCenterUV(Settings.bOverride_BloomConvolutionCenterUV)
-	//, bOverride_BloomConvolutionPreFilter_DEPRECATED(Settings.bOverride_BloomConvolutionPreFilter_DEPRECATED)
-	, bOverride_BloomConvolutionPreFilterMin(Settings.bOverride_BloomConvolutionPreFilterMin)
-	, bOverride_BloomConvolutionPreFilterMax(Settings.bOverride_BloomConvolutionPreFilterMax)
-	, bOverride_BloomConvolutionPreFilterMult(Settings.bOverride_BloomConvolutionPreFilterMult)
-	, bOverride_BloomConvolutionBufferScale(Settings.bOverride_BloomConvolutionBufferScale)
-	, bOverride_BloomDirtMaskIntensity(Settings.bOverride_BloomDirtMaskIntensity)
-	, bOverride_BloomDirtMaskTint(Settings.bOverride_BloomDirtMaskTint)
-	, bOverride_BloomDirtMask(Settings.bOverride_BloomDirtMask)
-	, bOverride_CameraShutterSpeed(Settings.bOverride_CameraShutterSpeed)
-	, bOverride_CameraISO(Settings.bOverride_CameraISO)
-	, bOverride_AutoExposureMethod(Settings.bOverride_AutoExposureMethod)
-	, bOverride_AutoExposureLowPercent(Settings.bOverride_AutoExposureLowPercent)
-	, bOverride_AutoExposureHighPercent(Settings.bOverride_AutoExposureHighPercent)
-	, bOverride_AutoExposureMinBrightness(Settings.bOverride_AutoExposureMinBrightness)
-	, bOverride_AutoExposureMaxBrightness(Settings.bOverride_AutoExposureMaxBrightness)
-	, bOverride_AutoExposureCalibrationConstant_DEPRECATED(Settings.bOverride_AutoExposureCalibrationConstant_DEPRECATED)
-	, bOverride_AutoExposureSpeedUp(Settings.bOverride_AutoExposureSpeedUp)
-	, bOverride_AutoExposureSpeedDown(Settings.bOverride_AutoExposureSpeedDown)
-	, bOverride_AutoExposureBias(Settings.bOverride_AutoExposureBias)
-	, bOverride_AutoExposureBiasCurve(Settings.bOverride_AutoExposureBiasCurve)
-	, bOverride_AutoExposureMeterMask(Settings.bOverride_AutoExposureMeterMask)
-	, bOverride_AutoExposureApplyPhysicalCameraExposure(Settings.bOverride_AutoExposureApplyPhysicalCameraExposure)
-	, bOverride_HistogramLogMin(Settings.bOverride_HistogramLogMin)
-	, bOverride_HistogramLogMax(Settings.bOverride_HistogramLogMax)
-	, bOverride_LocalExposureContrastScale_DEPRECATED(Settings.bOverride_LocalExposureContrastScale_DEPRECATED)
-	, bOverride_LocalExposureHighlightContrastScale(Settings.bOverride_LocalExposureHighlightContrastScale)
-	, bOverride_LocalExposureShadowContrastScale(Settings.bOverride_LocalExposureShadowContrastScale)
-	, bOverride_LocalExposureDetailStrength(Settings.bOverride_LocalExposureDetailStrength)
-	, bOverride_LocalExposureBlurredLuminanceBlend(Settings.bOverride_LocalExposureBlurredLuminanceBlend)
-	, bOverride_LocalExposureBlurredLuminanceKernelSizePercent(Settings.bOverride_LocalExposureBlurredLuminanceKernelSizePercent)
-	, bOverride_LocalExposureMiddleGreyBias(Settings.bOverride_LocalExposureMiddleGreyBias)
-	, bOverride_LensFlareIntensity(Settings.bOverride_LensFlareIntensity)
-	, bOverride_LensFlareTint(Settings.bOverride_LensFlareTint)
-	, bOverride_LensFlareTints(Settings.bOverride_LensFlareTints)
-	, bOverride_LensFlareBokehSize(Settings.bOverride_LensFlareBokehSize)
-	, bOverride_LensFlareBokehShape(Settings.bOverride_LensFlareBokehShape)
-	, bOverride_LensFlareThreshold(Settings.bOverride_LensFlareThreshold)
-	, bOverride_VignetteIntensity(Settings.bOverride_VignetteIntensity)
-	, bOverride_GrainIntensity_DEPRECATED(Settings.bOverride_GrainIntensity_DEPRECATED)
-	, bOverride_GrainJitter_DEPRECATED(Settings.bOverride_GrainJitter_DEPRECATED)
-	, bOverride_FilmGrainIntensity(Settings.bOverride_FilmGrainIntensity)
-	, bOverride_FilmGrainIntensityShadows(Settings.bOverride_FilmGrainIntensityShadows)
-	, bOverride_FilmGrainIntensityMidtones(Settings.bOverride_FilmGrainIntensityMidtones)
-	, bOverride_FilmGrainIntensityHighlights(Settings.bOverride_FilmGrainIntensityHighlights)
-	, bOverride_FilmGrainShadowsMax(Settings.bOverride_FilmGrainShadowsMax)
-	, bOverride_FilmGrainHighlightsMin(Settings.bOverride_FilmGrainHighlightsMin)
-	, bOverride_FilmGrainHighlightsMax(Settings.bOverride_FilmGrainHighlightsMax)
-	, bOverride_FilmGrainTexelSize(Settings.bOverride_FilmGrainTexelSize)
-	, bOverride_FilmGrainTexture(Settings.bOverride_FilmGrainTexture)
-	, bOverride_AmbientOcclusionIntensity(Settings.bOverride_AmbientOcclusionIntensity)
-	, bOverride_AmbientOcclusionStaticFraction(Settings.bOverride_AmbientOcclusionStaticFraction)
-	, bOverride_AmbientOcclusionRadius(Settings.bOverride_AmbientOcclusionRadius)
-	, bOverride_AmbientOcclusionFadeDistance(Settings.bOverride_AmbientOcclusionFadeDistance)
-	, bOverride_AmbientOcclusionFadeRadius(Settings.bOverride_AmbientOcclusionFadeRadius)
-	, bOverride_AmbientOcclusionDistance_DEPRECATED(Settings.bOverride_AmbientOcclusionDistance_DEPRECATED)
-	, bOverride_AmbientOcclusionRadiusInWS(Settings.bOverride_AmbientOcclusionRadiusInWS)
-	, bOverride_AmbientOcclusionPower(Settings.bOverride_AmbientOcclusionPower)
-	, bOverride_AmbientOcclusionBias(Settings.bOverride_AmbientOcclusionBias)
-	, bOverride_AmbientOcclusionQuality(Settings.bOverride_AmbientOcclusionQuality)
-	, bOverride_AmbientOcclusionMipBlend(Settings.bOverride_AmbientOcclusionMipBlend)
-	, bOverride_AmbientOcclusionMipScale(Settings.bOverride_AmbientOcclusionMipScale)
-	, bOverride_AmbientOcclusionMipThreshold(Settings.bOverride_AmbientOcclusionMipThreshold)
-	, bOverride_AmbientOcclusionTemporalBlendWeight(Settings.bOverride_AmbientOcclusionTemporalBlendWeight)
-	, bOverride_RayTracingAO(Settings.bOverride_RayTracingAO)
-	, bOverride_RayTracingAOSamplesPerPixel(Settings.bOverride_RayTracingAOSamplesPerPixel)
-	, bOverride_RayTracingAOIntensity(Settings.bOverride_RayTracingAOIntensity)
-	, bOverride_RayTracingAORadius(Settings.bOverride_RayTracingAORadius)
-	, bOverride_IndirectLightingColor(Settings.bOverride_IndirectLightingColor)
-	, bOverride_IndirectLightingIntensity(Settings.bOverride_IndirectLightingIntensity)
-	, bOverride_ColorGradingIntensity(Settings.bOverride_ColorGradingIntensity)
-	, bOverride_ColorGradingLUT(Settings.bOverride_ColorGradingLUT)
-	, bOverride_DepthOfFieldFocalDistance(Settings.bOverride_DepthOfFieldFocalDistance)
-	, bOverride_DepthOfFieldFstop(Settings.bOverride_DepthOfFieldFstop)
-	, bOverride_DepthOfFieldMinFstop(Settings.bOverride_DepthOfFieldMinFstop)
-	, bOverride_DepthOfFieldBladeCount(Settings.bOverride_DepthOfFieldBladeCount)
-	, bOverride_DepthOfFieldSensorWidth(Settings.bOverride_DepthOfFieldSensorWidth)
-	, bOverride_DepthOfFieldSqueezeFactor(Settings.bOverride_DepthOfFieldSqueezeFactor)
-	, bOverride_DepthOfFieldDepthBlurRadius(Settings.bOverride_DepthOfFieldDepthBlurRadius)
-	, bOverride_DepthOfFieldDepthBlurAmount(Settings.bOverride_DepthOfFieldDepthBlurAmount)
-	, bOverride_DepthOfFieldFocalRegion(Settings.bOverride_DepthOfFieldFocalRegion)
-	, bOverride_DepthOfFieldNearTransitionRegion(Settings.bOverride_DepthOfFieldNearTransitionRegion)
-	, bOverride_DepthOfFieldFarTransitionRegion(Settings.bOverride_DepthOfFieldFarTransitionRegion)
-	, bOverride_DepthOfFieldScale(Settings.bOverride_DepthOfFieldScale)
-	, bOverride_DepthOfFieldNearBlurSize(Settings.bOverride_DepthOfFieldNearBlurSize)
-	, bOverride_DepthOfFieldFarBlurSize(Settings.bOverride_DepthOfFieldFarBlurSize)
-	, bOverride_MobileHQGaussian(Settings.bOverride_MobileHQGaussian)
-	, bOverride_DepthOfFieldOcclusion(Settings.bOverride_DepthOfFieldOcclusion)
-	, bOverride_DepthOfFieldSkyFocusDistance(Settings.bOverride_DepthOfFieldSkyFocusDistance)
-	, bOverride_DepthOfFieldVignetteSize(Settings.bOverride_DepthOfFieldVignetteSize)
-	, bOverride_MotionBlurAmount(Settings.bOverride_MotionBlurAmount)
-	, bOverride_MotionBlurMax(Settings.bOverride_MotionBlurMax)
-	, bOverride_MotionBlurPerObjectSize(Settings.bOverride_MotionBlurPerObjectSize)
-	, bOverride_ScreenPercentage_DEPRECATED(Settings.bOverride_ScreenPercentage_DEPRECATED)
-	, bOverride_ReflectionMethod(Settings.bOverride_ReflectionMethod)
-	, bOverride_LumenReflectionQuality(Settings.bOverride_LumenReflectionQuality)
-	, bOverride_ScreenSpaceReflectionIntensity(Settings.bOverride_ScreenSpaceReflectionIntensity)
-	, bOverride_ScreenSpaceReflectionQuality(Settings.bOverride_ScreenSpaceReflectionQuality)
-	, bOverride_ScreenSpaceReflectionMaxRoughness(Settings.bOverride_ScreenSpaceReflectionMaxRoughness)
-	, bOverride_ScreenSpaceReflectionRoughnessScale(Settings.bOverride_ScreenSpaceReflectionRoughnessScale)
-	, bOverride_RayTracingReflectionsMaxRoughness(Settings.bOverride_RayTracingReflectionsMaxRoughness)
-	, bOverride_RayTracingReflectionsMaxBounces(Settings.bOverride_RayTracingReflectionsMaxBounces)
-	, bOverride_RayTracingReflectionsSamplesPerPixel(Settings.bOverride_RayTracingReflectionsSamplesPerPixel)
-	, bOverride_RayTracingReflectionsShadows(Settings.bOverride_RayTracingReflectionsShadows)
-	, bOverride_RayTracingReflectionsTranslucency(Settings.bOverride_RayTracingReflectionsTranslucency)
-	, bOverride_TranslucencyType(Settings.bOverride_TranslucencyType)
-	, bOverride_RayTracingTranslucencyMaxRoughness(Settings.bOverride_RayTracingTranslucencyMaxRoughness)
-	, bOverride_RayTracingTranslucencyRefractionRays(Settings.bOverride_RayTracingTranslucencyRefractionRays)
-	, bOverride_RayTracingTranslucencySamplesPerPixel(Settings.bOverride_RayTracingTranslucencySamplesPerPixel)
-	, bOverride_RayTracingTranslucencyShadows(Settings.bOverride_RayTracingTranslucencyShadows)
-	, bOverride_RayTracingTranslucencyRefraction(Settings.bOverride_RayTracingTranslucencyRefraction)
-	, bOverride_DynamicGlobalIlluminationMethod(Settings.bOverride_DynamicGlobalIlluminationMethod)
-    , bOverride_LumenSceneLightingQuality(Settings.bOverride_LumenSceneLightingQuality)
-    , bOverride_LumenSceneDetail(Settings.bOverride_LumenSceneDetail)
-	, bOverride_LumenSceneViewDistance(Settings.bOverride_LumenSceneViewDistance)
-	, bOverride_LumenSceneLightingUpdateSpeed(Settings.bOverride_LumenSceneLightingUpdateSpeed)
-    , bOverride_LumenFinalGatherQuality(Settings.bOverride_LumenFinalGatherQuality)
-    , bOverride_LumenFinalGatherLightingUpdateSpeed(Settings.bOverride_LumenFinalGatherLightingUpdateSpeed)
-    , bOverride_LumenMaxTraceDistance(Settings.bOverride_LumenMaxTraceDistance)
-	, bOverride_LumenDiffuseColorBoost(Settings.bOverride_LumenDiffuseColorBoost)
-	, bOverride_LumenSkylightLeaking(Settings.bOverride_LumenSkylightLeaking)
-	, bOverride_LumenFullSkylightLeakingDistance(Settings.bOverride_LumenFullSkylightLeakingDistance)
-	, bOverride_LumenRayLightingMode(Settings.bOverride_LumenRayLightingMode)
-	, bOverride_LumenFrontLayerTranslucencyReflections(Settings.bOverride_LumenFrontLayerTranslucencyReflections)
-	, bOverride_LumenMaxReflectionBounces(Settings.bOverride_LumenMaxReflectionBounces)
-	, bOverride_LumenSurfaceCacheResolution(Settings.bOverride_LumenSurfaceCacheResolution)
-	, bOverride_RayTracingGI(Settings.bOverride_RayTracingGI)
-	, bOverride_RayTracingGIMaxBounces(Settings.bOverride_RayTracingGIMaxBounces)
-	, bOverride_RayTracingGISamplesPerPixel(Settings.bOverride_RayTracingGISamplesPerPixel)
-	, bOverride_PathTracingMaxBounces(Settings.bOverride_PathTracingMaxBounces)
-	, bOverride_PathTracingSamplesPerPixel(Settings.bOverride_PathTracingSamplesPerPixel)
-	, bOverride_PathTracingFilterWidth(Settings.bOverride_PathTracingFilterWidth)
-	, bOverride_PathTracingMaxPathExposure(Settings.bOverride_PathTracingMaxPathExposure)
-	, bOverride_PathTracingEnableReferenceDOF(Settings.bOverride_PathTracingEnableReferenceDOF)
-	, bOverride_PathTracingEnableReferenceAtmosphere(Settings.bOverride_PathTracingEnableReferenceAtmosphere)
-	, bOverride_PathTracingEnableDenoiser(Settings.bOverride_PathTracingEnableDenoiser)
-	, bOverride_PathTracingIncludeEmissive(Settings.bOverride_PathTracingIncludeEmissive)
-	, bOverride_PathTracingIncludeIndirectEmissive(Settings.bOverride_PathTracingIncludeIndirectEmissive)
-	, bOverride_PathTracingIncludeDiffuse(Settings.bOverride_PathTracingIncludeDiffuse)
-	, bOverride_PathTracingIncludeIndirectDiffuse(Settings.bOverride_PathTracingIncludeIndirectDiffuse)
-	, bOverride_PathTracingIncludeSpecular(Settings.bOverride_PathTracingIncludeSpecular)
-	, bOverride_PathTracingIncludeIndirectSpecular(Settings.bOverride_PathTracingIncludeIndirectSpecular)
-	, bOverride_PathTracingIncludeVolume(Settings.bOverride_PathTracingIncludeVolume)
-	, bOverride_PathTracingIncludeIndirectVolume(Settings.bOverride_PathTracingIncludeIndirectVolume)
-
-	, bMobileHQGaussian(Settings.bMobileHQGaussian)
-	, BloomMethod(Settings.BloomMethod)
-	, AutoExposureMethod(Settings.AutoExposureMethod)
-	, TemperatureType(Settings.TemperatureType)
-	, WhiteTemp(Settings.WhiteTemp)
-	, WhiteTint(Settings.WhiteTint)
-	, ColorSaturation(Settings.ColorSaturation)
-	, ColorContrast(Settings.ColorContrast)
-	, ColorGamma(Settings.ColorGamma)
-	, ColorGain(Settings.ColorGain)
-	, ColorOffset(Settings.ColorOffset)
-	, ColorSaturationShadows(Settings.ColorSaturationShadows)
-	, ColorContrastShadows(Settings.ColorContrastShadows)
-	, ColorGammaShadows(Settings.ColorGammaShadows)
-	, ColorGainShadows(Settings.ColorGainShadows)
-	, ColorOffsetShadows(Settings.ColorOffsetShadows)
-	, ColorSaturationMidtones(Settings.ColorSaturationMidtones)
-	, ColorContrastMidtones(Settings.ColorContrastMidtones)
-	, ColorGammaMidtones(Settings.ColorGammaMidtones)
-	, ColorGainMidtones(Settings.ColorGainMidtones)
-	, ColorOffsetMidtones(Settings.ColorOffsetMidtones)
-	, ColorSaturationHighlights(Settings.ColorSaturationHighlights)
-	, ColorContrastHighlights(Settings.ColorContrastHighlights)
-	, ColorGammaHighlights(Settings.ColorGammaHighlights)
-	, ColorGainHighlights(Settings.ColorGainHighlights)
-	, ColorOffsetHighlights(Settings.ColorOffsetHighlights)
-	, ColorCorrectionHighlightsMin(Settings.ColorCorrectionHighlightsMin)
-	, ColorCorrectionHighlightsMax(Settings.ColorCorrectionHighlightsMax)
-	, ColorCorrectionShadowsMax(Settings.ColorCorrectionShadowsMax)
-	, BlueCorrection(Settings.BlueCorrection)
-	, ExpandGamut(Settings.ExpandGamut)
-	, ToneCurveAmount(Settings.ToneCurveAmount)
-	, FilmSlope(Settings.FilmSlope)
-	, FilmToe(Settings.FilmToe)
-	, FilmShoulder(Settings.FilmShoulder)
-	, FilmBlackClip(Settings.FilmBlackClip)
-	, FilmWhiteClip(Settings.FilmWhiteClip)
-	, SceneColorTint(Settings.SceneColorTint)
-	, SceneFringeIntensity(Settings.SceneFringeIntensity)
-	, ChromaticAberrationStartOffset(Settings.ChromaticAberrationStartOffset)
-	, BloomIntensity(Settings.BloomIntensity)
-	, BloomThreshold(Settings.BloomThreshold)
-	, BloomSizeScale(Settings.BloomSizeScale)
-	, Bloom1Size(Settings.Bloom1Size)
-	, Bloom2Size(Settings.Bloom2Size)
-	, Bloom3Size(Settings.Bloom3Size)
-	, Bloom4Size(Settings.Bloom4Size)
-	, Bloom5Size(Settings.Bloom5Size)
-	, Bloom6Size(Settings.Bloom6Size)
-	, Bloom1Tint(Settings.Bloom1Tint)
-	, Bloom2Tint(Settings.Bloom2Tint)
-	, Bloom3Tint(Settings.Bloom3Tint)
-	, Bloom4Tint(Settings.Bloom4Tint)
-	, Bloom5Tint(Settings.Bloom5Tint)
-	, Bloom6Tint(Settings.Bloom6Tint)
-	, BloomConvolutionScatterDispersion(Settings.BloomConvolutionScatterDispersion)
-	, BloomConvolutionSize(Settings.BloomConvolutionSize)
-	, BloomConvolutionTexture(Settings.BloomConvolutionTexture)
-	, BloomConvolutionCenterUV(Settings.BloomConvolutionCenterUV)
-	//, BloomConvolutionPreFilter_DEPRECATED(Settings.BloomConvolutionPreFilter_DEPRECATED)
-	, BloomConvolutionPreFilterMin(Settings.BloomConvolutionPreFilterMin)
-	, BloomConvolutionPreFilterMax(Settings.BloomConvolutionPreFilterMax)
-	, BloomConvolutionPreFilterMult(Settings.BloomConvolutionPreFilterMult)
-	, BloomConvolutionBufferScale(Settings.BloomConvolutionBufferScale)
-	, BloomDirtMask(Settings.BloomDirtMask)
-	, BloomDirtMaskIntensity(Settings.BloomDirtMaskIntensity)
-	, BloomDirtMaskTint(Settings.BloomDirtMaskTint)
-	, DynamicGlobalIlluminationMethod(Settings.DynamicGlobalIlluminationMethod)
-	, IndirectLightingColor(Settings.IndirectLightingColor)
-	, IndirectLightingIntensity(Settings.IndirectLightingIntensity)
-	, LumenSceneLightingQuality(Settings.LumenSceneLightingQuality)
-	, LumenSceneDetail(Settings.LumenSceneDetail)
-	, LumenSceneViewDistance(Settings.LumenSceneViewDistance)
-	, LumenSceneLightingUpdateSpeed(Settings.LumenSceneLightingUpdateSpeed)
-	, LumenFinalGatherQuality(Settings.LumenFinalGatherQuality)
-	, LumenFinalGatherLightingUpdateSpeed(Settings.LumenFinalGatherLightingUpdateSpeed)
-	, LumenMaxTraceDistance(Settings.LumenMaxTraceDistance)
-	, LumenDiffuseColorBoost(Settings.LumenDiffuseColorBoost)
-	, LumenSkylightLeaking(Settings.LumenSkylightLeaking)
-	, LumenFullSkylightLeakingDistance(Settings.LumenFullSkylightLeakingDistance)
-	, LumenSurfaceCacheResolution(Settings.LumenSurfaceCacheResolution)
-	, RayTracingGIType(Settings.RayTracingGIType)
-	, RayTracingGIMaxBounces(Settings.RayTracingGIMaxBounces)
-	, RayTracingGISamplesPerPixel(Settings.RayTracingGISamplesPerPixel)
-	, ReflectionMethod(Settings.ReflectionMethod)
-	, LumenReflectionQuality(Settings.LumenReflectionQuality)
-	, LumenRayLightingMode(Settings.LumenRayLightingMode)
-	, LumenFrontLayerTranslucencyReflections(Settings.LumenFrontLayerTranslucencyReflections)
-	, LumenMaxReflectionBounces(Settings.LumenMaxReflectionBounces)
-	, ScreenSpaceReflectionIntensity(Settings.ScreenSpaceReflectionIntensity)
-	, ScreenSpaceReflectionQuality(Settings.ScreenSpaceReflectionQuality)
-	, ScreenSpaceReflectionMaxRoughness(Settings.ScreenSpaceReflectionMaxRoughness)
-	, RayTracingReflectionsMaxRoughness(Settings.RayTracingReflectionsMaxRoughness)
-	, RayTracingReflectionsMaxBounces(Settings.RayTracingReflectionsMaxBounces)
-	, RayTracingReflectionsSamplesPerPixel(Settings.RayTracingReflectionsSamplesPerPixel)
-	, RayTracingReflectionsShadows(Settings.RayTracingReflectionsShadows)
-	, RayTracingReflectionsTranslucency(Settings.RayTracingReflectionsTranslucency)
-	, AmbientCubemapTint(Settings.AmbientCubemapTint)
-	, AmbientCubemapIntensity(Settings.AmbientCubemapIntensity)
-	, AmbientCubemap(Settings.AmbientCubemap)
-	, CameraShutterSpeed(Settings.CameraShutterSpeed)
-	, CameraISO(Settings.CameraISO)
-	, DepthOfFieldFstop(Settings.DepthOfFieldFstop)
-	, DepthOfFieldMinFstop(Settings.DepthOfFieldMinFstop)
-	, DepthOfFieldBladeCount(Settings.DepthOfFieldBladeCount)
-	, AutoExposureBias(Settings.AutoExposureBias)
-	, AutoExposureApplyPhysicalCameraExposure(Settings.AutoExposureApplyPhysicalCameraExposure)
-	, AutoExposureBiasCurve(Settings.AutoExposureBiasCurve)
-	, AutoExposureMeterMask(Settings.AutoExposureMeterMask)
-	, AutoExposureLowPercent(Settings.AutoExposureLowPercent)
-	, AutoExposureHighPercent(Settings.AutoExposureHighPercent)
-	, AutoExposureMinBrightness(Settings.AutoExposureMinBrightness)
-	, AutoExposureMaxBrightness(Settings.AutoExposureMaxBrightness)
-	, AutoExposureSpeedUp(Settings.AutoExposureSpeedUp)
-	, AutoExposureSpeedDown(Settings.AutoExposureSpeedDown)
-	, HistogramLogMin(Settings.HistogramLogMin)
-	, HistogramLogMax(Settings.HistogramLogMax)
-	, AutoExposureCalibrationConstant_DEPRECATED(Settings.AutoExposureCalibrationConstant_DEPRECATED)
-	, LocalExposureContrastScale_DEPRECATED(Settings.LocalExposureContrastScale_DEPRECATED)
-	, LocalExposureHighlightContrastScale(Settings.LocalExposureHighlightContrastScale)
-	, LocalExposureShadowContrastScale(Settings.LocalExposureShadowContrastScale)
-	, LocalExposureDetailStrength(Settings.LocalExposureDetailStrength)
-	, LocalExposureBlurredLuminanceBlend(Settings.LocalExposureBlurredLuminanceBlend)
-	, LocalExposureBlurredLuminanceKernelSizePercent(Settings.LocalExposureBlurredLuminanceKernelSizePercent)
-	, LocalExposureMiddleGreyBias(Settings.LocalExposureMiddleGreyBias)
-	, LensFlareIntensity(Settings.LensFlareIntensity)
-	, LensFlareTint(Settings.LensFlareTint)
-	, LensFlareBokehSize(Settings.LensFlareBokehSize)
-	, LensFlareThreshold(Settings.LensFlareThreshold)
-	, LensFlareBokehShape(Settings.LensFlareBokehShape)
-	, VignetteIntensity(Settings.VignetteIntensity)
-	, GrainJitter_DEPRECATED(Settings.GrainJitter_DEPRECATED)
-	, GrainIntensity_DEPRECATED(Settings.GrainIntensity_DEPRECATED)
-	, FilmGrainIntensity(Settings.FilmGrainIntensity)
-	, FilmGrainIntensityShadows(Settings.FilmGrainIntensityShadows)
-	, FilmGrainIntensityMidtones(Settings.FilmGrainIntensityMidtones)
-	, FilmGrainIntensityHighlights(Settings.FilmGrainIntensityHighlights)
-	, FilmGrainShadowsMax(Settings.FilmGrainShadowsMax)
-	, FilmGrainHighlightsMin(Settings.FilmGrainHighlightsMin)
-	, FilmGrainHighlightsMax(Settings.FilmGrainHighlightsMax)
-	, FilmGrainTexelSize(Settings.FilmGrainTexelSize)
-	, FilmGrainTexture(Settings.FilmGrainTexture)
-	, AmbientOcclusionIntensity(Settings.AmbientOcclusionIntensity)
-	, AmbientOcclusionStaticFraction(Settings.AmbientOcclusionStaticFraction)
-	, AmbientOcclusionRadius(Settings.AmbientOcclusionRadius)
-	, AmbientOcclusionRadiusInWS(Settings.AmbientOcclusionRadiusInWS)
-	, AmbientOcclusionFadeDistance(Settings.AmbientOcclusionFadeDistance)
-	, AmbientOcclusionFadeRadius(Settings.AmbientOcclusionFadeRadius)
-	, AmbientOcclusionDistance_DEPRECATED(Settings.AmbientOcclusionDistance_DEPRECATED)
-	, AmbientOcclusionPower(Settings.AmbientOcclusionPower)
-	, AmbientOcclusionBias(Settings.AmbientOcclusionBias)
-	, AmbientOcclusionQuality(Settings.AmbientOcclusionQuality)
-	, AmbientOcclusionMipBlend(Settings.AmbientOcclusionMipBlend)
-	, AmbientOcclusionMipScale(Settings.AmbientOcclusionMipScale)
-	, AmbientOcclusionMipThreshold(Settings.AmbientOcclusionMipThreshold)
-	, AmbientOcclusionTemporalBlendWeight(Settings.AmbientOcclusionTemporalBlendWeight)
-	, RayTracingAO(Settings.RayTracingAO)
-	, RayTracingAOSamplesPerPixel(Settings.RayTracingAOSamplesPerPixel)
-	, RayTracingAOIntensity(Settings.RayTracingAOIntensity)
-	, RayTracingAORadius(Settings.RayTracingAORadius)
-	, ColorGradingIntensity(Settings.ColorGradingIntensity)
-	, ColorGradingLUT(Settings.ColorGradingLUT)
-	, DepthOfFieldSensorWidth(Settings.DepthOfFieldSensorWidth)
-	, DepthOfFieldSqueezeFactor(Settings.DepthOfFieldSqueezeFactor)
-	, DepthOfFieldFocalDistance(Settings.DepthOfFieldFocalDistance)
-	, DepthOfFieldDepthBlurAmount(Settings.DepthOfFieldDepthBlurAmount)
-	, DepthOfFieldDepthBlurRadius(Settings.DepthOfFieldDepthBlurRadius)
-	, DepthOfFieldFocalRegion(Settings.DepthOfFieldFocalRegion)
-	, DepthOfFieldNearTransitionRegion(Settings.DepthOfFieldNearTransitionRegion)
-	, DepthOfFieldFarTransitionRegion(Settings.DepthOfFieldFarTransitionRegion)
-	, DepthOfFieldScale(Settings.DepthOfFieldScale)
-	, DepthOfFieldNearBlurSize(Settings.DepthOfFieldNearBlurSize)
-	, DepthOfFieldFarBlurSize(Settings.DepthOfFieldFarBlurSize)
-	, DepthOfFieldOcclusion(Settings.DepthOfFieldOcclusion)
-	, DepthOfFieldSkyFocusDistance(Settings.DepthOfFieldSkyFocusDistance)
-	, DepthOfFieldVignetteSize(Settings.DepthOfFieldVignetteSize)
-	, MotionBlurAmount(Settings.MotionBlurAmount)
-	, MotionBlurMax(Settings.MotionBlurMax)
-	, MotionBlurTargetFPS(Settings.MotionBlurTargetFPS)
-	, MotionBlurPerObjectSize(Settings.MotionBlurPerObjectSize)
-	, TranslucencyType(Settings.TranslucencyType)
-	, RayTracingTranslucencyMaxRoughness(Settings.RayTracingTranslucencyMaxRoughness)
-	, RayTracingTranslucencyRefractionRays(Settings.RayTracingTranslucencyRefractionRays)
-	, RayTracingTranslucencySamplesPerPixel(Settings.RayTracingTranslucencySamplesPerPixel)
-	, RayTracingTranslucencyShadows(Settings.RayTracingTranslucencyShadows)
-	, RayTracingTranslucencyRefraction(Settings.RayTracingTranslucencyRefraction)
-
-	, PathTracingMaxBounces(Settings.PathTracingMaxBounces)
-	, PathTracingSamplesPerPixel(Settings.PathTracingSamplesPerPixel)
-	, PathTracingFilterWidth(Settings.PathTracingFilterWidth)
-	, PathTracingMaxPathExposure(Settings.PathTracingMaxPathExposure)
-	, PathTracingEnableReferenceDOF(Settings.PathTracingEnableReferenceDOF)
-	, PathTracingEnableReferenceAtmosphere(Settings.PathTracingEnableReferenceAtmosphere)
-	, PathTracingEnableDenoiser(Settings.PathTracingEnableDenoiser)
-	, PathTracingIncludeEmissive(Settings.PathTracingIncludeEmissive)
-	, PathTracingIncludeIndirectEmissive(Settings.PathTracingIncludeIndirectEmissive)
-	, PathTracingIncludeDiffuse(Settings.PathTracingIncludeDiffuse)
-	, PathTracingIncludeIndirectDiffuse(Settings.PathTracingIncludeIndirectDiffuse)
-	, PathTracingIncludeSpecular(Settings.PathTracingIncludeSpecular)
-	, PathTracingIncludeIndirectSpecular(Settings.PathTracingIncludeIndirectSpecular)
-	, PathTracingIncludeVolume(Settings.PathTracingIncludeVolume)
-	, PathTracingIncludeIndirectVolume(Settings.PathTracingIncludeIndirectVolume)
-
-	, ScreenPercentage_DEPRECATED(Settings.ScreenPercentage_DEPRECATED)
-
-	, WeightedBlendables(Settings.WeightedBlendables)
-	//, Blendables_DEPRECATED(Settings.Blendables_DEPRECATED)
-{
-	for (int32 i = 0; i < UE_ARRAY_COUNT(LensFlareTints); i++)
-		LensFlareTints[i] = Settings.LensFlareTints[i];
-}
-	
+
 #if WITH_EDITORONLY_DATA
 bool FPostProcessSettings::Serialize(FArchive& Ar)
 {
