Blendspace Evaluator: Ensure a clamped normalized time stays within the 0-1 range

diff --git a/Engine/Source/Runtime/Engine/Private/Animation/BlendSpace.cpp b/Engine/Source/Runtime/Engine/Private/Animation/BlendSpace.cpp
--- a/Engine/Source/Runtime/Engine/Private/Animation/BlendSpace.cpp
+++ b/Engine/Source/Runtime/Engine/Private/Animation/BlendSpace.cpp
@@ -560,7 +560,8 @@ void UBlendSpace::TickAssetPlayer(FAnimTickRecord& Instance, struct FAnimNotifyQ
 				{
 					// When running under an evaluator the time is being set explicitly and we want to add on the deltas.
 					ClampedNormalizedPreviousTime -= ExtraNormalizedDeltaTime;
-					ClampedNormalizedPreviousTime = ClampedNormalizedPreviousTime < 0.0f ? ClampedNormalizedPreviousTime + 1.0f : ClampedNormalizedPreviousTime;
+					// Note that ExtraNormalizedDeltaTime can be negative
+					ClampedNormalizedPreviousTime = FMath::Wrap<float>(ClampedNormalizedPreviousTime, 0.0f, 1.0f);
 
 					// Also when under an evaluator, since the time is explicitly set before the update is called, the desired 
 					// current time is actually what we recorded before advancing time (effectively ignoring whatever was added).
