Fix bug preventing some materials from being moved to a different folder.

The issue was caused by the an invalid state of the material EditorOnlyData. The change fixes the bug and tries to restore the EOD valid state.

diff --git a/Engine/Source/Runtime/Engine/Private/Materials/MaterialInterface.cpp b/Engine/Source/Runtime/Engine/Private/Materials/MaterialInterface.cpp
--- a/Engine/Source/Runtime/Engine/Private/Materials/MaterialInterface.cpp
+++ b/Engine/Source/Runtime/Engine/Private/Materials/MaterialInterface.cpp
@@ -107,6 +107,18 @@ static EMaterialGetParameterValueFlags MakeParameterValueFlags(bool bOveriddenOn
 	return Result;
 }
 
+/** Deletes any invalid EditorOnlyData object with specified name that isn't the one assigned to specified interface. */
+static void DisposeInvalidEditorOnlyData(UMaterialInterface* MaterialInterface, const FString& EditorOnlyDataName)
+{
+	UObject* EditorOnlyExisting = StaticFindObject(/*Class=*/ nullptr, MaterialInterface, *EditorOnlyDataName, true);
+	if (EditorOnlyExisting && EditorOnlyExisting != MaterialInterface->GetEditorOnlyData())
+	{
+		FName UniqueDummyName = MakeUniqueObjectName(GetTransientPackage(), UMaterialInterfaceEditorOnlyData::StaticClass());
+		EditorOnlyExisting->Rename(*UniqueDummyName.ToString(), GetTransientPackage(), REN_NonTransactional | REN_DontCreateRedirectors | REN_ForceNoResetLoaders);
+		EditorOnlyExisting->MarkAsGarbage();
+	}
+}
+
 //////////////////////////////////////////////////////////////////////////
 
 /** Copies the material's relevance flags to a primitive's view relevance flags. */
@@ -716,15 +728,9 @@ void UMaterialInterface::PostDuplicate(bool bDuplicateForPIE)
 #if WITH_EDITORONLY_DATA
 	// Duplication will create a new editor only data but this MI already links a populated data. This makes sure that the
 	// duplicate instance EOD has the correct name (necessary for running the editor on cooked data).
-	UMaterialInterfaceEditorOnlyData* EditorOnly = GetEditorOnlyData();
 	FString EditorOnlyDataName = MaterialInterface::GetEditorOnlyDataName(*GetName());
-	UObject* EditorOnlyExisting = StaticFindObject(/*Class=*/ nullptr, this, *EditorOnlyDataName, true);
-	if (EditorOnlyExisting && EditorOnlyExisting != EditorOnly)
-	{
-		EditorOnlyExisting->Rename(*(EditorOnlyDataName + "_Unused"), nullptr, REN_NonTransactional | REN_DontCreateRedirectors);
-		EditorOnlyExisting->MarkAsGarbage();
-	}
-	if (EditorOnly)
+	DisposeInvalidEditorOnlyData(this, EditorOnlyDataName);
+	if (GetEditorOnlyData())
 	{
 		GetEditorOnlyData()->Rename(*EditorOnlyDataName, nullptr, REN_NonTransactional | REN_DontCreateRedirectors);
 	}
@@ -1548,11 +1554,7 @@ void UMaterialInterface::PreSave(FObjectPreSaveContext ObjectSaveContext)
 	if (!ObjectSaveContext.IsCooking() && EditorOnly && EditorOnly->GetName() != EditorOnlyDataName)
 	{
 		UE_LOG(LogMaterial, Display, TEXT("MaterialInterface %s has a incorrectly name EditorOnlyData '%s'. This may cause issues when running the editor on cooked data. Trying to rename it to the correct name '%s'."), *GetName(), *EditorOnly->GetName(), *EditorOnlyDataName);
-		UObject* EditorOnlyExisting = StaticFindObject(/*Class=*/ nullptr, this, *EditorOnlyDataName, true);
-		if (EditorOnlyExisting && EditorOnlyExisting != EditorOnly)
-		{
-			EditorOnlyExisting->Rename(*(EditorOnlyDataName + "_Unused"), nullptr, REN_NonTransactional | REN_DontCreateRedirectors);
-		}
+		DisposeInvalidEditorOnlyData(this, EditorOnlyDataName);
 		if (EditorOnly)
 		{
 			EditorOnly->Rename(*EditorOnlyDataName, nullptr, REN_NonTransactional | REN_DontCreateRedirectors);
@@ -1680,6 +1682,7 @@ bool UMaterialInterface::Rename(const TCHAR* NewName, UObject* NewOuter, ERename
 	if (bRenamed && NewName && EditorOnlyData)
 	{
 		FString EditorOnlyDataName = MaterialInterface::GetEditorOnlyDataName(NewName);
+		DisposeInvalidEditorOnlyData(this, EditorOnlyDataName);
 		bRenamed = EditorOnlyData->Rename(*EditorOnlyDataName, nullptr, Flags);
 	}
 #endif
