[shaders] Add "compile argument" functionality for generic parameters consumed by compile backends, along with helper to set both a compile arg AND define for cases which require both. The purpose of this is to move towards enforcing that preprocessor defines should _only_ be used during the preprocessing phase itself; the subset of defines that are currently used only as arguments to the compile backends (or used as both actual preprocessor defines and backend arguments) should use this mechanism instead. This allows us to explicitly identify and hash such parameters to be included in a stable key based primarily on a hash of preprocessed/deadstripped source; without it we would need to hash all defines which will defeat most of the deduplication benefit of doing this. A subsequent change will deprecate the GetDefinitions function on the shader compilation environment to enforce that defines cannot be read by anything outside of core shader system functionality.

diff --git a/Engine/Source/Developer/Apple/MetalShaderFormat/Private/MetalDerivedData.cpp b/Engine/Source/Developer/Apple/MetalShaderFormat/Private/MetalDerivedData.cpp
--- a/Engine/Source/Developer/Apple/MetalShaderFormat/Private/MetalDerivedData.cpp
+++ b/Engine/Source/Developer/Apple/MetalShaderFormat/Private/MetalDerivedData.cpp
@@ -210,13 +210,7 @@ bool DoCompileMetalShader(
 	const FString& Standard,
 	const FString& MinOSVersion)
 {
-	int32 IABTier = 0;
-
-	FString const* IABVersion = Input.Environment.GetDefinitions().Find(TEXT("METAL_INDIRECT_ARGUMENT_BUFFERS"));
-	if (VersionEnum >= 4 && IABVersion && IABVersion->IsNumeric())
-	{
-		LexFromString(IABTier, *(*IABVersion));
-	}
+	int32 IABTier = VersionEnum >= 4 ? Input.Environment.GetCompileArgument(TEXT("METAL_INDIRECT_ARGUMENT_BUFFERS"), 0) : 0;
 
 	Output.bSucceeded = false;
 
@@ -226,34 +220,11 @@ bool DoCompileMetalShader(
 	bool const bZeroInitialise = Input.Environment.CompilerFlags.Contains(CFLAG_ZeroInitialise);
 	bool const bBoundsChecks = Input.Environment.CompilerFlags.Contains(CFLAG_BoundsChecking);
 
-    bool bSupportAppleA8 = false;
-    FString const* SupportAppleA8 = Input.Environment.GetDefinitions().Find(TEXT("SUPPORT_APPLE_A8"));
-    if (SupportAppleA8)
-    {
-        LexFromString(bSupportAppleA8, *(*SupportAppleA8));
-    }
-    
-	bool bSwizzleSample = false;
-	FString const* Swizzle = Input.Environment.GetDefinitions().Find(TEXT("METAL_SWIZZLE_SAMPLES"));
-	if (Swizzle)
-	{
-		LexFromString(bSwizzleSample, *(*Swizzle));
-	}
-
-	bool bAllowFastIntriniscs = false;
-	FString const* FastIntrinsics = Input.Environment.GetDefinitions().Find(TEXT("METAL_USE_FAST_INTRINSICS"));
-	if (FastIntrinsics)
-	{
-		LexFromString(bAllowFastIntriniscs, *(*FastIntrinsics));
-	}
+	bool bSupportAppleA8 = Input.Environment.GetCompileArgument(TEXT("SUPPORT_APPLE_A8"), false);
+	bool bAllowFastIntrinsics = Input.Environment.GetCompileArgument(TEXT("METAL_USE_FAST_INTRINSICS"), false);
 
-	bool bForceInvariance = false;
-	FString const* UsingWPO = Input.Environment.GetDefinitions().Find(TEXT("USES_WORLD_POSITION_OFFSET"));
-	if (UsingWPO && FString("1") == *UsingWPO)
-	{
-		// WPO requires that we make all multiply/sincos instructions invariant :(
-		bForceInvariance = true;
-	}
+	// WPO requires that we make all multiply/sincos instructions invariant :(
+	bool bForceInvariance = Input.Environment.GetCompileArgument(TEXT("USES_WORLD_POSITION_OFFSET"), false);
 	
 	FMetalShaderOutputMetaData OutputData;
 	
@@ -1408,7 +1379,7 @@ bool DoCompileMetalShader(
 	if (Result != 0)
 	{
 		Output.Target = Input.Target;
-		BuildMetalShaderOutput(Output, Input, GUIDHash, CCFlags, MetalSource.c_str(), MetalSource.length(), CRCLen, CRC, VersionEnum, *Standard, *MinOSVersion, TypeMode, Output.Errors, OutputData.TypedBuffers, OutputData.InvariantBuffers, OutputData.TypedUAVs, OutputData.ConstantBuffers, OutputData.TypedBufferFormats, bAllowFastIntriniscs);
+		BuildMetalShaderOutput(Output, Input, GUIDHash, CCFlags, MetalSource.c_str(), MetalSource.length(), CRCLen, CRC, VersionEnum, *Standard, *MinOSVersion, TypeMode, Output.Errors, OutputData.TypedBuffers, OutputData.InvariantBuffers, OutputData.TypedUAVs, OutputData.ConstantBuffers, OutputData.TypedBufferFormats, bAllowFastIntrinsics);
 		return Output.bSucceeded;
 	}
 	else
diff --git a/Engine/Source/Developer/Apple/MetalShaderFormat/Private/MetalShaderCompiler.cpp b/Engine/Source/Developer/Apple/MetalShaderFormat/Private/MetalShaderCompiler.cpp
--- a/Engine/Source/Developer/Apple/MetalShaderFormat/Private/MetalShaderCompiler.cpp
+++ b/Engine/Source/Developer/Apple/MetalShaderFormat/Private/MetalShaderCompiler.cpp
@@ -270,8 +270,8 @@ void BuildMetalShaderOutput(
 	//TODO read from toolchain
 	const bool bIsMobile = (ShaderInput.Target.Platform == SP_METAL || ShaderInput.Target.Platform == SP_METAL_MRT || ShaderInput.Target.Platform == SP_METAL_TVOS || ShaderInput.Target.Platform == SP_METAL_MRT_TVOS);
 	bool bNoFastMath = ShaderInput.Environment.CompilerFlags.Contains(CFLAG_NoFastMath);
-	FString const* UsingWPO = ShaderInput.Environment.GetDefinitions().Find(TEXT("USES_WORLD_POSITION_OFFSET"));
-	if (UsingWPO && FString("1") == *UsingWPO && (ShaderInput.Target.Platform == SP_METAL_MRT || ShaderInput.Target.Platform == SP_METAL_MRT_TVOS) && Frequency == SF_Vertex)
+	const bool bUsingWPO = ShaderInput.Environment.GetCompileArgument(TEXT("USES_WORLD_POSITION_OFFSET"), false);
+	if (bUsingWPO && (ShaderInput.Target.Platform == SP_METAL_MRT || ShaderInput.Target.Platform == SP_METAL_MRT_TVOS) && Frequency == SF_Vertex)
 	{
 		// WPO requires that we make all multiply/sincos instructions invariant :(
 		bNoFastMath = true;
@@ -870,18 +870,12 @@ void CompileShader_Metal(const FShaderCompilerInput& _Input,FShaderCompilerOutpu
 	
 	EMetalGPUSemantics Semantics = EMetalGPUSemanticsMobile;
 	
-    const int32 VersionEnum = [&Input, &Output]() -> int32
+	uint32 VersionEnum = GMetalDefaultShadingLanguageVersion;
+	bool bFoundVersion = Input.Environment.GetCompileArgument(TEXT("SHADER_LANGUAGE_VERSION"), VersionEnum);
+	if (!bFoundVersion)
 	{
-		if (const FString* VersionEnumEntry = Input.Environment.GetDefinitions().Find(TEXT("SHADER_LANGUAGE_VERSION")))
-		{
-			return FCString::Atoi(*FString(*VersionEnumEntry));
-		}
-		else
-		{
-			new(Output.Errors) FShaderCompilerError(*FString::Printf(TEXT("Missing definition of SHADER_LANGUAGE_VERSION; Falling back to default value %d"), GMetalDefaultShadingLanguageVersion));
-			return GMetalDefaultShadingLanguageVersion;
-		}
-	}();
+		new(Output.Errors) FShaderCompilerError(*FString::Printf(TEXT("Missing SHADER_LANGUAGE_VERSION compile argument; Falling back to default value %d"), GMetalDefaultShadingLanguageVersion));
+	}
 
 	// TODO read from toolchain
 	bool bAppleTV = (Input.ShaderFormat == NAME_SF_METAL_TVOS || Input.ShaderFormat == NAME_SF_METAL_MRT_TVOS);
diff --git a/Engine/Source/Developer/ShaderFormatOpenGL/Private/OpenGLShaderCompiler.cpp b/Engine/Source/Developer/ShaderFormatOpenGL/Private/OpenGLShaderCompiler.cpp
--- a/Engine/Source/Developer/ShaderFormatOpenGL/Private/OpenGLShaderCompiler.cpp
+++ b/Engine/Source/Developer/ShaderFormatOpenGL/Private/OpenGLShaderCompiler.cpp
@@ -2671,35 +2671,23 @@ enum EDecalBlendFlags
 	Modulate		= 1 << 6,
 };
 
-bool EnvironmentHasMatchingKeyValue(const TMap<FString, FString>& Definitions, const FString& Key, const FString& Value)
-{
-	const FString* MatchedKey = Definitions.Find(Key);
-	if (MatchedKey)
-	{
-		return *MatchedKey == Value;
-	}
-	return false;
-}
-
 uint32_t GetDecalBlendFlags(const FShaderCompilerInput& Input)
 {
 	uint32_t Flags = 0;
 
-	const TMap<FString, FString>& Definitions = Input.Environment.GetDefinitions();
-	
-	if (EnvironmentHasMatchingKeyValue(Definitions, TEXT("DECAL_OUT_MRT0"), TEXT("1")))
+	if (Input.Environment.GetCompileArgument(TEXT("DECAL_OUT_MRT0"), false))
 	{
 		Flags |= EDecalBlendFlags::DecalOut_MRT0;
 	}
-	if (EnvironmentHasMatchingKeyValue(Definitions, TEXT("DECAL_OUT_MRT1"), TEXT("1")))
+	if (Input.Environment.GetCompileArgument(TEXT("DECAL_OUT_MRT1"), false))
 	{
 		Flags |= EDecalBlendFlags::DecalOut_MRT1;
 	}
-	if (EnvironmentHasMatchingKeyValue(Definitions, TEXT("DECAL_OUT_MRT2"), TEXT("1")))
+	if (Input.Environment.GetCompileArgument(TEXT("DECAL_OUT_MRT2"), false))
 	{
 		Flags |= EDecalBlendFlags::DecalOut_MRT2;
 	}
-	if (EnvironmentHasMatchingKeyValue(Definitions, TEXT("DECAL_OUT_MRT3"), TEXT("1")))
+	if (Input.Environment.GetCompileArgument(TEXT("DECAL_OUT_MRT3"), false))
 	{
 		Flags |= EDecalBlendFlags::DecalOut_MRT3;
 	}
@@ -2709,15 +2697,15 @@ uint32_t GetDecalBlendFlags(const FShaderCompilerInput& Input)
 		return 0;
 	}
 
-	if (EnvironmentHasMatchingKeyValue(Definitions, TEXT("MATERIALBLENDING_ALPHACOMPOSITE"), TEXT("1")))
+	if (Input.Environment.GetCompileArgument(TEXT("MATERIALBLENDING_ALPHACOMPOSITE"), false))
 	{
 		Flags |= EDecalBlendFlags::AlphaComposite;
 	}
-	else if (EnvironmentHasMatchingKeyValue(Definitions, TEXT("MATERIALBLENDING_MODULATE"), TEXT("1")))
+	else if (Input.Environment.GetCompileArgument(TEXT("MATERIALBLENDING_MODULATE"), false))
 	{
 		Flags |= EDecalBlendFlags::Modulate;
 	}
-	else if (EnvironmentHasMatchingKeyValue(Definitions, TEXT("MATERIALBLENDING_TRANSLUCENT"), TEXT("1")))
+	else if (Input.Environment.GetCompileArgument(TEXT("MATERIALBLENDING_TRANSLUCENT"), false))
 	{
 		Flags |= EDecalBlendFlags::Translucent;
 	}
@@ -3153,8 +3141,8 @@ static bool CompileToGlslWithShaderConductor(
 			//TargetDesc.CompileFlags.SetDefine(TEXT("force_temporary"), 1);
 
 			// If we have mobile multiview define set then set the view count and enable extension
-			const FString* MultiViewDefine = Input.Environment.GetDefinitions().Find(TEXT("MOBILE_MULTI_VIEW"));
-			if (Frequency == SF_Vertex && MultiViewDefine && *MultiViewDefine == "1")
+			const bool bMultiView = Input.Environment.GetCompileArgument(TEXT("MOBILE_MULTI_VIEW"), false);
+			if (Frequency == SF_Vertex && bMultiView)
 			{
 				TargetDesc.CompileFlags.SetDefine(TEXT("ovr_multiview_view_count"), 2);
 			}
@@ -3229,8 +3217,8 @@ static bool CompileToGlslWithShaderConductor(
 
 		// Handle PLS and FBF in OpenGL
 
-		if (EnvironmentHasMatchingKeyValue(Input.Environment.GetDefinitions(), TEXT("SHADING_PATH_MOBILE"), TEXT("1")) &&
-			EnvironmentHasMatchingKeyValue(Input.Environment.GetDefinitions(), TEXT("MOBILE_DEFERRED_SHADING"), TEXT("1")) &&
+		if (Input.Environment.GetCompileArgument(TEXT("SHADING_PATH_MOBILE"), false) &&
+			Input.Environment.GetCompileArgument(TEXT("MOBILE_DEFERRED_SHADING"), false) &&
 			Version == GLSL_ES3_1_ANDROID)
 		{
 			bCompilationFailed = !GenerateDeferredMobileShaders(GlslSource, GLSLCompileParams, SourceData, ReflectData, true, false, bEmulatedUBs, BlendFlags);
diff --git a/Engine/Source/Developer/VulkanShaderFormat/Private/VulkanShaderCompiler.cpp b/Engine/Source/Developer/VulkanShaderFormat/Private/VulkanShaderCompiler.cpp
--- a/Engine/Source/Developer/VulkanShaderFormat/Private/VulkanShaderCompiler.cpp
+++ b/Engine/Source/Developer/VulkanShaderFormat/Private/VulkanShaderCompiler.cpp
@@ -875,17 +875,8 @@ static void ConvertToNEWHeader(FOLDVulkanCodeHeader& OLDHeader,
 	OutHeader.DebugName = OLDHeader.ShaderName;
 #endif
 	OutHeader.InOutMask = OLDHeader.SerializedBindings.InOutMask.Bitmask;
-
-	OutHeader.RayTracingPayloadType = 0;
-	if (const FString* RTPayloadTypePtr = ShaderInput.Environment.GetDefinitions().Find(TEXT("RT_PAYLOAD_TYPE")))
-	{
-		OutHeader.RayTracingPayloadType = FCString::Atoi(**RTPayloadTypePtr);
-	}
-	OutHeader.RayTracingPayloadSize = 0;
-	if (const FString* RTPayloadSizePtr = ShaderInput.Environment.GetDefinitions().Find(TEXT("RT_PAYLOAD_MAX_SIZE")))
-	{
-		OutHeader.RayTracingPayloadSize = FCString::Atoi(**RTPayloadSizePtr);
-	}
+	OutHeader.RayTracingPayloadType = ShaderInput.Environment.GetCompileArgument(TEXT("RT_PAYLOAD_TYPE"), 0u);
+	OutHeader.RayTracingPayloadSize = ShaderInput.Environment.GetCompileArgument(TEXT("RT_PAYLOAD_MAX_SIZE"), 0u);
 }
 
 
@@ -2556,8 +2547,7 @@ void DoCompileVulkanShader(const FShaderCompilerInput& Input, FShaderCompilerOut
 	// By default we strip reflecion information for Android platform to avoid issues with older drivers
 	if (IsAndroidShaderFormat(Input.ShaderFormat))
 	{
-		const FString* StripReflect_Android = Input.Environment.GetDefinitions().Find(TEXT("STRIP_REFLECT_ANDROID"));
-		bStripReflect = !(StripReflect_Android && *StripReflect_Android == TEXT("0"));
+		bStripReflect = Input.Environment.GetCompileArgument(TEXT("STRIP_REFLECT_ANDROID"), true);
 	}
 
 	const CrossCompiler::FShaderConductorOptions::ETargetEnvironment MinTargetEnvironment = GetMinimumTargetEnvironment(Version);
diff --git a/Engine/Source/Developer/Windows/ShaderFormatD3D/Private/D3DShaderCompilerDXC.cpp b/Engine/Source/Developer/Windows/ShaderFormatD3D/Private/D3DShaderCompilerDXC.cpp
--- a/Engine/Source/Developer/Windows/ShaderFormatD3D/Private/D3DShaderCompilerDXC.cpp
+++ b/Engine/Source/Developer/Windows/ShaderFormatD3D/Private/D3DShaderCompilerDXC.cpp
@@ -1029,12 +1029,10 @@ bool CompileAndProcessD3DShaderDXC(FString& PreprocessedShaderSource,
 			uint32 RayTracingPayloadSize = 0;
 			if (bIsRayTracingShader)
 			{
-				const FString* RTPayloadTypePtr = Input.Environment.GetDefinitions().Find(TEXT("RT_PAYLOAD_TYPE"));
-				const FString* RTPayloadSizePtr = Input.Environment.GetDefinitions().Find(TEXT("RT_PAYLOAD_MAX_SIZE"));
-				checkf(RTPayloadTypePtr != nullptr, TEXT("Ray tracing shaders must provide a payload type as this information is required for offline RTPSO compilation. Check that FShaderType::ModifyCompilationEnvironment correctly set this value."));
-				checkf(RTPayloadSizePtr != nullptr, TEXT("Ray tracing shaders must provide a payload size as this information is required for offline RTPSO compilation. Check that FShaderType::ModifyCompilationEnvironment correctly set this value."));
-				RayTracingPayloadType = FCString::Atoi(**RTPayloadTypePtr);
-				RayTracingPayloadSize = FCString::Atoi(**RTPayloadSizePtr);
+				bool bArgFound = Input.Environment.GetCompileArgument(TEXT("RT_PAYLOAD_TYPE"), RayTracingPayloadType);
+				checkf(bArgFound, TEXT("Ray tracing shaders must provide a payload type as this information is required for offline RTPSO compilation. Check that FShaderType::ModifyCompilationEnvironment correctly set this value."));
+				bArgFound = Input.Environment.GetCompileArgument(TEXT("RT_PAYLOAD_MAX_SIZE"), RayTracingPayloadSize);
+				checkf(bArgFound, TEXT("Ray tracing shaders must provide a payload size as this information is required for offline RTPSO compilation. Check that FShaderType::ModifyCompilationEnvironment correctly set this value."));
 			}
 			auto PostSRTWriterCallback = [&](FMemoryWriter& Ar)
 			{
diff --git a/Engine/Source/Runtime/Engine/Private/Materials/HLSLMaterialTranslator.cpp b/Engine/Source/Runtime/Engine/Private/Materials/HLSLMaterialTranslator.cpp
--- a/Engine/Source/Runtime/Engine/Private/Materials/HLSLMaterialTranslator.cpp
+++ b/Engine/Source/Runtime/Engine/Private/Materials/HLSLMaterialTranslator.cpp
@@ -1991,7 +1991,11 @@ void FHLSLMaterialTranslator::GetMaterialEnvironment(EShaderPlatform InPlatform,
 	OutEnvironment.SetDefine(TEXT("NEEDS_INSTANCE_WORLD_TO_LOCAL_PS"), bUsesInstanceWorldToLocalPS);
 	OutEnvironment.SetDefine(TEXT("USES_TRANSFORM_VECTOR"), bUsesTransformVector);
 	OutEnvironment.SetDefine(TEXT("WANT_PIXEL_DEPTH_OFFSET"), bUsesPixelDepthOffset);
-	OutEnvironment.SetDefine(TEXT("USES_WORLD_POSITION_OFFSET"), bUsesWorldPositionOffset);
+	
+	// we want USES_WORLD_POSITION_OFFSET to be readable as a bool compile argument, hence the != 0 comparison 
+	// (bUsesWorldPositionOffset is actually a 1-bit uint32 bitfield member)
+	OutEnvironment.SetDefineAndCompileArgument(TEXT("USES_WORLD_POSITION_OFFSET"), (bUsesWorldPositionOffset != 0));
+
 	OutEnvironment.SetDefine(TEXT("USES_EMISSIVE_COLOR"), bUsesEmissiveColor);
 	// Distortion uses tangent space transform 
 	OutEnvironment.SetDefine(TEXT("USES_DISTORTION"), Material->IsDistorted());
diff --git a/Engine/Source/Runtime/Engine/Private/Materials/MaterialHLSLEmitter.cpp b/Engine/Source/Runtime/Engine/Private/Materials/MaterialHLSLEmitter.cpp
--- a/Engine/Source/Runtime/Engine/Private/Materials/MaterialHLSLEmitter.cpp
+++ b/Engine/Source/Runtime/Engine/Private/Materials/MaterialHLSLEmitter.cpp
@@ -489,7 +489,7 @@ static void GetMaterialEnvironment(EShaderPlatform InPlatform,
 	OutEnvironment.SetDefine(TEXT("NEEDS_PARTICLE_WORLD_TO_LOCAL"), false);// bUsesParticleWorldToLocal);
 	OutEnvironment.SetDefine(TEXT("USES_TRANSFORM_VECTOR"), false);// bUsesTransformVector);
 	OutEnvironment.SetDefine(TEXT("WANT_PIXEL_DEPTH_OFFSET"), false);// bUsesPixelDepthOffset);
-	OutEnvironment.SetDefine(TEXT("USES_WORLD_POSITION_OFFSET"), false);// bUsesWorldPositionOffset);
+	OutEnvironment.SetDefineAndCompileArgument(TEXT("USES_WORLD_POSITION_OFFSET"), false);// bUsesWorldPositionOffset);
 	OutEnvironment.SetDefine(TEXT("USES_EMISSIVE_COLOR"), false);// bUsesEmissiveColor);
 	// Distortion uses tangent space transform 
 	OutEnvironment.SetDefine(TEXT("USES_DISTORTION"), InMaterial.IsDistorted());
diff --git a/Engine/Source/Runtime/Engine/Private/Materials/MaterialShared.cpp b/Engine/Source/Runtime/Engine/Private/Materials/MaterialShared.cpp
--- a/Engine/Source/Runtime/Engine/Private/Materials/MaterialShared.cpp
+++ b/Engine/Source/Runtime/Engine/Private/Materials/MaterialShared.cpp
@@ -2359,24 +2359,24 @@ void FMaterial::SetupMaterialEnvironment(
 		case BLEND_AlphaComposite:
 		{
 			// Blend mode will reuse MATERIALBLENDING_TRANSLUCENT
-			OutEnvironment.SetDefine(TEXT("MATERIALBLENDING_ALPHACOMPOSITE"), TEXT("1"));
-			OutEnvironment.SetDefine(TEXT("MATERIALBLENDING_TRANSLUCENT"), TEXT("1"));
+			OutEnvironment.SetDefineAndCompileArgument(TEXT("MATERIALBLENDING_ALPHACOMPOSITE"), true);
+			OutEnvironment.SetDefineAndCompileArgument(TEXT("MATERIALBLENDING_TRANSLUCENT"), true);
 			break;
 		}
 		case BLEND_AlphaHoldout:
 		{
 			// Blend mode will reuse MATERIALBLENDING_TRANSLUCENT
 			OutEnvironment.SetDefine(TEXT("MATERIALBLENDING_ALPHAHOLDOUT"), TEXT("1"));
-			OutEnvironment.SetDefine(TEXT("MATERIALBLENDING_TRANSLUCENT"), TEXT("1"));
+			OutEnvironment.SetDefineAndCompileArgument(TEXT("MATERIALBLENDING_TRANSLUCENT"), true);
 			break;
 		}
 		case BLEND_TranslucentColoredTransmittance:
 		case BLEND_Translucent:
-			OutEnvironment.SetDefine(TEXT("MATERIALBLENDING_TRANSLUCENT"),TEXT("1")); 
+			OutEnvironment.SetDefineAndCompileArgument(TEXT("MATERIALBLENDING_TRANSLUCENT"), true); 
 			break;
 
-		case BLEND_Additive: OutEnvironment.SetDefine(TEXT("MATERIALBLENDING_ADDITIVE"),TEXT("1")); break;
-		case BLEND_Modulate: OutEnvironment.SetDefine(TEXT("MATERIALBLENDING_MODULATE"),TEXT("1")); break;
+		case BLEND_Additive: OutEnvironment.SetDefineAndCompileArgument(TEXT("MATERIALBLENDING_ADDITIVE"), true); break;
+		case BLEND_Modulate: OutEnvironment.SetDefineAndCompileArgument(TEXT("MATERIALBLENDING_MODULATE"), true); break;
 
 		default: 
 			UE_LOG(LogMaterial, Warning, TEXT("Unknown material blend mode: %u  Setting to BLEND_Opaque"),(int32)GetBlendMode());
@@ -2417,33 +2417,33 @@ void FMaterial::SetupMaterialEnvironment(
 		}
 		case BLEND_AlphaComposite:
 		{
-			OutEnvironment.SetDefine(TEXT("MATERIALBLENDING_ALPHACOMPOSITE"), TEXT("1"));
-			OutEnvironment.SetDefine(TEXT("MATERIALBLENDING_TRANSLUCENT"), TEXT("1"));
+			OutEnvironment.SetDefineAndCompileArgument(TEXT("MATERIALBLENDING_ALPHACOMPOSITE"), true);
+			OutEnvironment.SetDefineAndCompileArgument(TEXT("MATERIALBLENDING_TRANSLUCENT"), true);
 			OutEnvironment.SetDefine(TEXT("STRATA_BLENDING_TRANSLUCENT_GREYTRANSMITTANCE"), TEXT("1"));
 			break;
 		}
 		case BLEND_TranslucentGreyTransmittance:
 		{
-			OutEnvironment.SetDefine(TEXT("MATERIALBLENDING_TRANSLUCENT"), TEXT("1"));
+			OutEnvironment.SetDefineAndCompileArgument(TEXT("MATERIALBLENDING_TRANSLUCENT"), true);
 			OutEnvironment.SetDefine(TEXT("STRATA_BLENDING_TRANSLUCENT_GREYTRANSMITTANCE"), TEXT("1"));
 			break;
 		}
 		case BLEND_TranslucentColoredTransmittance:
 		{
-			OutEnvironment.SetDefine(TEXT("MATERIALBLENDING_TRANSLUCENT"), TEXT("1"));
+			OutEnvironment.SetDefineAndCompileArgument(TEXT("MATERIALBLENDING_TRANSLUCENT"), true);
 			OutEnvironment.SetDefine(TEXT("STRATA_BLENDING_TRANSLUCENT_COLOREDTRANSMITTANCE"), TEXT("1"));
 			break;
 		}
 		case BLEND_ColoredTransmittanceOnly:
 		{
-			OutEnvironment.SetDefine(TEXT("MATERIALBLENDING_MODULATE"), TEXT("1"));
+			OutEnvironment.SetDefineAndCompileArgument(TEXT("MATERIALBLENDING_MODULATE"), true);
 			OutEnvironment.SetDefine(TEXT("STRATA_BLENDING_COLOREDTRANSMITTANCEONLY"), TEXT("1"));
 			break;
 		}
 		case BLEND_AlphaHoldout:
 		{
+			OutEnvironment.SetDefineAndCompileArgument(TEXT("MATERIALBLENDING_TRANSLUCENT"), true);
 			OutEnvironment.SetDefine(TEXT("MATERIALBLENDING_ALPHAHOLDOUT"), TEXT("1"));
-			OutEnvironment.SetDefine(TEXT("MATERIALBLENDING_TRANSLUCENT"), TEXT("1"));
 			OutEnvironment.SetDefine(TEXT("STRATA_BLENDING_ALPHAHOLDOUT"), TEXT("1"));
 			break;
 		}
@@ -2567,7 +2567,7 @@ void FMaterial::SetupMaterialEnvironment(
 
 	if (GetShadingRate() != MSR_1x1)
 	{
-		OutEnvironment.SetDefine(TEXT("USING_VARIABLE_RATE_SHADING"), TEXT("1"));
+		OutEnvironment.SetCompileArgument(TEXT("USING_VARIABLE_RATE_SHADING"), true);
 	}
 
 	{
diff --git a/Engine/Source/Runtime/Engine/Private/ShaderCompiler/ShaderCompiler.cpp b/Engine/Source/Runtime/Engine/Private/ShaderCompiler/ShaderCompiler.cpp
--- a/Engine/Source/Runtime/Engine/Private/ShaderCompiler/ShaderCompiler.cpp
+++ b/Engine/Source/Runtime/Engine/Private/ShaderCompiler/ShaderCompiler.cpp
@@ -6066,10 +6066,10 @@ void GlobalBeginCompileShader(
 	const bool bUsingMobileRenderer = FSceneInterface::GetShadingPath(GetMaxSupportedFeatureLevel(ShaderPlatform)) == EShadingPath::Mobile;
 	if (bUsingMobileRenderer)
 	{
-		Input.Environment.SetDefine(TEXT("SHADING_PATH_MOBILE"), 1);
+		Input.Environment.SetDefineAndCompileArgument(TEXT("SHADING_PATH_MOBILE"), true);
 		
 		const bool bMobileDeferredShading = IsMobileDeferredShadingEnabled((EShaderPlatform)Target.Platform);
-		Input.Environment.SetDefine(TEXT("MOBILE_DEFERRED_SHADING"), bMobileDeferredShading ? 1 : 0);
+		Input.Environment.SetDefineAndCompileArgument(TEXT("MOBILE_DEFERRED_SHADING"), bMobileDeferredShading);
 
 		if (bMobileDeferredShading)
 		{
@@ -6087,7 +6087,7 @@ void GlobalBeginCompileShader(
 		GConfig->GetBool(TEXT("/Script/AndroidRuntimeSettings.AndroidRuntimeSettings"), TEXT("bStripShaderReflection"), bIsStripReflect, GEngineIni);
 		if (!bIsStripReflect)
 		{
-			Input.Environment.SetDefine(TEXT("STRIP_REFLECT_ANDROID"), false);
+			Input.Environment.SetCompileArgument(TEXT("STRIP_REFLECT_ANDROID"), false);
 		}
 	}
 
@@ -6096,9 +6096,9 @@ void GlobalBeginCompileShader(
 	{
 		const UE::StereoRenderUtils::FStereoShaderAspects Aspects(ShaderPlatform);
 
-		Input.Environment.SetDefine(TEXT("INSTANCED_STEREO"), Aspects.IsInstancedStereoEnabled());
-		Input.Environment.SetDefine(TEXT("MULTI_VIEW"), Aspects.IsInstancedMultiViewportEnabled());
-		Input.Environment.SetDefine(TEXT("MOBILE_MULTI_VIEW"), Aspects.IsMobileMultiViewEnabled());
+		Input.Environment.SetDefineAndCompileArgument(TEXT("INSTANCED_STEREO"), Aspects.IsInstancedStereoEnabled());
+		Input.Environment.SetDefineAndCompileArgument(TEXT("MULTI_VIEW"), Aspects.IsInstancedMultiViewportEnabled());
+		Input.Environment.SetDefineAndCompileArgument(TEXT("MOBILE_MULTI_VIEW"), Aspects.IsMobileMultiViewEnabled());
 
 		// Throw a warning if we are silently disabling ISR due to missing platform support (but don't have MMV enabled).
 		static const auto CVarInstancedStereo = IConsoleManager::Get().FindTConsoleVariableDataInt(TEXT("vr.InstancedStereo"));
@@ -6125,9 +6125,9 @@ void GlobalBeginCompileShader(
 	}
 	else
 	{
-		Input.Environment.SetDefine(TEXT("INSTANCED_STEREO"), 0);
-		Input.Environment.SetDefine(TEXT("MULTI_VIEW"), 0);
-		Input.Environment.SetDefine(TEXT("MOBILE_MULTI_VIEW"), 0);
+		Input.Environment.SetDefineAndCompileArgument(TEXT("INSTANCED_STEREO"), false);
+		Input.Environment.SetDefineAndCompileArgument(TEXT("MULTI_VIEW"), 0);
+		Input.Environment.SetDefineAndCompileArgument(TEXT("MOBILE_MULTI_VIEW"), false);
 	}
 
 	ShaderType->AddUniformBufferIncludesToEnvironment(Input.Environment, ShaderPlatform);
@@ -6299,11 +6299,11 @@ void GlobalBeginCompileShader(
 		
 		{
 			uint32 ShaderVersion = RHIGetMetalShaderLanguageVersion(EShaderPlatform(Target.Platform));
-			Input.Environment.SetDefine(TEXT("SHADER_LANGUAGE_VERSION"), ShaderVersion);
+			Input.Environment.SetCompileArgument(TEXT("SHADER_LANGUAGE_VERSION"), ShaderVersion);
 			
 			bool bAllowFastIntrinsics = false;
 			bool bForceFloats = false;
-			FString IndirectArgumentTier;
+			int32 IndirectArgumentTier = 0;
 			bool bEnableMathOptimisations = true;
             bool bSupportAppleA8 = false;
             
@@ -6311,7 +6311,7 @@ void GlobalBeginCompileShader(
 			{
 				GConfig->GetBool(TEXT("/Script/MacTargetPlatform.MacTargetSettings"), TEXT("UseFastIntrinsics"), bAllowFastIntrinsics, GEngineIni);
 				GConfig->GetBool(TEXT("/Script/MacTargetPlatform.MacTargetSettings"), TEXT("EnableMathOptimisations"), bEnableMathOptimisations, GEngineIni);
-				GConfig->GetString(TEXT("/Script/MacTargetPlatform.MacTargetSettings"), TEXT("IndirectArgumentTier"), IndirectArgumentTier, GEngineIni);
+				GConfig->GetInt(TEXT("/Script/MacTargetPlatform.MacTargetSettings"), TEXT("IndirectArgumentTier"), IndirectArgumentTier, GEngineIni);
                 
                 // No half precision support on MacOS at the moment
                 bForceFloats = true;
@@ -6321,7 +6321,7 @@ void GlobalBeginCompileShader(
 				GConfig->GetBool(TEXT("/Script/IOSRuntimeSettings.IOSRuntimeSettings"), TEXT("UseFastIntrinsics"), bAllowFastIntrinsics, GEngineIni);
 				GConfig->GetBool(TEXT("/Script/IOSRuntimeSettings.IOSRuntimeSettings"), TEXT("EnableMathOptimisations"), bEnableMathOptimisations, GEngineIni);
 				GConfig->GetBool(TEXT("/Script/IOSRuntimeSettings.IOSRuntimeSettings"), TEXT("ForceFloats"), bForceFloats, GEngineIni);
-				GConfig->GetString(TEXT("/Script/IOSRuntimeSettings.IOSRuntimeSettings"), TEXT("IndirectArgumentTier"), IndirectArgumentTier, GEngineIni);
+				GConfig->GetInt(TEXT("/Script/IOSRuntimeSettings.IOSRuntimeSettings"), TEXT("IndirectArgumentTier"), IndirectArgumentTier, GEngineIni);
                 GConfig->GetBool(TEXT("/Script/IOSRuntimeSettings.IOSRuntimeSettings"), TEXT("bSupportAppleA8"), bSupportAppleA8, GEngineIni);
                 
 				// Force no development shaders on iOS
@@ -6330,9 +6330,9 @@ void GlobalBeginCompileShader(
             
             Input.Environment.FullPrecisionInPS |= bForceFloats;
             
-			Input.Environment.SetDefine(TEXT("METAL_USE_FAST_INTRINSICS"), bAllowFastIntrinsics);
-			Input.Environment.SetDefine(TEXT("METAL_INDIRECT_ARGUMENT_BUFFERS"), IndirectArgumentTier);
-            Input.Environment.SetDefine(TEXT("SUPPORT_APPLE_A8"), bSupportAppleA8);
+			Input.Environment.SetCompileArgument(TEXT("METAL_USE_FAST_INTRINSICS"), bAllowFastIntrinsics);
+			Input.Environment.SetCompileArgument(TEXT("METAL_INDIRECT_ARGUMENT_BUFFERS"), IndirectArgumentTier);
+            Input.Environment.SetCompileArgument(TEXT("SUPPORT_APPLE_A8"), bSupportAppleA8);
 			
 			// Same as console-variable above, but that's global and this is per-platform, per-project
 			if (!bEnableMathOptimisations)
diff --git a/Engine/Source/Runtime/RenderCore/Private/Shader.cpp b/Engine/Source/Runtime/RenderCore/Private/Shader.cpp
--- a/Engine/Source/Runtime/RenderCore/Private/Shader.cpp
+++ b/Engine/Source/Runtime/RenderCore/Private/Shader.cpp
@@ -448,8 +448,8 @@ void FShaderType::ModifyCompilationEnvironment(const FShaderPermutationParameter
 	{
 		checkf(IsRayTracingPayloadRegistered(RayTracingPayloadType), TEXT("Raytracing shader %s is using a payload type (%u) which was never registered"), Name, RayTracingPayloadType);
 
-		OutEnvironment.SetDefine(TEXT("RT_PAYLOAD_TYPE"), static_cast<int32>(RayTracingPayloadType));
-		OutEnvironment.SetDefine(TEXT("RT_PAYLOAD_MAX_SIZE"), GetRayTracingPayloadTypeMaxSize(RayTracingPayloadType));
+		OutEnvironment.SetDefineAndCompileArgument(TEXT("RT_PAYLOAD_TYPE"), static_cast<uint32>(RayTracingPayloadType));
+		OutEnvironment.SetDefineAndCompileArgument(TEXT("RT_PAYLOAD_MAX_SIZE"), GetRayTracingPayloadTypeMaxSize(RayTracingPayloadType));
 
 		if (   (uint32(RayTracingPayloadType) & uint32(ERayTracingPayloadType::RayTracingMaterial))
 			|| (uint32(RayTracingPayloadType) & uint32(ERayTracingPayloadType::GPULightmass))
diff --git a/Engine/Source/Runtime/RenderCore/Public/ShaderCore.h b/Engine/Source/Runtime/RenderCore/Public/ShaderCore.h
--- a/Engine/Source/Runtime/RenderCore/Public/ShaderCore.h
+++ b/Engine/Source/Runtime/RenderCore/Public/ShaderCore.h
@@ -23,6 +23,7 @@
 #include "Misc/Optional.h"
 #include "Misc/Paths.h"
 #include "Misc/SecureHash.h"
+#include "Misc/TVariant.h"
 #include "PixelFormat.h"
 #include "RHIDefinitions.h"
 #include "Serialization/Archive.h"
@@ -567,6 +568,48 @@ struct FShaderCompilerEnvironment
 		}
 	}
 
+
+	// Sets a generic parameter which can be read in the various shader format backends to modify compilation
+	// behaviour. Intended to replace any usage of definitions after shader preprocessing.
+	template <typename ValueType> void SetCompileArgument(const TCHAR* Name, ValueType Value)
+	{
+		CompileArgs.Add(Name, TVariant<bool, float, int32, uint32>(TInPlaceType<ValueType>(), Value));
+	}
+
+	// Helper to set both a define and a compile argument to the same value. Useful for various parameters which
+	// need to be consumed both by preprocessing and in the shader format backends to modify compilation behaviour.
+	template <typename ValueType> void SetDefineAndCompileArgument(const TCHAR* Name, ValueType Value)
+	{
+		SetDefine(Name, Value);
+		SetCompileArgument(Name, Value);
+	}
+
+	// If a compile argument with the given name exists and is of the specified type, returns its value. Otherwise, 
+	// either the named argument doesn't exist or the type does not match, and the default value will be returned.
+	template <typename ValueType> ValueType GetCompileArgument(const TCHAR* Name, const ValueType& DefaultValue) const
+	{
+		const TVariant<bool, float, int32, uint32>* StoredValue = CompileArgs.Find(Name);
+		if (StoredValue && StoredValue->IsType<ValueType>())
+		{
+			return StoredValue->Get<ValueType>();
+		}
+		return DefaultValue;
+	}
+
+	// If a compile argument with the given name exists and is of the specified type, its value will be assigned to OutValue
+	// and the function will return true. Otherwise, either the named argument doesn't exist or the type does not match, the
+	// OutValue will be left unmodified and the function will return false.
+	template <typename ValueType> bool GetCompileArgument(const TCHAR* Name, ValueType& OutValue) const
+	{
+		const TVariant<bool, float, int32, uint32>* StoredValue = CompileArgs.Find(Name);
+		if (StoredValue && StoredValue->IsType<ValueType>())
+		{
+			OutValue = StoredValue->Get<ValueType>();
+			return true;
+		}
+		return false;
+	}
+
 	const TMap<FString,FString>& GetDefinitions() const
 	{
 		return Definitions.GetDefinitionMap();
@@ -581,6 +624,7 @@ struct FShaderCompilerEnvironment
 	inline void SerializeEverythingButFiles(FArchive& Ar)
 	{
 		Ar << Definitions;
+		Ar << CompileArgs;
 		Ar << CompilerFlags;
 		Ar << RenderTargetOutputFormatsMap;
 		Ar << ResourceTableMap;
@@ -624,6 +668,7 @@ struct FShaderCompilerEnvironment
 		ResourceTableMap.Append(Other.ResourceTableMap);
 		UniformBufferMap.Append(Other.UniformBufferMap);
 		Definitions.Merge(Other.Definitions);
+		CompileArgs.Append(Other.CompileArgs);
 		RenderTargetOutputFormatsMap.Append(Other.RenderTargetOutputFormatsMap);
 		RemoteServerData.Append(Other.RemoteServerData);
 		FullPrecisionInPS |= Other.FullPrecisionInPS;
@@ -632,6 +677,7 @@ struct FShaderCompilerEnvironment
 private:
 
 	FShaderCompilerDefinitions Definitions;
+	TMap<FString, TVariant<bool, float, int32, uint32>> CompileArgs;
 };
 
 struct FSharedShaderCompilerEnvironment final : public FShaderCompilerEnvironment, public FRefCountBase
diff --git a/Engine/Source/Runtime/Renderer/Private/DecalRenderingCommon.cpp b/Engine/Source/Runtime/Renderer/Private/DecalRenderingCommon.cpp
--- a/Engine/Source/Runtime/Renderer/Private/DecalRenderingCommon.cpp
+++ b/Engine/Source/Runtime/Renderer/Private/DecalRenderingCommon.cpp
@@ -900,10 +900,10 @@ namespace DecalRendering
 		OutEnvironment.SetDefine(TEXT("DECAL_RENDERTARGETMODE"), (uint32)RenderTargetMode);
 		OutEnvironment.SetDefine(TEXT("DECAL_RENDERTARGET_COUNT"), RenderTargetCount);
 
-		OutEnvironment.SetDefine(TEXT("DECAL_OUT_MRT0"), (RenderTargetWriteMask & 1) != 0 ? 1 : 0);
-		OutEnvironment.SetDefine(TEXT("DECAL_OUT_MRT1"), (RenderTargetWriteMask & 2) != 0 ? 1 : 0);
-		OutEnvironment.SetDefine(TEXT("DECAL_OUT_MRT2"), (RenderTargetWriteMask & 4) != 0 ? 1 : 0);
-		OutEnvironment.SetDefine(TEXT("DECAL_OUT_MRT3"), (RenderTargetWriteMask & 8) != 0 ? 1 : 0);
+		OutEnvironment.SetDefineAndCompileArgument(TEXT("DECAL_OUT_MRT0"), (RenderTargetWriteMask & 1) != 0 ? true : false);
+		OutEnvironment.SetDefineAndCompileArgument(TEXT("DECAL_OUT_MRT1"), (RenderTargetWriteMask & 2) != 0 ? true : false);
+		OutEnvironment.SetDefineAndCompileArgument(TEXT("DECAL_OUT_MRT2"), (RenderTargetWriteMask & 4) != 0 ? true : false);
+		OutEnvironment.SetDefineAndCompileArgument(TEXT("DECAL_OUT_MRT3"), (RenderTargetWriteMask & 8) != 0 ? true : false);
 
 		OutEnvironment.SetDefine(TEXT("DECAL_RENDERSTAGE_BEFOREBASEPASS"), (uint32)EDecalRenderStage::BeforeBasePass);
 		OutEnvironment.SetDefine(TEXT("DECAL_RENDERSTAGE_BEFORELIGHTING"), (uint32)EDecalRenderStage::BeforeLighting);
diff --git a/Engine/Source/Runtime/Renderer/Private/MobileBasePassRendering.h b/Engine/Source/Runtime/Renderer/Private/MobileBasePassRendering.h
--- a/Engine/Source/Runtime/Renderer/Private/MobileBasePassRendering.h
+++ b/Engine/Source/Runtime/Renderer/Private/MobileBasePassRendering.h
@@ -142,7 +142,7 @@ public:
 		// @lh-todo:
 		//		This is a workaround to prevent Z-fighting due to a mismatch between HLSLCC compiled shaders (from cooked materials) and DXC compiled shaders (from user materials).
 		//		It effectively disables CVar "r.Shaders.ForceDXC" which we partially enabled during the HLSLCC to DXC transition for certain platform backends.
-		OutEnvironment.SetDefine(TEXT("WORKAROUND_DISABLE_rShadersForceDXC"), 1);
+		OutEnvironment.SetCompileArgument(TEXT("WORKAROUND_DISABLE_rShadersForceDXC"), true);
 	}
 
 	void GetShaderBindings(
@@ -237,7 +237,7 @@ public:
 		ModifyCompilationEnvironmentForQualityLevel(Parameters.Platform, Parameters.MaterialParameters.QualityLevel, OutEnvironment);
 
 		// @lh-todo: Same workaround as for the VS of MobileBasePass. See TMobileBasePassVSPolicyParamType::ModifyCompilationEnvironment for details.
-		OutEnvironment.SetDefine(TEXT("WORKAROUND_DISABLE_rShadersForceDXC"), 1);
+		OutEnvironment.SetCompileArgument(TEXT("WORKAROUND_DISABLE_rShadersForceDXC"), true);
 	}
 
 	/** Initialization constructor. */
