Build fix for 24854056.

diff --git a/Engine/Source/Runtime/Engine/Private/Materials/MaterialInterface.cpp b/Engine/Source/Runtime/Engine/Private/Materials/MaterialInterface.cpp
--- a/Engine/Source/Runtime/Engine/Private/Materials/MaterialInterface.cpp
+++ b/Engine/Source/Runtime/Engine/Private/Materials/MaterialInterface.cpp
@@ -107,6 +107,8 @@ static EMaterialGetParameterValueFlags MakeParameterValueFlags(bool bOveriddenOn
 	return Result;
 }
 
+#if WITH_EDITORONLY_DATA
+
 /** Deletes any invalid EditorOnlyData object with specified name that isn't the one assigned to specified interface. */
 static void DisposeInvalidEditorOnlyData(UMaterialInterface* MaterialInterface, const FString& EditorOnlyDataName)
 {
@@ -119,6 +121,8 @@ static void DisposeInvalidEditorOnlyData(UMaterialInterface* MaterialInterface,
 	}
 }
 
+#endif
+
 //////////////////////////////////////////////////////////////////////////
 
 /** Copies the material's relevance flags to a primitive's view relevance flags. */
