Animation Attribute curves are not updated/reflecting reimported data
#fix make ::AddTypedCustomAttribute set the keys if the attribute _already_ exists

diff --git a/Engine/Source/Runtime/Engine/Classes/Animation/BuiltInAttributeTypes.h b/Engine/Source/Runtime/Engine/Classes/Animation/BuiltInAttributeTypes.h
--- a/Engine/Source/Runtime/Engine/Classes/Animation/BuiltInAttributeTypes.h
+++ b/Engine/Source/Runtime/Engine/Classes/Animation/BuiltInAttributeTypes.h
@@ -380,7 +380,7 @@ namespace UE
 			const FAnimationAttributeIdentifier Identifier = UAnimationAttributeIdentifierExtensions::CreateAttributeIdentifier(AnimSequenceBase, AttributeName, BoneName, AttributeType::StaticStruct());
 
 			IAnimationDataController& Controller = AnimSequenceBase->GetController();
-			if (Controller.AddAttribute(Identifier))
+			if (AnimSequenceBase->GetDataModelInterface()->FindAttribute(Identifier) || Controller.AddAttribute(Identifier))
 			{
 				TArray<AttributeType> AttributeValues; 
 				Algo::Transform(Values, AttributeValues, [](const ValueType& Value)
