Build fix related to 24854056.

diff --git a/Engine/Source/Runtime/Engine/Private/Materials/MaterialInterface.cpp b/Engine/Source/Runtime/Engine/Private/Materials/MaterialInterface.cpp
--- a/Engine/Source/Runtime/Engine/Private/Materials/MaterialInterface.cpp
+++ b/Engine/Source/Runtime/Engine/Private/Materials/MaterialInterface.cpp
@@ -37,6 +37,7 @@
 #include "MaterialCachedData.h"
 #include "Components/DecalComponent.h"
 #include "UObject/ArchiveCookContext.h"
+#include "UObject/Package.h"
 
 #if WITH_EDITOR
 #include "ObjectCacheEventSink.h"
