Fix the seamless travel handler to handle PIE renaming for the transition map the same way it does the destination map, this fixes a crash with multiplayer PIE where it tries to use the same map package

diff --git a/Engine/Source/Runtime/Engine/Classes/Engine/World.h b/Engine/Source/Runtime/Engine/Classes/Engine/World.h
--- a/Engine/Source/Runtime/Engine/Classes/Engine/World.h
+++ b/Engine/Source/Runtime/Engine/Classes/Engine/World.h
@@ -239,6 +239,9 @@ private:
 
 	void SetHandlerLoadedData(UObject* InLevelPackage, UWorld* InLoadedWorld);
 
+	/** Wrapper for loading transition or destination map, returns false if not found */
+	bool StartLoadingMap(FString MapPackageToLoadFrom);
+
 	/** called to kick off async loading of the destination map and any other packages it requires */
 	void StartLoadingDestination();
 
diff --git a/Engine/Source/Runtime/Engine/Private/World.cpp b/Engine/Source/Runtime/Engine/Private/World.cpp
--- a/Engine/Source/Runtime/Engine/Private/World.cpp
+++ b/Engine/Source/Runtime/Engine/Private/World.cpp
@@ -7051,23 +7051,11 @@ bool FSeamlessTravelHandler::StartTravel(UWorld* InCurrentWorld, const FURL& InU
 				}
 				else
 				{
-					// Set the world type in the static map, so that UWorld::PostLoad can set the world type
-					UWorld::WorldTypePreLoadMap.FindOrAdd(*TransitionMap) = CurrentWorld->WorldType;
-
-					// first, load the entry level package
+					// Load the transition map, possibly with PIE prefix
 					STAT_ADD_CUSTOMMESSAGE_NAME( STAT_NamedMarker, *(FString( TEXT( "StartTravel - " ) + TransitionMap )) );
 					TRACE_BOOKMARK(TEXT("StartTravel - %s"), *TransitionMap);
-					FPackagePath PackagePath;
-					if (FPackagePath::TryFromMountedName(TransitionMap, PackagePath))
-					{
-						LoadPackageAsync(PackagePath, 
-							FName(),
-							FLoadPackageAsyncDelegate::CreateRaw(this, &FSeamlessTravelHandler::SeamlessTravelLoadCallback),
-							(CurrentWorld->WorldType == EWorldType::PIE ? PKG_PlayInEditor : PKG_None),
-							Context.PIEInstance
-							);
-					}
-					else
+
+					if (!StartLoadingMap(TransitionMap))
 					{
 						UE_LOG(LogWorld, Error, TEXT("StartTravel: Invalid TransitionMap \"%s\""), *TransitionMap);
 					}
@@ -7144,59 +7132,68 @@ void FSeamlessTravelHandler::SetPauseAtMidpoint(bool bNowPaused)
 	}
 }
 
-void FSeamlessTravelHandler::StartLoadingDestination()
+bool FSeamlessTravelHandler::StartLoadingMap(FString MapPackageToLoadFrom)
 {
-	if (bTransitionInProgress && bSwitchedToDefaultMap)
+	if (MapPackageToLoadFrom.IsEmpty())
 	{
-		UE_LOG(LogWorld, Log, TEXT("StartLoadingDestination to: %s"), *PendingTravelURL.Map);
-
-		CurrentWorld->GetGameInstance()->PreloadContentForURL(PendingTravelURL);
+		return false;
+	}
 
-		const FName URLMapFName = FName(*PendingTravelURL.Map);
+	// In PIE we might want to mangle MapPackageName when traveling to a map loaded in the editor
+	FString MapPackageName = MapPackageToLoadFrom;
+	EPackageFlags PackageFlags = PKG_None;
+	int32 PIEInstanceID = INDEX_NONE;
 
-		// In PIE we might want to mangle MapPackageName when traveling to a map loaded in the editor
-		FString URLMapPackageName = PendingTravelURL.Map;
-		FString URLMapPackageToLoadFrom = PendingTravelURL.Map;
-		EPackageFlags PackageFlags = PKG_None;
-		int32 PIEInstanceID = INDEX_NONE;
-		
 #if WITH_EDITOR
-		if (GIsEditor)
-		{
-			FWorldContext &WorldContext = GEngine->GetWorldContextFromHandleChecked(WorldContextHandle);
+	if (GIsEditor)
+	{
+		FWorldContext& WorldContext = GEngine->GetWorldContextFromHandleChecked(WorldContextHandle);
 
-			PIEInstanceID = WorldContext.PIEInstance;
-			URLMapPackageName = UWorld::ConvertToPIEPackageName(URLMapPackageName, PIEInstanceID);
+		PIEInstanceID = WorldContext.PIEInstance;
+		MapPackageName = UWorld::ConvertToPIEPackageName(MapPackageName, PIEInstanceID);
 
-			if (WorldContext.WorldType == EWorldType::PIE)
-			{
-				PackageFlags |= PKG_PlayInEditor;
+		if (WorldContext.WorldType == EWorldType::PIE)
+		{
+			PackageFlags |= PKG_PlayInEditor;
 
-				// Prepare soft object paths for fixup
-				FSoftObjectPath::AddPIEPackageName(FName(*URLMapPackageName));
-			}
+			// Prepare soft object paths for fixup
+			FSoftObjectPath::AddPIEPackageName(FName(*MapPackageName));
 		}
+	}
 #endif
 
-		// Set the world type in the static map, so that UWorld::PostLoad can set the world type
-		FName PackageName(*URLMapPackageName);
-		UWorld::WorldTypePreLoadMap.FindOrAdd(PackageName) = CurrentWorld->WorldType;
+	// Set the world type in the static map, so that UWorld::PostLoad can set the world type
+	FName PackageFName(*MapPackageName);
+	UWorld::WorldTypePreLoadMap.FindOrAdd(PackageFName) = CurrentWorld->WorldType;
 
-		FStringView PackageNameToLoad(!URLMapPackageToLoadFrom.IsEmpty() ? FStringView(URLMapPackageToLoadFrom) : FStringView(URLMapPackageName));
-		FPackagePath PackagePath;
-		if (FPackagePath::TryFromMountedName(PackageNameToLoad, PackagePath))
-		{
-			LoadPackageAsync(
-				PackagePath,
-				PackageName,
-				FLoadPackageAsyncDelegate::CreateRaw(this, &FSeamlessTravelHandler::SeamlessTravelLoadCallback),
-				PackageFlags,
-				PIEInstanceID
-			);
-		}
-		else
+	FPackagePath PackagePath;
+	if (FPackagePath::TryFromMountedName(MapPackageToLoadFrom, PackagePath))
+	{
+		LoadPackageAsync(
+			PackagePath,
+			PackageFName,
+			FLoadPackageAsyncDelegate::CreateRaw(this, &FSeamlessTravelHandler::SeamlessTravelLoadCallback),
+			PackageFlags,
+			PIEInstanceID
+		);
+
+		return true;
+	}
+
+	return false;
+}
+
+void FSeamlessTravelHandler::StartLoadingDestination()
+{
+	if (bTransitionInProgress && bSwitchedToDefaultMap)
+	{
+		UE_LOG(LogWorld, Log, TEXT("StartLoadingDestination to: %s"), *PendingTravelURL.Map);
+
+		CurrentWorld->GetGameInstance()->PreloadContentForURL(PendingTravelURL);
+
+		if (!StartLoadingMap(PendingTravelURL.Map))
 		{
-			UE_LOG(LogWorld, Error, TEXT("StartLoadingDestination: Invalid PackagePath \"%.*s\""), PackageNameToLoad.Len(), PackageNameToLoad.GetData());
+			UE_LOG(LogWorld, Error, TEXT("StartLoadingDestination: Invalid destination map \"%s\""), *PendingTravelURL.Map);
 		}
 	}
 	else
