Engine session summary: removed any parameters from disallowed exec commands telemetry.

diff --git a/Engine/Source/Runtime/Engine/Private/Analytics/EngineAnalyticsSessionSummary.cpp b/Engine/Source/Runtime/Engine/Private/Analytics/EngineAnalyticsSessionSummary.cpp
--- a/Engine/Source/Runtime/Engine/Private/Analytics/EngineAnalyticsSessionSummary.cpp
+++ b/Engine/Source/Runtime/Engine/Private/Analytics/EngineAnalyticsSessionSummary.cpp
@@ -2,6 +2,7 @@
 
 #include "Analytics/EngineAnalyticsSessionSummary.h"
 #include "AnalyticsSessionSummaryManager.h"
+#include "Containers/StringView.h"
 #include "Engine/EngineTypes.h"
 #include "GeneralProjectSettings.h"
 #include "Engine/Engine.h"
@@ -121,7 +122,7 @@ namespace EngineAnalyticsProperties
 		return true; // Shutdown|Terminate|Abnormal can overwrite each other.
 	}
 
-	bool UpdateCommandValue(const TCHAR* Cmd, FString& Value)
+	bool UpdateCommandValue(const FStringView Cmd, FString& Value)
 	{
 		constexpr int32 MaxAllowedStringSize = 128;
 		if (Value.Len() >= MaxAllowedStringSize || Value.Contains(Cmd))
@@ -329,10 +330,19 @@ void FEngineAnalyticsSessionSummary::OnVanillaStateChanged(bool bIsVanilla)
 
 void FEngineAnalyticsSessionSummary::OnDisallowedExecCommandCalled(const TCHAR* Cmd)
 {
+	FStringView CmdWithoutParams{Cmd};
+
+	// Don't send any command parameters as they could contain sensitive data
+	int32 SpacePos = INDEX_NONE;
+	if (CmdWithoutParams.FindChar(TEXT(' '), SpacePos) || CmdWithoutParams.FindChar(TEXT('\t'), SpacePos))
+	{
+		CmdWithoutParams.LeftInline(SpacePos);
+	}
+	
 	EngineAnalyticsProperties::CalledDisallowedExecCommands.Update(GetStore(),
-		[Cmd](FString& Value)
+		[CmdWithoutParams](FString& Value)
 		{
-			return EngineAnalyticsProperties::UpdateCommandValue(Cmd, Value);
+			return EngineAnalyticsProperties::UpdateCommandValue(CmdWithoutParams, Value);
 		});
 }
 
