Fixes a bug where FDefaultDynamicResolutionStateProxy::BeginFrame() was being called even when GRHISupportsDynamicResolution=false

diff --git a/Engine/Source/Runtime/Engine/Private/DynamicResolution.cpp b/Engine/Source/Runtime/Engine/Private/DynamicResolution.cpp
--- a/Engine/Source/Runtime/Engine/Private/DynamicResolution.cpp
+++ b/Engine/Source/Runtime/Engine/Private/DynamicResolution.cpp
@@ -1141,7 +1141,7 @@ public:
 		if (Event == EDynamicResolutionStateEvent::BeginFrame)
 		{
 			check(bRecordThisFrame == false);
-			bRecordThisFrame = bIsEnabled;
+			bRecordThisFrame = bIsEnabled && IsSupported();
 		}
 
 		// Early return if not recording this frame.
