World Partition - Level Streaming: UWorldPartitionLevelStreamingDynamic now stores an hard ref to the world partition, instead of a weak ref.
* It was possible for the streaming levels to outlive the WP object

diff --git a/Engine/Source/Runtime/Engine/Private/WorldPartition/WorldPartitionLevelStreamingDynamic.cpp b/Engine/Source/Runtime/Engine/Private/WorldPartition/WorldPartitionLevelStreamingDynamic.cpp
--- a/Engine/Source/Runtime/Engine/Private/WorldPartition/WorldPartitionLevelStreamingDynamic.cpp
+++ b/Engine/Source/Runtime/Engine/Private/WorldPartition/WorldPartitionLevelStreamingDynamic.cpp
@@ -631,7 +631,7 @@ UWorld* UWorldPartitionLevelStreamingDynamic::GetStreamingWorld() const
 {
 	// For UWorldPartitionLevelStreamingDynamic the StreamingWorld is the world to which the OuterWorldPartition is outered.
 	// This World can be used to resolved SoftObjectPaths between cells.
-	check(OuterWorldPartition.IsValid());
+	check(OuterWorldPartition);
 	return OuterWorldPartition->GetTypedOuter<UWorld>();
 }
 
@@ -687,7 +687,7 @@ bool UWorldPartitionLevelStreamingDynamic::ShouldBlockOnUnload() const
 
 	// When world partition cannot stream (anymore), return true so that RemoveFromWorld of this level is not incremental. 
 	// This guarantees that unloaded instanced wp levels fully unload their cell levels.
-	if (OuterWorldPartition.IsValid() && !OuterWorldPartition->CanStream())
+	if (OuterWorldPartition && !OuterWorldPartition->CanStream())
 	{
 		return true;
 	}
@@ -698,7 +698,7 @@ bool UWorldPartitionLevelStreamingDynamic::ShouldBlockOnUnload() const
 void UWorldPartitionLevelStreamingDynamic::UpdateShouldSkipMakingVisibilityTransactionRequest()
 {
 	const UWorldPartitionRuntimeCell* Cell = GetWorldPartitionRuntimeCell();
-	if (ensure(Cell && OuterWorldPartition.IsValid()))
+	if (ensure(Cell && OuterWorldPartition))
 	{
 		// It is safe to skip client MakingVisibility transaction requests for cells without data layers when world partition server streaming is disabled
 		bSkipClientUseMakingVisibleTransactionRequest = !(OuterWorldPartition->IsServerStreamingEnabled() || Cell->HasDataLayers());
diff --git a/Engine/Source/Runtime/Engine/Public/WorldPartition/WorldPartitionLevelStreamingDynamic.h b/Engine/Source/Runtime/Engine/Public/WorldPartition/WorldPartitionLevelStreamingDynamic.h
--- a/Engine/Source/Runtime/Engine/Public/WorldPartition/WorldPartitionLevelStreamingDynamic.h
+++ b/Engine/Source/Runtime/Engine/Public/WorldPartition/WorldPartitionLevelStreamingDynamic.h
@@ -97,5 +97,5 @@ private:
 	TWeakObjectPtr<const UWorldPartitionRuntimeLevelStreamingCell> StreamingCell;
 
 	UPROPERTY()
-	TWeakObjectPtr<UWorldPartition> OuterWorldPartition;
+	TObjectPtr<UWorldPartition> OuterWorldPartition;
 };
\ No newline at end of file
