World Partition - HLOD: Avoid crash when unregistering HLOD actors from a world that is being streamed out.

diff --git a/Engine/Source/Runtime/Engine/Private/WorldPartition/HLOD/HLODSubsystem.cpp b/Engine/Source/Runtime/Engine/Private/WorldPartition/HLOD/HLODSubsystem.cpp
--- a/Engine/Source/Runtime/Engine/Private/WorldPartition/HLOD/HLODSubsystem.cpp
+++ b/Engine/Source/Runtime/Engine/Private/WorldPartition/HLOD/HLODSubsystem.cpp
@@ -401,12 +401,14 @@ const UHLODSubsystem::FCellData* UHLODSubsystem::GetCellData(const UWorldPartiti
 UHLODSubsystem::FCellData* UHLODSubsystem::GetCellData(const UWorldPartitionRuntimeCell* InCell)
 {
 	const UWorldPartition* WorldPartition = InCell->GetOuterWorld()->GetWorldPartition();
-	check(WorldPartition);
-	FWorldPartitionHLODRuntimeData* WorldPartitionHLODRuntimeData = WorldPartitionsHLODRuntimeData.Find(WorldPartition);
-	if (WorldPartitionHLODRuntimeData)
+	if (WorldPartition)
 	{
-		check(WorldPartition->IsStreamingEnabled());
-		return WorldPartitionHLODRuntimeData->CellsData.Find(InCell->GetGuid());
+		FWorldPartitionHLODRuntimeData* WorldPartitionHLODRuntimeData = WorldPartitionsHLODRuntimeData.Find(WorldPartition);
+		if (WorldPartitionHLODRuntimeData)
+		{
+			check(WorldPartition->IsStreamingEnabled());
+			return WorldPartitionHLODRuntimeData->CellsData.Find(InCell->GetGuid());
+		}
 	}
 
 	return nullptr;
@@ -415,13 +417,15 @@ UHLODSubsystem::FCellData* UHLODSubsystem::GetCellData(const UWorldPartitionRunt
 UHLODSubsystem::FCellData* UHLODSubsystem::GetCellData(AWorldPartitionHLOD* InWorldPartitionHLOD)
 {
 	const UWorldPartition* WorldPartition = FHLODSubsystem::GetWorldPartition(InWorldPartitionHLOD);
-	check(WorldPartition);
-	FWorldPartitionHLODRuntimeData* WorldPartitionHLODRuntimeData = WorldPartitionsHLODRuntimeData.Find(WorldPartition);
-	if (WorldPartitionHLODRuntimeData)
+	if (WorldPartition)
 	{
-		check(WorldPartition->IsStreamingEnabled());
-		const FGuid CellGuid = InWorldPartitionHLOD->GetSourceCellGuid();
-		return WorldPartitionHLODRuntimeData->CellsData.Find(CellGuid);
+		FWorldPartitionHLODRuntimeData* WorldPartitionHLODRuntimeData = WorldPartitionsHLODRuntimeData.Find(WorldPartition);
+		if (WorldPartitionHLODRuntimeData)
+		{
+			check(WorldPartition->IsStreamingEnabled());
+			const FGuid CellGuid = InWorldPartitionHLOD->GetSourceCellGuid();
+			return WorldPartitionHLODRuntimeData->CellsData.Find(CellGuid);
+		}
 	}
 
 	return nullptr;		
