World Partition Class Descriptors: don't create class descriptors for Blueprints that don't have an existing class descriptor when saving actors.

diff --git a/Engine/Source/Runtime/Engine/Private/WorldPartition/WorldPartitionClassDescRegistry.cpp b/Engine/Source/Runtime/Engine/Private/WorldPartition/WorldPartitionClassDescRegistry.cpp
--- a/Engine/Source/Runtime/Engine/Private/WorldPartition/WorldPartitionClassDescRegistry.cpp
+++ b/Engine/Source/Runtime/Engine/Private/WorldPartition/WorldPartitionClassDescRegistry.cpp
@@ -144,14 +144,14 @@ void FWorldPartitionClassDescRegistry::PrefetchClassDescs(const TArray<FTopLevel
 	// the shortest name.
 	auto SortBlueprintAssetsByNameLength = [](const FAssetData& Lhs, const FAssetData& Rhs) { return Lhs.AssetName.GetStringLength() < Rhs.AssetName.GetStringLength(); };
 
-	auto GetBlueprintAssets = [&AssetRegistry](const TArray<FString>& FilePaths, TArray<FAssetData>& Assets, bool bIncludeOnlyOnDiskAssets = false)
+	auto GetBlueprintAssets = [&AssetRegistry](const TArray<FString>& FilePaths, TArray<FAssetData>& Assets)
 	{
 		FARFilter Filter;
 		Filter.ClassPaths.Add(UBlueprint::StaticClass()->GetClassPathName());
 		Filter.ClassPaths.Add(UBlueprintGeneratedClass::StaticClass()->GetClassPathName());
 		Filter.ClassPaths.Add(UObjectRedirector::StaticClass()->GetClassPathName());
 		Filter.bRecursiveClasses = true;
-		Filter.bIncludeOnlyOnDiskAssets = bIncludeOnlyOnDiskAssets;
+		Filter.bIncludeOnlyOnDiskAssets = true;
 		Filter.PackageNames.Reserve(FilePaths.Num());
 		Algo::Transform(FilePaths, Filter.PackageNames, [](const FString& ClassPath) { return *ClassPath; });
 	
@@ -200,16 +200,8 @@ void FWorldPartitionClassDescRegistry::PrefetchClassDescs(const TArray<FTopLevel
 			if (RedirectAssets.IsEmpty())
 			{
 				UE_LOG(LogWorldPartition, Warning, TEXT("Failed to find redirected assets for '%s' from '%s'"), *AssetData.ToSoftObjectPath().ToString(), *AssetClassPath.ToString());
-
-				GetBlueprintAssets({ AssetClassPath.GetPackageName().ToString() }, RedirectAssets, true);
-
-				if (RedirectAssets.IsEmpty())
-				{
-					AssetData = FAssetData();
-					break;
-				}
-
-				UE_LOG(LogWorldPartition, Warning, TEXT("\tRecovered redirected assets using bIncludeOnlyOnDiskAssets=true"));
+				AssetData = FAssetData();
+				break;
 			}
 
 			ClassRedirects.Add(FTopLevelAssetPath(AssetData.ToSoftObjectPath().ToString()));
@@ -407,7 +399,7 @@ void FWorldPartitionClassDescRegistry::OnObjectPreSave(UObject* InObject, FObjec
 				{
 					if (UBlueprint* Blueprint = Cast<UBlueprint>(BlueprintGeneratedClass->ClassGeneratedBy))
 					{
-						UpdateClassDescriptor(Blueprint);
+						UpdateClassDescriptor(Blueprint, true);
 					}
 				}
 			}
@@ -419,7 +411,7 @@ void FWorldPartitionClassDescRegistry::OnObjectPreSave(UObject* InObject, FObjec
 			{
 				if (BlueprintGeneratedClass->IsChildOf<AActor>())
 				{
-					UpdateClassDescriptor(Blueprint);
+					UpdateClassDescriptor(Blueprint, false);
 				}
 			}
 		}
@@ -436,7 +428,7 @@ void FWorldPartitionClassDescRegistry::OnObjectPropertyChanged(UObject* InObject
 		if (Blueprint->GeneratedClass && Blueprint->GeneratedClass->IsChildOf<AActor>())
 		{
 			ValidateInternalState();
-			UpdateClassDescriptor(InObject);
+			UpdateClassDescriptor(InObject, true);
 			ValidateInternalState();
 		}
 	}
@@ -552,7 +544,7 @@ void FWorldPartitionClassDescRegistry::RegisterClasses()
 	ValidateInternalState();
 }
 
-void FWorldPartitionClassDescRegistry::UpdateClassDescriptor(UObject* InObject)
+void FWorldPartitionClassDescRegistry::UpdateClassDescriptor(UObject* InObject, bool bOnlyIfExists)
 {
 	check(IsInitialized());
 
@@ -570,18 +562,23 @@ void FWorldPartitionClassDescRegistry::UpdateClassDescriptor(UObject* InObject)
 
 		if (!*ExistingClassDesc)
 		{
-			FWorldPartitionActorDesc* NewActorDesc = ActorCDO->CreateActorDesc().Release();
-			RegisterClassDescriptor(NewActorDesc);
+			if (!bOnlyIfExists)
+			{
+				FWorldPartitionActorDesc* NewActorDesc = ActorCDO->CreateActorDesc().Release();
+				RegisterClassDescriptor(NewActorDesc);
+
+				ParentClassMap.Add(ClassPath, CurrentParentClassPath);
+			}
 		}
 		else
 		{
 			FWorldPartitionActorDescUtils::UpdateActorDescriptorFromActor(ActorCDO, **ExistingClassDesc);
-		}
 
-		if (PreviousParentClassPath != CurrentParentClassPath)
-		{
-			// We are reparenting a blueprint, update our parent map
-			ParentClassMap.Add(ClassPath, CurrentParentClassPath);
+			if (PreviousParentClassPath != CurrentParentClassPath)
+			{
+				// We are reparenting a blueprint, update our parent map
+				ParentClassMap.Add(ClassPath, CurrentParentClassPath);
+			}
 		}
 	}
 }
diff --git a/Engine/Source/Runtime/Engine/Public/WorldPartition/WorldPartitionClassDescRegistry.h b/Engine/Source/Runtime/Engine/Public/WorldPartition/WorldPartitionClassDescRegistry.h
--- a/Engine/Source/Runtime/Engine/Public/WorldPartition/WorldPartitionClassDescRegistry.h
+++ b/Engine/Source/Runtime/Engine/Public/WorldPartition/WorldPartitionClassDescRegistry.h
@@ -56,7 +56,7 @@ private:
 	void RegisterClass(UClass* Class);
 	void RegisterClasses();
 
-	void UpdateClassDescriptor(UObject* InObject);
+	void UpdateClassDescriptor(UObject* InObject, bool bOnlyIfExists);
 
 	void ValidateInternalState();
 
