Handle common math operations that have known valid constant results.  For example, multiplying by 0 will always be zero.
Also added a similar change to the step and clamp, min and max.

diff --git a/Engine/Source/Runtime/Engine/Private/Materials/HLSLMaterialTranslator.cpp b/Engine/Source/Runtime/Engine/Private/Materials/HLSLMaterialTranslator.cpp
--- a/Engine/Source/Runtime/Engine/Private/Materials/HLSLMaterialTranslator.cpp
+++ b/Engine/Source/Runtime/Engine/Private/Materials/HLSLMaterialTranslator.cpp
@@ -7951,7 +7951,214 @@ int32 FHLSLMaterialTranslator::VertexInterpolator(uint32 InterpolatorIndex)
 	{
 		return AddCodeChunk(Interpolator->InterpolatedType, *ValueCode);
 	}
+}
+
+FLinearColor FHLSLMaterialTranslator::CoerceConstantType(FLinearColor SourceValue, EMaterialValueType SourceType, EMaterialValueType DestType)
+{
+	EMaterialCastFlags CastFlags = EMaterialCastFlags::ReplicateScalar;
+	if (DestType == MCT_Float || DestType == MCT_Float1 || DestType == MCT_LWCScalar)
+	{
+		// CoerceValue allows truncating to scalar types only
+		CastFlags |= EMaterialCastFlags::AllowTruncate;
+	}
+	return CastConstantType(SourceValue, SourceType, DestType, CastFlags);
+}
+
+FLinearColor FHLSLMaterialTranslator::CastConstantType(FLinearColor SourceValue, EMaterialValueType SourceType, EMaterialValueType DestType, EMaterialCastFlags Flags)
+{
+	const bool bAllowTruncate = EnumHasAnyFlags(Flags, EMaterialCastFlags::AllowTruncate);
+	const bool bAllowAppendZeroes = EnumHasAnyFlags(Flags, EMaterialCastFlags::AllowAppendZeroes);
+	bool bReplicateScalar = EnumHasAnyFlags(Flags, EMaterialCastFlags::ReplicateScalar);
+
+	// Constants only function as floats
+	SourceType = MakeNonLWCType(SourceType);
+	DestType = MakeNonLWCType(DestType);
+
+	if (SourceType == DestType)
+	{
+		return SourceValue;
+	}
+	else if (IsFloatNumericType(SourceType) && IsFloatNumericType(DestType))
+	{
+		const uint32 NumSourceComponents = GetNumComponents(SourceType);
+		const uint32 NumDestComponents = GetNumComponents(DestType);
+		if (NumSourceComponents != 1)
+		{
+			bReplicateScalar = false;
+		}
+		if (!bReplicateScalar && !bAllowAppendZeroes && NumDestComponents > NumSourceComponents)
+		{
+			Errorf(TEXT("Cannot cast from smaller type %s to larger type %s."), DescribeType(SourceType), DescribeType(DestType));
+			return SourceValue;
+		}
+		if (!bReplicateScalar && !bAllowTruncate && NumDestComponents < NumSourceComponents)
+		{
+			Errorf(TEXT("Cannot cast from larger type %s to smaller type %s."), DescribeType(SourceType), DescribeType(DestType));
+			return SourceValue;
+		}
+
+		FString Result;
+		uint32 NumComponents = 0u;
+		if (NumSourceComponents == NumDestComponents)
+		{
+			return SourceValue;
+		}
+		else if (bReplicateScalar)
+		{
+			return GetTypeMaskedValue(DestType, FLinearColor(SourceValue.R, SourceValue.R, SourceValue.R, SourceValue.R), nullptr);
+		}
+		else
+		{
+			NumComponents = FMath::Min(NumSourceComponents, NumDestComponents);
+			if (NumComponents != NumDestComponents)
+			{
+				return GetTypeMaskedValue(DestType, SourceValue, nullptr);
+			}
+			return SourceValue;
+		}
+
+	}
+	else
+	{
+		Errorf(TEXT("Cannot cast between non-numeric types %s to %s."), DescribeType(SourceType), DescribeType(DestType));
+		return SourceValue;
+	}
+}
+
+bool FHLSLMaterialTranslator::GetConstParameterValue(FMaterialUniformExpression* Expression, FLinearColor& OutConstValue)
+{
+	if (Expression == nullptr || !Expression->IsConstant())
+	{
+		return false;
+	}
+
+	FMaterialRenderContext DummyContext(nullptr, *Material, nullptr);
+	Expression->GetNumberValue(DummyContext, OutConstValue);
+	return true;
+}
+
+FLinearColor FHLSLMaterialTranslator::GetTypeMaskedValue(EMaterialValueType Type, FLinearColor ConstValue, bool* OutSuccess)
+{
+	bool bValid = true;
+	FLinearColor OutValue = ConstValue;
+	switch (Type)
+	{
+	case MCT_Float:
+	case MCT_Float1:
+	case MCT_LWCScalar:
+		OutValue = FLinearColor(ConstValue.R, 0.0f, 0.0f, 0.0f);
+	case MCT_Float2:
+	case MCT_LWCVector2:
+		OutValue = FLinearColor(ConstValue.R, ConstValue.G, 0.0f, 0.0f);
+	case MCT_Float3:
+	case MCT_LWCVector3:
+		OutValue = FLinearColor(ConstValue.R, ConstValue.G, ConstValue.B, 0.0f);
+	case MCT_Float4:
+	case MCT_LWCVector4:
+		OutValue = ConstValue;
+	default:
+		bValid = false;
+	}
+
+	if (OutSuccess)
+	{
+		*OutSuccess = bValid;
+	}
+
+	return OutValue;
+}
+
+bool FHLSLMaterialTranslator::GetConstMaskedParameterValue(EMaterialValueType Type, FMaterialUniformExpression* Expression, FLinearColor& OutConstValue)
+{
+	if (Expression == nullptr || !Expression->IsConstant())
+	{
+		return false;
+	}
+
+	FMaterialRenderContext DummyContext(nullptr, *Material, nullptr);
+
+	FLinearColor ConstantValue;
+	Expression->GetNumberValue(DummyContext, ConstantValue);
+
+	bool bValid;
+	OutConstValue = GetTypeMaskedValue(Type, ConstantValue, &bValid);
+	return bValid;
+}
 
+int32 FHLSLMaterialTranslator::ConstResultValue(EMaterialValueType Type, FLinearColor ConstantValue)
+{
+	switch (Type)
+	{
+	case MCT_Float:
+	case MCT_Float1:
+	case MCT_LWCScalar:
+		return Constant(ConstantValue.R);
+	case MCT_Float2:
+	case MCT_LWCVector2:
+		return Constant2(ConstantValue.R, ConstantValue.G);
+	case MCT_Float3:
+	case MCT_LWCVector3:
+		return Constant3(ConstantValue.R, ConstantValue.G, ConstantValue.B);
+	case MCT_Float4:
+	case MCT_LWCVector4:
+		return Constant4(ConstantValue.R, ConstantValue.G, ConstantValue.B, ConstantValue.A);
+	default:
+		// Error would already be output by GetArithmeticResultType in this case, just return an invalid index
+		return INDEX_NONE;
+	}
+}
+
+int32 FHLSLMaterialTranslator::ConstResultValue(EMaterialValueType Type, float ConstantValue)
+{
+	return ConstResultValue(Type, FLinearColor(ConstantValue, ConstantValue, ConstantValue, ConstantValue));
+}
+
+int32 FHLSLMaterialTranslator::ConstArithmeticResultValue(int LeftEpression, int RightExpression, FLinearColor ConstantValue)
+{
+	check(LeftEpression != INDEX_NONE);
+	check(RightExpression != INDEX_NONE);
+
+	EMaterialValueType ResultType = GetArithmeticResultType(LeftEpression, RightExpression);
+	return ConstResultValue(ResultType, ConstantValue);
+}
+
+int32 FHLSLMaterialTranslator::ConstArithmeticResultValue(int LeftEpression, int RightExpression, float ConstantValue)
+{
+	check(LeftEpression != INDEX_NONE);
+	check(RightExpression != INDEX_NONE);
+
+	EMaterialValueType ResultType = GetArithmeticResultType(LeftEpression, RightExpression);
+	return ConstResultValue(ResultType, ConstantValue);
+}
+
+bool FHLSLMaterialTranslator::IsExpressionConstantValue(int Code, float ConstantValue)
+{
+	check(Code != INDEX_NONE);
+
+	FLinearColor Value;
+	FMaterialUniformExpression* Expression = GetParameterUniformExpression(Code);
+	if (GetConstParameterValue(Expression, Value))
+	{
+		EMaterialValueType CodeType = (*CurrentScopeChunks)[Code].Type;
+
+		switch (CodeType)
+		{
+		case MCT_Float:
+		case MCT_Float1:
+			return Value.R == ConstantValue;
+		case MCT_Float2:
+			return Value.R == ConstantValue && Value.G == ConstantValue;
+		case MCT_Float3:
+			return Value.R == ConstantValue && Value.G == ConstantValue && Value.B == ConstantValue;
+		case MCT_Float4:
+			return Value.R == ConstantValue && Value.G == ConstantValue && Value.B == ConstantValue && Value.A == ConstantValue;
+		default:
+			ensureAlwaysMsgf(false, TEXT("FHLSLMaterialTranslator::IsExpressionConstantValue - Unknown constant expression type"));
+			return false;
+		}
+	}
+
+	return false;
 }
 
 int32 FHLSLMaterialTranslator::Add(int32 A, int32 B)
@@ -7961,9 +8168,20 @@ int32 FHLSLMaterialTranslator::Add(int32 A, int32 B)
 		return INDEX_NONE;
 	}
 
+	FMaterialUniformExpression* ExpressionA = GetParameterUniformExpression(A);
+	FMaterialUniformExpression* ExpressionB = GetParameterUniformExpression(B);
 	const EMaterialValueType ResultType = GetArithmeticResultType(A, B);
-	if (GetParameterUniformExpression(A) && GetParameterUniformExpression(B))
+	if (ExpressionA && ExpressionB)
 	{
+		if (ExpressionA->IsConstant() && ExpressionB->IsConstant())
+		{
+			FLinearColor ValueA, ValueB;
+			GetConstParameterValue(ExpressionA, ValueA);
+			GetConstParameterValue(ExpressionB, ValueB);
+
+			FLinearColor ConstantValue = ValueA + ValueB;
+			return ConstResultValue(ResultType, ConstantValue);
+		}
 		return AddUniformExpression(new FMaterialUniformExpressionFoldedMath(GetParameterUniformExpression(A), GetParameterUniformExpression(B), FMO_Add), ResultType, TEXT("(%s + %s)"), *GetParameterCode(A), *GetParameterCode(B));
 	}
 	else
@@ -7993,9 +8211,20 @@ int32 FHLSLMaterialTranslator::Sub(int32 A, int32 B)
 		return INDEX_NONE;
 	}
 
+	FMaterialUniformExpression* ExpressionA = GetParameterUniformExpression(A);
+	FMaterialUniformExpression* ExpressionB = GetParameterUniformExpression(B);
 	const EMaterialValueType ResultType = GetArithmeticResultType(A, B);
-	if (GetParameterUniformExpression(A) && GetParameterUniformExpression(B))
+	if (ExpressionA && ExpressionB)
 	{
+		if (ExpressionA->IsConstant() && ExpressionB->IsConstant())
+		{
+			FLinearColor ValueA, ValueB;
+			GetConstParameterValue(ExpressionA, ValueA);
+			GetConstParameterValue(ExpressionB, ValueB);
+
+			FLinearColor ConstantValue = ValueA - ValueB;
+			return ConstResultValue(ResultType, ConstantValue);
+		}
 		return AddUniformExpression(new FMaterialUniformExpressionFoldedMath(GetParameterUniformExpression(A), GetParameterUniformExpression(B), FMO_Sub), ResultType, TEXT("(%s - %s)"), *GetParameterCode(A), *GetParameterCode(B));
 	}
 	else
@@ -8043,9 +8272,36 @@ int32 FHLSLMaterialTranslator::Mul(int32 A,int32 B)
 		return INDEX_NONE;
 	}
 
-	if(GetParameterUniformExpression(A) && GetParameterUniformExpression(B))
+	// If either the left or the right is const zero, assume constant zero
+	if (IsExpressionConstantValue(A, 0.0f) || IsExpressionConstantValue(B, 0.0f))
 	{
-		return AddUniformExpression(new FMaterialUniformExpressionFoldedMath(GetParameterUniformExpression(A),GetParameterUniformExpression(B),FMO_Mul),GetArithmeticResultType(A,B),TEXT("(%s * %s)"),*GetParameterCode(A),*GetParameterCode(B));
+		return ConstArithmeticResultValue(A, B, 0.0);
+	}
+	else if(IsExpressionConstantValue(A, 1.0f))
+	{
+		return B;
+	}
+	else if (IsExpressionConstantValue(B, 1.0f))
+	{
+		return A;
+	}
+
+	FMaterialUniformExpression* ExpressionA = GetParameterUniformExpression(A);
+	FMaterialUniformExpression* ExpressionB = GetParameterUniformExpression(B);
+	const EMaterialValueType ResultType = GetArithmeticResultType(A, B);
+	if (ExpressionA && ExpressionB)
+	{
+		if (ExpressionA->IsConstant() && ExpressionB->IsConstant())
+		{
+			FLinearColor ValueA, ValueB;
+			GetConstParameterValue(ExpressionA, ValueA);
+			GetConstParameterValue(ExpressionB, ValueB);
+
+			FLinearColor ConstantValue = ValueA * ValueB;
+			return ConstResultValue(ResultType, ConstantValue);
+		}
+
+		return AddUniformExpression(new FMaterialUniformExpressionFoldedMath(ExpressionA, ExpressionB, FMO_Mul),GetArithmeticResultType(A,B),TEXT("(%s * %s)"),*GetParameterCode(A),*GetParameterCode(B));
 	}
 	else
 	{
@@ -8067,17 +8323,37 @@ int32 FHLSLMaterialTranslator::Div(int32 A, int32 B)
 		return INDEX_NONE;
 	}
 
-	FMaterialUniformExpression* UniformA = GetParameterUniformExpression(A);
-	FMaterialUniformExpression* UniformB = GetParameterUniformExpression(B);
-	if (UniformA && UniformB)
+	if (IsExpressionConstantValue(A, 0.0f))
+	{
+		return ConstArithmeticResultValue(A, B, 0.0);
+	}
+	else if(IsExpressionConstantValue(B, 1.0f))
 	{
-		return AddUniformExpression(new FMaterialUniformExpressionFoldedMath(UniformA, UniformB, FMO_Div), GetArithmeticResultType(A, B), TEXT("(%s / %s)"), *GetParameterCode(A), *GetParameterCode(B));
+		return A;
 	}
-	else if (UniformB && !UniformB->IsConstant())
+
+	FMaterialUniformExpression* ExpressionA = GetParameterUniformExpression(A);
+	FMaterialUniformExpression* ExpressionB = GetParameterUniformExpression(B);
+	const EMaterialValueType ResultType = GetArithmeticResultType(A, B);
+	if (ExpressionA && ExpressionB)
+	{
+		if (ExpressionA->IsConstant() && ExpressionB->IsConstant())
+		{
+			FLinearColor ValueA, ValueB;
+			GetConstParameterValue(ExpressionA, ValueA);
+			GetConstParameterValue(ExpressionB, ValueB);
+
+			FLinearColor ConstantValue = ValueA / ValueB;
+			return ConstResultValue(ResultType, ConstantValue);
+		}
+
+		return AddUniformExpression(new FMaterialUniformExpressionFoldedMath(ExpressionA, ExpressionB, FMO_Div), GetArithmeticResultType(A, B), TEXT("(%s / %s)"), *GetParameterCode(A), *GetParameterCode(B));
+	}
+	else if (ExpressionB && !ExpressionB->IsConstant())
 	{
 		// Division is often optimized as multiplication by reciprocal
 		// If the divisor is a uniform expression, we can fold the reciprocal into the preshader
-		const int32 RcpB = AddUniformExpression(new FMaterialUniformExpressionRcp(UniformB), GetParameterType(B), TEXT("rcp(%s)"), *GetParameterCode(B));
+		const int32 RcpB = AddUniformExpression(new FMaterialUniformExpressionRcp(ExpressionB), GetParameterType(B), TEXT("rcp(%s)"), *GetParameterCode(B));
 		return Mul(A, RcpB);
 	}
 	else
@@ -8100,13 +8376,37 @@ int32 FHLSLMaterialTranslator::Dot(int32 A,int32 B)
 		return INDEX_NONE;
 	}
 
+	if (IsExpressionConstantValue(A, 0.0f) || IsExpressionConstantValue(B, 0.0f))
+	{
+		return Constant(0.0f);
+	}
+
 	FMaterialUniformExpression* ExpressionA = GetParameterUniformExpression(A);
 	FMaterialUniformExpression* ExpressionB = GetParameterUniformExpression(B);
-
 	EMaterialValueType TypeA = GetParameterType(A);
 	EMaterialValueType TypeB = GetParameterType(B);
 	if(ExpressionA && ExpressionB)
 	{
+		if (ExpressionA->IsConstant() && ExpressionB->IsConstant())
+		{
+			FLinearColor ValueA, ValueB;
+			GetConstMaskedParameterValue(TypeA, ExpressionA, ValueA);
+			GetConstMaskedParameterValue(TypeB, ExpressionB, ValueB);
+
+			// Promote scalar (or truncate the bigger type)
+			if (TypeA == MCT_Float || (TypeB != MCT_Float && GetNumComponents(TypeA) > GetNumComponents(TypeB)))
+			{
+				ValueA = CoerceConstantType(ValueA, TypeA, TypeB);
+			}
+			else
+			{
+				ValueB = CoerceConstantType(ValueB, TypeB, TypeA);
+			}
+
+			FVector4 VectorA(ValueA), VectorB(ValueB);
+			return ConstResultValue(MCT_Float, Dot4(VectorA, VectorB));
+		}
+
 		if (TypeA == MCT_Float && TypeB == MCT_Float)
 		{
 			return AddUniformExpression(new FMaterialUniformExpressionFoldedMath(ExpressionA,ExpressionB,FMO_Mul),MCT_Float,TEXT("(%s * %s)"),*GetParameterCode(A),*GetParameterCode(B));
@@ -8158,22 +8458,48 @@ int32 FHLSLMaterialTranslator::Cross(int32 A,int32 B)
 	{
 		return INDEX_NONE;
 	}
+	EMaterialValueType ResultType = GetArithmeticResultType(A, B);
 
-	if(GetParameterUniformExpression(A) && GetParameterUniformExpression(B))
-	{		
-		EMaterialValueType ResultType = GetArithmeticResultType(A,B);
-		if (ResultType == MCT_Float2 || (ResultType & MCT_Float) == 0)
+	bool bIsValid = ResultType != MCT_Float2 && (ResultType & MCT_Float) != 0;
+
+	if(bIsValid && (IsExpressionConstantValue(A, 0.0f) || IsExpressionConstantValue(B, 0.0f)))
+	{
+		return ConstArithmeticResultValue(A, B, 0.0f);
+	}
+	
+	FMaterialUniformExpression* ExpressionA = GetParameterUniformExpression(A);
+	FMaterialUniformExpression* ExpressionB = GetParameterUniformExpression(B);
+	if (ExpressionA && ExpressionB)
+	{
+		if (!bIsValid)
 		{
 			return Errorf(TEXT("Cross product requires 3-component vector input."));
 		}
+
+		if (ExpressionA->IsConstant() && ExpressionB->IsConstant())
+		{
+			EMaterialValueType TypeA = GetParameterType(A);
+			EMaterialValueType TypeB = GetParameterType(B);
+			FLinearColor ValueA, ValueB;
+			GetConstMaskedParameterValue(TypeA, ExpressionA, ValueA);
+			GetConstMaskedParameterValue(TypeB, ExpressionB, ValueB);
+
+			FVector VectorA(CoerceConstantType(ValueA, TypeA, MCT_Float3));
+			FVector VectorB(CoerceConstantType(ValueB, TypeB, MCT_Float3));
+
+			FLinearColor Result(VectorA.Cross(VectorB));
+			return ConstResultValue(MCT_Float3, Result);
+		}
+
 		return AddUniformExpression(new FMaterialUniformExpressionFoldedMath(GetParameterUniformExpression(A),GetParameterUniformExpression(B),FMO_Cross,ResultType),MCT_Float3,TEXT("cross(%s,%s)"),*GetParameterCode(A),*GetParameterCode(B));
 	}
+	else if(IsAnalyticDerivEnabled())
+	{
+		return DerivativeAutogen.GenerateExpressionFunc2(*this, FMaterialDerivativeAutogen::EFunc2::Cross, A, B);
+	}
 	else
 	{
-		if(IsAnalyticDerivEnabled())
-			return DerivativeAutogen.GenerateExpressionFunc2(*this, FMaterialDerivativeAutogen::EFunc2::Cross, A, B);
-		else
-		return AddCodeChunk(MCT_Float3,TEXT("cross(%s,%s)"),*CoerceParameter(A,MCT_Float3),*CoerceParameter(B,MCT_Float3));
+		return AddCodeChunk(MCT_Float3, TEXT("cross(%s,%s)"), *CoerceParameter(A, MCT_Float3), *CoerceParameter(B, MCT_Float3));
 	}
 }
 
@@ -8184,15 +8510,72 @@ int32 FHLSLMaterialTranslator::Power(int32 Base,int32 Exponent)
 		return INDEX_NONE;
 	}
 
+	FMaterialUniformExpression* ExponentExpression = GetParameterUniformExpression(Exponent);
+	FMaterialUniformExpression* BaseExpression = GetParameterUniformExpression(Base);
+
+	// Pow can only have known constant results if the exponent is const
+	// Having a constant base of 0 will result in zero, unless the exponent is 0, so it cannot
+	// be assumed
+	if (ExponentExpression && ExponentExpression->IsConstant())
+	{
+		if (IsExpressionConstantValue(Exponent, 1.0f))
+		{
+			return Base;
+		}
+		else if (IsExpressionConstantValue(Exponent, 0.0f))
+		{
+			return ConstArithmeticResultValue(Base, Exponent, 1.0f);
+		}
+		else if (IsExpressionConstantValue(Base, 0.0f))
+		{
+			return ConstArithmeticResultValue(Base, Exponent, 0.0f);
+		}
+		else if(BaseExpression && BaseExpression->IsConstant())
+		{
+			EMaterialValueType TypeA = GetParameterType(Base);
+			EMaterialValueType TypeB = GetParameterType(Exponent);
+			if (TypeA == TypeB)
+			{
+				auto PositiveClampedPow = [](float BaseComponent, float ExponentComponent)
+				{
+					return (BaseComponent <= 0.0f) ? 0.0f : pow(BaseComponent, ExponentComponent);
+				};
+
+				FLinearColor BaseValue;
+				FLinearColor ExponentValue;
+				
+				GetConstParameterValue(BaseExpression, BaseValue);
+				GetConstParameterValue(ExponentExpression, ExponentValue);
+
+				//See note below on the issues with the two paths below for non const results.
+				// Currently mimicking that functionallity, but this shoul get changed when that does
+				FLinearColor Result;
+				FLinearColor ExponentConverted = IsAnalyticDerivEnabled() ? ExponentValue : CoerceConstantType(ExponentValue, MCT_Float, TypeA);
+
+				Result = FLinearColor(
+					PositiveClampedPow(BaseValue.R, ExponentConverted.R),
+					PositiveClampedPow(BaseValue.G, ExponentConverted.G),
+					PositiveClampedPow(BaseValue.B, ExponentConverted.B),
+					PositiveClampedPow(BaseValue.A, ExponentConverted.A));
+
+				return ConstResultValue(TypeA, Result);
+			}
+		}
+	}
+
+
+	// The code in these two does not seem to match.
+	// The first path does not coerce the exponent value, which will result in all component getting pow for their corresponding exponent
+	// The second path forces the exponent to a float, which makes all components getting pow with only that component
 	if (IsAnalyticDerivEnabled())
 	{
 		return DerivativeAutogen.GenerateExpressionFunc2(*this,FMaterialDerivativeAutogen::EFunc2::PowPositiveClamped,Base,Exponent);
 	}
 	else
 	{
-	// Clamp Pow input to >= 0 to help avoid common NaN cases
-	return AddCodeChunk(GetParameterType(Base),TEXT("PositiveClampedPow(%s,%s)"),*GetParameterCode(Base),*CoerceParameter(Exponent,MCT_Float));
-}
+		// Clamp Pow input to >= 0 to help avoid common NaN cases
+		return AddCodeChunk(GetParameterType(Base),TEXT("PositiveClampedPow(%s,%s)"),*GetParameterCode(Base),*CoerceParameter(Exponent,MCT_Float));
+	}
 }
 	
 int32 FHLSLMaterialTranslator::Logarithm2(int32 X)
@@ -8201,7 +8584,23 @@ int32 FHLSLMaterialTranslator::Logarithm2(int32 X)
 	{
 		return INDEX_NONE;
 	}
-		
+
+	FLinearColor ValueX;
+	FMaterialUniformExpression* ExpressionX = GetParameterUniformExpression(X);
+
+	if (GetConstParameterValue(ExpressionX, ValueX))
+	{
+		EMaterialValueType ResultType = GetParameterType(X);
+
+		FLinearColor Result(
+			FMath::Log2(ValueX.R),
+			FMath::Log2(ValueX.G),
+			FMath::Log2(ValueX.B),
+			FMath::Log2(ValueX.A));
+
+		return ConstResultValue(ResultType, Result);
+	}
+
 	if(GetParameterUniformExpression(X))
 	{
 		return AddUniformExpression(new FMaterialUniformExpressionLogarithm2(GetParameterUniformExpression(X)),GetParameterType(X),TEXT("log2(%s)"),*GetParameterCode(X));
@@ -8222,6 +8621,20 @@ int32 FHLSLMaterialTranslator::Logarithm10(int32 X)
 	{
 		return INDEX_NONE;
 	}
+
+	FLinearColor ValueX;
+	FMaterialUniformExpression* ExpressionX = GetParameterUniformExpression(X);
+	if (GetConstParameterValue(ExpressionX, ValueX))
+	{
+		EMaterialValueType ResultType = GetParameterType(X);
+		FLinearColor Result(
+			FMath::LogX(10, ValueX.R),
+			FMath::LogX(10, ValueX.G),
+			FMath::LogX(10, ValueX.B),
+			FMath::LogX(10, ValueX.A));
+
+		return ConstResultValue(ResultType, Result);
+	}
 		
 	if(GetParameterUniformExpression(X))
 	{
@@ -8244,6 +8657,20 @@ int32 FHLSLMaterialTranslator::SquareRoot(int32 X)
 		return INDEX_NONE;
 	}
 
+	FLinearColor ValueX;
+	FMaterialUniformExpression* ExpressionX = GetParameterUniformExpression(X);
+	if (GetConstParameterValue(ExpressionX, ValueX))
+	{
+		EMaterialValueType ResultType = GetParameterType(X);
+		FLinearColor Result(
+			FMath::Sqrt(ValueX.R),
+			FMath::Sqrt(ValueX.G),
+			FMath::Sqrt(ValueX.B),
+			FMath::Sqrt(ValueX.A));
+
+		return ConstResultValue(ResultType, Result);
+	}
+
 	if(GetParameterUniformExpression(X))
 	{
 		return AddUniformExpression(new FMaterialUniformExpressionSquareRoot(GetParameterUniformExpression(X)),GetParameterType(X),TEXT("sqrt(%s)"),*GetParameterCode(X));
@@ -8268,6 +8695,22 @@ int32 FHLSLMaterialTranslator::Length(int32 X)
 		return INDEX_NONE;
 	}
 
+	FLinearColor ValueX;
+	FMaterialUniformExpression* ExpressionX = GetParameterUniformExpression(X);
+	if (GetConstParameterValue(ExpressionX, ValueX))
+	{
+		FVector4 Vector(ValueX);
+		float FloatResult = Vector.Size();
+		EMaterialValueType ResultType = GetParameterType(X);
+		FLinearColor Result(
+			FloatResult,
+			FloatResult,
+			FloatResult,
+			FloatResult);
+
+		return ConstResultValue(MCT_Float, Result);
+	}
+
 	if (GetParameterUniformExpression(X))
 	{
 		return AddUniformExpression(new FMaterialUniformExpressionLength(GetParameterUniformExpression(X), GetParameterType(X)), MCT_Float, TEXT("length(%s)"), *GetParameterCode(X));
@@ -8294,6 +8737,15 @@ int32 FHLSLMaterialTranslator::Normalize(int32 X)
 
 	const EMaterialValueType ResultType = MakeNonLWCType(GetParameterType(X));
 
+	FLinearColor ValueX;
+	FMaterialUniformExpression* ExpressionX = GetParameterUniformExpression(X);
+	if (GetConstParameterValue(ExpressionX, ValueX))
+	{
+		FVector4 Vector(GetTypeMaskedValue(ResultType, ValueX, nullptr));
+		FLinearColor Result(Vector.GetSafeNormal());
+		return ConstResultValue(ResultType, Result);
+	}
+
 	if (GetParameterUniformExpression(X))
 	{
 		return AddUniformExpression(new FMaterialUniformExpressionNormalize(GetParameterUniformExpression(X)), ResultType, TEXT("normalize(%s)"), *GetParameterCode(X));
@@ -8330,22 +8782,7 @@ int32 FHLSLMaterialTranslator::Step(int32 Y, int32 X)
 		if (ExpressionX == ExpressionY)
 		{
 			const int32 EqualResult = 1.0f;
-			if (ResultType == MCT_Float || ResultType == MCT_Float1)
-			{
-				return Constant(EqualResult);
-			}
-			if (ResultType == MCT_Float2)
-			{
-				return Constant2(EqualResult, EqualResult);
-			}
-			if (ResultType == MCT_Float3)
-			{
-				return Constant3(EqualResult, EqualResult, EqualResult);
-			}
-			if (ResultType == MCT_Float4)
-			{
-				return Constant4(EqualResult, EqualResult, EqualResult, EqualResult);
-			}
+			return ConstResultValue(ResultType, EqualResult);
 		}
 
 		if (ExpressionX->IsConstant() && ExpressionY->IsConstant())
@@ -8437,22 +8874,7 @@ int32 FHLSLMaterialTranslator::SmoothStep(int32 X, int32 Y, int32 A)
 
 	if (bExpressionsAreEqual)
 	{
-		if (ResultType == MCT_Float || ResultType == MCT_Float1)
-		{
-			return Constant(EqualResult);
-		}
-		if (ResultType == MCT_Float2)
-		{
-			return Constant2(EqualResult, EqualResult);
-		}
-		if (ResultType == MCT_Float3)
-		{
-			return Constant3(EqualResult, EqualResult, EqualResult);
-		}
-		if (ResultType == MCT_Float4)
-		{
-			return Constant4(EqualResult, EqualResult, EqualResult, EqualResult);
-		}
+		return ConstResultValue(ResultType, EqualResult);
 	}
 
 	// smoothstep( x, x, a ) could create a div by zero depending on implementation.
@@ -8560,22 +8982,7 @@ int32 FHLSLMaterialTranslator::InvLerp(int32 X, int32 Y, int32 A)
 
 	if (bExpressionsAreEqual)
 	{
-		if (ResultType == MCT_Float || ResultType == MCT_Float1)
-		{
-			return Constant(EqualResult);
-		}
-		if (ResultType == MCT_Float2)
-		{
-			return Constant2(EqualResult, EqualResult);
-		}
-		if (ResultType == MCT_Float3)
-		{
-			return Constant3(EqualResult, EqualResult, EqualResult);
-		}
-		if (ResultType == MCT_Float4)
-		{
-			return Constant4(EqualResult, EqualResult, EqualResult, EqualResult);
-		}
+		return ConstResultValue(ResultType, EqualResult);
 	}
 
 	// (a-x)/(x-x) will create a div by zero.
@@ -8719,9 +9126,28 @@ int32 FHLSLMaterialTranslator::Min(int32 A,int32 B)
 		return INDEX_NONE;
 	}
 
-	if(GetParameterUniformExpression(A) && GetParameterUniformExpression(B))
+	FMaterialUniformExpression* ExpressionA = GetParameterUniformExpression(A);
+	FMaterialUniformExpression* ExpressionB = GetParameterUniformExpression(B);
+
+	if(ExpressionA && ExpressionB)
 	{
-		return AddUniformExpression(new FMaterialUniformExpressionMin(GetParameterUniformExpression(A),GetParameterUniformExpression(B)),GetParameterType(A),TEXT("min(%s,%s)"),*GetParameterCode(A),*CoerceParameter(B,GetParameterType(A)));
+		FLinearColor ValueA, ValueB;
+		if (GetConstParameterValue(ExpressionA, ValueA) && GetConstParameterValue(ExpressionB, ValueB))
+		{
+			EMaterialValueType TypeA = GetParameterType(A);
+			EMaterialValueType TypeB = GetParameterType(B);
+			FLinearColor Result;
+			ValueB = CoerceConstantType(ValueB, TypeB, TypeA);
+
+			Result.R = FMath::Min(ValueA.R, ValueB.R);
+			Result.G = FMath::Min(ValueA.G, ValueB.G);
+			Result.B = FMath::Min(ValueA.B, ValueB.B);
+			Result.A = FMath::Min(ValueA.A, ValueB.A);
+
+			return ConstResultValue(TypeA, Result);
+		}
+		
+		return AddUniformExpression(new FMaterialUniformExpressionMin(ExpressionA, ExpressionB),GetParameterType(A),TEXT("min(%s,%s)"),*GetParameterCode(A),*CoerceParameter(B,GetParameterType(A)));
 	}
 	else
 	{
@@ -8731,10 +9157,10 @@ int32 FHLSLMaterialTranslator::Min(int32 A,int32 B)
 		}
 		else
 		{
-		return AddCodeChunk(GetParameterType(A),TEXT("min(%s,%s)"),*GetParameterCode(A),*CoerceParameter(B,GetParameterType(A)));
+			return AddCodeChunk(GetParameterType(A),TEXT("min(%s,%s)"),*GetParameterCode(A),*CoerceParameter(B,GetParameterType(A)));
+		}
 	}
 }
-}
 
 int32 FHLSLMaterialTranslator::Max(int32 A,int32 B)
 {
@@ -8743,9 +9169,29 @@ int32 FHLSLMaterialTranslator::Max(int32 A,int32 B)
 		return INDEX_NONE;
 	}
 
-	if(GetParameterUniformExpression(A) && GetParameterUniformExpression(B))
+	FMaterialUniformExpression* ExpressionA = GetParameterUniformExpression(A);
+	FMaterialUniformExpression* ExpressionB = GetParameterUniformExpression(B);
+
+	if (ExpressionA && ExpressionB)
 	{
-		return AddUniformExpression(new FMaterialUniformExpressionMax(GetParameterUniformExpression(A),GetParameterUniformExpression(B)),GetParameterType(A),TEXT("max(%s,%s)"),*GetParameterCode(A),*CoerceParameter(B,GetParameterType(A)));
+		FLinearColor ValueA, ValueB;
+		if (GetConstParameterValue(ExpressionA, ValueA) && GetConstParameterValue(ExpressionB, ValueB))
+		{
+			EMaterialValueType TypeA = GetParameterType(A);
+			EMaterialValueType TypeB = GetParameterType(B);
+			FLinearColor Result;
+			ValueB = CoerceConstantType(ValueB, TypeB, TypeA);
+
+			Result.R = FMath::Max(ValueA.R, ValueB.R);
+			Result.G = FMath::Max(ValueA.G, ValueB.G);
+			Result.B = FMath::Max(ValueA.B, ValueB.B);
+			Result.A = FMath::Max(ValueA.A, ValueB.A);
+
+			EMaterialValueType ResultType = GetArithmeticResultType(A, B);
+			return ConstResultValue(ResultType, Result);
+		}
+
+		return AddUniformExpression(new FMaterialUniformExpressionMax(ExpressionA, ExpressionB),GetParameterType(A),TEXT("max(%s,%s)"),*GetParameterCode(A),*CoerceParameter(B,GetParameterType(A)));
 	}
 	else
 	{
@@ -8755,10 +9201,10 @@ int32 FHLSLMaterialTranslator::Max(int32 A,int32 B)
 		}
 		else
 		{
-		return AddCodeChunk(GetParameterType(A),TEXT("max(%s,%s)"),*GetParameterCode(A),*CoerceParameter(B,GetParameterType(A)));
+			return AddCodeChunk(GetParameterType(A),TEXT("max(%s,%s)"),*GetParameterCode(A),*CoerceParameter(B,GetParameterType(A)));
+		}
 	}
 }
-}
 
 int32 FHLSLMaterialTranslator::Clamp(int32 X, int32 A, int32 B)
 {
@@ -8767,24 +9213,31 @@ int32 FHLSLMaterialTranslator::Clamp(int32 X, int32 A, int32 B)
 		return INDEX_NONE;
 	}
 
-	FMaterialUniformExpression* UniformX = GetParameterUniformExpression(X);
 	FMaterialUniformExpression* UniformA = GetParameterUniformExpression(A);
 	FMaterialUniformExpression* UniformB = GetParameterUniformExpression(B);
+	FMaterialUniformExpression* UniformX = GetParameterUniformExpression(X);
 
-	if (UniformA && UniformB && UniformA->IsConstant() && UniformB->IsConstant())
+	FLinearColor ValueA, ValueB, ValueX;
+	bool AIsConst = GetConstParameterValue(UniformA, ValueA);
+	bool BIsConst = GetConstParameterValue(UniformB, ValueB);
+	bool XIsConst = GetConstParameterValue(UniformX, ValueX);
+
+	if ((AIsConst || BIsConst) && XIsConst)
+	{
+		// Just use Min and Max, this will const collapse one or both max/min operations
+		// All parameters are const, just use max and mins const collapse
+		int32 AResult = Max(X, A);
+		return Min(AResult, B);
+	}
+	else if (AIsConst && BIsConst)
 	{
 		// Check to see if we're clamping between 0-1, in that case we can use Saturate, which has some additional optimizations
-		FMaterialRenderContext DummyContext(nullptr, *Material, nullptr);
-		FLinearColor ValueA;
-		FLinearColor ValueB;
-		UniformA->GetNumberValue(DummyContext, ValueA);
-		UniformB->GetNumberValue(DummyContext, ValueB);
 		if (ValueA == FLinearColor(0.0f, 0.0f, 0.0f, 0.0f) && ValueB == FLinearColor(1.0f, 1.0f, 1.0f, 1.0f))
 		{
 			return Saturate(X);
 		}
 	}
-	
+
 	if (UniformX && UniformA && UniformB)
 	{
 		return AddUniformExpression(new FMaterialUniformExpressionClamp(UniformX, UniformA, UniformB), GetParameterType(X), TEXT("min(max(%s,%s),%s)"), *GetParameterCode(X), *CoerceParameter(A, GetParameterType(X)), *CoerceParameter(B, GetParameterType(X)));
@@ -8811,9 +9264,20 @@ int32 FHLSLMaterialTranslator::Saturate(int32 X)
 		return INDEX_NONE;
 	}
 
-	if(GetParameterUniformExpression(X))
+	FMaterialUniformExpression* UniformX = GetParameterUniformExpression(X);
+	FLinearColor ValueX;
+
+	if (GetConstParameterValue(UniformX, ValueX))
+	{
+		// Just use Min and Max, this will const collapse one or both max/min operations
+		// All parameters are const, just use max and mins const collapse
+		int32 AResult = Max(X, Constant(0.0f));
+		return Min(AResult, Constant(1.0f));
+	}
+
+	if(UniformX)
 	{
-		return AddUniformExpression(new FMaterialUniformExpressionSaturate(GetParameterUniformExpression(X)),GetParameterType(X),TEXT("saturate(%s)"),*GetParameterCode(X));
+		return AddUniformExpression(new FMaterialUniformExpressionSaturate(UniformX),GetParameterType(X),TEXT("saturate(%s)"),*GetParameterCode(X));
 	}
 	else
 	{
diff --git a/Engine/Source/Runtime/Engine/Private/Materials/HLSLMaterialTranslator.h b/Engine/Source/Runtime/Engine/Private/Materials/HLSLMaterialTranslator.h
--- a/Engine/Source/Runtime/Engine/Private/Materials/HLSLMaterialTranslator.h
+++ b/Engine/Source/Runtime/Engine/Private/Materials/HLSLMaterialTranslator.h
@@ -724,6 +724,33 @@ protected:
 	// GetArithmeticResultType
 	EMaterialValueType GetArithmeticResultType(EMaterialValueType TypeA, EMaterialValueType TypeB);
 
+	// Same as GetConstParameterValue, but ensures that the value that comes out has zero in components that are not included in EMaterialValueType
+	FLinearColor GetTypeMaskedValue(EMaterialValueType Type, FLinearColor ConstValue, bool* OutSuccess);
+
+	// Same as GetConstParameterValue, but ensures that the value that comes out has zero in components that are not included in EMaterialValueType
+	bool GetConstMaskedParameterValue(EMaterialValueType Type, FMaterialUniformExpression* Expression, FLinearColor& OutConstValue);
+
+	// Attempts to fetch the constant value for a parameter code.
+	// returns true if the value was a valid const expression
+	// OutConstValue will only get filled in if it is valid
+	bool GetConstParameterValue(FMaterialUniformExpression* Expression, FLinearColor& OutConstValue);
+
+	// Properly casts a const type to another const type
+	FLinearColor CoerceConstantType(FLinearColor SourceValue, EMaterialValueType SourceType, EMaterialValueType DestinationType);
+	FLinearColor CastConstantType(FLinearColor SourceValue, EMaterialValueType SourceType, EMaterialValueType DestinationType, EMaterialCastFlags Flags);
+
+	// Returns a constant of the specified typ
+	int32 ConstResultValue(EMaterialValueType Type, FLinearColor ConstantValue);
+	int32 ConstResultValue(EMaterialValueType Type, float ConstantValue);
+	
+
+	// Returns a valid constant result for an arithmetic expression that has a known const value
+	int32 ConstArithmeticResultValue(int LeftEpression, int RightExpression, FLinearColor ConstantValue);
+	int32 ConstArithmeticResultValue(int LeftEpression, int RightExpression, float ConstantValue);
+
+	// Returns true of the expression results in a specific constant value
+	bool IsExpressionConstantValue(int Code, float ConstantValue);
+
 	int32 GenericSwitch(const TCHAR* Function, int32 IfTrue, int32 IfFalse);
 
 	FString StrataGetCastParameterCode(int32 Index, EMaterialValueType DestType);
