add some sanity check and safety measures to the UClusterUnionReplicatedProxyComponent surrounding how it defers adding to the parent cluster union

diff --git a/Engine/Source/Runtime/Engine/Classes/PhysicsEngine/ClusterUnionReplicatedProxyComponent.h b/Engine/Source/Runtime/Engine/Classes/PhysicsEngine/ClusterUnionReplicatedProxyComponent.h
--- a/Engine/Source/Runtime/Engine/Classes/PhysicsEngine/ClusterUnionReplicatedProxyComponent.h
+++ b/Engine/Source/Runtime/Engine/Classes/PhysicsEngine/ClusterUnionReplicatedProxyComponent.h
@@ -91,7 +91,8 @@ private:
 	UPROPERTY()
 	bool bIsPendingDeletion = false;
 
-	void DeferUntilChildClusteredComponentInParentUnion(TFunction<void()> Func);
+	FTimerHandle DeferSetChildToParentHandle;
+	void DeferSetChildToParentChildUntilClusteredComponentInParentUnion();
 
 	//~ Begin UActorComponent Interface
 public:
diff --git a/Engine/Source/Runtime/Engine/Private/PhysicsEngine/ClusterUnionReplicatedProxyComponent.cpp b/Engine/Source/Runtime/Engine/Private/PhysicsEngine/ClusterUnionReplicatedProxyComponent.cpp
--- a/Engine/Source/Runtime/Engine/Private/PhysicsEngine/ClusterUnionReplicatedProxyComponent.cpp
+++ b/Engine/Source/Runtime/Engine/Private/PhysicsEngine/ClusterUnionReplicatedProxyComponent.cpp
@@ -40,6 +40,13 @@ void UClusterUnionReplicatedProxyComponent::EndPlay(const EEndPlayReason::Type E
 {
 	UActorComponent::EndPlay(EndPlayReason);
 
+	MarkPendingDeletion();
+
+	if (AActor* Owner = GetOwner(); Owner && DeferSetChildToParentHandle.IsValid())
+	{
+		Owner->GetWorldTimerManager().ClearTimer(DeferSetChildToParentHandle);
+	}
+
 	if (!GetOwner()->HasAuthority() && ParentClusterUnion && ChildClusteredComponent)
 	{
 		ParentClusterUnion->RemoveComponentFromCluster(ChildClusteredComponent);
@@ -124,36 +131,31 @@ void UClusterUnionReplicatedProxyComponent::PostRepNotifies()
 
 	if (bIsValid && bNetUpdateParticleChildToParents && ParticleBoneIds.Num() == ParticleChildToParents.Num())
 	{
-		// This particular bit can't happen utnil *after* we add the component to the cluster union. There's an additional deferral
+		// This particular bit can't happen until *after* we add the component to the cluster union. There's an additional deferral
 		// in AddComponentToCluster that we have to wait for.
-		DeferUntilChildClusteredComponentInParentUnion(
-			[this]()
-			{
-				ParentClusterUnion->ForceSetChildToParent(ChildClusteredComponent, ParticleBoneIds, ParticleChildToParents);
-			}
-		);
+		if (!DeferSetChildToParentHandle.IsValid())
+		{
+			DeferSetChildToParentChildUntilClusteredComponentInParentUnion();
+		}
 		bNetUpdateParticleChildToParents = false;
 	}
 }
 
-void UClusterUnionReplicatedProxyComponent::DeferUntilChildClusteredComponentInParentUnion(TFunction<void()> Func)
+void UClusterUnionReplicatedProxyComponent::DeferSetChildToParentChildUntilClusteredComponentInParentUnion()
 {
-	if (!ParentClusterUnion || !ChildClusteredComponent)
+	if (!ParentClusterUnion || !ChildClusteredComponent || IsPendingDeletion())
 	{
 		return;
 	}
 
+	DeferSetChildToParentHandle.Invalidate();
+
 	if (ParentClusterUnion->IsComponentAdded(ChildClusteredComponent))
 	{
-		Func();
+		ParentClusterUnion->ForceSetChildToParent(ChildClusteredComponent, ParticleBoneIds, ParticleChildToParents);
 	}
-	else
+	else if (AActor* Owner = GetOwner())
 	{
-		GetOwner()->GetWorldTimerManager().SetTimerForNextTick(
-			[this, Func]()
-			{
-				DeferUntilChildClusteredComponentInParentUnion(Func);
-			}
-		);
+		DeferSetChildToParentHandle = Owner->GetWorldTimerManager().SetTimerForNextTick(this, &UClusterUnionReplicatedProxyComponent::DeferSetChildToParentChildUntilClusteredComponentInParentUnion);
 	}
 }
\ No newline at end of file
