Split IsContainerInstance/GetContainerInstance from IsContainerFilter/GetContainerFilter

#rnx

diff --git a/Engine/Source/Editor/SceneOutliner/Private/SceneOutlinerModule.cpp b/Engine/Source/Editor/SceneOutliner/Private/SceneOutlinerModule.cpp
--- a/Engine/Source/Editor/SceneOutliner/Private/SceneOutlinerModule.cpp
+++ b/Engine/Source/Editor/SceneOutliner/Private/SceneOutlinerModule.cpp
@@ -386,9 +386,9 @@ void FSceneOutlinerModule::CreateActorInfoColumns(FSceneOutlinerInitializationOp
 		{
 			if (const FWorldPartitionActorDesc* ActorDesc = ActorDescItem->ActorDescHandle.Get())
 			{
-				if (FName LevelPackage = ActorDesc->GetLevelPackage(); !LevelPackage.IsNone())
+				if (FName LevelPackage = ActorDesc->GetContainerPackage(); !LevelPackage.IsNone())
 				{
-					return ActorDesc->GetLevelPackage().ToString();
+					return ActorDesc->GetContainerPackage().ToString();
 				}
 			}
 		}
diff --git a/Engine/Source/Runtime/Engine/Private/WorldPartition/LevelInstance/LevelInstanceActorDesc.cpp b/Engine/Source/Runtime/Engine/Private/WorldPartition/LevelInstance/LevelInstanceActorDesc.cpp
--- a/Engine/Source/Runtime/Engine/Private/WorldPartition/LevelInstance/LevelInstanceActorDesc.cpp
+++ b/Engine/Source/Runtime/Engine/Private/WorldPartition/LevelInstance/LevelInstanceActorDesc.cpp
@@ -174,7 +174,7 @@ bool FLevelInstanceActorDesc::IsContainerInstanceInternal() const
 	return ULevelInstanceSubsystem::CanUsePackage(LevelPackage);
 }
 
-bool FLevelInstanceActorDesc::GetContainerInstance(const FGetContainerInstanceParams& InParams, FContainerInstance& OutContainerInstance) const
+bool FLevelInstanceActorDesc::GetContainerInstance(FContainerInstance& OutContainerInstance) const
 {
 	if (LevelInstanceContainer.IsValid())
 	{
@@ -187,15 +187,6 @@ bool FLevelInstanceActorDesc::GetContainerInstance(const FGetContainerInstancePa
 		// Apply level instance pivot offset
 		FTransform LevelInstancePivotOffsetTransform = FTransform(ULevel::GetLevelInstancePivotOffsetFromPackage(LevelInstanceContainer->GetContainerPackage()));
 		OutContainerInstance.Transform = LevelInstancePivotOffsetTransform * LevelInstanceTransform;		
-
-		if (InParams.bBuildFilter)
-		{
-			// Fill Container Instance Filter
-			UWorldPartitionSubsystem* WorldPartitionSubsystem = UWorld::GetSubsystem<UWorldPartitionSubsystem>(LevelInstanceContainerWorldContext.Get());
-			check(WorldPartitionSubsystem);
-
-			OutContainerInstance.FilteredActors = WorldPartitionSubsystem->GetFilteredActorsPerContainer(InParams.ContainerID, LevelInstanceContainer->GetContainerPackage().ToString(), Filter);
-		}
 		return true;
 	}
 
diff --git a/Engine/Source/Runtime/Engine/Private/WorldPartition/LoaderAdapter/LoaderAdapterActorList.cpp b/Engine/Source/Runtime/Engine/Private/WorldPartition/LoaderAdapter/LoaderAdapterActorList.cpp
--- a/Engine/Source/Runtime/Engine/Private/WorldPartition/LoaderAdapter/LoaderAdapterActorList.cpp
+++ b/Engine/Source/Runtime/Engine/Private/WorldPartition/LoaderAdapter/LoaderAdapterActorList.cpp
@@ -67,7 +67,7 @@ void FLoaderAdapterActorList::RemoveActors(const TArray<FWorldPartitionHandle>&
 		if (ActorHandle->IsContainerInstance())
 		{
 			FWorldPartitionActorDesc::FContainerInstance ContainerInstance;
-			if (ActorHandle->GetContainerInstance(FWorldPartitionActorDesc::FGetContainerInstanceParams(), ContainerInstance))
+			if (ActorHandle->GetContainerInstance(ContainerInstance))
 			{
 				if (ContainerInstance.bSupportsPartialEditorLoading)
 				{
diff --git a/Engine/Source/Runtime/Engine/Private/WorldPartition/LoaderAdapter/LoaderAdapterList.cpp b/Engine/Source/Runtime/Engine/Private/WorldPartition/LoaderAdapter/LoaderAdapterList.cpp
--- a/Engine/Source/Runtime/Engine/Private/WorldPartition/LoaderAdapter/LoaderAdapterList.cpp
+++ b/Engine/Source/Runtime/Engine/Private/WorldPartition/LoaderAdapter/LoaderAdapterList.cpp
@@ -28,7 +28,7 @@ void FLoaderAdapterList::HandleActorContainer(const FWorldPartitionHandle& InAct
 	if (InActor->IsContainerInstance())
 	{
 		FWorldPartitionActorDesc::FContainerInstance ContainerInstance;
-		if (InActor->GetContainerInstance(FWorldPartitionActorDesc::FGetContainerInstanceParams(), ContainerInstance))
+		if (InActor->GetContainerInstance(ContainerInstance))
 		{
 			if (ContainerInstance.bSupportsPartialEditorLoading)
 			{
diff --git a/Engine/Source/Runtime/Engine/Private/WorldPartition/LoaderAdapter/LoaderAdapterSpatial.cpp b/Engine/Source/Runtime/Engine/Private/WorldPartition/LoaderAdapter/LoaderAdapterSpatial.cpp
--- a/Engine/Source/Runtime/Engine/Private/WorldPartition/LoaderAdapter/LoaderAdapterSpatial.cpp
+++ b/Engine/Source/Runtime/Engine/Private/WorldPartition/LoaderAdapter/LoaderAdapterSpatial.cpp
@@ -41,7 +41,7 @@ void ILoaderAdapterSpatial::HandleIntersectingContainer(UWorldPartition* InWorld
 			if (ActorHandle->GetIsSpatiallyLoaded() && ActorHandle->IsContainerInstance())
 			{
 				FWorldPartitionActorDesc::FContainerInstance ContainerInstance;
-				if (ActorHandle->GetContainerInstance(FWorldPartitionActorDesc::FGetContainerInstanceParams(), ContainerInstance))
+				if (ActorHandle->GetContainerInstance(ContainerInstance))
 				{
 					if (ContainerInstance.bSupportsPartialEditorLoading)
 					{
diff --git a/Engine/Source/Runtime/Engine/Private/WorldPartition/WorldPartitionActorDescView.cpp b/Engine/Source/Runtime/Engine/Private/WorldPartition/WorldPartitionActorDescView.cpp
--- a/Engine/Source/Runtime/Engine/Private/WorldPartition/WorldPartitionActorDescView.cpp
+++ b/Engine/Source/Runtime/Engine/Private/WorldPartition/WorldPartitionActorDescView.cpp
@@ -179,14 +179,25 @@ bool FWorldPartitionActorDescView::IsContainerInstance() const
 	return ActorDesc->IsContainerInstance();
 }
 
-FName FWorldPartitionActorDescView::GetLevelPackage() const
+bool FWorldPartitionActorDescView::IsContainerFilter() const
 {
-	return ActorDesc->GetLevelPackage();
+	return ActorDesc->IsContainerFilter();
 }
 
-bool FWorldPartitionActorDescView::GetContainerInstance(const FWorldPartitionActorDesc::FGetContainerInstanceParams& InParams, FWorldPartitionActorDesc::FContainerInstance& OutContainerInstance) const
+FName FWorldPartitionActorDescView::GetContainerPackage() const
 {
-	return ActorDesc->GetContainerInstance(InParams, OutContainerInstance);
+	return ActorDesc->GetContainerPackage();
+}
+
+bool FWorldPartitionActorDescView::GetContainerInstance(FWorldPartitionActorDesc::FContainerInstance& OutContainerInstance) const
+{
+	return ActorDesc->GetContainerInstance(OutContainerInstance);
+}
+
+const FWorldPartitionActorFilter* FWorldPartitionActorDescView::GetContainerFilter() const
+{
+	check(IsContainerFilter());
+	return ActorDesc->GetContainerFilter();
 }
 
 void FWorldPartitionActorDescView::CheckForErrors(IStreamingGenerationErrorHandler* ErrorHandler) const
diff --git a/Engine/Source/Runtime/Engine/Private/WorldPartition/WorldPartitionStreamingGeneration.cpp b/Engine/Source/Runtime/Engine/Private/WorldPartition/WorldPartitionStreamingGeneration.cpp
--- a/Engine/Source/Runtime/Engine/Private/WorldPartition/WorldPartitionStreamingGeneration.cpp
+++ b/Engine/Source/Runtime/Engine/Private/WorldPartition/WorldPartitionStreamingGeneration.cpp
@@ -20,6 +20,7 @@
 #include "WorldPartition/ErrorHandling/WorldPartitionStreamingGenerationLogErrorHandler.h"
 #include "WorldPartition/ErrorHandling/WorldPartitionStreamingGenerationMapCheckErrorHandler.h"
 #include "WorldPartition/HLOD/HLODActor.h"
+#include "WorldPartition/WorldPartitionSubsystem.h"
 #include "HAL/FileManager.h"
 
 #define LOCTEXT_NAMESPACE "WorldPartition"
@@ -476,26 +477,35 @@ class FWorldPartitionStreamingGenerator
 		TArray<FWorldPartitionActorDescView> ContainerInstanceViews = ContainerDescriptor.ContainerInstanceViews;
 		for (const FWorldPartitionActorDescView& ContainerInstanceView : ContainerInstanceViews)
 		{
-			const FGuid ActorGuid = ContainerInstanceView.GetGuid();
-			const FActorContainerID ContainerID(InContainerInstanceDescriptor.ID, ActorGuid);
-
-			FWorldPartitionActorDesc::FGetContainerInstanceParams GetContainerInstanceParams = FWorldPartitionActorDesc::FGetContainerInstanceParams()
-				.SetContainerID(ContainerID)
-				.SetBuildFilter(InContainerInstanceDescriptor.ID.IsMainContainer()); // Build filter on top level child containers
-
 			FWorldPartitionActorDesc::FContainerInstance SubContainerInstance;
-			if (!ContainerInstanceView.GetContainerInstance(GetContainerInstanceParams, SubContainerInstance))
+			if (!ContainerInstanceView.GetContainerInstance(SubContainerInstance))
 			{
-				ErrorHandler->OnLevelInstanceInvalidWorldAsset(ContainerInstanceView, ContainerInstanceView.GetLevelPackage(), IStreamingGenerationErrorHandler::ELevelInstanceInvalidReason::WorldAssetHasInvalidContainer);
+				ErrorHandler->OnLevelInstanceInvalidWorldAsset(ContainerInstanceView, ContainerInstanceView.GetContainerPackage(), IStreamingGenerationErrorHandler::ELevelInstanceInvalidReason::WorldAssetHasInvalidContainer);
 				continue;
 			}
 
 			if (ContainerInstancesStack.Contains(SubContainerInstance.Container->ContainerPackageName))
 			{
-				ErrorHandler->OnLevelInstanceInvalidWorldAsset(ContainerInstanceView, ContainerInstanceView.GetLevelPackage(), IStreamingGenerationErrorHandler::ELevelInstanceInvalidReason::CirculalReference);
+				ErrorHandler->OnLevelInstanceInvalidWorldAsset(ContainerInstanceView, ContainerInstanceView.GetContainerPackage(), IStreamingGenerationErrorHandler::ELevelInstanceInvalidReason::CirculalReference);
 				continue;
 			}
 
+			const FGuid ActorGuid = ContainerInstanceView.GetGuid();
+			const FActorContainerID ContainerID(InContainerInstanceDescriptor.ID, ActorGuid);
+						
+			TMap<FActorContainerID, TSet<FGuid>> ContainerInstanceFilter;
+			if (InContainerInstanceDescriptor.ID.IsMainContainer() && ContainerInstanceView.IsContainerFilter())
+			{
+				if (const FWorldPartitionActorFilter* ContainerFilter = ContainerInstanceView.GetContainerFilter())
+				{
+					UWorld* OwningWorld = WorldPartitionContext ? WorldPartitionContext->GetWorld() : nullptr;
+					if (UWorldPartitionSubsystem* WorldPartitionSubsystem = UWorld::GetSubsystem<UWorldPartitionSubsystem>(OwningWorld))
+					{
+						ContainerInstanceFilter = WorldPartitionSubsystem->GetFilteredActorsPerContainer(ContainerID, ContainerInstanceView.GetContainerPackage().ToString(), *ContainerFilter);
+					}
+				}
+			}
+
 			ContainerInstancesStack.Add(SubContainerInstance.Container->ContainerPackageName);
 
 			check(SubContainerInstance.Container);
@@ -505,7 +515,7 @@ class FWorldPartitionStreamingGenerator
 			SubContainerInstanceDescriptor.ID = ContainerID;
 			SubContainerInstanceDescriptor.Container = SubContainerInstance.Container;
 			SubContainerInstanceDescriptor.Transform = SubContainerInstance.Transform * InContainerInstanceDescriptor.Transform;
-			SubContainerInstanceDescriptor.FilteredActors = SubContainerInstance.FilteredActors;
+			SubContainerInstanceDescriptor.FilteredActors = MoveTemp(ContainerInstanceFilter);
 			SubContainerInstanceDescriptor.ParentID = InContainerInstanceDescriptor.ID;
 			SubContainerInstanceDescriptor.OwnerName = *ContainerInstanceView.GetActorLabelOrName().ToString();
 
diff --git a/Engine/Source/Runtime/Engine/Private/WorldPartition/WorldPartitionSubsystem.cpp b/Engine/Source/Runtime/Engine/Private/WorldPartition/WorldPartitionSubsystem.cpp
--- a/Engine/Source/Runtime/Engine/Private/WorldPartition/WorldPartitionSubsystem.cpp
+++ b/Engine/Source/Runtime/Engine/Private/WorldPartition/WorldPartitionSubsystem.cpp
@@ -164,7 +164,7 @@ FWorldPartitionActorFilter UWorldPartitionSubsystem::GetWorldPartitionActorFilte
 			check(!WorldDataLayersActorDesc);
 			WorldDataLayersActorDesc = static_cast<const FWorldDataLayersActorDesc*>(*ActorDescIt);
 		}
-		else if (ActorDescIt->IsContainerInstance())
+		else if (ActorDescIt->IsContainerFilter())
 		{
 			ContainerActorDescs.Add(*ActorDescIt);
 		}
@@ -189,7 +189,7 @@ FWorldPartitionActorFilter UWorldPartitionSubsystem::GetWorldPartitionActorFilte
 		TSet<FString> VisitedPackagesCopy(InOutVisitedPackages);
 
 		// Get World Default Filter
-		FWorldPartitionActorFilter* ChildFilter = new FWorldPartitionActorFilter(GetWorldPartitionActorFilterInternal(ContainerActorDesc->GetLevelPackage().ToString(), VisitedPackagesCopy));
+		FWorldPartitionActorFilter* ChildFilter = new FWorldPartitionActorFilter(GetWorldPartitionActorFilterInternal(ContainerActorDesc->GetContainerPackage().ToString(), VisitedPackagesCopy));
 		ChildFilter->DisplayName = ContainerActorDesc->GetActorLabelOrName().ToString();
 
 		// Apply Filter to Default
@@ -231,7 +231,22 @@ TMap<FActorContainerID, TSet<FGuid>> UWorldPartitionSubsystem::GetFilteredActors
 
 	ProcessFilter(InContainerID, ContainerFilter);
 
-	TFunction<void(const FActorContainerID&, const UActorDescContainer*)> ProcessContainers = [&DataLayerFiltersPerContainer, &FilteredActors, &ProcessContainers](const FActorContainerID& InContainerID, const UActorDescContainer* InContainer)
+	// Keep track of registered containers to unregister them
+	TMap<FName, UActorDescContainer*> RegisteredContainers;
+
+	TFunction<UActorDescContainer* (FName)> FindOrRegisterContainer = [this, &RegisteredContainers](FName ContainerPackage)
+	{
+		if (UActorDescContainer** FoundContainer = RegisteredContainers.Find(ContainerPackage))
+		{
+			return *FoundContainer;
+		}
+		
+		UActorDescContainer* RegisteredContainer = RegisterContainer(ContainerPackage);
+		RegisteredContainers.Add(ContainerPackage, RegisteredContainer);
+		return RegisteredContainer;
+	};
+
+	TFunction<void(const FActorContainerID&, const UActorDescContainer*)> ProcessContainers = [&FindOrRegisterContainer, &DataLayerFiltersPerContainer, &FilteredActors, &ProcessContainers](const FActorContainerID& InContainerID, const UActorDescContainer* InContainer)
 	{
 		const TMap<FSoftObjectPath, FWorldPartitionActorFilter::FDataLayerFilter>& DataLayerFilters = DataLayerFiltersPerContainer.FindChecked(InContainerID);
 		for (FActorDescList::TConstIterator<> ActorDescIt(InContainer); ActorDescIt; ++ActorDescIt)
@@ -262,20 +277,23 @@ TMap<FActorContainerID, TSet<FGuid>> UWorldPartitionSubsystem::GetFilteredActors
 				}
 			}
 
-			if (ActorDescIt->IsContainerInstance())
+			if (ActorDescIt->IsContainerFilter())
 			{
-				FWorldPartitionActorDesc::FContainerInstance ChildContainerInstance;
-				if (ActorDescIt->GetContainerInstance(FWorldPartitionActorDesc::FGetContainerInstanceParams(), ChildContainerInstance))
-				{
-					ProcessContainers(FActorContainerID(InContainerID, ActorDescIt->GetGuid()), ChildContainerInstance.Container);
-				}
+				UActorDescContainer* ChildContainer = FindOrRegisterContainer(ActorDescIt->GetContainerPackage());
+				check(ChildContainer);
+				ProcessContainers(FActorContainerID(InContainerID, ActorDescIt->GetGuid()), ChildContainer);
 			}
 		}
 	};
 
-	UActorDescContainer* Container = RegisterContainer(*InWorldPackage);
+	UActorDescContainer* Container = FindOrRegisterContainer(*InWorldPackage);
 	ProcessContainers(InContainerID, Container);
-	UnregisterContainer(Container);
+	
+	// Unregister Containers
+	for (auto& [Name, RegisteredContainer] : RegisteredContainers)
+	{
+		UnregisterContainer(RegisteredContainer);
+	}
 
 	return FilteredActors;
 }
diff --git a/Engine/Source/Runtime/Engine/Private/WorldPartitionBlueprintLibrary.cpp b/Engine/Source/Runtime/Engine/Private/WorldPartitionBlueprintLibrary.cpp
--- a/Engine/Source/Runtime/Engine/Private/WorldPartitionBlueprintLibrary.cpp
+++ b/Engine/Source/Runtime/Engine/Private/WorldPartitionBlueprintLibrary.cpp
@@ -102,7 +102,7 @@ bool UWorldPartitionBlueprintLibrary::GetActorDescs(const UActorDescContainer* I
 		if (ActorDescIt->IsContainerInstance())
 		{
 			FWorldPartitionActorDesc::FContainerInstance ContainerInstance;
-			if (ActorDescIt->GetContainerInstance(FWorldPartitionActorDesc::FGetContainerInstanceParams(), ContainerInstance))
+			if (ActorDescIt->GetContainerInstance(ContainerInstance))
 			{
 				bResult &= GetActorDescs(ContainerInstance.Container, ContainerInstance.Transform * InTransform, OutActorDescs);
 			}
@@ -127,7 +127,7 @@ bool UWorldPartitionBlueprintLibrary::HandleIntersectingActorDesc(const FWorldPa
 	if (ActorDesc->IsContainerInstance())
 	{
 		FWorldPartitionActorDesc::FContainerInstance ContainerInstance;
-		if (ActorDesc->GetContainerInstance(FWorldPartitionActorDesc::FGetContainerInstanceParams(), ContainerInstance))
+		if (ActorDesc->GetContainerInstance(ContainerInstance))
 		{
 			bResult &= GetIntersectingActorDescs(ContainerInstance.Container, InBox, ContainerInstance.Transform * InTransform, OutActorDescs);
 		}
diff --git a/Engine/Source/Runtime/Engine/Public/WorldPartition/LevelInstance/LevelInstanceActorDesc.h b/Engine/Source/Runtime/Engine/Public/WorldPartition/LevelInstance/LevelInstanceActorDesc.h
--- a/Engine/Source/Runtime/Engine/Public/WorldPartition/LevelInstance/LevelInstanceActorDesc.h
+++ b/Engine/Source/Runtime/Engine/Public/WorldPartition/LevelInstance/LevelInstanceActorDesc.h
@@ -26,8 +26,9 @@ public:
 	virtual ~FLevelInstanceActorDesc() override;
 
 	virtual bool IsContainerInstance() const override;
-	virtual FName GetLevelPackage() const override { return LevelPackage; }
-	virtual bool GetContainerInstance(const FGetContainerInstanceParams& InParams, FContainerInstance& OutContainerInstance) const override;
+	virtual bool IsContainerFilter() const override { return IsContainerInstance(); }
+	virtual FName GetContainerPackage() const override { return LevelPackage; }
+	virtual bool GetContainerInstance(FContainerInstance& OutContainerInstance) const override;
 	virtual const FWorldPartitionActorFilter* GetContainerFilter() const override { return &Filter; }
 	virtual void CheckForErrors(IStreamingGenerationErrorHandler* ErrorHandler) const override;
 
diff --git a/Engine/Source/Runtime/Engine/Public/WorldPartition/WorldPartitionActorDesc.h b/Engine/Source/Runtime/Engine/Public/WorldPartition/WorldPartitionActorDesc.h
--- a/Engine/Source/Runtime/Engine/Public/WorldPartition/WorldPartitionActorDesc.h
+++ b/Engine/Source/Runtime/Engine/Public/WorldPartition/WorldPartitionActorDesc.h
@@ -89,15 +89,6 @@ public:
 		TMap<FActorContainerID, TSet<FGuid>> FilteredActors;
 	};
 
-	struct FGetContainerInstanceParams
-	{
-		FActorContainerID ContainerID;
-		bool bBuildFilter = false;
-
-		FGetContainerInstanceParams& SetContainerID(FActorContainerID InContainerID) { ContainerID = InContainerID; return *this; }
-		FGetContainerInstanceParams& SetBuildFilter(bool bInBuildFilter) { bBuildFilter = bInBuildFilter; return *this; }
-	};
-
 	virtual ~FWorldPartitionActorDesc() {}
 
 	inline const FGuid& GetGuid() const { return Guid; }
@@ -153,9 +144,11 @@ PRAGMA_ENABLE_DEPRECATION_WARNINGS
 	inline bool IsDefaultActorDesc() const { return bIsDefaultActorDesc; }
 
 	virtual bool IsContainerInstance() const { return false; }
-	virtual FName GetLevelPackage() const { return NAME_None; }
+	virtual bool IsContainerFilter() const { return false; }
+	virtual FName GetContainerPackage() const { return NAME_None; }
+	
 	virtual const FWorldPartitionActorFilter* GetContainerFilter() const { return nullptr; }
-	virtual bool GetContainerInstance(const FGetContainerInstanceParams& InParams, FContainerInstance& OutContainerInstance) const { return false; }
+	virtual bool GetContainerInstance(FContainerInstance& OutContainerInstance) const { return false; }
 
 	FGuid GetContentBundleGuid() const;
 
@@ -168,6 +161,9 @@ PRAGMA_ENABLE_DEPRECATION_WARNINGS
 	UE_DEPRECATED(5.2, "ShouldValidateRuntimeGrid is deprecated and should not be used.")
 	virtual bool ShouldValidateRuntimeGrid() const { return true; }
 
+	UE_DEPRECATED(5.3, "GetLevelPackage is deprecated use GetContainerPackage instead.")
+	virtual FName GetLevelPackage() const { return GetContainerPackage(); }
+
 	bool operator==(const FWorldPartitionActorDesc& Other) const
 	{
 		return Guid == Other.Guid;
diff --git a/Engine/Source/Runtime/Engine/Public/WorldPartition/WorldPartitionActorDescView.h b/Engine/Source/Runtime/Engine/Public/WorldPartition/WorldPartitionActorDescView.h
--- a/Engine/Source/Runtime/Engine/Public/WorldPartition/WorldPartitionActorDescView.h
+++ b/Engine/Source/Runtime/Engine/Public/WorldPartition/WorldPartitionActorDescView.h
@@ -12,6 +12,7 @@ class AActor;
 class IStreamingGenerationErrorHandler;
 class UActorDescContainer;
 class UWorldPartition;
+struct FWorldPartitionActorFilter;
 
 /**
  * A view on top of an actor desc, used to cache information that can be (potentially) different than the actor desc
@@ -58,10 +59,14 @@ public:
 	const FGuid& GetFolderGuid() const;
 
 	FGuid GetContentBundleGuid() const;
-
+	FName GetContainerPackage() const;
 	bool IsContainerInstance() const;
-	FName GetLevelPackage() const;
-	bool GetContainerInstance(const FWorldPartitionActorDesc::FGetContainerInstanceParams& InParams, FWorldPartitionActorDesc::FContainerInstance& OutContainerInstance) const;
+	bool GetContainerInstance(FWorldPartitionActorDesc::FContainerInstance& OutContainerInstance) const;
+	bool IsContainerFilter() const;
+	const FWorldPartitionActorFilter* GetContainerFilter() const;
+	
+	UE_DEPRECATED(5.3, "GetLevelPackage is deprecated use GetContainerPackage instead.")
+	FName GetLevelPackage() const { return GetContainerPackage(); }
 
 	void CheckForErrors(IStreamingGenerationErrorHandler* ErrorHandler) const;
 
