World Partition Streaming Generation Logs:
- Output generation logs under the Saved/Logs folder so they are part of build machines artifacts.
- Output cell filename in streaming generation logs.
- Reenable outputing generation logs on build machines.

diff --git a/Engine/Source/Runtime/Engine/Private/WorldPartition/WorldPartitionRuntimeSpatialHash.cpp b/Engine/Source/Runtime/Engine/Private/WorldPartition/WorldPartitionRuntimeSpatialHash.cpp
--- a/Engine/Source/Runtime/Engine/Private/WorldPartition/WorldPartitionRuntimeSpatialHash.cpp
+++ b/Engine/Source/Runtime/Engine/Private/WorldPartition/WorldPartitionRuntimeSpatialHash.cpp
@@ -1250,7 +1250,7 @@ void UWorldPartitionRuntimeSpatialHash::DumpStateLog(FHierarchicalLogArchive& Ar
 
 				for (UWorldPartitionRuntimeCell* Cell : Cells)
 				{
-					FHierarchicalLogArchive::FIndentScope CellIndentScope = Ar.PrintfIndent(TEXT("Content of Cell %s"), *Cell->GetDebugName());
+					FHierarchicalLogArchive::FIndentScope CellIndentScope = Ar.PrintfIndent(TEXT("Content of Cell %s (%s)"), *Cell->GetDebugName(), *Cell->GetName());
 					Cell->DumpStateLog(Ar);
 				}
 			}
diff --git a/Engine/Source/Runtime/Engine/Private/WorldPartition/WorldPartitionStreamingGeneration.cpp b/Engine/Source/Runtime/Engine/Private/WorldPartition/WorldPartitionStreamingGeneration.cpp
--- a/Engine/Source/Runtime/Engine/Private/WorldPartition/WorldPartitionStreamingGeneration.cpp
+++ b/Engine/Source/Runtime/Engine/Private/WorldPartition/WorldPartitionStreamingGeneration.cpp
@@ -932,20 +932,15 @@ public:
 
 	static TUniquePtr<FArchive> CreateDumpStateLogArchive(const TCHAR* Suffix, bool bTimeStamped = true)
 	{
-		if (!GIsBuildMachine)
-		{
-			FString StateLogOutputFilename = FPaths::ProjectSavedDir() / TEXT("WorldPartition") / FString::Printf(TEXT("StreamingGeneration-%s"), Suffix);
-
-			if (bTimeStamped)
-			{			
-				StateLogOutputFilename += FString::Printf(TEXT("-%08x-%s"), FPlatformProcess::GetCurrentProcessId(), *FDateTime::Now().ToIso8601().Replace(TEXT(":"), TEXT(".")));
-			}
+		FString StateLogOutputFilename = FPaths::ProjectLogDir() / TEXT("WorldPartition") / FString::Printf(TEXT("StreamingGeneration-%s"), Suffix);
 
-			StateLogOutputFilename += TEXT(".log");
-			return TUniquePtr<FArchive>(IFileManager::Get().CreateFileWriter(*StateLogOutputFilename));
+		if (bTimeStamped)
+		{			
+			StateLogOutputFilename += FString::Printf(TEXT("-%08x-%s"), FPlatformProcess::GetCurrentProcessId(), *FDateTime::Now().ToIso8601().Replace(TEXT(":"), TEXT(".")));
 		}
 
-		return MakeUnique<FArchive>();
+		StateLogOutputFilename += TEXT(".log");
+		return TUniquePtr<FArchive>(IFileManager::Get().CreateFileWriter(*StateLogOutputFilename));
 	}
 
 	void DumpStateLog(FHierarchicalLogArchive& Ar)
