[shaders] deprecate GetDefinitions accessor on FShaderCompilerEnvironment; encapsulate internal usage of definitions object in friend classes and remove some superfluous code calling the function unnecessarily

diff --git a/Engine/Source/Developer/ShaderPreprocessor/Private/ShaderPreprocessor.cpp b/Engine/Source/Developer/ShaderPreprocessor/Private/ShaderPreprocessor.cpp
--- a/Engine/Source/Developer/ShaderPreprocessor/Private/ShaderPreprocessor.cpp
+++ b/Engine/Source/Developer/ShaderPreprocessor/Private/ShaderPreprocessor.cpp
@@ -209,22 +209,42 @@ FMcppAllocator GMcppAlloc;
 
 #endif
 
-static void DumpShaderDefinesAsCommentedCode(const FShaderCompilerInput& ShaderInput, FString* OutDefines)
+static void AddStbDefine(stb_arena* MacroArena, macro_definition**& StbDefines, const TCHAR* Name, const TCHAR* Value);
+static void AddStbDefines(stb_arena* MacroArena, macro_definition**& StbDefines, TMap<FString, FString> DefinitionsMap);
+
+class FShaderPreprocessorUtilities
 {
-	const TMap<FString, FString>& Definitions = ShaderInput.Environment.GetDefinitions();
+public:
+	static void DumpShaderDefinesAsCommentedCode(const FShaderCompilerInput& ShaderInput, FString* OutDefines)
+	{
+		const TMap<FString, FString>& Definitions = ShaderInput.Environment.Definitions.GetDefinitionMap();
 
-	TArray<FString> Keys;
-	Definitions.GetKeys(/* out */ Keys);
-	Keys.Sort();
+		TArray<FString> Keys;
+		Definitions.GetKeys(/* out */ Keys);
+		Keys.Sort();
 
-	FString Defines;
-	for (const FString& Key : Keys)
+		FString Defines;
+		for (const FString& Key : Keys)
+		{
+			Defines += FString::Printf(TEXT("// #define %s %s\n"), *Key, *Definitions[Key]);
+		}
+
+		*OutDefines = MakeInjectedShaderCodeBlock(TEXT("DumpShaderDefinesAsCommentedCode"), Defines);
+	}
+
+	static void PopulateDefinesMcpp(const FShaderCompilerInput& Input, const FShaderCompilerDefinitions& AdditionalDefines, TArray<TArray<ANSICHAR>>& OutDefines)
 	{
-		Defines += FString::Printf(TEXT("// #define %s %s\n"), *Key, *Definitions[Key]);
+		AddMcppDefines(OutDefines, Input.Environment.Definitions.GetDefinitionMap());
+		AddMcppDefines(OutDefines, AdditionalDefines.GetDefinitionMap());
 	}
 
-	*OutDefines = MakeInjectedShaderCodeBlock(TEXT("DumpShaderDefinesAsCommentedCode"), Defines);
-}
+	static void PopulateDefinesStb(const FShaderCompilerInput& Input, const FShaderCompilerDefinitions& AdditionalDefines, stb_arena* MacroArena, macro_definition**& OutDefines)
+	{
+		AddStbDefines(MacroArena, OutDefines, Input.Environment.Definitions.GetDefinitionMap());
+		AddStbDefines(MacroArena, OutDefines, AdditionalDefines.GetDefinitionMap());
+		AddStbDefine(MacroArena, OutDefines, TEXT("_STB_PREPROCESS"), TEXT("1"));
+	}
+};
 
 //////////////////////////////////////////////////////////////////////////
 
@@ -244,8 +264,7 @@ bool InnerPreprocessShaderMcpp(
 		FMcppFileLoader FileLoader(ShaderInput, ShaderOutput);
 
 		TArray<TArray<ANSICHAR>> McppOptions;
-		AddMcppDefines(McppOptions, ShaderInput.Environment.GetDefinitions());
-		AddMcppDefines(McppOptions, AdditionalDefines.GetDefinitionMap());
+		FShaderPreprocessorUtilities::PopulateDefinesMcpp(ShaderInput, AdditionalDefines, McppOptions);
 
 		// MCPP is not threadsafe.
 
@@ -569,10 +588,7 @@ bool InnerPreprocessShaderStb(
 {
 	stb_arena MacroArena = { 0 };
 	macro_definition** StbDefines = nullptr;
-
-	AddStbDefines(&MacroArena, StbDefines, ShaderInput.Environment.GetDefinitions());
-	AddStbDefines(&MacroArena, StbDefines, AdditionalDefines.GetDefinitionMap());
-	AddStbDefine(&MacroArena, StbDefines, TEXT("_STB_PREPROCESS"), TEXT("1"));
+	FShaderPreprocessorUtilities::PopulateDefinesStb(ShaderInput, AdditionalDefines, &MacroArena, StbDefines);
 
 	FStbPreprocessContext Context{ ShaderInput };
 
@@ -670,7 +686,7 @@ bool PreprocessShader(
 	// List the defines used for compilation in the preprocessed shaders, especially to know witch permutation vector this shader is.
 	if (DefinesPolicy == EDumpShaderDefines::AlwaysIncludeDefines || (DefinesPolicy == EDumpShaderDefines::DontCare && ShaderInput.DumpDebugInfoPath.Len() > 0))
 	{
-		DumpShaderDefinesAsCommentedCode(ShaderInput, &OutPreprocessedShader);
+		FShaderPreprocessorUtilities::DumpShaderDefinesAsCommentedCode(ShaderInput, &OutPreprocessedShader);
 	}
 
 	OutPreprocessedShader += PreprocessorOutput;
diff --git a/Engine/Source/Runtime/Engine/Private/ShaderCompiler/ShaderGenerationUtil.cpp b/Engine/Source/Runtime/Engine/Private/ShaderCompiler/ShaderGenerationUtil.cpp
--- a/Engine/Source/Runtime/Engine/Private/ShaderCompiler/ShaderGenerationUtil.cpp
+++ b/Engine/Source/Runtime/Engine/Private/ShaderCompiler/ShaderGenerationUtil.cpp
@@ -41,12 +41,12 @@ bool NeedsVelocityDepth(EShaderPlatform TargetPlatform)
 
 #if WITH_EDITOR
 
-static int FetchCompileInt(const FShaderCompilerEnvironment& OutEnvironment, const char* SrcName)
+int FShaderCompileUtilities::FetchCompileInt(const FShaderCompilerEnvironment& OutEnvironment, const char* SrcName)
 {
 	int32 Ret = 0;
-	if (OutEnvironment.GetDefinitions().Contains(SrcName))
+	if (OutEnvironment.Definitions.GetDefinitionMap().Contains(SrcName))
 	{
-		const FString* Str = OutEnvironment.GetDefinitions().Find(SrcName);
+		const FString* Str = OutEnvironment.Definitions.GetDefinitionMap().Find(SrcName);
 
 		Ret = atoi((const char*)Str->GetCharArray().GetData());
 	}
@@ -57,15 +57,10 @@ static int FetchCompileInt(const FShaderCompilerEnvironment& OutEnvironment, con
 
 #define SET_COMPILE_BOOL_IF_TRUE(X) { if (DerivedDefines.X) { OutEnvironment.SetDefine(TEXT(#X),TEXT("1")); } }
 
-#define FETCH_COMPILE_BOOL(X) { if (OutEnvironment.GetDefinitions().Contains(#X)) SrcDefines.X = FetchCompileInt(OutEnvironment,#X) != 0 ? 1 : 0; }
-#define FETCH_COMPILE_INT(X) { if (OutEnvironment.GetDefinitions().Contains(#X)) SrcDefines.X = FetchCompileInt(OutEnvironment,#X); }
+#define FETCH_COMPILE_BOOL(X) { if (Environment.Definitions.GetDefinitionMap().Contains(#X)) SrcDefines.X = FetchCompileInt(Environment,#X) != 0 ? 1 : 0; }
+#define FETCH_COMPILE_INT(X) { if (Environment.Definitions.GetDefinitionMap().Contains(#X)) SrcDefines.X = FetchCompileInt(Environment,#X); }
 
-#define MERGE_COMPILE_BOOL(X) (Lhs.X = (Lhs.X | Rhs.X))
-
-// these definitions can't be merged, so one of them should be cleared at zero.
-#define MERGE_COMPILE_INT(X) { check(Lhs.X == 0 || Rhs.X == 0); Lhs.X = Lhs.X > Rhs.X ? Lhs.X : Rhs.X; }
-
-void FShaderCompileUtilities::ApplyFetchEnvironment(FShaderGlobalDefines& SrcDefines, FShaderCompilerEnvironment& OutEnvironment, const EShaderPlatform Platform)
+void FShaderCompileUtilities::ApplyFetchEnvironment(FShaderGlobalDefines& SrcDefines, const FShaderCompilerEnvironment& Environment, const EShaderPlatform Platform)
 {
 	FETCH_COMPILE_BOOL(GBUFFER_HAS_VELOCITY);
 	FETCH_COMPILE_BOOL(GBUFFER_HAS_TANGENT);
@@ -98,7 +93,7 @@ void FShaderCompileUtilities::ApplyFetchEnvironment(FShaderGlobalDefines& SrcDef
 	}
 }
 
-void FShaderCompileUtilities::ApplyFetchEnvironment(FShaderLightmapPropertyDefines& SrcDefines, FShaderCompilerEnvironment& OutEnvironment)
+void FShaderCompileUtilities::ApplyFetchEnvironment(FShaderLightmapPropertyDefines& SrcDefines, const FShaderCompilerEnvironment& Environment)
 {
 	FETCH_COMPILE_BOOL(LQ_TEXTURE_LIGHTMAP);
 	FETCH_COMPILE_BOOL(HQ_TEXTURE_LIGHTMAP);
@@ -125,7 +120,7 @@ void FShaderCompileUtilities::ApplyFetchEnvironment(FShaderLightmapPropertyDefin
 
 }
 
-void FShaderCompileUtilities::ApplyFetchEnvironment(FShaderMaterialPropertyDefines& SrcDefines, FShaderCompilerEnvironment& OutEnvironment)
+void FShaderCompileUtilities::ApplyFetchEnvironment(FShaderMaterialPropertyDefines& SrcDefines, const FShaderCompilerEnvironment& Environment)
 {
 	FETCH_COMPILE_BOOL(MATERIAL_ENABLE_TRANSLUCENCY_FOGGING);
 	FETCH_COMPILE_BOOL(MATERIALBLENDING_ANY_TRANSLUCENT);
@@ -241,7 +236,7 @@ void FShaderCompileUtilities::ApplyFetchEnvironment(FShaderMaterialPropertyDefin
 	FETCH_COMPILE_INT(GBUFFER_LAYOUT);
 }
 
-void FShaderCompileUtilities::ApplyFetchEnvironment(FShaderCompilerDefines& SrcDefines, FShaderCompilerEnvironment& OutEnvironment)
+void FShaderCompileUtilities::ApplyFetchEnvironment(FShaderCompilerDefines& SrcDefines, const FShaderCompilerEnvironment& Environment)
 {
 	FETCH_COMPILE_BOOL(COMPILER_GLSL_ES3_1);
 	FETCH_COMPILE_BOOL(ES3_1_PROFILE);
diff --git a/Engine/Source/Runtime/Engine/Public/ShaderCompiler.h b/Engine/Source/Runtime/Engine/Public/ShaderCompiler.h
--- a/Engine/Source/Runtime/Engine/Public/ShaderCompiler.h
+++ b/Engine/Source/Runtime/Engine/Public/ShaderCompiler.h
@@ -359,37 +359,36 @@ private:
 	virtual void OnMachineResourcesChanged() override;
 };
 
-namespace FShaderCompileUtilities
+class FShaderCompileUtilities
 {
-	bool DoWriteTasks(const TArray<FShaderCommonCompileJobPtr>& QueuedJobs, FArchive& TransferFile, IDistributedBuildController* BuildDistributionController = nullptr, bool bUseRelativePaths = false, bool bCompressTaskFile = false);
-	void DoReadTaskResults(const TArray<FShaderCommonCompileJobPtr>& QueuedJobs, FArchive& OutputFile);
+public:
+	static bool DoWriteTasks(const TArray<FShaderCommonCompileJobPtr>& QueuedJobs, FArchive& TransferFile, IDistributedBuildController* BuildDistributionController = nullptr, bool bUseRelativePaths = false, bool bCompressTaskFile = false);
+	static void DoReadTaskResults(const TArray<FShaderCommonCompileJobPtr>& QueuedJobs, FArchive& OutputFile);
 
 	/** Execute the specified (single or pipeline) shader compile job. */
-	void ExecuteShaderCompileJob(FShaderCommonCompileJob& Job);
-
-	class FArchive* CreateFileHelper(const FString& Filename);
-	void MoveFileHelper(const FString& To, const FString& From);
-	void DeleteFileHelper(const FString& Filename);
+	static void ExecuteShaderCompileJob(FShaderCommonCompileJob& Job);
 
-	ENGINE_API void GenerateBrdfHeaders(const EShaderPlatform Platform);
-	ENGINE_API void GenerateBrdfHeaders(const FName& ShaderFormat);
-	void ApplyDerivedDefines(FShaderCompilerEnvironment& OutEnvironment, FShaderCompilerEnvironment * SharedEnvironment, const EShaderPlatform Platform);
-	void AppendGBufferDDCKeyString(const EShaderPlatform Platform, FString& KeyString);
-	ENGINE_API void WriteGBufferInfoAutogen(EShaderPlatform TargetPlatform, ERHIFeatureLevel::Type FeatureLevel);
+	static class FArchive* CreateFileHelper(const FString& Filename);
+	static void MoveFileHelper(const FString& To, const FString& From);
+	static void DeleteFileHelper(const FString& Filename);
 
-	void ApplyFetchEnvironment(FShaderMaterialPropertyDefines& DefineData, FShaderCompilerEnvironment& OutEnvironment);
-	void ApplyFetchEnvironment(FShaderGlobalDefines& DefineData, FShaderCompilerEnvironment& OutEnvironment, const EShaderPlatform Platform);
-	void ApplyFetchEnvironment(FShaderLightmapPropertyDefines& DefineData, FShaderCompilerEnvironment& OutEnvironment);
-	void ApplyFetchEnvironment(FShaderCompilerDefines& DefineData, FShaderCompilerEnvironment& OutEnvironment);
+	static ENGINE_API void GenerateBrdfHeaders(const EShaderPlatform Platform);
+	static ENGINE_API void GenerateBrdfHeaders(const FName& ShaderFormat);
+	static void ApplyDerivedDefines(FShaderCompilerEnvironment& OutEnvironment, FShaderCompilerEnvironment* SharedEnvironment, const EShaderPlatform Platform);
+	static void AppendGBufferDDCKeyString(const EShaderPlatform Platform, FString& KeyString);
+	static ENGINE_API void WriteGBufferInfoAutogen(EShaderPlatform TargetPlatform, ERHIFeatureLevel::Type FeatureLevel);
 
-	void ApplyModifyEnvironment(const FShaderMaterialDerivedDefines& DefineData, FShaderCompilerEnvironment& OutEnvironment);
+	static int FetchCompileInt(const FShaderCompilerEnvironment& Enviroment, const char* SrcName);
+	static void ApplyFetchEnvironment(FShaderMaterialPropertyDefines& DefineData, const FShaderCompilerEnvironment& Environment);
+	static void ApplyFetchEnvironment(FShaderGlobalDefines& DefineData, const FShaderCompilerEnvironment& Environment, const EShaderPlatform Platform);
+	static void ApplyFetchEnvironment(FShaderLightmapPropertyDefines& DefineData, const FShaderCompilerEnvironment& Environment);
+	static void ApplyFetchEnvironment(FShaderCompilerDefines& DefineData, const FShaderCompilerEnvironment& Environment);
 
-	ENGINE_API EGBufferLayout FetchGBufferLayout(const FShaderCompilerEnvironment& Environment);
+	static ENGINE_API EGBufferLayout FetchGBufferLayout(const FShaderCompilerEnvironment& Environment);
 
-	ENGINE_API FGBufferParams FetchGBufferParamsRuntime(EShaderPlatform Platform, EGBufferLayout Layout); // this function is called from renderer
-	FGBufferParams FetchGBufferParamsPipeline(EShaderPlatform Platform, EGBufferLayout Layout);
-
-}
+	static ENGINE_API FGBufferParams FetchGBufferParamsRuntime(EShaderPlatform Platform, EGBufferLayout Layout); // this function is called from renderer
+	static FGBufferParams FetchGBufferParamsPipeline(EShaderPlatform Platform, EGBufferLayout Layout);
+};
 
 class FShaderCompileDistributedThreadRunnable_Interface : public FShaderCompileThreadRunnableBase
 {
diff --git a/Engine/Source/Runtime/RenderCore/Public/ShaderCore.h b/Engine/Source/Runtime/RenderCore/Public/ShaderCore.h
--- a/Engine/Source/Runtime/RenderCore/Public/ShaderCore.h
+++ b/Engine/Source/Runtime/RenderCore/Public/ShaderCore.h
@@ -43,6 +43,8 @@ class Error;
 class FMemoryImageWriter;
 class FMemoryUnfreezeContent;
 class FPointerTableBase;
+class FShaderCompileUtilities;
+class FShaderPreprocessorUtilities;
 class FSHA1;
 class ITargetPlatform;
 
@@ -610,6 +612,7 @@ struct FShaderCompilerEnvironment
 		return false;
 	}
 
+	UE_DEPRECATED(5.3, "GetDefinitions is deprecated; preprocessor defines must now only be accessed by core shader system code. Use Get/SetCompileArgument for generic params instead.")
 	const TMap<FString,FString>& GetDefinitions() const
 	{
 		return Definitions.GetDefinitionMap();
@@ -676,6 +679,9 @@ struct FShaderCompilerEnvironment
 
 private:
 
+	friend class FShaderCompileUtilities;
+	friend class FShaderPreprocessorUtilities;
+
 	FShaderCompilerDefinitions Definitions;
 	TMap<FString, TVariant<bool, float, int32, uint32>> CompileArgs;
 };
