LevelInstanceActorDesc: DeltaSerialize Filter

#rnx

diff --git a/Engine/Source/Runtime/Core/Public/UObject/FortniteMainBranchObjectVersion.h b/Engine/Source/Runtime/Core/Public/UObject/FortniteMainBranchObjectVersion.h
--- a/Engine/Source/Runtime/Core/Public/UObject/FortniteMainBranchObjectVersion.h
+++ b/Engine/Source/Runtime/Core/Public/UObject/FortniteMainBranchObjectVersion.h
@@ -289,6 +289,9 @@ struct CORE_API FFortniteMainBranchObjectVersion
 		// Reduce size and improve behaviour of Chaos::FImplicitObjectUnion
 		ChaosImplicitObjectUnionBVHRefactor,
 
+		// FLevelInstanceActorDesc DeltaSerialize Filter
+		LevelInstanceActorDescDeltaSerializeFilter,
+
 		// -----<new versions can be added above this line>-------------------------------------------------
 		VersionPlusOne,
 		LatestVersion = VersionPlusOne - 1
diff --git a/Engine/Source/Runtime/Engine/Private/WorldPartition/LevelInstance/LevelInstanceActorDesc.cpp b/Engine/Source/Runtime/Engine/Private/WorldPartition/LevelInstance/LevelInstanceActorDesc.cpp
--- a/Engine/Source/Runtime/Engine/Private/WorldPartition/LevelInstance/LevelInstanceActorDesc.cpp
+++ b/Engine/Source/Runtime/Engine/Private/WorldPartition/LevelInstance/LevelInstanceActorDesc.cpp
@@ -292,7 +292,14 @@ void FLevelInstanceActorDesc::Serialize(FArchive& Ar)
 
 	if (Ar.CustomVer(FFortniteMainBranchObjectVersion::GUID) >= FFortniteMainBranchObjectVersion::WorldPartitionActorFilter)
 	{
-		Ar << Filter;
+		if(Ar.CustomVer(FFortniteMainBranchObjectVersion::GUID) >= FFortniteMainBranchObjectVersion::LevelInstanceActorDescDeltaSerializeFilter)
+		{
+			Ar << TDeltaSerialize<FWorldPartitionActorFilter>(Filter);
+		}
+		else
+		{
+			Ar << Filter;
+		}
 	}
 }
 
