Improve the adaptive timeslice:
- Add an additional (default) method that computes the timeslice based on total time for all levels in flight rather than just remaining time. This avoids us backing off before before the current level has completed. This seems more elegant than relying on TimeSliceMaxReducePerSecond to slowly reduce the timeslice.
- Discount frames where we're not using the full timeslice rather than basing this off the queue size, and remove the corresponding cvar

diff --git a/Engine/Source/Runtime/Engine/Classes/Engine/Level.h b/Engine/Source/Runtime/Engine/Classes/Engine/Level.h
--- a/Engine/Source/Runtime/Engine/Classes/Engine/Level.h
+++ b/Engine/Source/Runtime/Engine/Classes/Engine/Level.h
@@ -1310,8 +1310,10 @@ public:
 	virtual UAssetUserData* GetAssetUserDataOfClass(TSubclassOf<UAssetUserData> InUserDataClass) override;
 	//~ End IInterface_AssetUserData Interface
 
-	/** Estimate the amount of work remaining for AddToWorld for this level. Used by the adaptive level streaming timeslice (see s.AdaptiveAddToWorld.Enabled) */
+	/** Estimate the amount of AddToWorld work for this level. Used by the adaptive level streaming timeslice (see s.AdaptiveAddToWorld.Enabled) */
 	int32 GetEstimatedAddToWorldWorkUnitsRemaining() const;
+	/** Estimate the total amount of AddToWorld work for this level. Used by the adaptive level streaming timeslice (see s.AdaptiveAddToWorld.Enabled) */
+	int32 GetEstimatedAddToWorldWorkUnitsTotal() const;
 
 #if WITH_EDITOR
 	/** meant to be called only from editor, calculating and storing static geometry to be used with off-line and/or on-line navigation building */
diff --git a/Engine/Source/Runtime/Engine/Private/Level.cpp b/Engine/Source/Runtime/Engine/Private/Level.cpp
--- a/Engine/Source/Runtime/Engine/Private/Level.cpp
+++ b/Engine/Source/Runtime/Engine/Private/Level.cpp
@@ -3301,6 +3301,12 @@ void ULevel::ResetRouteActorInitializationState()
 	 RouteActorInitializationIndex = 0;
 }
 
+int32 ULevel::GetEstimatedAddToWorldWorkUnitsTotal() const
+{
+	const int32 ActorCount = Actors.Num();
+	return ActorCount + ActorCount * GRouteActorInitializationWorkUnitWeighting;
+}
+
 int32 ULevel::GetEstimatedAddToWorldWorkUnitsRemaining() const
 {
 	// We count "work units" as the number of actors requiring component update, plus the number of actors requiring initialization
diff --git a/Engine/Source/Runtime/Engine/Private/World.cpp b/Engine/Source/Runtime/Engine/Private/World.cpp
--- a/Engine/Source/Runtime/Engine/Private/World.cpp
+++ b/Engine/Source/Runtime/Engine/Private/World.cpp
@@ -155,9 +155,11 @@ static int32 GAdaptiveAddToWorldEnabled = 0;
 FAutoConsoleVariableRef CVarAdaptiveAddToWorldEnabled(TEXT("s.AdaptiveAddToWorld.Enabled"), GAdaptiveAddToWorldEnabled, 
 	TEXT("Enables the adaptive AddToWorld timeslice (replaces s.LevelStreamingActorsUpdateTimeLimit) (default: off)"));
 
-static int32 GAdaptiveAddToWorldMinWorkUnitCount = 100;
-FAutoConsoleVariableRef CVarAdaptiveAddToWorldMinWorkUnitCount(TEXT("s.AdaptiveAddToWorld.MinWorkUnitCount"), GAdaptiveAddToWorldMinWorkUnitCount, 
-	TEXT("Work unit threshold at which the adaptive timeslice kicks in. If we're below this, we'll just use the min timeslice"));
+static int32 GAdaptiveAddToWorldMethod = 1;
+FAutoConsoleVariableRef CVarAdaptiveAddToWorldMethod(TEXT("s.AdaptiveAddToWorld.Method"), GAdaptiveAddToWorldMethod,
+	TEXT("Heuristic to use for the adaptive timeslice\n")
+	TEXT("0 - compute the target timeslice based on remaining work time\n")
+	TEXT("1 - compute the target timeslice based on total work time for levels in flight (this avoids slowing before a level completes)\n"));
 
 static float GAdaptiveAddToWorldTimeSliceMin = 1.0f;
 FAutoConsoleVariableRef CVarAdaptiveAddToWorldTimeSliceMin(TEXT("s.AdaptiveAddToWorld.AddToWorldTimeSliceMin"), GAdaptiveAddToWorldTimeSliceMin, 
@@ -210,6 +212,7 @@ public:
 		{
 		    WorkUnitsProcessedThisUpdate = 0;
 		    TotalWorkUnitsRemaining = 0;
+			TotalWorkUnitsForLevelsInFlight = 0;
 		}
 	}
 
@@ -225,7 +228,8 @@ public:
 		// Only count frames where we have a reasonable number of work units to process
 		double CurrentTimestamp = FPlatformTime::Seconds();
 		float DeltaT = float(CurrentTimestamp - LastUpdateTimestamp);
-		if (TotalWorkUnitsRemaining >= GAdaptiveAddToWorldMinWorkUnitCount)
+		// Only consider frames where we're using the full timeslice
+		if (TotalWorkUnitsRemaining > 0)
 		{
 			FHistoryEntry& HistoryEntry = UpdateHistory[HistoryFrameIndex];
 			HistoryEntry.WorkUnitsProcessed = WorkUnitsProcessedThisUpdate;
@@ -233,34 +237,45 @@ public:
 			HistoryEntry.TimeSlice = GetTimeSlice();
 			HistoryFrameIndex = (HistoryFrameIndex + 1) % HistorySize;
 
-			// Make sure we have a complete history buffer before 
+			// Make sure we have a complete history buffer before calculating
 			if (ValidFrameCount < HistorySize)
 			{
+				BaseTimeSlice = GAdaptiveAddToWorldTimeSliceMin;
 				ValidFrameCount++;
 			}
 			else
 			{
-				// Compute the rolling average
+				// Compute the rolling averages
 				int32 TotalWorkUnits = 0;
 				float TotalTime = 0.0f;
-				float TotalWorkTimeMs = 0.0f;
+				float TotalTimeSliceMs = 0.0f;
 				for (int i = 0; i < HistorySize; i++)
 				{
 					TotalWorkUnits += UpdateHistory[i].WorkUnitsProcessed;
 					TotalTime += UpdateHistory[i].DeltaT;
-					TotalWorkTimeMs += UpdateHistory[i].TimeSlice;
+					TotalTimeSliceMs += UpdateHistory[i].TimeSlice;
 				}
 
-				float AverageTimeSliceMs = TotalWorkTimeMs / float(HistorySize);
+				float AverageTimeSliceMs = TotalTimeSliceMs / float(HistorySize);
 
 				float AverageWorkUnitsPerSecondWith1MsTimeSlice = (float)TotalWorkUnits / ( TotalTime * AverageTimeSliceMs );
 				CSV_CUSTOM_STAT(LevelStreamingAdaptiveDetail, AverageWorkUnitsPerSecondWith1MsTimeSlice, AverageWorkUnitsPerSecondWith1MsTimeSlice, ECsvCustomStatOp::Set);
 
 				// Work out what the timeslice should be
-				float EstimatedTimeRemainingFor1MsTimeSlice = (float)TotalWorkUnitsRemaining / AverageWorkUnitsPerSecondWith1MsTimeSlice;
-				CSV_CUSTOM_STAT(LevelStreamingAdaptiveDetail, EstimatedTimeRemainingFor1MsTimeSlice, EstimatedTimeRemainingFor1MsTimeSlice, ECsvCustomStatOp::Set);
+				float TargetTimeSlice = 0.0f;
+				if (GAdaptiveAddToWorldMethod == 1)
+				{
+					float EstimatedTotalTimeForCurrentLevelsWith1MsTimeSlice = (float)TotalWorkUnitsForLevelsInFlight / AverageWorkUnitsPerSecondWith1MsTimeSlice;
+					CSV_CUSTOM_STAT(LevelStreamingAdaptiveDetail, EstimatedTotalTimeForCurrentLevelsWith1MsTimeSlice, EstimatedTotalTimeForCurrentLevelsWith1MsTimeSlice, ECsvCustomStatOp::Set);
+					TargetTimeSlice = EstimatedTotalTimeForCurrentLevelsWith1MsTimeSlice / GAdaptiveAddToWorldTargetMaxTimeRemaining - ExtraTimeSlice;
+				}
+				else
+				{
+					float EstimatedTimeRemainingWith1MsTimeSlice = (float)TotalWorkUnitsRemaining / AverageWorkUnitsPerSecondWith1MsTimeSlice;
+					CSV_CUSTOM_STAT(LevelStreamingAdaptiveDetail, EstimatedTimeRemainingWith1MsTimeSlice, EstimatedTimeRemainingWith1MsTimeSlice, ECsvCustomStatOp::Set);
+					TargetTimeSlice = EstimatedTimeRemainingWith1MsTimeSlice / GAdaptiveAddToWorldTargetMaxTimeRemaining - ExtraTimeSlice;
+				}
 
-				float TargetTimeSlice = EstimatedTimeRemainingFor1MsTimeSlice / GAdaptiveAddToWorldTargetMaxTimeRemaining - ExtraTimeSlice;
 				float DesiredBaseTimeSlice = FMath::Clamp(TargetTimeSlice, GAdaptiveAddToWorldTimeSliceMin, GAdaptiveAddToWorldTimeSliceMax);
 
 				// Limit the timeslice change based on MaxIncreasePerSecond/ MaxReducePerSecond
@@ -285,16 +300,18 @@ public:
 		CSV_CUSTOM_STAT(LevelStreamingAdaptiveDetail, WorkUnitsProcessedThisUpdate, WorkUnitsProcessedThisUpdate, ECsvCustomStatOp::Set);
 		CSV_CUSTOM_STAT(LevelStreamingAdaptive, TimeSlice, GetTimeSlice(), ECsvCustomStatOp::Set);
 		CSV_CUSTOM_STAT(LevelStreamingAdaptive, WorkUnitsRemaining, TotalWorkUnitsRemaining, ECsvCustomStatOp::Set);
+		CSV_CUSTOM_STAT(LevelStreamingAdaptive, TotalWorkUnitsForLevelsInFlight, TotalWorkUnitsForLevelsInFlight, ECsvCustomStatOp::Set);
 		LastUpdateTimestamp = CurrentTimestamp;
 	}
 
-	void RegisterAddToWorldWorkUnits(int32 StartWorkUnits, int32 EndWorkUnits)
+	void RegisterAddToWorldWork(int32 StartWorkUnits, int32 EndWorkUnits, int32 TotalWorkUnits)
 	{
 		WorkUnitsProcessedThisUpdate += StartWorkUnits - EndWorkUnits;
 		TotalWorkUnitsRemaining += EndWorkUnits;
+		TotalWorkUnitsForLevelsInFlight += TotalWorkUnits;
 	}
 
-	float GetTimeSlice()
+	float GetTimeSlice() const
 	{
 		return ExtraTimeSlice + BaseTimeSlice;
 	}
@@ -310,6 +327,7 @@ private:
 
 	int32 WorkUnitsProcessedThisUpdate = 0;
 	int32 TotalWorkUnitsRemaining = 0;
+	int32 TotalWorkUnitsForLevelsInFlight = 0;
 	int32 HistoryFrameIndex = 0;
 	FHistoryEntry UpdateHistory[HistorySize];
 	float BaseTimeSlice = 1.0f;
@@ -3354,15 +3372,15 @@ void UWorld::AddToWorld( ULevel* Level, const FTransform& LevelTransform, bool b
 	}
 #endif // PERF_TRACK_DETAILED_ASYNC_STATS
 
-    // Register work completed with adaptive level streaming
+	// Delta time in ms.
+	double DeltaTime = (FPlatformTime::Seconds() - StartTime) * 1000; 
+	GAddToWorldTimeCumul += DeltaTime;
+
+	// Register work completed with adaptive level streaming
 	if (GAdaptiveAddToWorld.IsEnabled())
 	{
-		GAdaptiveAddToWorld.RegisterAddToWorldWorkUnits(StartWorkUnitsRemaining, Level->GetEstimatedAddToWorldWorkUnitsRemaining());
+		GAdaptiveAddToWorld.RegisterAddToWorldWork(StartWorkUnitsRemaining, Level->GetEstimatedAddToWorldWorkUnitsRemaining(), Level->GetEstimatedAddToWorldWorkUnitsTotal());
 	}
-
-	// Delta time in ms.
-	double DeltaTime = (FPlatformTime::Seconds() - StartTime) * 1000;
-	GAddToWorldTimeCumul += DeltaTime;
 }
 
 void UWorld::BeginTearingDown()
@@ -9340,4 +9358,4 @@ FString ENGINE_API ToString(ENetMode NetMode)
 	}
 }
 
-#undef LOCTEXT_NAMESPACE 
+#undef LOCTEXT_NAMESPACE 
\ No newline at end of file
