fix potential crash due to the cluster union marshaling a GT particle pointer from the PT to the GT

diff --git a/Engine/Source/Runtime/Engine/Classes/PhysicsEngine/ClusterUnionComponent.h b/Engine/Source/Runtime/Engine/Classes/PhysicsEngine/ClusterUnionComponent.h
--- a/Engine/Source/Runtime/Engine/Classes/PhysicsEngine/ClusterUnionComponent.h
+++ b/Engine/Source/Runtime/Engine/Classes/PhysicsEngine/ClusterUnionComponent.h
@@ -31,8 +31,12 @@ struct FClusteredComponentData
 {
 	GENERATED_BODY()
 
+	// Set of physics objects that we actually added into the cluster union.
 	TSet<Chaos::FPhysicsObjectHandle> PhysicsObjects;
 
+	// Every physics object associated with this particular component.
+	TArray<Chaos::FPhysicsObjectHandle> AllPhysicsObjects;
+
 	// Using a TWeakObjectPtr here because the UClusterUnionReplicatedProxyComponent will have a pointer back
 	// and we don't want to get into a situation where a circular reference occurs.
 	UPROPERTY()
@@ -81,6 +85,22 @@ struct FClusterUnionPendingAddData
 	TArray<int32> BoneIds;
 };
 
+/**
+ * For every possible particle that could ever possibly be added into the cluster union,
+ * keep track of its component and its bone id.
+ */
+USTRUCT()
+struct FClusterUnionParticleCandidateData
+{
+	GENERATED_BODY()
+
+	UPROPERTY()
+	TObjectPtr<UPrimitiveComponent> Component;
+
+	UPROPERTY()
+	int32 BoneId;
+};
+
 /**
  * This does the bulk of the work exposing a physics cluster union to the game thread.
  * This component needs to be a primitive component primarily because of how physics
@@ -147,11 +167,6 @@ private:
 	// side benefit here is being able to track which components are clustered.
 	TMap<TObjectKey<UPrimitiveComponent>, FClusteredComponentData> ComponentToPhysicsObjects;
 
-	// Similarly, we need to keep track of the mapping from physics object handles to primitive components.
-	// The physics proxy will only have data about these handles, so when we get an update from the proxy
-	// about a particular handle we will want to map it back to a component here.
-	TMap<Chaos::FPhysicsObjectHandle, UPrimitiveComponent*> PhysicsObjectToComponent;
-
 	// Also keep track of which actors we are clustering and their components. We make modifications on
 	// actors that get clustered so we need to make sure we undo those changes only once all its clustered
 	// components are removed from the cluster.
@@ -161,6 +176,12 @@ private:
 	// Make sure we don't try to add the component multiples times while the add is pending.
 	TMap<TObjectKey<UPrimitiveComponent>, FClusterUnionPendingAddData> PendingComponentsToAdd;
 
+	// Given a unique index of a particle that we're adding to the cluster union - map it back to the component that owns it.
+	// This works decently because we assume that when we're using a cluster union component, we will only try to add to the
+	// cluster union via the GT so we can guarantee to have a decent mapping here.
+	UPROPERTY()
+	TMap<int32, FClusterUnionParticleCandidateData> UniqueIdxToComponent;
+
 	// Data that can be changed at runtime to keep state about the cluster union consistent between the server and client.
 	UPROPERTY(ReplicatedUsing=OnRep_RigidState)
 	FClusterUnionReplicatedData ReplicatedRigidState;
diff --git a/Engine/Source/Runtime/Engine/Private/PhysicsEngine/ClusterUnionComponent.cpp b/Engine/Source/Runtime/Engine/Private/PhysicsEngine/ClusterUnionComponent.cpp
--- a/Engine/Source/Runtime/Engine/Private/PhysicsEngine/ClusterUnionComponent.cpp
+++ b/Engine/Source/Runtime/Engine/Private/PhysicsEngine/ClusterUnionComponent.cpp
@@ -14,27 +14,6 @@
 
 DEFINE_LOG_CATEGORY(LogClusterUnion);
 
-namespace
-{
-	// This can probably be generalized into the physics engine somewhere.
-	UPrimitiveComponent* GetComponentFromUserData(void* UserData)
-	{
-		if (!UserData)
-		{
-			return nullptr;
-		}
-		else if (FBodyInstance* BodyInstance = FChaosUserData::Get<FBodyInstance>(UserData))
-		{
-			return BodyInstance->OwnerComponent.Get();
-		}
-		else if (UPrimitiveComponent* Component = FChaosUserData::Get<UPrimitiveComponent>(UserData))
-		{
-			return Component;
-		}
-		return nullptr;
-	}
-}
-
 UClusterUnionComponent::UClusterUnionComponent(const FObjectInitializer& ObjectInitializer)
 	: UPrimitiveComponent(ObjectInitializer)
 {
@@ -75,11 +54,27 @@ void UClusterUnionComponent::AddComponentToCluster(UPrimitiveComponent* InCompon
 
 	PendingComponentsToAdd.Remove(InComponent);
 
-	TArray<Chaos::FPhysicsObjectHandle> Objects;
+	TArray<Chaos::FPhysicsObjectHandle> AllObjects = InComponent->GetAllPhysicsObjects();
+	FLockedReadPhysicsObjectExternalInterface Interface = FPhysicsObjectExternalInterface::LockRead(AllObjects);
 
+	for (Chaos::FPhysicsObjectHandle Handle : AllObjects)
+	{
+		TArray<Chaos::FGeometryParticle*> Particles = Interface->GetAllParticles({ &Handle, 1 });
+		if (Particles.IsEmpty() || !Particles[0])
+		{
+			continue;
+		}
+
+		FClusterUnionParticleCandidateData Data;
+		Data.Component = InComponent;
+		Data.BoneId = Chaos::FPhysicsObjectInterface::GetId(Handle);
+		UniqueIdxToComponent.Add(Particles[0]->UniqueIdx().Idx, Data);
+	}
+
+	TArray<Chaos::FPhysicsObjectHandle> Objects;
 	if (BoneIds.IsEmpty())
 	{
-		Objects = InComponent->GetAllPhysicsObjects();
+		Objects = AllObjects;
 	}
 	else
 	{
@@ -90,8 +85,6 @@ void UClusterUnionComponent::AddComponentToCluster(UPrimitiveComponent* InCompon
 		}
 	}
 
-	FLockedReadPhysicsObjectExternalInterface Interface = FPhysicsObjectExternalInterface::LockRead(Objects);
-
 	if (BoneIds.IsEmpty())
 	{
 		Objects = Objects.FilterByPredicate(
@@ -351,10 +344,9 @@ void UClusterUnionComponent::SyncClusterUnionFromProxy()
 	TMap<TObjectKey<UPrimitiveComponent>, TMap<int32, FTransform>> MappedData;
 	for (const Chaos::FClusterUnionChildData& ChildData : FullData.ChildParticles)
 	{
-		if (UPrimitiveComponent* Component = GetComponentFromUserData(ChildData.Particle->UserData()))
+		if (FClusterUnionParticleCandidateData* Data = UniqueIdxToComponent.Find(ChildData.ParticleIdx.Idx))
 		{
-			const int32 BoneId = Component->GetIdFromGTParticle(ChildData.Particle);
-			MappedData.FindOrAdd(Component).Add(BoneId, ChildData.ChildToParent);
+			MappedData.FindOrAdd(Data->Component).Add(Data->BoneId, ChildData.ChildToParent);
 		}
 	}
 
@@ -467,8 +459,10 @@ void UClusterUnionComponent::HandleAddOrModifiedClusteredComponent(UPrimitiveCom
 	{
 		Chaos::FPhysicsObjectHandle PhysicsObject = ChangedComponent->GetPhysicsObjectById(Kvp.Key);
 		ComponentData.PhysicsObjects.Add(PhysicsObject);
-		PhysicsObjectToComponent.Add(PhysicsObject, ChangedComponent);
 	}
+
+	ComponentData.AllPhysicsObjects.Empty();
+	ComponentData.AllPhysicsObjects.Append(ChangedComponent->GetAllPhysicsObjects());
 }
 
 void UClusterUnionComponent::HandleRemovedClusteredComponent(UPrimitiveComponent* ChangedComponent, bool bDestroyReplicatedProxy)
@@ -488,9 +482,15 @@ void UClusterUnionComponent::HandleRemovedClusteredComponent(UPrimitiveComponent
 
 	if (FClusteredComponentData* ComponentData = ComponentToPhysicsObjects.Find(ChangedComponent))
 	{
-		for (Chaos::FPhysicsObjectHandle Handle : ComponentData->PhysicsObjects)
+		FLockedReadPhysicsObjectExternalInterface Interface = FPhysicsObjectExternalInterface::LockRead(ComponentData->AllPhysicsObjects);
+		for (Chaos::FGeometryParticle* Particle : Interface->GetAllParticles(ComponentData->AllPhysicsObjects))
 		{
-			PhysicsObjectToComponent.Remove(Handle);
+			if (!Particle)
+			{
+				continue;
+			}
+
+			UniqueIdxToComponent.Remove(Particle->UniqueIdx().Idx);
 		}
 
 		if (IsAuthority())
diff --git a/Engine/Source/Runtime/Experimental/Chaos/Private/PhysicsProxy/ClusterUnionPhysicsProxy.cpp b/Engine/Source/Runtime/Experimental/Chaos/Private/PhysicsProxy/ClusterUnionPhysicsProxy.cpp
--- a/Engine/Source/Runtime/Experimental/Chaos/Private/PhysicsProxy/ClusterUnionPhysicsProxy.cpp
+++ b/Engine/Source/Runtime/Experimental/Chaos/Private/PhysicsProxy/ClusterUnionPhysicsProxy.cpp
@@ -294,7 +294,7 @@ namespace Chaos
 		for (const FDirtyClusterUnionParticleData& InData : PullData.ChildParticles)
 		{
 			FClusterUnionChildData ConvertedData;
-			ConvertedData.Particle = InData.Particle;
+			ConvertedData.ParticleIdx = InData.ParticleIdx;
 			ConvertedData.ChildToParent = InData.ChildToParent;
 			SyncedData_External.ChildParticles.Add(ConvertedData);
 		}
@@ -401,13 +401,13 @@ namespace Chaos
 
 			for (FPBDRigidParticleHandle* Particle : ClusterUnion->ChildParticles)
 			{
-				if (!ensure(Particle->GTGeometryParticle()))
+				if (!ensure(Particle))
 				{
 					continue;
 				}
 
 				FDirtyClusterUnionParticleData Data;
-				Data.Particle = Particle->GTGeometryParticle()->CastToRigidParticle();
+				Data.ParticleIdx = Particle->UniqueIdx();
 				if (FPBDRigidClusteredParticleHandle* ClusteredParticle = Particle->CastToClustered())
 				{
 					Data.ChildToParent = ClusteredParticle->ChildToParent();
@@ -432,7 +432,7 @@ namespace Chaos
 		for (const FClusterUnionChildData& Data : SyncedData_External.ChildParticles)
 		{
 			FDirtyClusterUnionParticleData ConvertedData;
-			ConvertedData.Particle = Data.Particle;
+			ConvertedData.ParticleIdx = Data.ParticleIdx;
 			ConvertedData.ChildToParent = Data.ChildToParent;
 			BufferData.ChildParticles.Add(ConvertedData);
 		}
diff --git a/Engine/Source/Runtime/Experimental/Chaos/Public/Chaos/PullPhysicsDataImp.h b/Engine/Source/Runtime/Experimental/Chaos/Public/Chaos/PullPhysicsDataImp.h
--- a/Engine/Source/Runtime/Experimental/Chaos/Public/Chaos/PullPhysicsDataImp.h
+++ b/Engine/Source/Runtime/Experimental/Chaos/Public/Chaos/PullPhysicsDataImp.h
@@ -65,7 +65,7 @@ struct FDirtyGeometryCollectionData : public TBasePullData<FGeometryCollectionPh
 
 struct FDirtyClusterUnionParticleData
 {
-	FPBDRigidParticle* Particle = nullptr;
+	FUniqueIdx ParticleIdx;
 	FRigidTransform3 ChildToParent;
 };
 
diff --git a/Engine/Source/Runtime/Experimental/Chaos/Public/PhysicsProxy/ClusterUnionPhysicsProxy.h b/Engine/Source/Runtime/Experimental/Chaos/Public/PhysicsProxy/ClusterUnionPhysicsProxy.h
--- a/Engine/Source/Runtime/Experimental/Chaos/Public/PhysicsProxy/ClusterUnionPhysicsProxy.h
+++ b/Engine/Source/Runtime/Experimental/Chaos/Public/PhysicsProxy/ClusterUnionPhysicsProxy.h
@@ -19,7 +19,7 @@ namespace Chaos
 
 	struct FClusterUnionChildData
 	{
-		FPBDRigidParticle* Particle = nullptr;
+		FUniqueIdx ParticleIdx;
 		FTransform ChildToParent;
 	};
 
