Undo 24749866 to make a different fix in 5.3 instead

diff --git a/Engine/Source/Runtime/Engine/Private/UnrealEngine.cpp b/Engine/Source/Runtime/Engine/Private/UnrealEngine.cpp
--- a/Engine/Source/Runtime/Engine/Private/UnrealEngine.cpp
+++ b/Engine/Source/Runtime/Engine/Private/UnrealEngine.cpp
@@ -2712,7 +2712,7 @@ void UEngine::ReinitializeCustomTimeStep()
 	{
 		if (bIsCurrentCustomTimeStepInitialized)
 		{
-			CustomTimeStep->Shutdown(this);
+		CustomTimeStep->Shutdown(this);
 		}
 		bIsCurrentCustomTimeStepInitialized = CustomTimeStep->Initialize(this);
 	}
@@ -2720,23 +2720,22 @@ void UEngine::ReinitializeCustomTimeStep()
 
 bool UEngine::SetCustomTimeStep(UEngineCustomTimeStep* InCustomTimeStep)
 {
-	if (CustomTimeStep && bIsCurrentCustomTimeStepInitialized)
+	if (InCustomTimeStep != CustomTimeStep)
 	{
-		CustomTimeStep->Shutdown(this);
-
-		/** Some custom timesteps need some time to release their resources when we shut them down.
-			Until a delegate is added to the custom time steps to know when their shutdown is complete we need to sleep in order for the hardware to complete its uninitialization. **/
-		FPlatformProcess::Sleep(0.2);
-	}
+		if (CustomTimeStep && bIsCurrentCustomTimeStepInitialized)
+		{
+			CustomTimeStep->Shutdown(this);
+		}
 
-	bIsCurrentCustomTimeStepInitialized = false;
-	CustomTimeStep = IsValid(InCustomTimeStep) ? InCustomTimeStep : nullptr;
+		bIsCurrentCustomTimeStepInitialized = false;
+		CustomTimeStep = IsValid(InCustomTimeStep) ? InCustomTimeStep : nullptr;
 
-	if (CustomTimeStep)
-	{
-		bIsCurrentCustomTimeStepInitialized = CustomTimeStep->Initialize(this);
+		if (CustomTimeStep)
+			{
+			bIsCurrentCustomTimeStepInitialized = CustomTimeStep->Initialize(this);
+		}
+		OnCustomTimeStepChanged().Broadcast();
 	}
-	OnCustomTimeStepChanged().Broadcast();
 	return bIsCurrentCustomTimeStepInitialized;
 }
 
@@ -2747,7 +2746,7 @@ void UEngine::ReinitializeTimecodeProvider()
 	{
 		if (bIsCurrentTimecodeProviderInitialized)
 		{
-			Provider->Shutdown(this);
+	Provider->Shutdown(this);
 		}
 		bIsCurrentTimecodeProviderInitialized = Provider->Initialize(this);
 	}
@@ -2755,24 +2754,22 @@ void UEngine::ReinitializeTimecodeProvider()
 
 bool UEngine::SetTimecodeProvider(UTimecodeProvider* InTimecodeProvider)
 {
-	if (TimecodeProvider && bIsCurrentTimecodeProviderInitialized)
-	{
-		TimecodeProvider->Shutdown(this);
-		/** Some timecode providers need some time to release their resources when we shut them down.
-			Until a delegate is added to the timecode providers to know when their shutdown is complete, we need to sleep in order for the hardware to complete its uninitialization. **/
-		// 
-		FPlatformProcess::Sleep(0.2);
-	}
+	if (InTimecodeProvider != TimecodeProvider)
+			{
+		if (TimecodeProvider && bIsCurrentTimecodeProviderInitialized)
+		{
+			TimecodeProvider->Shutdown(this);
+		}
 
-	bIsCurrentTimecodeProviderInitialized = false;
-	TimecodeProvider = IsValid(InTimecodeProvider) ? InTimecodeProvider : nullptr;
+		bIsCurrentTimecodeProviderInitialized = false;
+		TimecodeProvider = IsValid(InTimecodeProvider) ? InTimecodeProvider : nullptr;
 
-	if (TimecodeProvider)
-	{
-		bIsCurrentTimecodeProviderInitialized = TimecodeProvider->Initialize(this);
+		if (TimecodeProvider)
+			{
+			bIsCurrentTimecodeProviderInitialized = TimecodeProvider->Initialize(this);
+		}
+		OnTimecodeProviderChanged().Broadcast();
 	}
-	OnTimecodeProviderChanged().Broadcast();
-
 	return bIsCurrentTimecodeProviderInitialized;
 }
 
@@ -2941,32 +2938,7 @@ static void LoadCustomTimeStep(UEngine* Engine)
 
 static void ResetCustomTimeStep()
 {
-	if (UEngineCustomTimeStep* CustomTimeStep = GEngine->GetCustomTimeStep())
-	{
-#if WITH_EDITOR
-		if (CustomTimeStep->GetClass()->HasAnyClassFlags(CLASS_NewerVersionExists) || CustomTimeStep->GetClass()->HasAnyClassFlags(CLASS_CompiledFromBlueprint))
-		{
-			// Class was recompiled so we need to reload it.
-			UClass* CustomTimeStepClass = CustomTimeStep->GetClass();
-
-			if (UBlueprint* Blueprint = Cast<UBlueprint>(CustomTimeStepClass->ClassGeneratedBy))
-			{
-				if (UClass* NewTimeStepClass = Blueprint->GeneratedClass)
-				{
-					UEngineCustomTimeStep* NewCustomTimeStep = NewObject<UEngineCustomTimeStep>(GEngine, NewTimeStepClass);
-					GEngine->SetCustomTimeStep(NewCustomTimeStep);
-					return;
-				}
-			}
-
-		}
-#endif
-		GEngine->SetCustomTimeStep(CustomTimeStep);
-	}
-	else
-	{
-		LoadCustomTimeStep(GEngine);
-	}
+	LoadCustomTimeStep(GEngine);
 }
 
 static FAutoConsoleCommand GResetCustomTimeStep(
@@ -2980,10 +2952,10 @@ static void LoadTimecodeProvider(UEngine* Engine)
 	if (Engine->TimecodeProviderClassName.IsValid())
 {
 		if (UClass* TimecodeProviderClass = Engine->TimecodeProviderClassName.TryLoadClass<UTimecodeProvider>())
-		{
+	{
 			UTimecodeProvider* NewTimecodeProvider = NewObject<UTimecodeProvider>(Engine, TimecodeProviderClass);
 			if (!Engine->SetTimecodeProvider(NewTimecodeProvider))
-			{
+		{
 				UE_LOG(LogEngine, Warning, TEXT("InitializeTimecodeProvider - Failed to intialize TimecodeProvider '%s'."), *NewTimecodeProvider->GetPathName());
 			}
 		}
@@ -3016,31 +2988,7 @@ static void LoadTimecodeProvider(UEngine* Engine)
 
 static void ResetTimecode()
 {
-	if (UTimecodeProvider* TimecodeProvider = GEngine->GetTimecodeProvider())
-	{
-#if WITH_EDITOR
-		if (TimecodeProvider->GetClass()->HasAnyClassFlags(CLASS_NewerVersionExists) || TimecodeProvider->GetClass()->HasAnyClassFlags(CLASS_CompiledFromBlueprint))
-		{
-			// Class was recompiled so we need to reload it.
-			UClass* TCProviderClass = TimecodeProvider->GetClass();
-
-			if (UBlueprint* Blueprint = Cast<UBlueprint>(TCProviderClass->ClassGeneratedBy))
-			{
-				if (UClass* NewTimecodeClass = Blueprint->GeneratedClass)
-				{
-					UTimecodeProvider* NewTimecodeProvider = NewObject<UTimecodeProvider>(GEngine, NewTimecodeClass);
-					GEngine->SetTimecodeProvider(NewTimecodeProvider);
-					return;
-				}
-			}
-		}
-#endif
-		GEngine->SetTimecodeProvider(TimecodeProvider);
-	}
-	else
-	{
-		LoadTimecodeProvider(GEngine);
-	}
+	LoadTimecodeProvider(GEngine);
 }
 
 static FAutoConsoleCommand GResetTimecode(
