additional trace scopes for OnSyncBodies

diff --git a/Engine/Source/Runtime/Engine/Private/PhysicsEngine/Experimental/PhysScene_Chaos.cpp b/Engine/Source/Runtime/Engine/Private/PhysicsEngine/Experimental/PhysScene_Chaos.cpp
--- a/Engine/Source/Runtime/Engine/Private/PhysicsEngine/Experimental/PhysScene_Chaos.cpp
+++ b/Engine/Source/Runtime/Engine/Private/PhysicsEngine/Experimental/PhysScene_Chaos.cpp
@@ -1759,6 +1759,11 @@ void FPhysScene_Chaos::KillVisualDebugger()
 
 }
 
+DECLARE_CYCLE_STAT(TEXT("FPhysScene_Chaos::OnSyncBodies-FSingleParticlePhysicsProxy"), STAT_SyncBodiesSingleParticlePhysicsProxy, STATGROUP_Chaos);
+DECLARE_CYCLE_STAT(TEXT("FPhysScene_Chaos::OnSyncBodies-FJointConstraintPhysicsProxy"), STAT_SyncBodiesJointConstraintPhysicsProxy, STATGROUP_Chaos);
+DECLARE_CYCLE_STAT(TEXT("FPhysScene_Chaos::OnSyncBodies-FGeometryCollectionPhysicsProxy"), STAT_SyncBodiesGeometryCollectionPhysicsProxy, STATGROUP_Chaos);
+DECLARE_CYCLE_STAT(TEXT("FPhysScene_Chaos::OnSyncBodies-FClusterUnionPhysicsProxy"), STAT_SyncBodiesClusterUnionPhysicsProxy, STATGROUP_Chaos);
+
 void FPhysScene_Chaos::OnSyncBodies(Chaos::FPhysicsSolverBase* Solver)
 {
 	using namespace Chaos;
@@ -1772,6 +1777,7 @@ void FPhysScene_Chaos::OnSyncBodies(Chaos::FPhysicsSolverBase* Solver)
 
 		void operator()(FSingleParticlePhysicsProxy* Proxy)
 		{
+			SCOPE_CYCLE_COUNTER(STAT_SyncBodiesSingleParticlePhysicsProxy);
 			FPBDRigidParticle* DirtyParticle = Proxy->GetRigidParticleUnsafe();
 
 			if (FBodyInstance* BodyInstance = FPhysicsUserData::Get<FBodyInstance>(DirtyParticle->UserData()))
@@ -1815,6 +1821,7 @@ void FPhysScene_Chaos::OnSyncBodies(Chaos::FPhysicsSolverBase* Solver)
 
 		void operator()(FJointConstraintPhysicsProxy* Proxy)
 		{
+			SCOPE_CYCLE_COUNTER(STAT_SyncBodiesJointConstraintPhysicsProxy);
 			Chaos::FJointConstraint* Constraint = Proxy->GetConstraint();
 
 			if (Constraint->GetOutputData().bIsBreaking)
@@ -1842,6 +1849,7 @@ void FPhysScene_Chaos::OnSyncBodies(Chaos::FPhysicsSolverBase* Solver)
 
 		void operator()(FGeometryCollectionPhysicsProxy* Proxy)
 		{
+			SCOPE_CYCLE_COUNTER(STAT_SyncBodiesGeometryCollectionPhysicsProxy);
 			// Don't pass in anything here so we don't end up locking anything because we can assume the scene is already locked.
 			FLockedWritePhysicsObjectExternalInterface Interface = FPhysicsObjectExternalInterface::LockWrite({});
 
@@ -1879,6 +1887,7 @@ void FPhysScene_Chaos::OnSyncBodies(Chaos::FPhysicsSolverBase* Solver)
 
 		void operator()(FClusterUnionPhysicsProxy* Proxy)
 		{
+			SCOPE_CYCLE_COUNTER(STAT_SyncBodiesClusterUnionPhysicsProxy);
 			FLockedWritePhysicsObjectExternalInterface Interface = FPhysicsObjectExternalInterface::LockWrite({});
 			Chaos::FPhysicsObjectHandle Handle = Proxy->GetPhysicsObjectHandle();
 
