[ReadFixedCompressedFloat]

Fix compile error on newer VS toolchains where you get an unspecified error when MaxValue > MaxBitValue.

diff --git a/Engine/Source/Runtime/Engine/Classes/Engine/NetSerialization.h b/Engine/Source/Runtime/Engine/Classes/Engine/NetSerialization.h
--- a/Engine/Source/Runtime/Engine/Classes/Engine/NetSerialization.h
+++ b/Engine/Source/Runtime/Engine/Classes/Engine/NetSerialization.h
@@ -339,13 +339,13 @@ bool ReadFixedCompressedFloat(T& Value, FArchive& Ar)
 	if constexpr (MaxValue > Details::MaxBitValue)
 	{
 		// We have to scale down, scale needs to be a float:
-		const T InvScale = MaxValue / (T)Details::MaxBitValue;
+		constexpr T InvScale = MaxValue / (T)Details::MaxBitValue;
 		Value = UnscaledValue * InvScale;
 	}
 	else
 	{
-		enum { Scale = Details::MaxBitValue / MaxValue };
-		const T InvScale = T(1) / (T)Scale;
+		constexpr int32 Scale = Details::MaxBitValue / MaxValue;
+		constexpr T InvScale = T(1) / (T)Scale;
 
 		Value = UnscaledValue * InvScale;
 	}
