Added TestFramework group to the EngineTest buildgraph
Changed System.Automation to TestFramework.
Added TestFramework plugin. Connected with EngineTest

diff --git a/Engine/Plugins/Tests/AutomationDriverTests/Source/AutomationDriverTests/Private/AutomationDriver.spec.cpp b/Engine/Plugins/Tests/AutomationDriverTests/Source/AutomationDriverTests/Private/AutomationDriver.spec.cpp
--- a/Engine/Plugins/Tests/AutomationDriverTests/Source/AutomationDriverTests/Private/AutomationDriver.spec.cpp
+++ b/Engine/Plugins/Tests/AutomationDriverTests/Source/AutomationDriverTests/Private/AutomationDriver.spec.cpp
@@ -25,7 +25,7 @@
 #define EPIC_TEST_BOOLEAN_(text, expression, expected) \
 	TestEqual(text, expression, expected);
 	  
-BEGIN_DEFINE_SPEC(FAutomationDriverSpec, "System.Automation.Driver", EAutomationTestFlags::ProductFilter | EAutomationTestFlags::ApplicationContextMask)
+BEGIN_DEFINE_SPEC(FAutomationDriverSpec, "TestFramework.Driver", EAutomationTestFlags::ProductFilter | EAutomationTestFlags::ApplicationContextMask)
 	TSharedPtr<SWindow> SuiteWindow;
 	TSharedPtr<SAutomationDriverSpecSuite> SuiteWidget;
 	TSharedPtr<IAutomationDriverSpecSuiteViewModel> SuiteViewModel;
diff --git a/Engine/Plugins/Tests/RuntimeTests/Source/RuntimeTests/Private/AutomationExpectedMessage.spec.cpp b/Engine/Plugins/Tests/RuntimeTests/Source/RuntimeTests/Private/AutomationExpectedMessage.spec.cpp
deleted file mode 100644
diff --git a/Engine/Plugins/Tests/RuntimeTests/Source/RuntimeTests/Private/AutomationSpec.spec.cpp b/Engine/Plugins/Tests/RuntimeTests/Source/RuntimeTests/Private/AutomationSpec.spec.cpp
deleted file mode 100644
diff --git a/Engine/Plugins/Tests/TestFramework/Source/TestFramework/Private/AutomationExpectedMessage.spec.cpp b/Engine/Plugins/Tests/TestFramework/Source/TestFramework/Private/AutomationExpectedMessage.spec.cpp
new file mode 100644
--- /dev/null
+++ b/Engine/Plugins/Tests/TestFramework/Source/TestFramework/Private/AutomationExpectedMessage.spec.cpp
@@ -0,0 +1,214 @@
+// Copyright Epic Games, Inc. All Rights Reserved.
+
+#include "Misc/AutomationTest.h"
+
+BEGIN_DEFINE_SPEC(FAutomationExpectedErrorTest, "TestFramework.ExpectedError", EAutomationTestFlags::EngineFilter | EAutomationTestFlags::ApplicationContextMask)
+END_DEFINE_SPEC(FAutomationExpectedErrorTest)
+void FAutomationExpectedErrorTest::Define()
+{
+	Describe("A defined expected error in a test", [this]()
+	{
+
+		It("will not add an error with a number of occurrences less than zero", [this]()
+		{
+			// Suppress error logged when adding entry with invalid occurrence count
+			AddExpectedError(TEXT("number of expected occurrences must be >= 0"), EAutomationExpectedErrorFlags::Contains, 1);
+
+			AddExpectedError(TEXT("The two values are not equal"), EAutomationExpectedErrorFlags::Contains,  -1);
+
+			TArray<FAutomationExpectedMessage> Errors;
+			GetExpectedMessages(Errors, ELogVerbosity::Warning);
+
+			// Only the first expected error should exist in the list.
+			TestEqual("Expected Errors Count", Errors.Num(), 1);
+		});
+
+		It("will add an error with a number of occurrences equal to zero", [this]()
+		{
+			AddExpectedError(TEXT("Expected Error"), EAutomationExpectedErrorFlags::Contains, 0);
+
+			TArray<FAutomationExpectedMessage> Errors;
+			GetExpectedMessages(Errors, ELogVerbosity::Warning);
+			TestEqual("Expected Errors Count", Errors.Num(), 1);
+
+			// Add the expected error to ensure all test conditions pass
+			AddError(TEXT("Expected Error"));
+		});
+
+		It("will not duplicate an existing expected error using the same matcher", [this]()
+		{
+			// Suppress warning logged when adding duplicate value
+			AddExpectedError(TEXT("cannot add duplicate entries"), EAutomationExpectedErrorFlags::Contains, 1);
+			
+			AddExpectedError(TEXT("Expected Error"), EAutomationExpectedErrorFlags::Contains, 1);
+			AddExpectedError(TEXT("Expected Error"), EAutomationExpectedErrorFlags::Contains, 1);
+
+			TArray<FAutomationExpectedMessage> Errors;
+			GetExpectedMessages(Errors, ELogVerbosity::Warning);
+			TestEqual("Expected Errors Count", Errors.Num(), 2);
+
+			// Add the expected error to ensure all test conditions pass
+			AddError(TEXT("Expected Error"));
+		});
+
+		It("will not duplicate an expected error using a different matcher", [this]()
+		{
+			// Suppress warning logged when adding duplicate value
+			AddExpectedError(TEXT("cannot add duplicate entries"), EAutomationExpectedErrorFlags::Contains, 2);
+
+			AddExpectedError(TEXT("Expected Exact Error"), EAutomationExpectedErrorFlags::Exact, 1);
+			AddExpectedError(TEXT("Expected Exact Error"), EAutomationExpectedErrorFlags::Contains, 1);
+
+			AddExpectedError(TEXT("Expected Contains Error"), EAutomationExpectedErrorFlags::Contains, 1);
+			AddExpectedError(TEXT("Expected Contains Error"), EAutomationExpectedErrorFlags::Exact, 1);
+
+			TArray<FAutomationExpectedMessage> Errors;
+			GetExpectedMessages(Errors, ELogVerbosity::Warning);
+			TestEqual("Expected Errors Count", Errors.Num(), 3);
+
+			// Add the expected errors to ensure all test conditions pass
+			AddError(TEXT("Expected Exact Error"));
+			AddError(TEXT("Expected Contains Error"));
+		});
+
+		// Disabled until fix for UE-44340 (crash creating invalid regex) is merged
+		xIt("will not add an error with an invalid regex pattern", [this]()
+		{
+			AddExpectedError(TEXT("invalid regex }])([{"), EAutomationExpectedErrorFlags::Contains, 0);
+
+			TArray<FAutomationExpectedMessage> Errors;
+			GetExpectedMessages(Errors, ELogVerbosity::Warning);
+
+			TestEqual("Expected Errors Count", Errors.Num(), 0);
+		});
+
+		It("will match both Error and Warning messages", [this]()
+		{
+			AddExpectedError(TEXT("Expected Message"), EAutomationExpectedErrorFlags::Contains, 0);
+			AddError(TEXT("Expected Message"));
+			AddWarning(TEXT("Expected Message"));
+		});
+
+		It("will not fail or warn if encountered", [this]()
+		{
+			AddExpectedError(TEXT("Expected Error"), EAutomationExpectedErrorFlags::Contains, 1);
+			AddExpectedError(TEXT("Expected Warning"), EAutomationExpectedErrorFlags::Contains, 1);
+			AddError(TEXT("Expected Error"));
+			AddWarning(TEXT("Expected Warning"));
+		});
+
+		It("will not match multiple occurrences in the same message when using Contains matcher", [this]()
+		{
+			AddExpectedError(TEXT("Expected"), EAutomationExpectedErrorFlags::Contains, 1);
+			AddError(TEXT("ExpectedExpectedExpected ExpectedExpectedExpected"));
+		});
+
+		It("will match different messages that fit the regex pattern", [this]()
+		{
+			AddExpectedError(TEXT("Response \\(-?\\d+\\)"), EAutomationExpectedErrorFlags::Contains, 4);
+			AddError(TEXT("Response (0)"));
+			AddError(TEXT("Response (1)"));
+			AddError(FString::Printf(TEXT("Response (%d)"), MIN_int64));
+			AddError(FString::Printf(TEXT("Response (%d)"), MAX_uint64));
+		});
+
+	});
+
+}
+
+// Tests for cases where expected errors will fail a test.
+// IMPORTANT: The pass condition for these tests is that they FAIL. To prevent
+// the expected failures from interfering with regular test runs, these tests
+// must be run manually.
+BEGIN_DEFINE_SPEC(FAutomationExpectedErrorFailureTest, "TestFramework.ExpectedError", EAutomationTestFlags::NegativeFilter | EAutomationTestFlags::ApplicationContextMask | EAutomationTestFlags::RequiresUser)
+END_DEFINE_SPEC(FAutomationExpectedErrorFailureTest)
+void FAutomationExpectedErrorFailureTest::Define()
+{
+	Describe("An expected error failure", [this]()
+	{
+		It("will occur if expected a specific number of times and NOT encountered.", [this]()
+		{
+			AddExpectedError(TEXT("Expected Error"), EAutomationExpectedErrorFlags::Exact, 1);
+		});
+
+		It("will occur if expected a specific number of times and is encountered too few times.", [this]()
+		{
+			AddExpectedError(TEXT("Expected Error"), EAutomationExpectedErrorFlags::Exact, 2);
+			AddError(TEXT("Expected Error"));
+		});
+
+		It("will occur if expected a specific number of times and is encountered too many times.", [this]()
+		{
+			AddExpectedError(TEXT("Expected Error"), EAutomationExpectedErrorFlags::Exact, 1);
+			AddError(TEXT("Expected Error"));
+			AddError(TEXT("Expected Error"));
+		});
+
+		It("will occur if expected any number of times and is not encountered.", [this]()
+		{
+			AddExpectedError(TEXT("Expected Error"), EAutomationExpectedErrorFlags::Exact, 0);
+		});
+
+		It("will occur if multiple expected errors are NOT encountered.", [this]()
+		{
+			AddExpectedError(TEXT("Expected Error 1"), EAutomationExpectedErrorFlags::Exact, 1);
+			AddExpectedError(TEXT("Expected Error 2"), EAutomationExpectedErrorFlags::Contains, 1);
+		});
+
+		It("will occur if not all expected errors are encountered.", [this]()
+		{
+			AddExpectedError(TEXT("Expected error 1"), EAutomationExpectedErrorFlags::Exact, 1);
+			AddExpectedError(TEXT("Expected error 2"), EAutomationExpectedErrorFlags::Contains, 1);
+			AddError(TEXT("Expected error 1"));
+		});
+
+		It("will occur if only partial matches are encountered when using Exact matcher", [this]()
+		{
+			AddExpectedError(TEXT("Expected"), EAutomationExpectedErrorFlags::Exact, 1);
+			AddError(TEXT("Expected error"));
+		});
+	});
+}
+
+BEGIN_DEFINE_SPEC(FAutomationExpectedMessageTest, "TestFramework.ExpectedMessage", EAutomationTestFlags::EngineFilter | EAutomationTestFlags::ApplicationContextMask)
+END_DEFINE_SPEC(FAutomationExpectedMessageTest)
+void FAutomationExpectedMessageTest::Define()
+{
+	Describe("A defined expected message with a specified verbosity ", [this]()
+	{
+
+		It("will include error and fatal categories when checking for warning messages", [this]()
+		{
+			AddExpectedMessage(TEXT("suppress this error message"), ELogVerbosity::Error, EAutomationExpectedMessageFlags::Contains, 1);
+			AddExpectedMessage(TEXT("suppress this fatal message"), ELogVerbosity::Fatal, EAutomationExpectedMessageFlags::Contains, 1);
+
+			TArray<FAutomationExpectedMessage> Messages;
+			GetExpectedMessages(Messages, ELogVerbosity::Warning);
+
+			TestEqual(TEXT("Both the error and fatal messages are returned when Warning is specified as the maximum verbosity"), Messages.Num(), 2);
+
+			TestTrue(TEXT("An Error message also counts as an error"), LogCategoryMatchesSeverityInclusive(Messages[0].Verbosity, ELogVerbosity::Warning));
+			TestTrue(TEXT("A Fatal message also counts as an error"), LogCategoryMatchesSeverityInclusive(Messages[1].Verbosity, ELogVerbosity::Warning));
+
+			AddError(TEXT("suppress this error message please"));
+			AddError(TEXT("also suppress this fatal message"));
+		});
+
+
+		// An expected message with a verbosity of "all" behaves as "any"
+		It("will be included when the message is \"all\" but the queried verbosity is NOT \"all\"", [this]()
+		{
+			AddExpectedMessage(TEXT("suppress this message"), ELogVerbosity::All, EAutomationExpectedMessageFlags::Contains, 1);
+
+			// Message with verbosity "all" should always be returned regardless of requested verbosity here
+			TArray<FAutomationExpectedMessage> Messages;
+			GetExpectedMessages(Messages, ELogVerbosity::Warning);
+
+			TestEqual(TEXT("The message with the intention of being any category is returned when something other than \"all\" is specified as the maximum verbosity to query"), Messages.Num(), 1);
+
+			AddInfo(TEXT("suppress this message please"));
+		});
+
+	});
+
+}
diff --git a/Engine/Plugins/Tests/TestFramework/Source/TestFramework/Private/AutomationSpec.spec.cpp b/Engine/Plugins/Tests/TestFramework/Source/TestFramework/Private/AutomationSpec.spec.cpp
new file mode 100644
--- /dev/null
+++ b/Engine/Plugins/Tests/TestFramework/Source/TestFramework/Private/AutomationSpec.spec.cpp
@@ -0,0 +1,141 @@
+// Copyright Epic Games, Inc. All Rights Reserved.
+
+#include "Misc/AutomationTest.h"
+
+BEGIN_DEFINE_SPEC(AutomationSpec, "TestFramework.Spec", EAutomationTestFlags::EngineFilter | EAutomationTestFlags::ApplicationContextMask)
+	bool Foo;
+	FString RunOrder; 
+END_DEFINE_SPEC(AutomationSpec)
+void AutomationSpec::Define()
+{
+	Describe("A Describe 1", [this]()
+	{
+		It("contains a spec with an expectation", [this]()
+		{
+			Foo = true;
+			TestEqual("Foo", Foo, true);
+		});
+	});
+
+	Describe("A Describe 2", [this]()
+	{
+		BeforeEach([&]()
+		{
+			Foo = false;
+		});
+
+		xIt("will not run disabled specs", [this]()
+		{
+			TestEqual("Foo", Foo, true);
+		});
+
+		xDescribe("with disabled nested Describes", [this]()
+		{
+			It("will not run specs within those Describes", [this]()
+			{
+				TestEqual("Foo", Foo, true);
+			});
+		});
+	});
+
+	Describe("A spec using BeforeEach and AfterEach", [this]()
+	{
+		BeforeEach([this]()
+		{
+			RunOrder = TEXT("A");
+		});
+
+		It("will run code before each spec in the Describe and after each spec in the Describe", [this]()
+		{
+			TestEqual("RunOrder", RunOrder, TEXT("A"));
+		});
+
+		AfterEach([this]()
+		{
+			RunOrder += TEXT("Z");
+			TestEqual("RunOrder", RunOrder, TEXT("AZ"));
+		});
+	});
+
+	Describe("A spec using BeforeEach and AfterEach", [this]()
+	{
+		AfterEach([this]()
+		{
+			RunOrder += TEXT("Z");
+			TestEqual("RunOrder", RunOrder, TEXT("ABYZ"));
+		});
+
+		Describe("while nested inside another Describe", [this]()
+		{
+			It("will run all BeforeEach blocks and all AfterEach blocks", [this]()
+			{
+				TestEqual("RunOrder", RunOrder, TEXT("AB"));
+			});
+
+			AfterEach([this]()
+			{
+				RunOrder += TEXT("Y");
+			});
+
+			BeforeEach([this]()
+			{
+				RunOrder += TEXT("B");
+			});
+		});
+
+		BeforeEach([this]()
+		{
+			RunOrder = TEXT("A");
+		});
+	});
+
+	Describe("A spec using BeforeEach and AfterEach", [this]()
+	{
+		BeforeEach([this]()
+		{
+			RunOrder = TEXT("A");
+		});
+
+		AfterEach([this]()
+		{
+			RunOrder += TEXT("Z");
+			TestEqual("RunOrder", RunOrder, TEXT("ABCDXYZ"));
+		});
+
+		BeforeEach([this]()
+		{
+			RunOrder += TEXT("B");
+		});
+
+		Describe("while nested inside another Describe", [this]()
+		{
+			AfterEach([this]()
+			{
+				RunOrder += TEXT("Y");
+			});
+
+			BeforeEach([this]()
+			{
+				RunOrder += TEXT("C");
+			});
+
+			Describe("while nested inside yet another Describe", [this]()
+			{
+				It("will run all BeforeEach blocks and all AfterEach blocks", [this]()
+				{
+					TestEqual("RunOrder", RunOrder, TEXT("ABCD"));
+				});
+
+				AfterEach([this]()
+				{
+					RunOrder += TEXT("X");
+				});
+
+				BeforeEach([this]()
+				{
+					RunOrder += TEXT("D");
+				});
+			});
+		});
+	});
+}
diff --git a/Engine/Plugins/Tests/TestFramework/Source/TestFramework/Private/TestFramework.cpp b/Engine/Plugins/Tests/TestFramework/Source/TestFramework/Private/TestFramework.cpp
new file mode 100644
--- /dev/null
+++ b/Engine/Plugins/Tests/TestFramework/Source/TestFramework/Private/TestFramework.cpp
@@ -0,0 +1,20 @@
+// Copyright Epic Games, Inc. All Rights Reserved.
+
+#include "TestFramework.h"
+
+#define LOCTEXT_NAMESPACE "FTestFrameworkModule"
+
+void FTestFrameworkModule::StartupModule()
+{
+	// This code will execute after your module is loaded into memory; the exact timing is specified in the .uplugin file per-module
+}
+
+void FTestFrameworkModule::ShutdownModule()
+{
+	// This function may be called during shutdown to clean up your module.  For modules that support dynamic reloading,
+	// we call this function before unloading the module.
+}
+
+#undef LOCTEXT_NAMESPACE
+	
+IMPLEMENT_MODULE(FTestFrameworkModule, TestFramework)
\ No newline at end of file
diff --git a/Engine/Plugins/Tests/TestFramework/Source/TestFramework/Public/TestFramework.h b/Engine/Plugins/Tests/TestFramework/Source/TestFramework/Public/TestFramework.h
new file mode 100644
--- /dev/null
+++ b/Engine/Plugins/Tests/TestFramework/Source/TestFramework/Public/TestFramework.h
@@ -0,0 +1,16 @@
+// Copyright Epic Games, Inc. All Rights Reserved.
+
+#pragma once
+
+#include "CoreMinimal.h"
+#include "Modules/ModuleManager.h"
+#include "Misc/AutomationTest.h"
+
+class FTestFrameworkModule : public IModuleInterface
+{
+public:
+
+	/** IModuleInterface implementation */
+	virtual void StartupModule() override;
+	virtual void ShutdownModule() override;
+};
diff --git a/Engine/Plugins/Tests/TestFramework/Source/TestFramework/TestFramework.Build.cs b/Engine/Plugins/Tests/TestFramework/Source/TestFramework/TestFramework.Build.cs
new file mode 100644
--- /dev/null
+++ b/Engine/Plugins/Tests/TestFramework/Source/TestFramework/TestFramework.Build.cs
@@ -0,0 +1,27 @@
+// Copyright Epic Games, Inc. All Rights Reserved.
+
+using UnrealBuildTool;
+
+public class TestFramework : ModuleRules
+{
+	public TestFramework(ReadOnlyTargetRules Target) : base(Target)
+	{
+		PublicDependencyModuleNames.AddRange(
+			new string[]
+			{
+				"Core",
+				"CoreUObject",
+				// ... add other public dependencies that you statically link with here ...
+			}
+			);
+
+
+		PrivateDependencyModuleNames.AddRange(
+			new string[]
+			{
+				"Engine",
+				// ... add private dependencies that you statically link with here ...	
+			}
+			);
+	}
+}
diff --git a/Engine/Plugins/Tests/TestFramework/TestFramework.uplugin b/Engine/Plugins/Tests/TestFramework/TestFramework.uplugin
new file mode 100644
--- /dev/null
+++ b/Engine/Plugins/Tests/TestFramework/TestFramework.uplugin
@@ -0,0 +1,24 @@
+{
+	"FileVersion": 3,
+	"Version": 1,
+	"VersionName": "1.0",
+	"FriendlyName": "TestFramework",
+	"Description": "",
+	"Category": "Testing",
+	"CreatedBy": "Epic Games, Inc.",
+	"CreatedByURL": "https://epicgames.com",
+	"DocsURL": "",
+	"MarketplaceURL": "",
+	"SupportURL": "",
+	"EnabledByDefault": false,
+	"CanContainContent": false,
+	"IsBetaVersion": false,
+	"Installed": false,
+	"Modules": [
+		{
+			"Name": "TestFramework",
+			"Type": "DeveloperTool",
+			"LoadingPhase": "Default"
+		}
+	]
+}
\ No newline at end of file
diff --git a/Engine/Source/Runtime/Engine/Private/Tests/EngineAutomationTests.cpp b/Engine/Source/Runtime/Engine/Private/Tests/EngineAutomationTests.cpp
--- a/Engine/Source/Runtime/Engine/Private/Tests/EngineAutomationTests.cpp
+++ b/Engine/Source/Runtime/Engine/Private/Tests/EngineAutomationTests.cpp
@@ -227,7 +227,7 @@ bool FSaveGameTest::RunTest(const FString& Parameters)
 	return LoadedData == SavedData;
 }
 
-IMPLEMENT_SIMPLE_AUTOMATION_TEST(FAutomationLogAddMessage, "System.Automation.Log.Add Log Message", EAutomationTestFlags::ApplicationContextMask | EAutomationTestFlags::EngineFilter)
+IMPLEMENT_SIMPLE_AUTOMATION_TEST(FAutomationLogAddMessage, "TestFramework.Log.Add Log Message", EAutomationTestFlags::ApplicationContextMask | EAutomationTestFlags::EngineFilter)
 
 bool FAutomationLogAddMessage::RunTest(const FString& Parameters)
 {
@@ -244,7 +244,7 @@ bool FAutomationLogAddMessage::RunTest(const FString& Parameters)
 	return true;
 }
 
-IMPLEMENT_SIMPLE_AUTOMATION_TEST(FAutomationLogAddWarning, "System.Automation.Log.Add Warning Message", EAutomationTestFlags::EditorContext | EAutomationTestFlags::EngineFilter)
+IMPLEMENT_SIMPLE_AUTOMATION_TEST(FAutomationLogAddWarning, "TestFramework.Log.Add Warning Message", EAutomationTestFlags::ApplicationContextMask | EAutomationTestFlags::EngineFilter)
 
 bool FAutomationLogAddWarning::RunTest(const FString& Parameters)
 {
@@ -261,7 +261,7 @@ bool FAutomationLogAddWarning::RunTest(const FString& Parameters)
 	return true;
 }
 
-IMPLEMENT_SIMPLE_AUTOMATION_TEST(FAutomationLogAddError, "System.Automation.Log.Add Error Message", EAutomationTestFlags::EditorContext | EAutomationTestFlags::EngineFilter)
+IMPLEMENT_SIMPLE_AUTOMATION_TEST(FAutomationLogAddError, "TestFramework.Log.Add Error Message", EAutomationTestFlags::ApplicationContextMask | EAutomationTestFlags::EngineFilter)
 
 bool FAutomationLogAddError::RunTest(const FString& Parameters)
 {
@@ -600,7 +600,7 @@ const FString FAutomationNearlyEqualTest::ActualValueIsIncreasedByFormatString(T
 const FString FAutomationNearlyEqualTest::ExpectedValueIsIncreasedByFormatString(TEXT("{0} (expected value is increased by {1})"));
 const FString FAutomationNearlyEqualTest::DifferenceAndOrToleranceAreNotValidFormatString(TEXT("Difference and/or Tolerance are not valid. Difference: {0}, Tolerance: {1}"));
 
-IMPLEMENT_CUSTOM_SIMPLE_AUTOMATION_TEST(FAutomationTestNearlyEqualFloatPositive, FAutomationNearlyEqualTest, "System.Automation.Validation.TestNearlyEqualFloatPositive", EAutomationTestFlags::EditorContext | EAutomationTestFlags::EngineFilter);
+IMPLEMENT_CUSTOM_SIMPLE_AUTOMATION_TEST(FAutomationTestNearlyEqualFloatPositive, FAutomationNearlyEqualTest, "TestFramework.Validation.TestNearlyEqualFloatPositive", EAutomationTestFlags::ApplicationContextMask | EAutomationTestFlags::EngineFilter);
 bool FAutomationTestNearlyEqualFloatPositive::RunTest(const FString& Parameters)
 {
 	//** TEST **//
@@ -627,7 +627,7 @@ bool FAutomationTestNearlyEqualFloatPositive::RunTest(const FString& Parameters)
 	return true;
 }
 
-IMPLEMENT_CUSTOM_SIMPLE_AUTOMATION_TEST(FAutomationTestNearlyEqualFloatNegative, FAutomationNearlyEqualTest, "System.Automation.Validation.TestNearlyEqualFloatNegative", EAutomationTestFlags::EditorContext | EAutomationTestFlags::EngineFilter);
+IMPLEMENT_CUSTOM_SIMPLE_AUTOMATION_TEST(FAutomationTestNearlyEqualFloatNegative, FAutomationNearlyEqualTest, "TestFramework.Validation.TestNearlyEqualFloatNegative", EAutomationTestFlags::ApplicationContextMask | EAutomationTestFlags::EngineFilter);
 bool FAutomationTestNearlyEqualFloatNegative::RunTest(const FString& Parameters)
 {
 	int32 CasesCheckedTotal = 0;
@@ -650,7 +650,7 @@ bool FAutomationTestNearlyEqualFloatNegative::RunTest(const FString& Parameters)
 	return true;
 }
 
-IMPLEMENT_CUSTOM_SIMPLE_AUTOMATION_TEST(FAutomationTestNearlyEqualDoublePositive, FAutomationNearlyEqualTest, "System.Automation.Validation.TestNearlyEqualDoublePositive", EAutomationTestFlags::EditorContext | EAutomationTestFlags::EngineFilter);
+IMPLEMENT_CUSTOM_SIMPLE_AUTOMATION_TEST(FAutomationTestNearlyEqualDoublePositive, FAutomationNearlyEqualTest, "TestFramework.Validation.TestNearlyEqualDoublePositive", EAutomationTestFlags::ApplicationContextMask | EAutomationTestFlags::EngineFilter);
 bool FAutomationTestNearlyEqualDoublePositive::RunTest(const FString& Parameters)
 {
 	//** TEST **//
@@ -676,7 +676,7 @@ bool FAutomationTestNearlyEqualDoublePositive::RunTest(const FString& Parameters
 	return true;
 }
 
-IMPLEMENT_CUSTOM_SIMPLE_AUTOMATION_TEST(FAutomationTestNearlyEqualDoubleNegative, FAutomationNearlyEqualTest, "System.Automation.Validation.TestNearlyEqualDoubleNegative", EAutomationTestFlags::EditorContext | EAutomationTestFlags::EngineFilter);
+IMPLEMENT_CUSTOM_SIMPLE_AUTOMATION_TEST(FAutomationTestNearlyEqualDoubleNegative, FAutomationNearlyEqualTest, "TestFramework.Validation.TestNearlyEqualDoubleNegative", EAutomationTestFlags::ApplicationContextMask | EAutomationTestFlags::EngineFilter);
 bool FAutomationTestNearlyEqualDoubleNegative::RunTest(const FString& Parameters)
 {
 	int32 CasesCheckedTotal = 0;
@@ -699,7 +699,7 @@ bool FAutomationTestNearlyEqualDoubleNegative::RunTest(const FString& Parameters
 	return true;
 }
 
-IMPLEMENT_CUSTOM_SIMPLE_AUTOMATION_TEST(FAutomationTestNearlyEqualFVectorPositive, FAutomationNearlyEqualTest, "System.Automation.Validation.TestNearlyEqualFVectorPositive", EAutomationTestFlags::EditorContext | EAutomationTestFlags::EngineFilter);
+IMPLEMENT_CUSTOM_SIMPLE_AUTOMATION_TEST(FAutomationTestNearlyEqualFVectorPositive, FAutomationNearlyEqualTest, "TestFramework.Validation.TestNearlyEqualFVectorPositive", EAutomationTestFlags::ApplicationContextMask | EAutomationTestFlags::EngineFilter);
 bool FAutomationTestNearlyEqualFVectorPositive::RunTest(const FString& Parameters)
 {
 	//** TEST **//
@@ -725,7 +725,7 @@ bool FAutomationTestNearlyEqualFVectorPositive::RunTest(const FString& Parameter
 	return true;
 }
 
-IMPLEMENT_CUSTOM_SIMPLE_AUTOMATION_TEST(FAutomationTestNearlyEqualFVectorNegative, FAutomationNearlyEqualTest, "System.Automation.Validation.TestNearlyEqualFVectorNegative", EAutomationTestFlags::EditorContext | EAutomationTestFlags::EngineFilter);
+IMPLEMENT_CUSTOM_SIMPLE_AUTOMATION_TEST(FAutomationTestNearlyEqualFVectorNegative, FAutomationNearlyEqualTest, "TestFramework.Validation.TestNearlyEqualFVectorNegative", EAutomationTestFlags::ApplicationContextMask | EAutomationTestFlags::EngineFilter);
 bool FAutomationTestNearlyEqualFVectorNegative::RunTest(const FString& Parameters)
 {
 	int32 CasesCheckedTotal = 0;
@@ -744,7 +744,7 @@ bool FAutomationTestNearlyEqualFVectorNegative::RunTest(const FString& Parameter
 	return true;
 }
 
-IMPLEMENT_CUSTOM_SIMPLE_AUTOMATION_TEST(FAutomationTestNearlyEqualFRotatorPositive, FAutomationNearlyEqualTest, "System.Automation.Validation.TestNearlyEqualFRotatorPositive", EAutomationTestFlags::EditorContext | EAutomationTestFlags::EngineFilter);
+IMPLEMENT_CUSTOM_SIMPLE_AUTOMATION_TEST(FAutomationTestNearlyEqualFRotatorPositive, FAutomationNearlyEqualTest, "TestFramework.Validation.TestNearlyEqualFRotatorPositive", EAutomationTestFlags::ApplicationContextMask | EAutomationTestFlags::EngineFilter);
 bool FAutomationTestNearlyEqualFRotatorPositive::RunTest(const FString& Parameters)
 {
 	//** TEST **//
@@ -770,7 +770,7 @@ bool FAutomationTestNearlyEqualFRotatorPositive::RunTest(const FString& Paramete
 	return true;
 }
 
-IMPLEMENT_CUSTOM_SIMPLE_AUTOMATION_TEST(FAutomationTestNearlyEqualFRotatorNegative, FAutomationNearlyEqualTest, "System.Automation.Validation.TestNearlyEqualFRotatorNegative", EAutomationTestFlags::EditorContext | EAutomationTestFlags::EngineFilter);
+IMPLEMENT_CUSTOM_SIMPLE_AUTOMATION_TEST(FAutomationTestNearlyEqualFRotatorNegative, FAutomationNearlyEqualTest, "TestFramework.Validation.TestNearlyEqualFRotatorNegative", EAutomationTestFlags::ApplicationContextMask | EAutomationTestFlags::EngineFilter);
 bool FAutomationTestNearlyEqualFRotatorNegative::RunTest(const FString& Parameters)
 {
 	int32 CasesCheckedTotal = 0;
@@ -789,7 +789,7 @@ bool FAutomationTestNearlyEqualFRotatorNegative::RunTest(const FString& Paramete
 	return true;
 }
 
-IMPLEMENT_CUSTOM_SIMPLE_AUTOMATION_TEST(FAutomationTestNearlyEqualFTransformPositive, FAutomationNearlyEqualTest, "System.Automation.Validation.TestNearlyEqualFTransformPositive", EAutomationTestFlags::EditorContext | EAutomationTestFlags::EngineFilter);
+IMPLEMENT_CUSTOM_SIMPLE_AUTOMATION_TEST(FAutomationTestNearlyEqualFTransformPositive, FAutomationNearlyEqualTest, "TestFramework.Validation.TestNearlyEqualFTransformPositive", EAutomationTestFlags::ApplicationContextMask | EAutomationTestFlags::EngineFilter);
 bool FAutomationTestNearlyEqualFTransformPositive::RunTest(const FString& Parameters)
 {
 	//** TEST **//
@@ -815,7 +815,7 @@ bool FAutomationTestNearlyEqualFTransformPositive::RunTest(const FString& Parame
 	return true;
 }
 
-IMPLEMENT_CUSTOM_SIMPLE_AUTOMATION_TEST(FAutomationTestNearlyEqualFTransformNegative, FAutomationNearlyEqualTest, "System.Automation.Validation.TestNearlyEqualFTransformNegative", EAutomationTestFlags::EditorContext | EAutomationTestFlags::EngineFilter);
+IMPLEMENT_CUSTOM_SIMPLE_AUTOMATION_TEST(FAutomationTestNearlyEqualFTransformNegative, FAutomationNearlyEqualTest, "TestFramework.Validation.TestNearlyEqualFTransformNegative", EAutomationTestFlags::ApplicationContextMask | EAutomationTestFlags::EngineFilter);
 bool FAutomationTestNearlyEqualFTransformNegative::RunTest(const FString& Parameters)
 {
 	int32 CasesCheckedTotal = 0;
@@ -861,7 +861,7 @@ const float FAutomationUTestMacrosExpr::ExpectedFloatValueOutOfTolerance(ActualF
 const FString FAutomationUTestMacrosExpr::ActualFStringValue(TEXT("EQUALS"));
 const FString FAutomationUTestMacrosExpr::ExpectedFStringValueLowerCase(TEXT("equals"));
 
-IMPLEMENT_CUSTOM_SIMPLE_AUTOMATION_TEST(FAutomationEqualEXPR, FAutomationUTestMacrosExpr, "System.Automation.Validation.UTestEqual", EAutomationTestFlags::ApplicationContextMask | EAutomationTestFlags::EngineFilter);
+IMPLEMENT_CUSTOM_SIMPLE_AUTOMATION_TEST(FAutomationEqualEXPR, FAutomationUTestMacrosExpr, "TestFramework.Validation.UTestEqual", EAutomationTestFlags::ApplicationContextMask | EAutomationTestFlags::EngineFilter);
 bool FAutomationEqualEXPR::RunTest(const FString& Parameters)
 {
 
@@ -874,7 +874,7 @@ bool FAutomationEqualEXPR::RunTest(const FString& Parameters)
 	return true;
 }
 
-IMPLEMENT_CUSTOM_SIMPLE_AUTOMATION_TEST(FAutomationSameNotSameEXPR, FAutomationUTestMacrosExpr, "System.Automation.Validation.UTestSameNotSame", EAutomationTestFlags::ApplicationContextMask | EAutomationTestFlags::EngineFilter);
+IMPLEMENT_CUSTOM_SIMPLE_AUTOMATION_TEST(FAutomationSameNotSameEXPR, FAutomationUTestMacrosExpr, "TestFramework.Validation.UTestSameNotSame", EAutomationTestFlags::ApplicationContextMask | EAutomationTestFlags::EngineFilter);
 bool FAutomationSameNotSameEXPR::RunTest(const FString& Parameters)
 {
 
@@ -884,7 +884,7 @@ bool FAutomationSameNotSameEXPR::RunTest(const FString& Parameters)
 	return true;
 }
 
-IMPLEMENT_CUSTOM_SIMPLE_AUTOMATION_TEST(FAutomationTrueFalseEXPR, FAutomationUTestMacrosExpr, "System.Automation.Validation.UTestTrueFalse", EAutomationTestFlags::ApplicationContextMask | EAutomationTestFlags::EngineFilter);
+IMPLEMENT_CUSTOM_SIMPLE_AUTOMATION_TEST(FAutomationTrueFalseEXPR, FAutomationUTestMacrosExpr, "TestFramework.Validation.UTestTrueFalse", EAutomationTestFlags::ApplicationContextMask | EAutomationTestFlags::EngineFilter);
 bool FAutomationTrueFalseEXPR::RunTest(const FString& Parameters)
 {
 
@@ -894,7 +894,7 @@ bool FAutomationTrueFalseEXPR::RunTest(const FString& Parameters)
 	return true;
 }
 
-IMPLEMENT_CUSTOM_SIMPLE_AUTOMATION_TEST(FAutomationValidInvalidEXPR, FAutomationUTestMacrosExpr, "System.Automation.Validation.UTestValidInvalid", EAutomationTestFlags::ApplicationContextMask | EAutomationTestFlags::EngineFilter);
+IMPLEMENT_CUSTOM_SIMPLE_AUTOMATION_TEST(FAutomationValidInvalidEXPR, FAutomationUTestMacrosExpr, "TestFramework.Validation.UTestValidInvalid", EAutomationTestFlags::ApplicationContextMask | EAutomationTestFlags::EngineFilter);
 bool FAutomationValidInvalidEXPR::RunTest(const FString& Parameters)
 {
 	struct HandleResumeContext
@@ -909,7 +909,7 @@ bool FAutomationValidInvalidEXPR::RunTest(const FString& Parameters)
 	return true;
 }
 
-IMPLEMENT_CUSTOM_SIMPLE_AUTOMATION_TEST(FAutomationNullNotNullPtrEXPR, FAutomationUTestMacrosExpr, "System.Automation.Validation.UTestNullNotNull", EAutomationTestFlags::ApplicationContextMask | EAutomationTestFlags::EngineFilter);
+IMPLEMENT_CUSTOM_SIMPLE_AUTOMATION_TEST(FAutomationNullNotNullPtrEXPR, FAutomationUTestMacrosExpr, "TestFramework.Validation.UTestNullNotNull", EAutomationTestFlags::ApplicationContextMask | EAutomationTestFlags::EngineFilter);
 bool FAutomationNullNotNullPtrEXPR::RunTest(const FString& Parameters)
 {
 	UWorld* World = UWorld::CreateWorld(EWorldType::Game, false);
