Added explicit moves/copies to FOverridableReplayVersionData inside the disabled deprecation block to avoid compiler deprecation warnings in compiler-generated special member functions.

diff --git a/Engine/Source/Runtime/Engine/Public/ReplayTypes.h b/Engine/Source/Runtime/Engine/Public/ReplayTypes.h
--- a/Engine/Source/Runtime/Engine/Public/ReplayTypes.h
+++ b/Engine/Source/Runtime/Engine/Public/ReplayTypes.h
@@ -523,6 +523,11 @@ public:
 	int32 PackageVersionLicenseeUE;       // Licensee package version on which the replay was recorded
 
 	PRAGMA_DISABLE_DEPRECATION_WARNINGS
+	FOverridableReplayVersionData(FOverridableReplayVersionData&&) = default;
+	FOverridableReplayVersionData(const FOverridableReplayVersionData&) = default;
+	FOverridableReplayVersionData& operator=(FOverridableReplayVersionData&&) = default;
+	FOverridableReplayVersionData& operator=(const FOverridableReplayVersionData&) = default;
+
 	// Init with Demo Header Version Data
 	FOverridableReplayVersionData(const FNetworkDemoHeader& DemoHeader)
 		: Version                     (DemoHeader.Version)
