Fixed deprecation warnings in compiler-generated copy/move operators by making FCompressedAnimSequence non-copyable.

diff --git a/Engine/Source/Runtime/Engine/Public/Animation/AnimCompressionTypes.h b/Engine/Source/Runtime/Engine/Public/Animation/AnimCompressionTypes.h
--- a/Engine/Source/Runtime/Engine/Public/Animation/AnimCompressionTypes.h
+++ b/Engine/Source/Runtime/Engine/Public/Animation/AnimCompressionTypes.h
@@ -769,6 +769,7 @@ struct FAnimCompressedCurveIndexedName
 struct ENGINE_API FCompressedAnimSequence
 {
 public:
+	UE_NONCOPYABLE(FCompressedAnimSequence)
 
 	/**
 	 * Version of TrackToSkeletonMapTable for the compressed tracks. Due to baking additive data
