Change Subsutrate sRGB encoding.

* Default now uses accurate sRGB value encoding.
* Lower platforms using lower shading quality (i.e., r.Substrate.ShadingQuality > 1), will continue to use the cheap sqrt/square encoding/decoding

Using proper sRGB encoding increases occupangy by 8 vgpr on ps4 and a fullscreen quad during the base pass by 0.2ms due to lower occupancy. So keep it disabled on lower platforms.

diff --git a/Engine/Shaders/Private/Strata/Strata.ush b/Engine/Shaders/Private/Strata/Strata.ush
--- a/Engine/Shaders/Private/Strata/Strata.ush
+++ b/Engine/Shaders/Private/Strata/Strata.ush
@@ -43,6 +43,13 @@
 #define STRATA_INLINE_SHADING 0
 #endif
 
+// Strata STRATA_SHADING_QUALITY define the lighting accuarcy/approximation
+// 1: accurante lighting
+// 2: approximation lighting
+#ifndef STRATA_SHADING_QUALITY
+#define STRATA_SHADING_QUALITY 1
+#endif
+
 #ifndef STRATA_DEFERRED_SHADING
 #define STRATA_DEFERRED_SHADING (STRATA_INLINE_SHADING == 0)
 #endif 
@@ -3107,7 +3114,7 @@ float3 sqrtFast3(float3 In) { return float3(sqrtFast(In.x), sqrtFast(In.y), sqrt
 
 half3 StrataLinearToSrgb(half3 In)
 {
-#if STRATA_USE_ACCURATE_SRGB
+#if STRATA_SHADING_QUALITY <= 1
 	return LinearToSrgb(In);
 #else
 	return sqrtFast3(In); // Simple gamma 2.0
@@ -3116,7 +3123,7 @@ half3 StrataLinearToSrgb(half3 In)
 
 half3 StrataSrgbToLinear(half3 In)
 {
-#if STRATA_USE_ACCURATE_SRGB
+#if STRATA_SHADING_QUALITY <= 1
 	return sRGBToLinear(In);
 #else
 	return In * In; // Simple gamma 2.0
diff --git a/Engine/Shaders/Private/Strata/StrataEvaluation.ush b/Engine/Shaders/Private/Strata/StrataEvaluation.ush
--- a/Engine/Shaders/Private/Strata/StrataEvaluation.ush
+++ b/Engine/Shaders/Private/Strata/StrataEvaluation.ush
@@ -40,13 +40,6 @@
 #define STRATA_SSS_TRANSMISSION 0
 #endif
 
-// Strata STRATA_SHADING_QUALITY define the lighting accuarcy/approximation
-// 1: accurante lighting
-// 2: approximation lighting
-#ifndef STRATA_SHADING_QUALITY
-#define STRATA_SHADING_QUALITY 1
-#endif
-
 ///////////////////////////////////////////////////////////////////////////////
 // Strata Macros for reducing permutation boiler plate
 // 
diff --git a/Engine/Source/Runtime/Engine/Private/ShaderCompiler/ShaderCompiler.cpp b/Engine/Source/Runtime/Engine/Private/ShaderCompiler/ShaderCompiler.cpp
--- a/Engine/Source/Runtime/Engine/Private/ShaderCompiler/ShaderCompiler.cpp
+++ b/Engine/Source/Runtime/Engine/Private/ShaderCompiler/ShaderCompiler.cpp
@@ -6535,10 +6535,9 @@ void GlobalBeginCompileShader(
 		{
 			const uint32 StrataShadingQuality = Strata::GetShadingQuality(Target.GetPlatform());
 			Input.Environment.SetDefine(TEXT("STRATA_SHADING_QUALITY"), StrataShadingQuality);
-			if (StrataShadingQuality > 1)
-			{
-				Input.Environment.SetDefine(TEXT("USE_ACHROMATIC_BXDF_ENERGY"), 1u);
-			}
+
+			const bool bLowQuality = StrataShadingQuality > 1;
+			Input.Environment.SetDefine(TEXT("USE_ACHROMATIC_BXDF_ENERGY"), bLowQuality ? 1u : 0u);
 
 			const uint32 StrataNormalQuality = Strata::GetNormalQuality();
 			Input.Environment.SetDefine(TEXT("STRATA_NORMAL_QUALITY"), StrataNormalQuality);
@@ -6553,9 +6552,6 @@ void GlobalBeginCompileShader(
 			Input.Environment.SetDefine(TEXT("STRATA_USE_DBUFFER_PASS"), bStrataDBufferPass ? 1 : 0);
 		}
 
-		const bool bStrataUseAccurateSRGB = bStrata && Strata::IsAccurateSRGBEnabled();
-		Input.Environment.SetDefine(TEXT("STRATA_USE_ACCURATE_SRGB"), bStrataUseAccurateSRGB ? 1 : 0);
-
 		const bool bStrataBackCompatibility = bStrata && Strata::IsBackCompatibilityEnabled();
 		Input.Environment.SetDefine(TEXT("PROJECT_STRATA_BACKCOMPATIBILITY"), bStrataBackCompatibility ? 1 : 0);
 
diff --git a/Engine/Source/Runtime/RenderCore/Private/RenderUtils.cpp b/Engine/Source/Runtime/RenderCore/Private/RenderUtils.cpp
--- a/Engine/Source/Runtime/RenderCore/Private/RenderUtils.cpp
+++ b/Engine/Source/Runtime/RenderCore/Private/RenderUtils.cpp
@@ -1588,12 +1588,6 @@ static TAutoConsoleVariable<int32> CVarSubstrateTileCoord8Bits(
 	TEXT("Format of tile coord. This variable is read-only."),
 	ECVF_ReadOnly | ECVF_RenderThreadSafe);
 
-static TAutoConsoleVariable<int32> CVarSubstrateAccurateSRGB(
-	TEXT("r.Substrate.AccurateSRGB"),
-	0,
-	TEXT("Enable accurate sRGB encoding for pixel data encoding/decoding during rendering passes (base pass, lighting, ...). Otherwise use a simple 'Gamma2' approximation."),
-	ECVF_ReadOnly | ECVF_RenderThreadSafe);
-
 namespace Strata
 {
 	bool IsStrataEnabled()
@@ -1686,11 +1680,6 @@ namespace Strata
 		return CVarSubstrateTileCoord8Bits.GetValueOnAnyThread() > 0 ? 1 : 0;
 	}
 
-	bool IsAccurateSRGBEnabled()
-	{
-		return CVarSubstrateAccurateSRGB.GetValueOnAnyThread() > 0 ? 1 : 0;
-	}
-
 	uint32 GetShadingQuality()
 	{
 		return CVarSubstrateShadingQuality.GetValueOnAnyThread();
diff --git a/Engine/Source/Runtime/RenderCore/Private/Shader.cpp b/Engine/Source/Runtime/RenderCore/Private/Shader.cpp
--- a/Engine/Source/Runtime/RenderCore/Private/Shader.cpp
+++ b/Engine/Source/Runtime/RenderCore/Private/Shader.cpp
@@ -2002,9 +2002,8 @@ void ShaderMapAppendKeyString(EShaderPlatform Platform, FString& KeyString)
 			KeyString += FString::Printf(TEXT("_ADVDEBUG"));
 		}
 
-		if (Strata::IsAccurateSRGBEnabled())
 		{
-			KeyString += FString::Printf(TEXT("_SRGB"));
+			KeyString += FString::Printf(TEXT("_STSHQL%u"), Strata::GetShadingQuality(Platform));
 		}
 	}
 
diff --git a/Engine/Source/Runtime/RenderCore/Public/RenderUtils.h b/Engine/Source/Runtime/RenderCore/Public/RenderUtils.h
--- a/Engine/Source/Runtime/RenderCore/Public/RenderUtils.h
+++ b/Engine/Source/Runtime/RenderCore/Public/RenderUtils.h
@@ -489,7 +489,6 @@ namespace Strata
 	RENDERCORE_API bool IsOpaqueRoughRefractionEnabled();
 	RENDERCORE_API bool IsAdvancedVisualizationEnabled();
 	RENDERCORE_API bool Is8bitTileCoordEnabled();
-	RENDERCORE_API bool IsAccurateSRGBEnabled();
 
 	RENDERCORE_API uint32 GetRayTracingMaterialPayloadSizeInBytes(bool bFullySimplifiedMaterial);
 
