Minor improvements in shader backends.

#rnx

diff --git a/Engine/Source/Developer/ShaderCompilerCommon/Private/ShaderCompilerCommon.cpp b/Engine/Source/Developer/ShaderCompilerCommon/Private/ShaderCompilerCommon.cpp
--- a/Engine/Source/Developer/ShaderCompilerCommon/Private/ShaderCompilerCommon.cpp
+++ b/Engine/Source/Developer/ShaderCompilerCommon/Private/ShaderCompilerCommon.cpp
@@ -1264,9 +1264,9 @@ SHADERCOMPILERCOMMON_API void WriteShaderConductorCommandLine(const FShaderCompi
 	}
 }
 
-static int Mali_ExtractNumberInstructions(const FString &MaliOutput)
+static uint32 Mali_ExtractNumberInstructions(const FString &MaliOutput)
 {
-	int ReturnedNum = 0;
+	uint32 ReturnedNum = 0;
 
 	// Parse the instruction count
 	int32 InstructionStringLength = FPlatformString::Strlen(TEXT("Instructions Emitted:"));
@@ -1304,7 +1304,7 @@ static int Mali_ExtractNumberInstructions(const FString &MaliOutput)
 					bFoundNrStart = false;
 					const FString NumberString = MaliOutput.Mid(NumberIndex, StartIndex - NumberIndex);
 					const float fNrInstructions = FCString::Atof(*NumberString);
-					ReturnedNum += ceil(fNrInstructions);
+					ReturnedNum += (uint32)FMath::Max(0.0, ceil(fNrInstructions));
 				}
 
 				++StartIndex;
diff --git a/Engine/Source/Developer/VulkanShaderFormat/Private/VulkanShaderCompiler.cpp b/Engine/Source/Developer/VulkanShaderFormat/Private/VulkanShaderCompiler.cpp
--- a/Engine/Source/Developer/VulkanShaderFormat/Private/VulkanShaderCompiler.cpp
+++ b/Engine/Source/Developer/VulkanShaderFormat/Private/VulkanShaderCompiler.cpp
@@ -479,7 +479,7 @@ static void AddUniformBuffer(FOLDVulkanCodeHeader& OLDHeader,
 	else
 	{
 		// If we're using real uniform buffers we have to have resources at least
-		check(!UBInfo.bOnlyHasResources);
+		checkf(!UBInfo.bOnlyHasResources, TEXT("UBName = %s, BindingIndex = %d"), *UBName, (int32)BindingIndex);
 	}
 
 	// Currently we don't support mismatched uniform buffer layouts/cbuffers with resources!
diff --git a/Engine/Source/Runtime/Engine/Private/ShaderCompiler/ShaderCompiler.cpp b/Engine/Source/Runtime/Engine/Private/ShaderCompiler/ShaderCompiler.cpp
--- a/Engine/Source/Runtime/Engine/Private/ShaderCompiler/ShaderCompiler.cpp
+++ b/Engine/Source/Runtime/Engine/Private/ShaderCompiler/ShaderCompiler.cpp
@@ -1381,16 +1381,17 @@ static void SplitJobsByType(const TArray<FShaderCommonCompileJobPtr>& QueuedJobs
 	for (int32 Index = 0; Index < QueuedJobs.Num(); ++Index)
 	{
 		FShaderCommonCompileJobPtr CommonJob = QueuedJobs[Index];
-		FShaderPipelineCompileJob* PipelineJob = CommonJob->GetShaderPipelineJob();
-		if (PipelineJob)
+		if (FShaderCompileJob* SingleJob = CommonJob->GetSingleShaderJob())
+		{
+			OutQueuedSingleJobs.Add(SingleJob);
+		}
+		else if (FShaderPipelineCompileJob* PipelineJob = CommonJob->GetShaderPipelineJob())
 		{
 			OutQueuedPipelineJobs.Add(PipelineJob);
 		}
 		else
 		{
-			FShaderCompileJob* SingleJob = CommonJob->GetSingleShaderJob();
-			check(SingleJob);
-			OutQueuedSingleJobs.Add(SingleJob);
+			checkf(0, TEXT("FShaderCommonCompileJob::Type=%d is not a valid type for a shader compile job"), (int32)CommonJob->Type);
 		}
 	}
 }
@@ -4883,15 +4884,18 @@ void FShaderCompilingManager::ProcessCompiledShaderMaps(
 					const bool bCheckSucceeded = CheckSingleJob(*SingleJob, Errors);
 					bSuccess = bCheckSucceeded && bSuccess;
 				}
-				else
+				else if (FShaderPipelineCompileJob* PipelineJob = CurrentJob.GetShaderPipelineJob())
 				{
-					FShaderPipelineCompileJob* PipelineJob = CurrentJob.GetShaderPipelineJob();
 					for (int32 Index = 0; Index < PipelineJob->StageJobs.Num(); ++Index)
 					{
 						const bool bCheckSucceeded = CheckSingleJob(*PipelineJob->StageJobs[Index], Errors);
 						bSuccess = PipelineJob->StageJobs[Index]->bSucceeded && bCheckSucceeded && bSuccess;
 					}
 				}
+				else
+				{
+					checkf(0, TEXT("FShaderCommonCompileJob::Type=%d is not a valid type for a shader compile job"), (int32)CurrentJob.Type);
+				}
 			}
 
 			FMaterialShaderMap* ShaderMapToUseForRendering = nullptr;
diff --git a/Engine/Source/Runtime/RenderCore/Private/ShaderParameterMetadata.cpp b/Engine/Source/Runtime/RenderCore/Private/ShaderParameterMetadata.cpp
--- a/Engine/Source/Runtime/RenderCore/Private/ShaderParameterMetadata.cpp
+++ b/Engine/Source/Runtime/RenderCore/Private/ShaderParameterMetadata.cpp
@@ -577,7 +577,7 @@ void FShaderParametersMetadata::InitializeLayout(FRHIUniformBufferLayoutInitiali
 				}
 				else if (ChildStruct->GetUseCase() != EUseCase::ShaderParameterStruct && UseCase == EUseCase::ShaderParameterStruct)
 				{
-					UE_LOG(LogRendererCore, Fatal, TEXT("%sCan only nests or include shader parameter struct define with BEGIN_SHADER_PARAMETER_STRUCT(), but %s is not."), *GetMemberErrorPrefix(), ChildStruct->GetStructTypeName());
+					UE_LOG(LogRendererCore, Fatal, TEXT("%sCan only nest or include shader parameter struct define with BEGIN_SHADER_PARAMETER_STRUCT(), but %s is not."), *GetMemberErrorPrefix(), ChildStruct->GetStructTypeName());
 				}
 			}
 
