Substrate - fixed not working name reroute propagating Substrate data.

diff --git a/Engine/Source/Runtime/Engine/Classes/Materials/MaterialExpressionRerouteBase.h b/Engine/Source/Runtime/Engine/Classes/Materials/MaterialExpressionRerouteBase.h
--- a/Engine/Source/Runtime/Engine/Classes/Materials/MaterialExpressionRerouteBase.h
+++ b/Engine/Source/Runtime/Engine/Classes/Materials/MaterialExpressionRerouteBase.h
@@ -28,6 +28,8 @@ class ENGINE_API UMaterialExpressionRerouteBase : public UMaterialExpression
 	virtual uint32 GetOutputType(int32 OutputIndex) override;
 	virtual bool IsResultMaterialAttributes(int32 OutputIndex) override;
 	virtual bool IsResultStrataMaterial(int32 OutputIndex) override;
+	void GatherStrataMaterialInfo(FStrataMaterialInfo& StrataMaterialInfo, int32 OutputIndex) override;
+	FStrataOperator* StrataGenerateMaterialTopologyTree(class FMaterialCompiler* Compiler, class UMaterialExpression* Parent, int32 OutputIndex) override;
 #endif
 	//~ End UMaterialExpression Interface
 
diff --git a/Engine/Source/Runtime/Engine/Private/Materials/MaterialExpressions.cpp b/Engine/Source/Runtime/Engine/Private/Materials/MaterialExpressions.cpp
--- a/Engine/Source/Runtime/Engine/Private/Materials/MaterialExpressions.cpp
+++ b/Engine/Source/Runtime/Engine/Private/Materials/MaterialExpressions.cpp
@@ -18225,6 +18225,45 @@ bool UMaterialExpressionRerouteBase::IsResultStrataMaterial(int32 OutputIndex)
 	return false;
 }
 
+void UMaterialExpressionRerouteBase::GatherStrataMaterialInfo(FStrataMaterialInfo& StrataMaterialInfo, int32 OutputIndex)
+{
+	FExpressionInput Input;
+	if (GetRerouteInput(Input))
+	{
+		// Most code checks to make sure that there aren't loops before going here. In our case, we rely on the fact that
+		// UMaterialExpressionReroute's implementation of TraceInputsToRealExpression is resistant to input loops.
+		if (Input.IsConnected() && Input.Expression != nullptr && OutputIndex == 0)
+		{
+			int32 RealExpressionOutputIndex = -1;
+			UMaterialExpression* RealExpression = TraceInputsToRealExpression(RealExpressionOutputIndex);
+			if (RealExpression != nullptr)
+			{
+				RealExpression->GatherStrataMaterialInfo(StrataMaterialInfo, RealExpressionOutputIndex);
+			}
+		}
+	}
+}
+
+FStrataOperator* UMaterialExpressionRerouteBase::StrataGenerateMaterialTopologyTree(class FMaterialCompiler* Compiler, class UMaterialExpression* Parent, int32 OutputIndex)
+{
+	FExpressionInput Input;
+	if (GetRerouteInput(Input))
+	{
+		// Most code checks to make sure that there aren't loops before going here. In our case, we rely on the fact that
+		// UMaterialExpressionReroute's implementation of TraceInputsToRealExpression is resistant to input loops.
+		if (Input.IsConnected() && Input.Expression != nullptr && OutputIndex == 0)
+		{
+			int32 RealExpressionOutputIndex = -1;
+			UMaterialExpression* RealExpression = TraceInputsToRealExpression(RealExpressionOutputIndex);
+			if (RealExpression != nullptr)
+			{
+				RealExpression->StrataGenerateMaterialTopologyTree(Compiler, Parent, RealExpressionOutputIndex);
+			}
+		}
+	}
+	return nullptr;
+}
+
 #endif // WITH_EDITOR
 
 
