Re-parenting:
- avoid reparenting to one of the bone's children
- avoid reparenting the root
- check that the re-parented branch is not already part of the parent branch
Removing:
- avoid removing the root

diff --git a/Engine/Plugins/Experimental/MeshModelingToolsetExp/Source/MeshModelingToolsExp/Private/SkeletalMesh/SkeletonEditingTool.cpp b/Engine/Plugins/Experimental/MeshModelingToolsetExp/Source/MeshModelingToolsExp/Private/SkeletalMesh/SkeletonEditingTool.cpp
--- a/Engine/Plugins/Experimental/MeshModelingToolsetExp/Source/MeshModelingToolsExp/Private/SkeletalMesh/SkeletonEditingTool.cpp
+++ b/Engine/Plugins/Experimental/MeshModelingToolsetExp/Source/MeshModelingToolsExp/Private/SkeletalMesh/SkeletonEditingTool.cpp
@@ -356,13 +356,13 @@ void USkeletonEditingTool::ParentBones(const FName& InParentName)
 	const bool bBonesParented = SkeletonModifier.ParentBones(GetSelectedBones(), {InParentName});
 	if (bBonesParented)
 	{
-		Operation = EEditingOperation::Select;
 		EndChange();
-		return;
 	}
-
+	else
+	{
+		CancelChange();
+	}
 	Operation = EEditingOperation::Select;
-	CancelChange();
 }
 
 void USkeletonEditingTool::MoveBones()
diff --git a/Engine/Source/Runtime/Engine/Private/Animation/ReferenceSkeleton.cpp b/Engine/Source/Runtime/Engine/Private/Animation/ReferenceSkeleton.cpp
--- a/Engine/Source/Runtime/Engine/Private/Animation/ReferenceSkeleton.cpp
+++ b/Engine/Source/Runtime/Engine/Private/Animation/ReferenceSkeleton.cpp
@@ -100,6 +100,13 @@ void FReferenceSkeleton::Remove(const FName InBoneName, const bool bRemoveChildr
 	const int32 RawBoneIndex = FindRawBoneIndex(InBoneName);
 	if (RawBoneIndex == INDEX_NONE)
 	{
+		UE_LOG(LogAnimation, Error, TEXT("Remove: '%s' not found."), *InBoneName.ToString());
+		return;
+	}
+	
+	if (RawBoneIndex == 0)
+	{
+		UE_LOG(LogAnimation, Error, TEXT("Remove: cannot remove root bone."));
 		return;
 	}
 
@@ -209,13 +216,15 @@ int32 FReferenceSkeleton::SetParent(const FName InBoneName, const FName InParent
 
 	const int32 BoneIndex = FindRawBoneIndex(InBoneName);
 	if (BoneIndex < 1)
-	{ // we do not allow to change the root's parent 
-		return BoneIndex;
+	{ // we do not allow to change the root's parent
+		UE_LOG(LogAnimation, Error, TEXT("SetParent: cannot re-parent root bone."));
+		return INDEX_NONE;
 	}
 	
 	const int32 NewParentIndex = FindRawBoneIndex(InParentName);
 	if (NewParentIndex == INDEX_NONE && InParentName != NAME_None)
 	{
+		UE_LOG(LogAnimation, Error, TEXT("SetParent: '%s' not found."), *InParentName.ToString());
 		return INDEX_NONE;
 	}
 
@@ -224,6 +233,17 @@ int32 FReferenceSkeleton::SetParent(const FName InBoneName, const FName InParent
 		return RawRefBoneInfo.IsValidIndex(InBoneIndex) ? RawRefBoneInfo[InBoneIndex].ParentIndex : INDEX_NONE;
 	};
 
+	int32 ParentParentIndex = GetParentIndex(NewParentIndex); 
+	while (ParentParentIndex > INDEX_NONE)
+	{
+		if (ParentParentIndex == BoneIndex)
+		{ // we do not allow to reparent to one of bone's children
+			UE_LOG(LogAnimation, Error, TEXT("SetParent: cannot parent '%s' to one of its children."), *InBoneName.ToString());
+			return INDEX_NONE;
+		}
+		ParentParentIndex = GetParentIndex(ParentParentIndex);
+	}
+	
 	const int32 NumBones = GetRawBoneNum();
 	auto GetEndOfBranch = [&](const int32 InBoneIndex) -> int32
 	{
@@ -283,7 +303,20 @@ int32 FReferenceSkeleton::SetParent(const FName InBoneName, const FName InParent
 	else
 	{
 		// add all bones before the end of the new parent branch
-		AddBones(0, EndOfNewParentBranch);
+		if (EndOfNewParentBranch > BoneIndex)
+		{
+			for (int32 Index = 0; Index < EndOfNewParentBranch; Index++)
+			{
+				if (Index < BoneIndex || Index >= EndOfBoneBranch)
+				{
+					Names.Add(RawRefBoneInfo[Index].Name);
+				}
+			}
+		}
+		else
+		{
+			AddBones(0, EndOfNewParentBranch);
+		}
 		// add bone branch
 		AddBones(BoneIndex, EndOfBoneBranch);
 		// add remaining bones
