Substrate: Hair BSDF should not be emissive by default

diff --git a/Engine/Source/Runtime/Engine/Private/Materials/MaterialExpressions.cpp b/Engine/Source/Runtime/Engine/Private/Materials/MaterialExpressions.cpp
--- a/Engine/Source/Runtime/Engine/Private/Materials/MaterialExpressions.cpp
+++ b/Engine/Source/Runtime/Engine/Private/Materials/MaterialExpressions.cpp
@@ -24267,7 +24267,7 @@ int32 UMaterialExpressionStrataHairBSDF::Compile(class FMaterialCompiler* Compil
 		CompileWithDefaultFloat1(Compiler, Specular,	0.5f),
 		CompileWithDefaultFloat1(Compiler, Roughness,	0.5f),
 		CompileWithDefaultFloat1(Compiler, Backlit,		0.0f),
-		CompileWithDefaultFloat3(Compiler, EmissiveColor,1.0f, 0.0f, 0.0f),
+		CompileWithDefaultFloat3(Compiler, EmissiveColor,0.0f, 0.0f, 0.0f),
 		TangentCodeChunk,
 		Compiler->GetStrataSharedLocalBasisIndexMacro(NewRegisteredSharedLocalBasis),
 		&StrataOperator);
