Substrate - fix for multiple BSDF reseting gather info from switch nodes and losing subsurface profile.

diff --git a/Engine/Source/Runtime/Engine/Classes/Engine/EngineTypes.h b/Engine/Source/Runtime/Engine/Classes/Engine/EngineTypes.h
--- a/Engine/Source/Runtime/Engine/Classes/Engine/EngineTypes.h
+++ b/Engine/Source/Runtime/Engine/Classes/Engine/EngineTypes.h
@@ -725,6 +725,7 @@ public:
 
 	// Shading model
 	void AddShadingModel(EStrataShadingModel InShadingModel) { check(InShadingModel < SSM_NUM); ShadingModelField |= (1 << (uint16)InShadingModel); }
+	void SetSingleShadingModel(EStrataShadingModel InShadingModel) { check(InShadingModel < SSM_NUM); ShadingModelField = (1 << (uint16)InShadingModel); }
 	bool HasShadingModel(EStrataShadingModel InShadingModel) const { return (ShadingModelField & (1 << (uint16)InShadingModel)) != 0; }
 	bool HasOnlyShadingModel(EStrataShadingModel InShadingModel) const { return ShadingModelField == (1 << (uint16)InShadingModel); }
 	uint16 GetShadingModelField() const { return ShadingModelField; }
diff --git a/Engine/Source/Runtime/Engine/Private/Materials/Material.cpp b/Engine/Source/Runtime/Engine/Private/Materials/Material.cpp
--- a/Engine/Source/Runtime/Engine/Private/Materials/Material.cpp
+++ b/Engine/Source/Runtime/Engine/Private/Materials/Material.cpp
@@ -4856,9 +4856,9 @@ void UMaterial::RebuildShadingModelField()
 			{
 				if (BlendMode == EBlendMode::BLEND_Opaque || BlendMode == EBlendMode::BLEND_Masked) 
 				{
-					// We only consider SSS subsurface post process for opaque materials)
-					StrataMaterialInfo = FStrataMaterialInfo();
-					StrataMaterialInfo.AddShadingModel(SSM_SubsurfaceLit);
+					// We only consider SSS subsurface post process for opaque materials.
+					// STRATA_TODO what does this case cover?
+					StrataMaterialInfo.SetSingleShadingModel(SSM_SubsurfaceLit);
 				}
 				else
 				{
@@ -4866,6 +4866,22 @@ void UMaterial::RebuildShadingModelField()
 					bSanitizeMaterial = true;
 				}
 			}
+			else if (StrataMaterialInfo.CountShadingModels() > 1 && StrataMaterialInfo.HasShadingModel(EStrataShadingModel::SSM_SubsurfaceLit))
+			{
+				// When we gather the worst case slab material, we can hit multiple shading models when gathering from multiple BSDF possibilities of switch nodes.
+				// We do not want to sanitize the material in this case, but only keep the worst case: subsurface with all potential features on a slab.
+				if (BlendMode == EBlendMode::BLEND_Opaque || BlendMode == EBlendMode::BLEND_Masked) 
+				{
+					// We only consider SSS subsurface post process for opaque materials. This can also be triggered by switch node accumulating multiple slab node.
+					// We still keep all the other material info to consider the worst case when encountering multiple shading models. And we also keep around the encountered subsurface profiles.
+					StrataMaterialInfo.SetSingleShadingModel(SSM_SubsurfaceLit);
+				}
+				else
+				{
+					// For transparent, we will fall back to use DefaultLit worst case with simple volumetric.
+					bSanitizeMaterial = true;
+				}
+			}
 			else
 			{
 				// Clear the material to default Lit
