Compilation fix.

diff --git a/Engine/Source/Runtime/Engine/Classes/Engine/EngineTypes.h b/Engine/Source/Runtime/Engine/Classes/Engine/EngineTypes.h
--- a/Engine/Source/Runtime/Engine/Classes/Engine/EngineTypes.h
+++ b/Engine/Source/Runtime/Engine/Classes/Engine/EngineTypes.h
@@ -724,8 +724,8 @@ public:
 	FStrataMaterialInfo(EStrataShadingModel InShadingModel) { AddShadingModel(InShadingModel); }
 
 	// Shading model
-	void AddShadingModel(EStrataShadingModel InShadingModel) { check(InShadingModel < SSM_NUM); ShadingModelField |= (1 << (uint16)InShadingModel); }
-	void SetSingleShadingModel(EStrataShadingModel InShadingModel) { check(InShadingModel < SSM_NUM); ShadingModelField = (1 << (uint16)InShadingModel); }
+	void AddShadingModel(EStrataShadingModel InShadingModel) { check(InShadingModel < SSM_NUM); ShadingModelField |= (uint16)(1 << (uint16)InShadingModel); }
+	void SetSingleShadingModel(EStrataShadingModel InShadingModel) { check(InShadingModel < SSM_NUM); ShadingModelField = (uint16)(1 << (uint16)InShadingModel); }
 	bool HasShadingModel(EStrataShadingModel InShadingModel) const { return (ShadingModelField & (1 << (uint16)InShadingModel)) != 0; }
 	bool HasOnlyShadingModel(EStrataShadingModel InShadingModel) const { return ShadingModelField == (1 << (uint16)InShadingModel); }
 	uint16 GetShadingModelField() const { return ShadingModelField; }
