Exposed and adapted FWorldPartitionHelpers::ConvertEditorPathToRuntimePath to work in engine/PIE/-game and fixed path conversion for actors part of PersistentLevel.

diff --git a/Engine/Source/Runtime/Engine/Private/WorldPartition/WorldPartitionHelpers.cpp b/Engine/Source/Runtime/Engine/Private/WorldPartition/WorldPartitionHelpers.cpp
--- a/Engine/Source/Runtime/Engine/Private/WorldPartition/WorldPartitionHelpers.cpp
+++ b/Engine/Source/Runtime/Engine/Private/WorldPartition/WorldPartitionHelpers.cpp
@@ -291,10 +291,17 @@ bool FWorldPartitionHelpers::ConvertRuntimePathToEditorPath(const FSoftObjectPat
 			}
 		}
 	}
+	else
+	{
+		OutPath = UWorld::RemovePIEPrefix(InPath.ToString());
+		return true;
+	}
 
 	return false;
 }
 
+#endif // #if WITH_EDITOR
+
 bool FWorldPartitionHelpers::ConvertEditorPathToRuntimePath(const FSoftObjectPath& InPath, FSoftObjectPath& OutPath)
 {
 	//
@@ -313,6 +320,4 @@ bool FWorldPartitionHelpers::ConvertEditorPathToRuntimePath(const FSoftObjectPat
 	}
 
 	return false;
-}
-
-#endif // #if WITH_EDITOR
\ No newline at end of file
+}
\ No newline at end of file
diff --git a/Engine/Source/Runtime/Engine/Private/WorldPartition/WorldPartitionLevelStreamingPolicy.cpp b/Engine/Source/Runtime/Engine/Private/WorldPartition/WorldPartitionLevelStreamingPolicy.cpp
--- a/Engine/Source/Runtime/Engine/Private/WorldPartition/WorldPartitionLevelStreamingPolicy.cpp
+++ b/Engine/Source/Runtime/Engine/Private/WorldPartition/WorldPartitionLevelStreamingPolicy.cpp
@@ -13,17 +13,16 @@
 #include "Engine/Level.h"
 #include "Engine/Canvas.h"
 #include "Engine/LevelStreamingGCHelper.h"
-
-#include UE_INLINE_GENERATED_CPP_BY_NAME(WorldPartitionLevelStreamingPolicy)
-
 #if WITH_EDITOR
 #include "Misc/PackageName.h"
+#endif
+
+#include UE_INLINE_GENERATED_CPP_BY_NAME(WorldPartitionLevelStreamingPolicy)
 
 namespace WorldPartitionLevelStreamingPolicy
 {
 	static FName PersistentCellName(TEXT("PersistentCell"));
 }
-#endif
 
 int32 UWorldPartitionLevelStreamingPolicy::GetCellLoadingCount() const
 {
@@ -228,61 +227,78 @@ void UWorldPartitionLevelStreamingPolicy::RemapSoftObjectPath(FSoftObjectPath& O
 				{
 					ObjectPath.FixupForPIE(PIEInstanceID);
 				}
-			}	
+			}
 		}
 	}
 }
+#endif
 
 bool UWorldPartitionLevelStreamingPolicy::ConvertEditorPathToRuntimePath(const FSoftObjectPath& InPath, FSoftObjectPath& OutPath) const
 {
 	const UWorld* OuterWorld = WorldPartition->GetTypedOuter<UWorld>();
 	const UPackage* OuterWorldPackage = OuterWorld->GetPackage();
-	const FTopLevelAssetPath WorldAssetPath(OuterWorld);
+	FTopLevelAssetPath WorldAssetPath(OuterWorld);
+	int32 PIEInstanceID = INDEX_NONE;
+#if WITH_EDITOR
+	WorldAssetPath = UWorld::RemovePIEPrefix(WorldAssetPath.ToString(), &PIEInstanceID);
+	check(PIEInstanceID == INDEX_NONE || OuterWorldPackage->HasAnyPackageFlags(PKG_PlayInEditor));
+	check(PIEInstanceID == OuterWorldPackage->GetPIEInstanceID());
+#endif
 
-	if (!OuterWorldPackage->HasAnyPackageFlags(PKG_PlayInEditor))
+	if (InPath.GetAssetPath() == WorldAssetPath)
 	{
-		if (InPath.GetAssetPath() == WorldAssetPath)
+		FString SubAssetName;
+		FString SubAssetContext;
+		if (InPath.GetSubPathString().Split(TEXT("."), &SubAssetContext, &SubAssetName))
 		{
-			FString SubAssetName;
-			FString SubAssetContext;
-			if (InPath.GetSubPathString().Split(TEXT("."), &SubAssetContext, &SubAssetName))
+			if (SubAssetContext == TEXT("PersistentLevel"))
 			{
-				if (SubAssetContext == TEXT("PersistentLevel"))
+				FString SubObjectName;
+				FString SubObjectContext(SubAssetName);
+				SubAssetName.Split(TEXT("."), &SubObjectContext, &SubObjectName);
+
+				FString LevelInstanceTag;
+				FString WorldAssetPackageName = *WorldAssetPath.GetPackageName().ToString();
+				FString WorldAssetName = WorldAssetPath.GetAssetName().ToString();
+
+#if WITH_EDITOR
+				// In the editor, the _LevelInstance_ID is appended to the persistent level, while at runtime it is appended to each cell package, so we need to remap it there if present.
+				const int32 LevelInstancePos = WorldAssetPackageName.Find(TEXT("_LevelInstance_"), ESearchCase::IgnoreCase, ESearchDir::FromEnd);
+				if (LevelInstancePos != INDEX_NONE)
 				{
-					FString SubObjectName;
-					FString SubObjectContext(SubAssetName);
-					SubAssetName.Split(TEXT("."), &SubObjectContext, &SubObjectName);
-
-					FString LevelInstanceTag;
-					FString WorldAssetPackageName = *WorldAssetPath.GetPackageName().ToString();
-					FString WorldAssetName = WorldAssetPath.GetAssetName().ToString();
-
-					// In the editor, the _LevelInstance_ID is appended to the persistent level, while at runtime it is appended to each cell package, so we need to remap it there if present.
-					const int32 LevelInstancePos = WorldAssetPackageName.Find(TEXT("_LevelInstance_"), ESearchCase::IgnoreCase, ESearchDir::FromEnd);
-					if (LevelInstancePos != INDEX_NONE)
-					{
-						LevelInstanceTag = WorldAssetPackageName.RightChop(LevelInstancePos);
-						WorldAssetPackageName.LeftInline(LevelInstancePos);
-					}
-
-					// Try to find the corresponding streaming cell, if it doesn't exists the actor must be in the persistent level.
-					if (const FName* CellName = SubObjectsToCellRemapping.Find(*SubObjectContext))
-					{
-						OutPath = FString::Printf(TEXT("%s/_Generated_/%s%s.%s:%s"), *WorldAssetPackageName, *(CellName->ToString()), *LevelInstanceTag, *WorldAssetPath.GetAssetName().ToString(), *InPath.GetSubPathString());
-					}
-					else
-					{
-						OutPath = FString::Printf(TEXT("%s/_Generated_/%s%s.%s:%s"), *WorldAssetPackageName, *WorldAssetName, *LevelInstanceTag, *WorldAssetPath.GetAssetName().ToString(), *InPath.GetSubPathString());
-					}
-					return true;
+					LevelInstanceTag = WorldAssetPackageName.RightChop(LevelInstancePos);
+					WorldAssetPackageName.LeftInline(LevelInstancePos);
 				}
+#endif
+
+				// Try to find the corresponding streaming cell, if it doesn't exists the actor must be in the persistent level.
+				const FName* CellName = SubObjectsToCellRemapping.Find(*SubObjectContext);
+				if (!CellName  || (*CellName == WorldPartitionLevelStreamingPolicy::PersistentCellName))
+				{
+					OutPath = FSoftObjectPath(WorldAssetPath, InPath.GetSubPathString());
+				}
+#if WITH_EDITOR
+				else if (OuterWorld->IsGameWorld())
+				{
+					const FString PackagePath = UWorldPartitionLevelStreamingPolicy::GetCellPackagePath(*CellName, OuterWorld);
+					OutPath = FString::Printf(TEXT("%s.%s:%s"), *PackagePath, *OuterWorld->GetName(), *InPath.GetSubPathString());
+				}
+#endif
+				else
+				{
+					OutPath = FString::Printf(TEXT("%s/_Generated_/%s%s.%s:%s"), *WorldAssetPackageName, *(CellName->ToString()), *LevelInstanceTag, *WorldAssetPath.GetAssetName().ToString(), *InPath.GetSubPathString());
+				}
+
+#if WITH_EDITOR
+				OutPath.FixupForPIE(PIEInstanceID);
+#endif
+				return true;
 			}
 		}
 	}
 
 	return false;
 }
-#endif
 
 UObject* UWorldPartitionLevelStreamingPolicy::GetSubObject(const TCHAR* SubObjectPath)
 {
@@ -290,7 +306,7 @@ UObject* UWorldPartitionLevelStreamingPolicy::GetSubObject(const TCHAR* SubObjec
 
 	// Support for subobjects such as Actor.Component
 	FString SubObjectName;
-	FString SubObjectContext;	
+	FString SubObjectContext;
 	if (!FString(SubObjectPath).Split(TEXT("."), &SubObjectContext, &SubObjectName))
 	{
 		SubObjectContext = SubObjectPath;
diff --git a/Engine/Source/Runtime/Engine/Public/WorldPartition/WorldPartitionHelpers.h b/Engine/Source/Runtime/Engine/Public/WorldPartition/WorldPartitionHelpers.h
--- a/Engine/Source/Runtime/Engine/Public/WorldPartition/WorldPartitionHelpers.h
+++ b/Engine/Source/Runtime/Engine/Public/WorldPartition/WorldPartitionHelpers.h
@@ -136,10 +136,12 @@ public:
 	// Simulate an engine frame tick
 	static void FakeEngineTick(UWorld* World);
 
-	// Editor/Runtime conversions
+	// Runtime/Editor conversions
 	static bool ConvertRuntimePathToEditorPath(const FSoftObjectPath& InPath, FSoftObjectPath& OutPath);
-	static bool ConvertEditorPathToRuntimePath(const FSoftObjectPath& InPath, FSoftObjectPath& OutPath);
 #endif // WITH_EDITOR
+
+	// Editor/Runtime conversions
+	static bool ConvertEditorPathToRuntimePath(const FSoftObjectPath& InPath, FSoftObjectPath& OutPath);
 };
 
 
