Important navmesh tile generation optimization when using HISM (fix to avoid gathering unnecessary geometry when building tiles).

diff --git a/Engine/Source/Runtime/Engine/Private/HierarchicalInstancedStaticMesh.cpp b/Engine/Source/Runtime/Engine/Private/HierarchicalInstancedStaticMesh.cpp
--- a/Engine/Source/Runtime/Engine/Private/HierarchicalInstancedStaticMesh.cpp
+++ b/Engine/Source/Runtime/Engine/Private/HierarchicalInstancedStaticMesh.cpp
@@ -3492,7 +3492,7 @@ void UHierarchicalInstancedStaticMeshComponent::GetNavigationPerInstanceTransfor
 		const TArray<FClusterNode>& ClusterTree = *ClusterTreePtr;
 		if (ClusterTree.Num())
 		{
-			GatherInstanceTransformsInArea(*this, AreaBox, 0, InstanceData);
+			GetOverlappingBoxTransforms(AreaBox, InstanceData);
 		}
 	}
 	else
