Added explicit move and copies to FOutBunch and FInBunch in a non-deprecated block to prevent warnings in those functions.

diff --git a/Engine/Source/Runtime/Engine/Public/Net/DataBunch.h b/Engine/Source/Runtime/Engine/Public/Net/DataBunch.h
--- a/Engine/Source/Runtime/Engine/Public/Net/DataBunch.h
+++ b/Engine/Source/Runtime/Engine/Public/Net/DataBunch.h
@@ -75,7 +75,9 @@ public:
 	FOutBunch( UPackageMap * PackageMap, int64 InMaxBits = 1024 );
 
 	PRAGMA_DISABLE_DEPRECATION_WARNINGS
+	FOutBunch(FOutBunch&&) = default;
 	FOutBunch(const FOutBunch&) = default;
+	FOutBunch& operator=(FOutBunch&&) = default;
 	FOutBunch& operator=(const FOutBunch&) = default;
 	PRAGMA_ENABLE_DEPRECATION_WARNINGS
 
@@ -173,7 +175,10 @@ public:
 	FInBunch( FInBunch &InBunch, bool CopyBuffer );
 
 	PRAGMA_DISABLE_DEPRECATION_WARNINGS
+	FInBunch(FInBunch&&) = default;
 	FInBunch(const FInBunch&) = default;
+	FInBunch& operator=(FInBunch&&) = default;
+	FInBunch& operator=(const FInBunch&) = default;
 	PRAGMA_ENABLE_DEPRECATION_WARNINGS
 
 	virtual void CountMemory(FArchive& Ar) const override;
