Added explicit move and copies to FNetGUIDCache in a non-deprecated block to prevent warnings in those functions.

diff --git a/Engine/Source/Runtime/Engine/Classes/Engine/PackageMapClient.h b/Engine/Source/Runtime/Engine/Classes/Engine/PackageMapClient.h
--- a/Engine/Source/Runtime/Engine/Classes/Engine/PackageMapClient.h
+++ b/Engine/Source/Runtime/Engine/Classes/Engine/PackageMapClient.h
@@ -151,6 +151,13 @@ class ENGINE_API FNetGUIDCache
 public:
 	FNetGUIDCache( UNetDriver * InDriver );
 
+	PRAGMA_DISABLE_DEPRECATION_WARNINGS
+	FNetGUIDCache(FNetGUIDCache&&) = default;
+	FNetGUIDCache(const FNetGUIDCache&) = default;
+	FNetGUIDCache& operator=(FNetGUIDCache&&) = default;
+	FNetGUIDCache& operator=(const FNetGUIDCache&) = default;
+	PRAGMA_ENABLE_DEPRECATION_WARNINGS
+
 	enum class ENetworkChecksumMode : uint8
 	{
 		None			= 0,		// Don't use checksums
