Deleted FNetGUIDCache's copy functions because its FNetAsyncLoadDelinquencyAnalytics member is non-copyable.

diff --git a/Engine/Source/Runtime/Engine/Classes/Engine/PackageMapClient.h b/Engine/Source/Runtime/Engine/Classes/Engine/PackageMapClient.h
--- a/Engine/Source/Runtime/Engine/Classes/Engine/PackageMapClient.h
+++ b/Engine/Source/Runtime/Engine/Classes/Engine/PackageMapClient.h
@@ -153,9 +153,9 @@ public:
 
 	PRAGMA_DISABLE_DEPRECATION_WARNINGS
 	FNetGUIDCache(FNetGUIDCache&&) = default;
-	FNetGUIDCache(const FNetGUIDCache&) = default;
+	FNetGUIDCache(const FNetGUIDCache&) = delete;
 	FNetGUIDCache& operator=(FNetGUIDCache&&) = default;
-	FNetGUIDCache& operator=(const FNetGUIDCache&) = default;
+	FNetGUIDCache& operator=(const FNetGUIDCache&) = delete;
 	PRAGMA_ENABLE_DEPRECATION_WARNINGS
 
 	enum class ENetworkChecksumMode : uint8
