NetDriver
* Added configuration option for Iris ReplicationSystem parameters

diff --git a/Engine/Source/Runtime/Engine/Classes/Engine/NetDriver.h b/Engine/Source/Runtime/Engine/Classes/Engine/NetDriver.h
--- a/Engine/Source/Runtime/Engine/Classes/Engine/NetDriver.h
+++ b/Engine/Source/Runtime/Engine/Classes/Engine/NetDriver.h
@@ -604,6 +604,25 @@ public:
 	}
 };
 
+/** Used to configure the replication system default values */
+USTRUCT()
+struct FNetDriverReplicationSystemConfig
+{
+	GENERATED_USTRUCT_BODY()
+
+	/** Override the max object count. If 0 use the default system value. */
+	UPROPERTY()
+	uint32 MaxReplicatedObjectCount = 0;
+	
+	/** Override the max compressed object count. If 0 use the default system value. */
+	UPROPERTY()
+	uint32 MaxDeltaCompressedObjectCount = 0;
+	
+	/** Override the max group count. If 0 use the default system value. */
+	UPROPERTY()
+	uint32 MaxNetObjectGroupCount = 0;
+};
+
 //
 // Priority sortable list.
 //
@@ -750,6 +769,10 @@ public:
 	/** Used to specify the class to use for ReplicationBridge */
 	UPROPERTY(Config)
 	FString ReplicationBridgeClassName;
+	
+	/** Can be used to configure settings for the ReplicationSystem */
+	UPROPERTY(Config)
+	FNetDriverReplicationSystemConfig ReplicationSystemConfig;
 
 	/** @todo document */
 	UPROPERTY(Config)
diff --git a/Engine/Source/Runtime/Engine/Private/NetDriver.cpp b/Engine/Source/Runtime/Engine/Private/NetDriver.cpp
--- a/Engine/Source/Runtime/Engine/Private/NetDriver.cpp
+++ b/Engine/Source/Runtime/Engine/Private/NetDriver.cpp
@@ -289,6 +289,25 @@ namespace UE::Net::Private
 
 	static bool bIgnoreStaticActorDestruction = false;
 
+
+	void ApplyReplicationSystemConfig(const FNetDriverReplicationSystemConfig& ReplicationSystemConfig, UReplicationSystem::FReplicationSystemParams& OutParams)
+	{
+		if (ReplicationSystemConfig.MaxReplicatedObjectCount != 0)
+		{
+			OutParams.MaxReplicatedObjectCount = ReplicationSystemConfig.MaxReplicatedObjectCount;
+		}
+
+		if (ReplicationSystemConfig.MaxDeltaCompressedObjectCount != 0)
+		{
+			OutParams.MaxDeltaCompressedObjectCount = ReplicationSystemConfig.MaxDeltaCompressedObjectCount;
+		}
+
+		if (ReplicationSystemConfig.MaxNetObjectGroupCount != 0)
+		{
+			OutParams.MaxNetObjectGroupCount = ReplicationSystemConfig.MaxNetObjectGroupCount;
+		}
+	}
+
 } //namespace UE::Net::Private
 
 namespace UE::Net
@@ -1475,6 +1494,8 @@ bool UNetDriver::InitBase(bool bInitAsClient, FNetworkNotify* InNotify, const FU
 				Params.bIsServer = !bInitAsClient;
 				Params.bAllowObjectReplication = !bInitAsClient;
 
+				UE::Net::Private::ApplyReplicationSystemConfig(ReplicationSystemConfig, Params);
+
 				SetReplicationSystem(UE::Net::FReplicationSystemFactory::CreateReplicationSystem(Params));
 			}
 			else
