Fix FWorldPartitionActorDesc::Equals missing some properties (spatially loaded and content bundle id).

diff --git a/Engine/Source/Runtime/Engine/Private/WorldPartition/WorldPartitionActorDesc.cpp b/Engine/Source/Runtime/Engine/Private/WorldPartition/WorldPartitionActorDesc.cpp
--- a/Engine/Source/Runtime/Engine/Private/WorldPartition/WorldPartitionActorDesc.cpp
+++ b/Engine/Source/Runtime/Engine/Private/WorldPartition/WorldPartitionActorDesc.cpp
@@ -244,6 +244,7 @@ bool FWorldPartitionActorDesc::Equals(const FWorldPartitionActorDesc* Other) con
 		BoundsLocation.Equals(Other->BoundsLocation, 0.1f) &&
 		BoundsExtent.Equals(Other->BoundsExtent, 0.1f) &&
 		RuntimeGrid == Other->RuntimeGrid &&
+		bIsSpatiallyLoaded == Other->bIsSpatiallyLoaded &&
 		bActorIsEditorOnly == Other->bActorIsEditorOnly &&
 		bActorIsRuntimeOnly == Other->bActorIsRuntimeOnly &&
 		bActorIsHLODRelevant == Other->bActorIsHLODRelevant &&
@@ -252,6 +253,7 @@ bool FWorldPartitionActorDesc::Equals(const FWorldPartitionActorDesc* Other) con
 		FolderPath == Other->FolderPath &&
 		FolderGuid == Other->FolderGuid &&
 		ParentActor == Other->ParentActor &&
+		ContentBundleGuid == Other->ContentBundleGuid &&
 		CompareUnsortedArrays(DataLayers, Other->DataLayers) &&
 		CompareUnsortedArrays(References, Other->References) &&
 		CompareUnsortedArrays(Tags, Other->Tags) &&
