Speculative fix for VT hanging in FVirtualTextureProducer::Release waiting for tasks that seem to have never been spawned

#rnx
#tests aibotsoak on xb1 during the weekend. Checked dev and test builds

diff --git a/Engine/Source/Runtime/Engine/Private/VT/VirtualTextureTranscodeCache.cpp b/Engine/Source/Runtime/Engine/Private/VT/VirtualTextureTranscodeCache.cpp
--- a/Engine/Source/Runtime/Engine/Private/VT/VirtualTextureTranscodeCache.cpp
+++ b/Engine/Source/Runtime/Engine/Private/VT/VirtualTextureTranscodeCache.cpp
@@ -379,6 +379,13 @@ FVTTranscodeTileHandle FVirtualTextureTranscodeCache::SubmitTask(
 	TaskEntry.FrameSubmitted = GFrameNumberRenderThread;
 	FMemory::Memzero(TaskEntry.StageTileHandle);
 
+	check(!TaskEntry.GraphEvent || TaskEntry.GraphEvent->IsComplete());
+	
+	if (TaskEntry.GraphEvent)
+	{
+		TaskEntry.GraphEvent.SafeRelease();
+	}
+
 	const FGraphEventArray* Prerequisites = InPrerequisites;
 	if (Prerequisites)
 	{
