add a default value for FClusterUnionParticleCandidateData::BoneId

diff --git a/Engine/Source/Runtime/Engine/Classes/PhysicsEngine/ClusterUnionComponent.h b/Engine/Source/Runtime/Engine/Classes/PhysicsEngine/ClusterUnionComponent.h
--- a/Engine/Source/Runtime/Engine/Classes/PhysicsEngine/ClusterUnionComponent.h
+++ b/Engine/Source/Runtime/Engine/Classes/PhysicsEngine/ClusterUnionComponent.h
@@ -98,7 +98,7 @@ struct FClusterUnionParticleCandidateData
 	TObjectPtr<UPrimitiveComponent> Component;
 
 	UPROPERTY()
-	int32 BoneId;
+	int32 BoneId = INDEX_NONE;
 };
 
 /**
