add editor visualization to the cluster union component

diff --git a/Engine/Source/Runtime/Engine/Private/PhysicsEngine/ClusterUnionComponent.cpp b/Engine/Source/Runtime/Engine/Private/PhysicsEngine/ClusterUnionComponent.cpp
--- a/Engine/Source/Runtime/Engine/Private/PhysicsEngine/ClusterUnionComponent.cpp
+++ b/Engine/Source/Runtime/Engine/Private/PhysicsEngine/ClusterUnionComponent.cpp
@@ -19,6 +19,10 @@ UClusterUnionComponent::UClusterUnionComponent(const FObjectInitializer& ObjectI
 {
 	PhysicsProxy = nullptr;
 	SetIsReplicatedByDefault(true);
+
+#if WITH_EDITORONLY_DATA
+	bVisualizeComponent = true;
+#endif
 }
 
 FPhysScene_Chaos* UClusterUnionComponent::GetChaosScene() const
