Pack FTexture2DMipMap fields

diff --git a/Engine/Plugins/Experimental/AnimToTexture/Source/AnimToTextureEditor/Public/AnimToTextureUtils.h b/Engine/Plugins/Experimental/AnimToTexture/Source/AnimToTextureEditor/Public/AnimToTextureUtils.h
--- a/Engine/Plugins/Experimental/AnimToTexture/Source/AnimToTextureEditor/Public/AnimToTextureUtils.h
+++ b/Engine/Plugins/Experimental/AnimToTexture/Source/AnimToTextureEditor/Public/AnimToTextureUtils.h
@@ -233,7 +233,7 @@ FORCEINLINE bool AnimToTexture_Private::WriteToTexture(
 	FTexture2DMipMap* Mip;
 	if (PlatformData->Mips.IsEmpty())
 	{
-		Mip = new FTexture2DMipMap();
+		Mip = new FTexture2DMipMap(0, 0);
 		PlatformData->Mips.Add(Mip);
 	}
 	else
diff --git a/Engine/Plugins/Experimental/Mutable/Source/CustomizableObject/Private/MuCO/CustomizableObjectInstance.cpp b/Engine/Plugins/Experimental/Mutable/Source/CustomizableObject/Private/MuCO/CustomizableObjectInstance.cpp
--- a/Engine/Plugins/Experimental/Mutable/Source/CustomizableObject/Private/MuCO/CustomizableObjectInstance.cpp
+++ b/Engine/Plugins/Experimental/Mutable/Source/CustomizableObject/Private/MuCO/CustomizableObjectInstance.cpp
@@ -3061,11 +3061,6 @@ FTexturePlatformData* UCustomizableInstancePrivateData::MutableCreateImagePlatfo
 	{
 		int32 MipLevelMutable = MipLevelUE - MipsToSkip;		
 
-		PlatformData->Mips.Add(new FTexture2DMipMap());
-		FTexture2DMipMap* Mip = &PlatformData->Mips.Last(); //new(Texture->PlatformData->Mips) FTexture2DMipMap();
-		Mip->SizeX = SizeX;
-		Mip->SizeY = SizeY;
-
 		// Unlike Mutable, UE expects MIPs sizes to be at least the size of the compression block.
 		// For example, a 8x8 PF_DXT1 texture will have the following MIPs:
 		// Mutable    Unreal Engine
@@ -3075,9 +3070,10 @@ FTexturePlatformData* UCustomizableInstancePrivateData::MutableCreateImagePlatfo
 		// 1x1        4x4
 		//
 		// Notice that even though Mutable reports MIP smaller than the block size, the actual data contains at least a block.
-		Mip->SizeX = FMath::Max(Mip->SizeX, GPixelFormats[PlatformFormat].BlockSizeX);
-		Mip->SizeY = FMath::Max(Mip->SizeY, GPixelFormats[PlatformFormat].BlockSizeY);
+		FTexture2DMipMap* Mip = new FTexture2DMipMap( FMath::Max(SizeX, GPixelFormats[PlatformFormat].BlockSizeX)
+													, FMath::Max(SizeY, GPixelFormats[PlatformFormat].BlockSizeY));
 
+		PlatformData->Mips.Add(Mip);
 		if(MipLevelUE >= MipsToSkip || OnlyLOD>=0)
 		{
 			check(MipLevelMutable >= 0);
diff --git a/Engine/Plugins/Runtime/HairStrands/Source/HairStrandsCore/Private/GroomResources.cpp b/Engine/Plugins/Runtime/HairStrands/Source/HairStrandsCore/Private/GroomResources.cpp
--- a/Engine/Plugins/Runtime/HairStrands/Source/HairStrandsCore/Private/GroomResources.cpp
+++ b/Engine/Plugins/Runtime/HairStrands/Source/HairStrandsCore/Private/GroomResources.cpp
@@ -542,10 +542,8 @@ static UTexture2D* CreateCardTexture(FIntPoint Resolution)
 		const uint32 MipResolutionY = Resolution.Y >> MipIt;
 		const uint32 SizeInBytes = sizeof(float) * MipResolutionX * MipResolutionY;
 
-		FTexture2DMipMap* MipMap = new FTexture2DMipMap();
+		FTexture2DMipMap* MipMap = new FTexture2DMipMap(MipResolutionX, MipResolutionY);
 		Out->GetPlatformData()->Mips.Add(MipMap);
-		MipMap->SizeX = MipResolutionX;
-		MipMap->SizeY = MipResolutionY;
 		MipMap->BulkData.Lock(LOCK_READ_WRITE);
 		float* MipMemory = (float*)MipMap->BulkData.Realloc(SizeInBytes);
 		for (uint32 Y = 0; Y < MipResolutionY; Y++)
diff --git a/Engine/Plugins/Runtime/MeshModelingToolset/Source/ModelingComponents/Private/AssetUtils/Texture2DBuilder.cpp b/Engine/Plugins/Runtime/MeshModelingToolset/Source/ModelingComponents/Private/AssetUtils/Texture2DBuilder.cpp
--- a/Engine/Plugins/Runtime/MeshModelingToolset/Source/ModelingComponents/Private/AssetUtils/Texture2DBuilder.cpp
+++ b/Engine/Plugins/Runtime/MeshModelingToolset/Source/ModelingComponents/Private/AssetUtils/Texture2DBuilder.cpp
@@ -43,10 +43,8 @@ bool FTexture2DBuilder::InitializeAndReplaceExistingTexture(UTexture2D* Existing
 		// Allocate first mipmap.
 		int32 NumBlocksX = InSizeX / GPixelFormats[InFormat].BlockSizeX;
 		int32 NumBlocksY = InSizeY / GPixelFormats[InFormat].BlockSizeY;
-		FTexture2DMipMap* Mip = new FTexture2DMipMap();
+		FTexture2DMipMap* Mip = new FTexture2DMipMap(InSizeX, InSizeY);
 		NewTexture->GetPlatformData()->Mips.Add(Mip);
-		Mip->SizeX = InSizeX;
-		Mip->SizeY = InSizeY;
 		Mip->BulkData.Lock(LOCK_READ_WRITE);
 		Mip->BulkData.Realloc(NumBlocksX * NumBlocksY * GPixelFormats[InFormat].BlockBytes);
 		Mip->BulkData.Unlock();
@@ -89,10 +87,8 @@ bool FTexture2DBuilder::InitializeAndReplaceExistingTexture(UTexture2D* Existing
 		// Allocate first mipmap.
 		const int32 NumBlocksX = InSizeX / GPixelFormats[InFormat].BlockSizeX;
 		const int32 NumBlocksY = InSizeY / GPixelFormats[InFormat].BlockSizeY;
-		FTexture2DMipMap* Mip = new FTexture2DMipMap();
+		FTexture2DMipMap* Mip = new FTexture2DMipMap(InSizeX, InSizeY);
 		NewTexture->GetPlatformData()->Mips.Add(Mip);
-		Mip->SizeX = InSizeX;
-		Mip->SizeY = InSizeY;
 		Mip->BulkData.Lock(LOCK_READ_WRITE);
 		Mip->BulkData.Realloc(NumBlocksX * NumBlocksY * GPixelFormats[InFormat].BlockBytes);
 		Mip->BulkData.Unlock();
diff --git a/Engine/Source/Runtime/Engine/Classes/Engine/Texture2DArray.cpp b/Engine/Source/Runtime/Engine/Classes/Engine/Texture2DArray.cpp
--- a/Engine/Source/Runtime/Engine/Classes/Engine/Texture2DArray.cpp
+++ b/Engine/Source/Runtime/Engine/Classes/Engine/Texture2DArray.cpp
@@ -66,11 +66,8 @@ UTexture2DArray* UTexture2DArray::CreateTransient(int32 InSizeX, int32 InSizeY,
 		// Allocate first mipmap.
 		int32 NumBlocksX = InSizeX / GPixelFormats[InFormat].BlockSizeX;
 		int32 NumBlocksY = InSizeY / GPixelFormats[InFormat].BlockSizeY;
-		FTexture2DMipMap* Mip = new FTexture2DMipMap();
+		FTexture2DMipMap* Mip = new FTexture2DMipMap(InSizeX, InSizeY, InArraySize);
 		NewTexture->GetPlatformData()->Mips.Add(Mip);
-		Mip->SizeX = InSizeX;
-		Mip->SizeY = InSizeY;
-		Mip->SizeZ = InArraySize;
 		Mip->BulkData.Lock(LOCK_READ_WRITE);
 		Mip->BulkData.Realloc((int64)GPixelFormats[InFormat].BlockBytes * NumBlocksX * NumBlocksY * InArraySize);
 		Mip->BulkData.Unlock();
diff --git a/Engine/Source/Runtime/Engine/Private/Streaming/Texture2DStreamIn.cpp b/Engine/Source/Runtime/Engine/Private/Streaming/Texture2DStreamIn.cpp
--- a/Engine/Source/Runtime/Engine/Private/Streaming/Texture2DStreamIn.cpp
+++ b/Engine/Source/Runtime/Engine/Private/Streaming/Texture2DStreamIn.cpp
@@ -128,8 +128,8 @@ void FTexture2DStreamIn::DoAsyncCreateWithNewMips(const FContext& Context)
 		const FTexture2DMipMap& RequestedMipMap = *Context.MipsView[PendingFirstLODIdx];
 
 		// old textures have sizes padded up to multiple of 4; that's wrong, should be real size unpadded
-		check( RequestedMipMap.SizeX == FMath::Max(1, Context.MipsView[0]->SizeX >> PendingFirstLODIdx) );
-		check( RequestedMipMap.SizeY == FMath::Max(1, Context.MipsView[0]->SizeY >> PendingFirstLODIdx) );
+		check( RequestedMipMap.SizeX == FMath::Max(1, (int32)Context.MipsView[0]->SizeX >> PendingFirstLODIdx) );
+		check( RequestedMipMap.SizeY == FMath::Max(1, (int32)Context.MipsView[0]->SizeY >> PendingFirstLODIdx) );
 
 		check( PendingFirstLODIdx+ResourceState.NumRequestedLODs <= MipData.Num() );
 
diff --git a/Engine/Source/Runtime/Engine/Private/Texture2D.cpp b/Engine/Source/Runtime/Engine/Private/Texture2D.cpp
--- a/Engine/Source/Runtime/Engine/Private/Texture2D.cpp
+++ b/Engine/Source/Runtime/Engine/Private/Texture2D.cpp
@@ -139,6 +139,10 @@ extern bool TrackTextureEvent( FStreamingRenderAsset* StreamingTexture, UStreama
 	FTexture2DMipMap
 -----------------------------------------------------------------------------*/
 
+#if !WITH_EDITORONLY_DATA
+static_assert(sizeof(FTexture2DMipMap) <= 80, "FTexture2DMipMap was packed to reduce its memory footprint and fit into an 80 bytes bin of MallocBinned2/3");
+#endif
+
 void FTexture2DMipMap::Serialize(FArchive& Ar, UObject* Owner, int32 MipIdx)
 {
 #if WITH_EDITORONLY_DATA
@@ -147,9 +151,15 @@ void FTexture2DMipMap::Serialize(FArchive& Ar, UObject* Owner, int32 MipIdx)
 	BulkData.Serialize(Ar, Owner, MipIdx, false);
 #endif
 
-	Ar << SizeX;
-	Ar << SizeY;
-	Ar << SizeZ;
+	int32 XSize = this->SizeX;
+	int32 YSize = this->SizeY;
+	int32 ZSize = this->SizeZ;
+	Ar << XSize;
+	Ar << YSize;
+	Ar << ZSize;
+	this->SizeX = XSize;
+	this->SizeY = YSize;
+	this->SizeZ = ZSize;
 
 #if WITH_EDITORONLY_DATA
 	if (!Ar.IsFilterEditorOnly())
@@ -1074,11 +1084,8 @@ UTexture2D* UTexture2D::CreateTransient(int32 InSizeX, int32 InSizeY, EPixelForm
 		// Allocate first mipmap.
 		int32 NumBlocksX = InSizeX / GPixelFormats[InFormat].BlockSizeX;
 		int32 NumBlocksY = InSizeY / GPixelFormats[InFormat].BlockSizeY;
-		FTexture2DMipMap* Mip = new FTexture2DMipMap();
+		FTexture2DMipMap* Mip = new FTexture2DMipMap(InSizeX, InSizeY, 1);
 		NewTexture->GetPlatformData()->Mips.Add(Mip);
-		Mip->SizeX = InSizeX;
-		Mip->SizeY = InSizeY;
-		Mip->SizeZ = 1;
 		Mip->BulkData.Lock(LOCK_READ_WRITE);
 		Mip->BulkData.Realloc((int64)NumBlocksX * NumBlocksY * GPixelFormats[InFormat].BlockBytes);
 		Mip->BulkData.Unlock();
diff --git a/Engine/Source/Runtime/Engine/Private/TextureCube.cpp b/Engine/Source/Runtime/Engine/Private/TextureCube.cpp
--- a/Engine/Source/Runtime/Engine/Private/TextureCube.cpp
+++ b/Engine/Source/Runtime/Engine/Private/TextureCube.cpp
@@ -48,11 +48,8 @@ UTextureCube* UTextureCube::CreateTransient(int32 InSizeX, int32 InSizeY, EPixel
 		// Allocate first mipmap.
 		int32 NumBlocksX = InSizeX / GPixelFormats[InFormat].BlockSizeX;
 		int32 NumBlocksY = InSizeY / GPixelFormats[InFormat].BlockSizeY;
-		FTexture2DMipMap* Mip = new FTexture2DMipMap();
+		FTexture2DMipMap* Mip = new FTexture2DMipMap(InSizeX, InSizeY, 1);
 		NewTexture->GetPlatformData()->Mips.Add(Mip);
-		Mip->SizeX = InSizeX;
-		Mip->SizeY = InSizeY;
-		Mip->SizeZ = 1;
 		Mip->BulkData.Lock(LOCK_READ_WRITE);
 		Mip->BulkData.Realloc((int64)6 * NumBlocksX * NumBlocksY * GPixelFormats[InFormat].BlockBytes);
 		Mip->BulkData.Unlock();
diff --git a/Engine/Source/Runtime/Engine/Private/TextureCubeArray.cpp b/Engine/Source/Runtime/Engine/Private/TextureCubeArray.cpp
--- a/Engine/Source/Runtime/Engine/Private/TextureCubeArray.cpp
+++ b/Engine/Source/Runtime/Engine/Private/TextureCubeArray.cpp
@@ -329,11 +329,8 @@ UTextureCubeArray* UTextureCubeArray::CreateTransient(int32 InSizeX, int32 InSiz
 		// Allocate first mipmap.
 		int32 NumBlocksX = InSizeX / GPixelFormats[InFormat].BlockSizeX;
 		int32 NumBlocksY = InSizeY / GPixelFormats[InFormat].BlockSizeY;
-		FTexture2DMipMap* Mip = new FTexture2DMipMap();
+		FTexture2DMipMap* Mip = new FTexture2DMipMap(InSizeX, InSizeY, InArraySize * 6);
 		NewTexture->GetPlatformData()->Mips.Add(Mip);
-		Mip->SizeX = InSizeX;
-		Mip->SizeY = InSizeY;
-		Mip->SizeZ = InArraySize * 6;
 		Mip->BulkData.Lock(LOCK_READ_WRITE);
 		Mip->BulkData.Realloc((int64)GPixelFormats[InFormat].BlockBytes * NumBlocksX * NumBlocksY * 6 * InArraySize);
 		Mip->BulkData.Unlock();
diff --git a/Engine/Source/Runtime/Engine/Private/TextureDerivedData.cpp b/Engine/Source/Runtime/Engine/Private/TextureDerivedData.cpp
--- a/Engine/Source/Runtime/Engine/Private/TextureDerivedData.cpp
+++ b/Engine/Source/Runtime/Engine/Private/TextureDerivedData.cpp
@@ -2241,11 +2241,11 @@ int64 FTexturePlatformData::GetPayloadSize(int32 MipBias) const
 	{
 		for (int32 MipIndex = MipBias; MipIndex < Mips.Num(); MipIndex++)
 		{
-			int32 BlockSizeX = FMath::DivideAndRoundUp(Mips[MipIndex].SizeX, GPixelFormats[PixelFormat].BlockSizeX);
-			int32 BlockSizeY = FMath::DivideAndRoundUp(Mips[MipIndex].SizeY, GPixelFormats[PixelFormat].BlockSizeY);
+			int32 BlockSizeX = FMath::DivideAndRoundUp((int32)Mips[MipIndex].SizeX, GPixelFormats[PixelFormat].BlockSizeX);
+			int32 BlockSizeY = FMath::DivideAndRoundUp((int32)Mips[MipIndex].SizeY, GPixelFormats[PixelFormat].BlockSizeY);
 			// for TextureCube and TextureCubeArray all the mipmaps contain the same number of slices, which is encoded in the PackedData member
 			// at the same time we can not just use SizeZ of a TextureCube mipmap, because for compatibility reasons it is always set to 1 and not 6 (which is the actual number of slices)
-			int32 BlockSizeZ = FMath::DivideAndRoundUp(FMath::Max(IsCubemap() ? GetNumSlices() : Mips[MipIndex].SizeZ, 1), GPixelFormats[PixelFormat].BlockSizeZ);
+			int32 BlockSizeZ = FMath::DivideAndRoundUp(FMath::Max(IsCubemap() ? GetNumSlices() : (int32)Mips[MipIndex].SizeZ, 1), GPixelFormats[PixelFormat].BlockSizeZ);
 			PayloadSize += (int64)GPixelFormats[PixelFormat].BlockBytes * BlockSizeX * BlockSizeY * BlockSizeZ;
 		}
 	}
@@ -2637,7 +2637,7 @@ static void SerializePlatformData(
 				PlatformData->Mips.Empty(NumMips);
 				for (int32 MipIndex = 0; MipIndex < NumMips; ++MipIndex)
 				{
-					PlatformData->Mips.Add(new FTexture2DMipMap());
+					PlatformData->Mips.Add(new FTexture2DMipMap(0, 0));
 					PlatformData->Mips.Last().BulkData.RemoveBulkData();
 				}
 			}
@@ -2802,7 +2802,7 @@ static void SerializePlatformData(
 		PlatformData->Mips.Empty(NumMips);
 		for (int32 MipIndex = 0; MipIndex < NumMips; ++MipIndex)
 		{
-			PlatformData->Mips.Add(new FTexture2DMipMap());
+			PlatformData->Mips.Add(new FTexture2DMipMap(0, 0));
 		}
 	}
 
diff --git a/Engine/Source/Runtime/Engine/Private/TextureDerivedDataTask.cpp b/Engine/Source/Runtime/Engine/Private/TextureDerivedDataTask.cpp
--- a/Engine/Source/Runtime/Engine/Private/TextureDerivedDataTask.cpp
+++ b/Engine/Source/Runtime/Engine/Private/TextureDerivedDataTask.cpp
@@ -539,11 +539,8 @@ static void DDC1_StoreClassicTextureInDerivedData(
 	for (int32 MipIndex = 0; MipIndex < MipCount; ++MipIndex)
 	{
 		const FCompressedImage2D& CompressedImage = CompressedMips[MipIndex];
-		FTexture2DMipMap* NewMip = new FTexture2DMipMap();
+		FTexture2DMipMap* NewMip = new FTexture2DMipMap(CompressedImage.SizeX, CompressedImage.SizeY, CompressedImage.SizeZ);
 		DerivedData->Mips.Add(NewMip);
-		NewMip->SizeX = CompressedImage.SizeX;
-		NewMip->SizeY = CompressedImage.SizeY;
-		NewMip->SizeZ = CompressedImage.SizeZ;
 		NewMip->FileRegionType = FFileRegion::SelectType(EPixelFormat(CompressedImage.PixelFormat));
 		check(NewMip->SizeZ == 1 || bVolume || bTextureArray); // Only volume & arrays can have SizeZ != 1
 
@@ -2113,14 +2110,11 @@ static bool UnpackPlatformDataFromBuild(FTexturePlatformData& OutPlatformData, U
 	uint64 CurrentMipTailOffset = 0;
 	for (int32 MipIndex = 0; MipIndex < EncodedTextureDescription.NumMips; MipIndex++)
 	{
-		FTexture2DMipMap* NewMip = new FTexture2DMipMap();
+		const FIntVector3 MipDims = EncodedTextureDescription.GetMipDimensions(MipIndex);
+		FTexture2DMipMap* NewMip = new FTexture2DMipMap(MipDims.X, MipDims.Y, MipDims.Z);
 		OutPlatformData.Mips.Add(NewMip);
 
 		NewMip->FileRegionType = FileRegion;
-		FIntVector3 MipDims = EncodedTextureDescription.GetMipDimensions(MipIndex);
-		NewMip->SizeX = MipDims.X;
-		NewMip->SizeY = MipDims.Y;
-		NewMip->SizeZ = MipDims.Z;
 		if (EncodedTextureDescription.bTextureArray)
 		{
 			// FTexture2DMipMap expects SizeZ to be the array count, potentially with cubemap slices.
diff --git a/Engine/Source/Runtime/Engine/Private/VolumeTexture.cpp b/Engine/Source/Runtime/Engine/Private/VolumeTexture.cpp
--- a/Engine/Source/Runtime/Engine/Private/VolumeTexture.cpp
+++ b/Engine/Source/Runtime/Engine/Private/VolumeTexture.cpp
@@ -78,11 +78,8 @@ UVolumeTexture* UVolumeTexture::CreateTransient(int32 InSizeX, int32 InSizeY, in
 		int32 NumBlocksX = InSizeX / GPixelFormats[InFormat].BlockSizeX;
 		int32 NumBlocksY = InSizeY / GPixelFormats[InFormat].BlockSizeY;
 		int32 NumBlocksZ = InSizeZ / GPixelFormats[InFormat].BlockSizeZ;
-		FTexture2DMipMap* Mip = new FTexture2DMipMap();
+		FTexture2DMipMap* Mip = new FTexture2DMipMap(InSizeX, InSizeY, InSizeZ);
 		NewTexture->GetPlatformData()->Mips.Add(Mip);
-		Mip->SizeX = InSizeX;
-		Mip->SizeY = InSizeY;
-		Mip->SizeZ = InSizeZ;
 		Mip->BulkData.Lock(LOCK_READ_WRITE);
 		Mip->BulkData.Realloc((int64)GPixelFormats[InFormat].BlockBytes * NumBlocksX * NumBlocksY * NumBlocksZ);
 		Mip->BulkData.Unlock();
diff --git a/Engine/Source/Runtime/Engine/Public/TextureResource.h b/Engine/Source/Runtime/Engine/Public/TextureResource.h
--- a/Engine/Source/Runtime/Engine/Public/TextureResource.h
+++ b/Engine/Source/Runtime/Engine/Public/TextureResource.h
@@ -43,21 +43,24 @@ struct IPooledRenderTarget;
  */
 struct FTexture2DMipMap
 {
-	/** Width of the mip-map. */
-	int32 SizeX = 0;
-	/** Height of the mip-map. */
-	int32 SizeY = 0;
-	/** Depth of the mip-map. */
-	int32 SizeZ = 0;
-
 	/** Reference to the data for the mip if it can be streamed. */
 	UE::FDerivedData DerivedData;
-
 	/** Stores the data for the mip when it is loaded. */
 	FByteBulkData BulkData;
 
+	/** Width of the mip-map. */
+	uint16 SizeX;
+	/** Height of the mip-map. */
+	uint16 SizeY;
+	/** Depth of the mip-map. */
+	uint16 SizeZ;
+
 	PRAGMA_DISABLE_DEPRECATION_WARNINGS
-	FTexture2DMipMap() = default;
+	FTexture2DMipMap(uint32 InSizeX, uint32 InSizeY, uint32 InSizeZ = 0)
+		: SizeX(InSizeX), SizeY(InSizeY), SizeZ(InSizeZ)
+	{
+		check(InSizeX <= 0xFFFF && InSizeY <= 0xFFFF && InSizeZ <= 0xFFFF);
+	}
 	FTexture2DMipMap(FTexture2DMipMap&&) = default;
 	FTexture2DMipMap(const FTexture2DMipMap&) = default;
 	FTexture2DMipMap& operator=(FTexture2DMipMap&&) = default;
diff --git a/Engine/Source/Runtime/RuntimeAssetCache/Private/RuntimeAssetCacheBuilders.cpp b/Engine/Source/Runtime/RuntimeAssetCache/Private/RuntimeAssetCacheBuilders.cpp
--- a/Engine/Source/Runtime/RuntimeAssetCache/Private/RuntimeAssetCacheBuilders.cpp
+++ b/Engine/Source/Runtime/RuntimeAssetCache/Private/RuntimeAssetCacheBuilders.cpp
@@ -184,7 +184,7 @@ void UExampleTextureCacheBuilder::SerializeAsset(FArchive& Ar)
 			PlatformData->Mips.Empty(NumMips);
 			for (int32 MipIndex = 0; MipIndex < NumMips; ++MipIndex)
 			{
-				PlatformData->Mips.Add(new FTexture2DMipMap());
+				PlatformData->Mips.Add(new FTexture2DMipMap(0, 0));
 			}
 		}
 
@@ -192,8 +192,12 @@ void UExampleTextureCacheBuilder::SerializeAsset(FArchive& Ar)
 		for (int32 MipIndex = FirstMip; MipIndex < LastMip; ++MipIndex)
 		{
 			FTexture2DMipMap& Mip = PlatformData->Mips[MipIndex];
-			Ar << Mip.SizeX;
-			Ar << Mip.SizeY;
+			int32 SizeX = Mip.SizeX;
+			int32 SizeY = Mip.SizeY;
+			Ar << SizeX;
+			Ar << SizeY;
+			Mip.SizeX = SizeX;
+			Mip.SizeY = SizeY;
 
 			int32 BulkDataSizeInBytes = Mip.BulkData.GetBulkDataSize();
 			Ar << BulkDataSizeInBytes;
