Fixed compilation in TextureResource.h

diff --git a/Engine/Source/Runtime/Engine/Public/TextureResource.h b/Engine/Source/Runtime/Engine/Public/TextureResource.h
--- a/Engine/Source/Runtime/Engine/Public/TextureResource.h
+++ b/Engine/Source/Runtime/Engine/Public/TextureResource.h
@@ -57,7 +57,7 @@ struct FTexture2DMipMap
 
 	PRAGMA_DISABLE_DEPRECATION_WARNINGS
 	FTexture2DMipMap(uint32 InSizeX, uint32 InSizeY, uint32 InSizeZ = 0)
-		: SizeX(InSizeX), SizeY(InSizeY), SizeZ(InSizeZ)
+		: SizeX((uint16)InSizeX), SizeY((uint16)InSizeY), SizeZ((uint16)InSizeZ)
 	{
 		check(InSizeX <= 0xFFFF && InSizeY <= 0xFFFF && InSizeZ <= 0xFFFF);
 	}
