Fix crash occuring after deleting a DataLayerAsset linked to a DataLayerInstance

diff --git a/Engine/Source/Runtime/Engine/Private/WorldPartition/DataLayer/WorldDataLayersActorDesc.cpp b/Engine/Source/Runtime/Engine/Private/WorldPartition/DataLayer/WorldDataLayersActorDesc.cpp
--- a/Engine/Source/Runtime/Engine/Private/WorldPartition/DataLayer/WorldDataLayersActorDesc.cpp
+++ b/Engine/Source/Runtime/Engine/Private/WorldPartition/DataLayer/WorldDataLayersActorDesc.cpp
@@ -153,7 +153,20 @@ FString FDataLayerInstanceDesc::GetShortName() const
 
 bool FDataLayerInstanceDesc::SupportsActorFilters() const
 {
-	return bIsPrivate ? bPrivateDataLayerSupportsActorFilter : GetAsset()->SupportsActorFilters();
+	if (bIsPrivate)
+	{
+		return bPrivateDataLayerSupportsActorFilter;
+	}
+
+	if (IsUsingAsset())
+	{
+		if (UDataLayerAsset* DataLayerAsset = GetAsset())
+		{
+			return DataLayerAsset->SupportsActorFilters();
+		}
+	}
+	
+	return false;
 }
 
 FDataLayerInstanceDesc::FDataLayerInstanceDesc()
