Allow arithmatic operations between a scalar and vector types. This was already supported by all code paths, but for some reason this code still assumes that only the float generic type can be used.

diff --git a/Engine/Source/Runtime/Engine/Private/Materials/HLSLMaterialTranslator.cpp b/Engine/Source/Runtime/Engine/Private/Materials/HLSLMaterialTranslator.cpp
--- a/Engine/Source/Runtime/Engine/Private/Materials/HLSLMaterialTranslator.cpp
+++ b/Engine/Source/Runtime/Engine/Private/Materials/HLSLMaterialTranslator.cpp
@@ -3892,17 +3892,13 @@ EMaterialValueType FHLSLMaterialTranslator::GetArithmeticResultType(EMaterialVal
 			return LWCTypeA;
 		}
 	}
-	else if(TypeA & TypeB)
+	else if (TypeA == MCT_Float || TypeA == MCT_Float1)
 	{
-		if(TypeA == MCT_Float)
-		{
-			return TypeB;
-		}
-		else
-		{
-			check(TypeB == MCT_Float);
-			return TypeA;
-		}
+		return TypeB;
+	}
+	else if(TypeB == MCT_Float || TypeB == MCT_Float1)
+	{
+		return TypeA;
 	}
 
 	Errorf(TEXT("Arithmetic between types %s and %s are undefined"), DescribeType(TypeA), DescribeType(TypeB));
@@ -7956,7 +7952,7 @@ int32 FHLSLMaterialTranslator::VertexInterpolator(uint32 InterpolatorIndex)
 	}
 }
 
-FLinearColor FHLSLMaterialTranslator::CoerceConstantType(FLinearColor SourceValue, EMaterialValueType SourceType, EMaterialValueType DestType)
+bool FHLSLMaterialTranslator::CoerceConstantType(FLinearColor SourceValue, EMaterialValueType SourceType, EMaterialValueType DestType, FLinearColor& OutResult)
 {
 	EMaterialCastFlags CastFlags = EMaterialCastFlags::ReplicateScalar;
 	if (DestType == MCT_Float || DestType == MCT_Float1 || DestType == MCT_LWCScalar)
@@ -7964,10 +7960,10 @@ FLinearColor FHLSLMaterialTranslator::CoerceConstantType(FLinearColor SourceValu
 		// CoerceValue allows truncating to scalar types only
 		CastFlags |= EMaterialCastFlags::AllowTruncate;
 	}
-	return CastConstantType(SourceValue, SourceType, DestType, CastFlags);
+	return CastConstantType(SourceValue, SourceType, DestType, CastFlags, OutResult);
 }
 
-FLinearColor FHLSLMaterialTranslator::CastConstantType(FLinearColor SourceValue, EMaterialValueType SourceType, EMaterialValueType DestType, EMaterialCastFlags Flags)
+bool FHLSLMaterialTranslator::CastConstantType(FLinearColor SourceValue, EMaterialValueType SourceType, EMaterialValueType DestType, EMaterialCastFlags Flags, FLinearColor& OutResult)
 {
 	const bool bAllowTruncate = EnumHasAnyFlags(Flags, EMaterialCastFlags::AllowTruncate);
 	const bool bAllowAppendZeroes = EnumHasAnyFlags(Flags, EMaterialCastFlags::AllowAppendZeroes);
@@ -7977,9 +7973,12 @@ FLinearColor FHLSLMaterialTranslator::CastConstantType(FLinearColor SourceValue,
 	SourceType = MakeNonLWCType(SourceType);
 	DestType = MakeNonLWCType(DestType);
 
+	// Assumed default value
+	OutResult = SourceValue;
+
 	if (SourceType == DestType)
 	{
-		return SourceValue;
+		return true;
 	}
 	else if (IsFloatNumericType(SourceType) && IsFloatNumericType(DestType))
 	{
@@ -7992,40 +7991,39 @@ FLinearColor FHLSLMaterialTranslator::CastConstantType(FLinearColor SourceValue,
 		if (!bReplicateScalar && !bAllowAppendZeroes && NumDestComponents > NumSourceComponents)
 		{
 			Errorf(TEXT("Cannot cast from smaller type %s to larger type %s."), DescribeType(SourceType), DescribeType(DestType));
-			return SourceValue;
+			return false;
 		}
 		if (!bReplicateScalar && !bAllowTruncate && NumDestComponents < NumSourceComponents)
 		{
 			Errorf(TEXT("Cannot cast from larger type %s to smaller type %s."), DescribeType(SourceType), DescribeType(DestType));
-			return SourceValue;
+			return false;
 		}
 
 		FString Result;
 		uint32 NumComponents = 0u;
 		if (NumSourceComponents == NumDestComponents)
 		{
-			return SourceValue;
+			return true;
 		}
 		else if (bReplicateScalar)
 		{
-			return GetTypeMaskedValue(DestType, FLinearColor(SourceValue.R, SourceValue.R, SourceValue.R, SourceValue.R), nullptr);
+			OutResult = GetTypeMaskedValue(DestType, FLinearColor(SourceValue.R, SourceValue.R, SourceValue.R, SourceValue.R), nullptr);
+			return true;
 		}
 		else
 		{
 			NumComponents = FMath::Min(NumSourceComponents, NumDestComponents);
 			if (NumComponents != NumDestComponents)
 			{
-				return GetTypeMaskedValue(DestType, SourceValue, nullptr);
+				OutResult = GetTypeMaskedValue(DestType, SourceValue, nullptr);
 			}
-			return SourceValue;
+			return true;
 		}
 
 	}
-	else
-	{
-		Errorf(TEXT("Cannot cast between non-numeric types %s to %s."), DescribeType(SourceType), DescribeType(DestType));
-		return SourceValue;
-	}
+
+	Errorf(TEXT("Cannot cast between non-numeric types %s to %s."), DescribeType(SourceType), DescribeType(DestType));
+	return false;
 }
 
 bool FHLSLMaterialTranslator::GetConstParameterValue(FMaterialUniformExpression* Expression, FLinearColor& OutConstValue)
@@ -8275,6 +8273,15 @@ int32 FHLSLMaterialTranslator::Mul(int32 A,int32 B)
 		return INDEX_NONE;
 	}
 
+	const EMaterialValueType ResultType = GetArithmeticResultType(A, B);
+
+	if(ResultType == MCT_Unknown)
+	{
+		return INDEX_NONE;
+	}
+
+	int32 ConstOneReturnValue = INDEX_NONE;
+
 	// If either the left or the right is const zero, assume constant zero
 	if (IsExpressionConstantValue(A, 0.0f) || IsExpressionConstantValue(B, 0.0f))
 	{
@@ -8282,16 +8289,29 @@ int32 FHLSLMaterialTranslator::Mul(int32 A,int32 B)
 	}
 	else if(IsExpressionConstantValue(A, 1.0f))
 	{
-		return B;
+		
+		ConstOneReturnValue = B;
 	}
 	else if (IsExpressionConstantValue(B, 1.0f))
 	{
-		return A;
+		ConstOneReturnValue = A;
 	}
 
+
+	if(ConstOneReturnValue != INDEX_NONE)
+	{
+		int32 Return = ConstOneReturnValue;
+		// This should only be possible with scalar x vector arithmatic, so only attempt to append scalar onto the vector result
+		while (GetNumComponents(GetParameterType(Return)) < GetNumComponents(ResultType))
+		{
+			Return = AppendVector(Return, ConstOneReturnValue);
+		}
+		return Return;
+	}
+
+
 	FMaterialUniformExpression* ExpressionA = GetParameterUniformExpression(A);
 	FMaterialUniformExpression* ExpressionB = GetParameterUniformExpression(B);
-	const EMaterialValueType ResultType = GetArithmeticResultType(A, B);
 	if (ExpressionA && ExpressionB)
 	{
 		if (ExpressionA->IsConstant() && ExpressionB->IsConstant())
@@ -8326,21 +8346,29 @@ int32 FHLSLMaterialTranslator::Div(int32 A, int32 B)
 		return INDEX_NONE;
 	}
 
+	const EMaterialValueType ResultType = GetArithmeticResultType(A, B);
+
 	if (IsExpressionConstantValue(A, 0.0f))
 	{
 		return ConstArithmeticResultValue(A, B, 0.0);
 	}
 	else if(IsExpressionConstantValue(B, 1.0f))
 	{
-		return A;
+		int32 Return = A;
+		// This should only be possible with scalar x vector arithmatic, so only attempt to append scalar onto the vector result
+		while (GetNumComponents(GetParameterType(Return)) < GetNumComponents(ResultType))
+		{
+			Return = AppendVector(Return, A);
+		}
+		return Return;
 	}
 
 	FMaterialUniformExpression* ExpressionA = GetParameterUniformExpression(A);
 	FMaterialUniformExpression* ExpressionB = GetParameterUniformExpression(B);
-	const EMaterialValueType ResultType = GetArithmeticResultType(A, B);
 	if (ExpressionA && ExpressionB)
 	{
-		if (ExpressionA->IsConstant() && ExpressionB->IsConstant())
+		// Const division, unless it would result in division by zero, the hlsl compiler does not like inf outputs
+		if (ExpressionA->IsConstant() && ExpressionB->IsConstant() && !IsExpressionConstantValue(B, 0.0f))
 		{
 			FLinearColor ValueA, ValueB;
 			GetConstParameterValue(ExpressionA, ValueA);
@@ -8399,11 +8427,17 @@ int32 FHLSLMaterialTranslator::Dot(int32 A,int32 B)
 			// Promote scalar (or truncate the bigger type)
 			if (TypeA == MCT_Float || (TypeB != MCT_Float && GetNumComponents(TypeA) > GetNumComponents(TypeB)))
 			{
-				ValueA = CoerceConstantType(ValueA, TypeA, TypeB);
+				if(!CoerceConstantType(ValueA, TypeA, TypeB, ValueA))
+				{
+					return INDEX_NONE;
+				}
 			}
 			else
 			{
-				ValueB = CoerceConstantType(ValueB, TypeB, TypeA);
+				if(!CoerceConstantType(ValueB, TypeB, TypeA, ValueB))
+				{
+					return INDEX_NONE;
+				}
 			}
 
 			FVector4 VectorA(ValueA), VectorB(ValueB);
@@ -8487,8 +8521,13 @@ int32 FHLSLMaterialTranslator::Cross(int32 A,int32 B)
 			GetConstMaskedParameterValue(TypeA, ExpressionA, ValueA);
 			GetConstMaskedParameterValue(TypeB, ExpressionB, ValueB);
 
-			FVector VectorA(CoerceConstantType(ValueA, TypeA, MCT_Float3));
-			FVector VectorB(CoerceConstantType(ValueB, TypeB, MCT_Float3));
+			if(!CoerceConstantType(ValueA, TypeA, MCT_Float3, ValueA) || !CoerceConstantType(ValueB, TypeB, MCT_Float3, ValueB))
+			{
+				return INDEX_NONE;
+			}
+
+			FVector VectorA(ValueA);
+			FVector VectorB(ValueB);
 
 			FLinearColor Result(VectorA.Cross(VectorB));
 			return ConstResultValue(MCT_Float3, Result);
@@ -8521,23 +8560,31 @@ int32 FHLSLMaterialTranslator::Power(int32 Base,int32 Exponent)
 	// be assumed
 	if (ExponentExpression && ExponentExpression->IsConstant())
 	{
+		// This will validate that the parameter types are valid, even though we already
+		// know the result type must be the base type
+		if (GetArithmeticResultType(Base, Exponent) == MCT_Unknown)
+		{
+			return INDEX_NONE;
+		}
+
+		int32 ConstantResult = INDEX_NONE;
 		if (IsExpressionConstantValue(Exponent, 1.0f))
 		{
-			return Base;
+			ConstantResult = Base;
 		}
 		else if (IsExpressionConstantValue(Exponent, 0.0f))
 		{
-			return ConstArithmeticResultValue(Base, Exponent, 1.0f);
+			ConstantResult = Constant(1.0f);
 		}
 		else if (IsExpressionConstantValue(Base, 0.0f))
 		{
-			return ConstArithmeticResultValue(Base, Exponent, 0.0f);
+			ConstantResult = Constant(0.0f);
 		}
 		else if(BaseExpression && BaseExpression->IsConstant())
 		{
-			EMaterialValueType TypeA = GetParameterType(Base);
-			EMaterialValueType TypeB = GetParameterType(Exponent);
-			if (TypeA == TypeB)
+			EMaterialValueType BaseType = GetParameterType(Base);
+			EMaterialValueType ExponentType = GetParameterType(Exponent);
+			if (BaseType == ExponentType)
 			{
 				auto PositiveClampedPow = [](float BaseComponent, float ExponentComponent)
 				{
@@ -8553,7 +8600,15 @@ int32 FHLSLMaterialTranslator::Power(int32 Base,int32 Exponent)
 				//See note below on the issues with the two paths below for non const results.
 				// Currently mimicking that functionallity, but this shoul get changed when that does
 				FLinearColor Result;
-				FLinearColor ExponentConverted = IsAnalyticDerivEnabled() ? ExponentValue : CoerceConstantType(ExponentValue, MCT_Float, TypeA);
+				FLinearColor ExponentConverted = ExponentValue;
+
+				if(!IsAnalyticDerivEnabled())
+				{
+					if (!CoerceConstantType(ExponentValue, MCT_Float, BaseType, ExponentConverted))
+					{
+						return INDEX_NONE;
+					}
+				} 
 
 				Result = FLinearColor(
 					PositiveClampedPow(BaseValue.R, ExponentConverted.R),
@@ -8561,9 +8616,28 @@ int32 FHLSLMaterialTranslator::Power(int32 Base,int32 Exponent)
 					PositiveClampedPow(BaseValue.B, ExponentConverted.B),
 					PositiveClampedPow(BaseValue.A, ExponentConverted.A));
 
-				return ConstResultValue(TypeA, Result);
+				return ConstResultValue(BaseType, Result);
 			}
 		}
+
+		if(ConstantResult != INDEX_NONE)
+		{
+			EMaterialValueType ConstResultType = GetParameterType(Base);
+			int32 Return = ConstantResult;
+			// This should only be possible with scalar x vector arithmatic, so only attempt to append scalar onto the vector result
+			while (GetNumComponents(GetParameterType(Return)) < GetNumComponents(ConstResultType))
+			{
+				Return = AppendVector(Return, ConstantResult);
+			}
+
+			if(GetNumComponents(GetParameterType(Return)) < GetNumComponents(ConstResultType))
+			{
+				int32 NumChannels = GetNumComponents(ConstResultType);
+				Return = ComponentMask(Return, true, NumChannels > 1, NumChannels > 2, NumChannels > 3);
+			}
+
+			return Return;
+		}
 	}
 
 
@@ -8704,7 +8778,6 @@ int32 FHLSLMaterialTranslator::Length(int32 X)
 	{
 		FVector4 Vector(ValueX);
 		float FloatResult = Vector.Size();
-		EMaterialValueType ResultType = GetParameterType(X);
 		FLinearColor Result(
 			FloatResult,
 			FloatResult,
@@ -9140,7 +9213,10 @@ int32 FHLSLMaterialTranslator::Min(int32 A,int32 B)
 			EMaterialValueType TypeA = GetParameterType(A);
 			EMaterialValueType TypeB = GetParameterType(B);
 			FLinearColor Result;
-			ValueB = CoerceConstantType(ValueB, TypeB, TypeA);
+			if (!CoerceConstantType(ValueB, TypeB, TypeA, ValueB))
+			{
+				return INDEX_NONE;
+			}
 
 			Result.R = FMath::Min(ValueA.R, ValueB.R);
 			Result.G = FMath::Min(ValueA.G, ValueB.G);
@@ -9183,7 +9259,10 @@ int32 FHLSLMaterialTranslator::Max(int32 A,int32 B)
 			EMaterialValueType TypeA = GetParameterType(A);
 			EMaterialValueType TypeB = GetParameterType(B);
 			FLinearColor Result;
-			ValueB = CoerceConstantType(ValueB, TypeB, TypeA);
+			if (!CoerceConstantType(ValueB, TypeB, TypeA, ValueB))
+			{
+				return INDEX_NONE;
+			}
 
 			Result.R = FMath::Max(ValueA.R, ValueB.R);
 			Result.G = FMath::Max(ValueA.G, ValueB.G);
diff --git a/Engine/Source/Runtime/Engine/Private/Materials/HLSLMaterialTranslator.h b/Engine/Source/Runtime/Engine/Private/Materials/HLSLMaterialTranslator.h
--- a/Engine/Source/Runtime/Engine/Private/Materials/HLSLMaterialTranslator.h
+++ b/Engine/Source/Runtime/Engine/Private/Materials/HLSLMaterialTranslator.h
@@ -736,8 +736,9 @@ protected:
 	bool GetConstParameterValue(FMaterialUniformExpression* Expression, FLinearColor& OutConstValue);
 
 	// Properly casts a const type to another const type
-	FLinearColor CoerceConstantType(FLinearColor SourceValue, EMaterialValueType SourceType, EMaterialValueType DestinationType);
-	FLinearColor CastConstantType(FLinearColor SourceValue, EMaterialValueType SourceType, EMaterialValueType DestinationType, EMaterialCastFlags Flags);
+	// Returns false if the case fails
+	bool CoerceConstantType(FLinearColor SourceValue, EMaterialValueType SourceType, EMaterialValueType DestinationType, FLinearColor& OutResult);
+	bool CastConstantType(FLinearColor SourceValue, EMaterialValueType SourceType, EMaterialValueType DestinationType, EMaterialCastFlags Flags, FLinearColor& OutResult);
 
 	// Returns a constant of the specified typ
 	int32 ConstResultValue(EMaterialValueType Type, FLinearColor ConstantValue);
