Added support for Equirect stereo layer radius

diff --git a/Engine/Plugins/Runtime/OpenXR/Source/OpenXRHMD/Private/OpenXRHMD_Layer.cpp b/Engine/Plugins/Runtime/OpenXR/Source/OpenXRHMD/Private/OpenXRHMD_Layer.cpp
--- a/Engine/Plugins/Runtime/OpenXR/Source/OpenXRHMD/Private/OpenXRHMD_Layer.cpp
+++ b/Engine/Plugins/Runtime/OpenXR/Source/OpenXRHMD/Private/OpenXRHMD_Layer.cpp
@@ -179,12 +179,14 @@ void FOpenXRLayer::CreateOpenXREquirectLayer(bool bIsStereo, bool bNoAlpha, FTra
 	Equirect.space = Space;
 	Equirect.subImage.imageArrayIndex = 0;
 	Equirect.pose = ToXrPose(Desc.Transform * PositionTransform, WorldToMeters);
-	// The radius is momentarily fixed to 0 as the editor does not support its insertion yet.
-	// An equirect layer with a radius of 0 is an infinite sphere.
-	Equirect.radius = 0.0f; 
 
 	const FEquirectLayer& EquirectProps = Desc.GetShape<FEquirectLayer>();
 
+	// An equirect layer with a radius of 0 is an infinite sphere.
+	// As of UE 5.3, equirect layers are supported only by the Oculus OpenXR runtime and 
+	// only with a radius of 0. Other radius values will be ignored.
+	Equirect.radius = FMath::Abs(EquirectProps.Radius / WorldToMeters);
+
 	FXrCompositionLayerUnion LayerUnion;
 	LayerUnion.Equirect = Equirect;
 
diff --git a/Engine/Source/Runtime/Engine/Classes/Components/StereoLayerComponent.h b/Engine/Source/Runtime/Engine/Classes/Components/StereoLayerComponent.h
--- a/Engine/Source/Runtime/Engine/Classes/Components/StereoLayerComponent.h
+++ b/Engine/Source/Runtime/Engine/Classes/Components/StereoLayerComponent.h
@@ -147,6 +147,10 @@ public:
 	/** Right eye's texture coordinate bias after mapping to 2D. */
 	FVector2D RightBias;
 
+	UPROPERTY(EditAnywhere, BlueprintReadOnly, export, Category = "StereoLayer | Equirect Layer Properties")
+	/** Sphere radius. As of UE 5.3, equirect layers are supported only by the Oculus OpenXR runtime and only with a radius of 0 (infinite sphere).*/
+	float Radius;
+
 public:
 
 	FEquirectProps()
@@ -156,15 +160,17 @@ public:
 		, RightScale(FVector2D(1.0f, 1.0f))
 		, LeftBias(FVector2D(0.0f, 0.0f))
 		, RightBias(FVector2D(0.0f, 0.0f))
+		, Radius(0.0f)
 	{}
 
-	FEquirectProps(FBox2D InLeftUVRect, FBox2D InRightUVRect, FVector2D InLeftScale, FVector2D InRightScale, FVector2D InLeftBias, FVector2D InRightBias)
+	FEquirectProps(FBox2D InLeftUVRect, FBox2D InRightUVRect, FVector2D InLeftScale, FVector2D InRightScale, FVector2D InLeftBias, FVector2D InRightBias, float Radius)
 		: LeftUVRect(InLeftUVRect)
 		, RightUVRect(InRightUVRect)
 		, LeftScale(InLeftScale)
 		, RightScale(InRightScale)
 		, LeftBias(InLeftBias)
 		, RightBias(InRightBias)
+		, Radius(Radius)
 	{ }
 
 	/**
@@ -198,6 +204,7 @@ public:
 		, RightScale(FVector2D(1.0f, 1.0f))
 		, LeftBias(FVector2D(0.0f, 0.0f))
 		, RightBias(FVector2D(0.0f, 0.0f))
+		, Radius(0.0)
 	{}
 
 	UPROPERTY(EditAnywhere, BlueprintReadOnly, Category = "Equirect Properties")
@@ -224,17 +231,25 @@ public:
 	/** Right eye's texture coordinate bias after mapping to 2D. */
 	FVector2D RightBias;
 
+	UPROPERTY(EditAnywhere, BlueprintReadOnly, Category = "Equirect Properties")
+	/** Sphere radius. As of UE 5.3, equirect layers are supported only by the Oculus OpenXR runtime and only with a radius of 0 (infinite sphere).*/
+	float Radius;
+
 	/**
 	 * Set Equirect layer properties: UVRect, Scale, and Bias
 	 * @param	LeftScale: Scale for left eye
 	 * @param	LeftBias: Bias for left eye
 	 * @param	RightScale: Scale for right eye
 	 * @param	RightBias: Bias for right eye
+	 * @param	Radius: Sphere radius
 	 */
 	UFUNCTION(BlueprintCallable, Category = "Components|Stereo Layer")
 	void SetEquirectProps(FEquirectProps InScaleBiases);
 
 	virtual void ApplyShape(IStereoLayers::FLayerDesc& LayerDesc) override;
+#if WITH_EDITOR
+	virtual void DrawShapeVisualization(const class FSceneView* View, class FPrimitiveDrawInterface* PDI) override;
+#endif
 };
 
 /** 
@@ -305,14 +320,15 @@ public:
 	FBox2D GetUVRect() const { return UVRect; }
 
 	/**
-	 * Set Equirect layer properties: UVRect, Scale, and Bias
+	 * Set Equirect layer properties: UVRect, Scale, Bias and Radius.
 	 * @param	LeftScale: Scale for left eye
 	 * @param	LeftBias: Bias for left eye
 	 * @param	RightScale: Scale for right eye
 	 * @param	RightBias: Bias for right eye
+	 * @param	Radius: sphere radius
 	 */
 	UFUNCTION(BlueprintCallable, Category = "Components|Stereo Layer", meta = (DeprecatedFunction, DeprecationMessage = "Use UStereoLayerShapeEquirect::SetEquirectProps() instead."))
-	void SetEquirectProps(FEquirectProps InScaleBiases);
+	void SetEquirectProps(FEquirectProps InEquirectProps);
 
 	/** 
 	 * Change the layer's render priority, higher priorities render on top of lower priorities
diff --git a/Engine/Source/Runtime/Engine/Private/Components/StereoLayerComponent.cpp b/Engine/Source/Runtime/Engine/Private/Components/StereoLayerComponent.cpp
--- a/Engine/Source/Runtime/Engine/Private/Components/StereoLayerComponent.cpp
+++ b/Engine/Source/Runtime/Engine/Private/Components/StereoLayerComponent.cpp
@@ -267,6 +267,7 @@ void UStereoLayerShapeEquirect::SetEquirectProps(FEquirectProps InEquirectProps)
 	RightScale = InEquirectProps.RightScale;
 	LeftBias = InEquirectProps.LeftBias;
 	RightBias = InEquirectProps.RightBias;
+	Radius = InEquirectProps.Radius;
 
 	MarkStereoLayerDirty();
 }
@@ -343,7 +344,7 @@ void UStereoLayerShapeCubemap::ApplyShape(IStereoLayers::FLayerDesc& LayerDesc)
 
 void UStereoLayerShapeEquirect::ApplyShape(IStereoLayers::FLayerDesc& LayerDesc)
 {
-	LayerDesc.SetShape<FEquirectLayer>(LeftUVRect, RightUVRect, LeftScale, RightScale, LeftBias, RightBias);
+	LayerDesc.SetShape<FEquirectLayer>(LeftUVRect, RightUVRect, LeftScale, RightScale, LeftBias, RightBias, Radius);
 }
 
 void UStereoLayerShapeQuad::ApplyShape(IStereoLayers::FLayerDesc& LayerDesc)
@@ -396,15 +397,25 @@ void UStereoLayerShapeCylinder::DrawShapeVisualization(const class FSceneView* V
 
 	PDI->DrawLine(RightVertex - HalfHeight, RightVertex + HalfHeight, YellowColor, 0);
 }
+
+void UStereoLayerShapeEquirect::DrawShapeVisualization(const class FSceneView* View, class FPrimitiveDrawInterface* PDI)
+{
+	FLinearColor YellowColor = FColor(231, 239, 0, 255);
+	check(GetOuter()->IsA<UStereoLayerComponent>());
+
+	auto StereoLayerComp = Cast<UStereoLayerComponent>(GetOuter());
+
+	DrawWireSphere(PDI, StereoLayerComp->GetComponentTransform().GetTranslation(), YellowColor, (double) Radius, 32, 0);
+}
 #endif
 
 
 bool FEquirectProps::operator==(const class UStereoLayerShapeEquirect& Other) const
 {
-	return (LeftUVRect == Other.LeftUVRect) && (RightUVRect == Other.RightUVRect) && (LeftScale == Other.LeftScale) && (RightScale == Other.RightScale) && (LeftBias == Other.LeftBias) && (RightBias == Other.RightBias);
+	return (LeftUVRect == Other.LeftUVRect) && (RightUVRect == Other.RightUVRect) && (LeftScale == Other.LeftScale) && (RightScale == Other.RightScale) && (LeftBias == Other.LeftBias) && (RightBias == Other.RightBias) && (Radius == Other.Radius);
 }
 
 bool FEquirectProps::operator==(const FEquirectProps& Other) const
 {
-	return (LeftUVRect == Other.LeftUVRect) && (RightUVRect == Other.RightUVRect) && (LeftScale == Other.LeftScale) && (RightScale == Other.RightScale) && (LeftBias == Other.LeftBias) && (RightBias == Other.RightBias);
+	return (LeftUVRect == Other.LeftUVRect) && (RightUVRect == Other.RightUVRect) && (LeftScale == Other.LeftScale) && (RightScale == Other.RightScale) && (LeftBias == Other.LeftBias) && (RightBias == Other.RightBias) && (Radius == Other.Radius);
 }
diff --git a/Engine/Source/Runtime/Engine/Public/StereoLayerShapes.h b/Engine/Source/Runtime/Engine/Public/StereoLayerShapes.h
--- a/Engine/Source/Runtime/Engine/Public/StereoLayerShapes.h
+++ b/Engine/Source/Runtime/Engine/Public/StereoLayerShapes.h
@@ -60,13 +60,14 @@ class ENGINE_API FEquirectLayer : public IStereoLayerShape
 public:
 
 	FEquirectLayer() {}
-	FEquirectLayer(FBox2D InLeftUVRect, FBox2D InRightUVRect, FVector2D InLeftScale, FVector2D InRightScale, FVector2D InLeftBias, FVector2D InRightBias) :
+	FEquirectLayer(FBox2D InLeftUVRect, FBox2D InRightUVRect, FVector2D InLeftScale, FVector2D InRightScale, FVector2D InLeftBias, FVector2D InRightBias, float InRadius) :
 		LeftUVRect(InLeftUVRect),
 		RightUVRect(InRightUVRect),
 		LeftScale(InLeftScale),
 		RightScale(InRightScale),
 		LeftBias(InLeftBias),
-		RightBias(InRightBias)
+		RightBias(InRightBias),
+		Radius(InRadius)
 	{}
 
 	/** Left source texture UVRect, specifying portion of input texture corresponding to left eye. */
@@ -86,4 +87,7 @@ public:
 
 	/** Right eye's texture coordinate bias after mapping to 2D. */
 	FVector2D RightBias;
+
+	/** Sphere radius. As of UE 5.3, equirect layers are supported only by the Oculus OpenXR runtime and only with a radius of 0 (infinite sphere).*/
+	float Radius;
 };
