Added removed default constructor to FTexture2DMipMap

diff --git a/Engine/Source/Runtime/Engine/Public/TextureResource.h b/Engine/Source/Runtime/Engine/Public/TextureResource.h
--- a/Engine/Source/Runtime/Engine/Public/TextureResource.h
+++ b/Engine/Source/Runtime/Engine/Public/TextureResource.h
@@ -49,13 +49,14 @@ struct FTexture2DMipMap
 	FByteBulkData BulkData;
 
 	/** Width of the mip-map. */
-	uint16 SizeX;
+	uint16 SizeX = 0;
 	/** Height of the mip-map. */
-	uint16 SizeY;
+	uint16 SizeY = 0;
 	/** Depth of the mip-map. */
-	uint16 SizeZ;
+	uint16 SizeZ = 0;
 
 	PRAGMA_DISABLE_DEPRECATION_WARNINGS
+	FTexture2DMipMap() = default;
 	FTexture2DMipMap(uint32 InSizeX, uint32 InSizeY, uint32 InSizeZ = 0)
 		: SizeX((uint16)InSizeX), SizeY((uint16)InSizeY), SizeZ((uint16)InSizeZ)
 	{
