Rearranged fields of FTransitionEvent to avoid excessive padding

diff --git a/Engine/Source/Runtime/Engine/Classes/Animation/AnimStateMachineTypes.h b/Engine/Source/Runtime/Engine/Classes/Animation/AnimStateMachineTypes.h
--- a/Engine/Source/Runtime/Engine/Classes/Animation/AnimStateMachineTypes.h
+++ b/Engine/Source/Runtime/Engine/Classes/Animation/AnimStateMachineTypes.h
@@ -54,15 +54,15 @@ namespace ETransitionLogicType
 struct FTransitionEvent
 {
 	TArray<int32, TInlineAllocator<8>> ConsumedTransitions;
-	FName EventName;
 	double CreationTime;
 	double TimeToLive;
+	FName EventName;
 	ETransitionRequestQueueMode QueueMode;
 	ETransitionRequestOverwriteMode OverwriteMode;
 
 	FTransitionEvent(const FName& InEventName, const double InTimeToLive, const ETransitionRequestQueueMode& InQueueMode, const ETransitionRequestOverwriteMode& InOverwriteMode)
-		: EventName(InEventName)
-		, TimeToLive(InTimeToLive)
+		: TimeToLive(InTimeToLive)
+		, EventName(InEventName)
 		, QueueMode(InQueueMode)
 		, OverwriteMode(InOverwriteMode)
 	{
