Cameras: refactor camera shakes' time management

The previous design didn't turn out as good as intended, as it was trying to offload too much logic into the base classes. This change does the following:

- FCameraShakeState is now just an opt-in struct for shake patterns who want an out-of-the-box time management utility. It's not imposed on everybody through the base class anymore.
- Having it previously at the base class level was making it awkward in places for those sub-classes that didn't need it, and ended up making the code more complicated than needed (for instance, the implementation of IsFinished(), or of the composite shake pattern, is now much shorter).
- The state struct also now uses separate elapsed/blend-in/blend-out times to fix some issues with inifinite shakes.
- USimpleCamerashakePattern uses the state struct.
- BlendingWeight has been removed, as it's not useful anymore.
- Semantics for "active" vs "finished" are clarified ("active" is true between "start" and "teardown")
- Added more tests for infinite shakes, especially for blending out.

Breaking changes:

- Subclasses of USimpleCameraShakePattern need to upcall their virtual methods and apply the BlendingWeight manually.

diff --git a/Engine/Plugins/Cameras/GameplayCameras/Source/GameplayCameras/Private/CompositeCameraShakePattern.cpp b/Engine/Plugins/Cameras/GameplayCameras/Source/GameplayCameras/Private/CompositeCameraShakePattern.cpp
--- a/Engine/Plugins/Cameras/GameplayCameras/Source/GameplayCameras/Private/CompositeCameraShakePattern.cpp
+++ b/Engine/Plugins/Cameras/GameplayCameras/Source/GameplayCameras/Private/CompositeCameraShakePattern.cpp
@@ -1,6 +1,7 @@
 // Copyright Epic Games, Inc. All Rights Reserved.
 
 #include "CompositeCameraShakePattern.h"
+#include "Camera/PlayerCameraManager.h"
 
 #include UE_INLINE_GENERATED_CPP_BY_NAME(CompositeCameraShakePattern)
 
@@ -37,21 +38,10 @@ void UCompositeCameraShakePattern::GetShakePatternInfoImpl(FCameraShakeInfo& Out
 
 void UCompositeCameraShakePattern::StartShakePatternImpl(const FCameraShakeStartParams& Params)
 {
-	// Create states for all our children patterns.
-	ChildStates.Reset();
-	ChildStates.Reserve(ChildPatterns.Num());
-
 	for (UCameraShakePattern* Pattern : ChildPatterns)
 	{
-		FCameraShakeState& ChildState = ChildStates.Emplace_GetRef();
-
 		if (Pattern != nullptr)
 		{
-			// Initialize the new child state.
-			FCameraShakeInfo ChildInfo;
-			Pattern->GetShakePatternInfo(ChildInfo);
-			ChildState.Initialize(ChildInfo);
-
 			// Start the child pattern.
 			Pattern->StartShakePattern(Params);
 		}
@@ -63,43 +53,24 @@ void UCompositeCameraShakePattern::UpdateShakePatternImpl(const FCameraShakeUpda
 	UCameraShakeBase* ShakeInstance = GetShakeInstance();
 	checkf(ShakeInstance, TEXT("Running a shake pattern without an outer shake instance"));
 
-	// Update each of our children states.
 	FCameraShakeUpdateParams ChildParams(Params);
 
-	for (uint32 Index = 0, Num = ChildPatterns.Num(); Index < Num; ++Index)
+	for (UCameraShakePattern* Pattern : ChildPatterns)
 	{
-		if (Index < (uint32)ChildStates.Num())
+		if (Pattern != nullptr && !Pattern->IsFinished())
 		{
-			FCameraShakeState& PatternState = ChildStates[Index];
-			UCameraShakePattern* Pattern = ChildPatterns[Index];
-			if (Pattern != nullptr)
+			// Let the child pattern run on the current result, with its own blending weight.
+			FCameraShakeUpdateResult ChildResult;
+			Pattern->UpdateShakePattern(ChildParams, ChildResult);
+
+			if (!Pattern->IsFinished())
 			{
-				// Update the child state.
-				float ChildBlendingWeight = PatternState.Update(Params.DeltaTime);
-				if (!PatternState.IsActive())
-				{
-					continue;
-				}
-
-				// Let the child pattern run on the current result, with its own blending weight.
-				ChildParams.BlendingWeight = Params.BlendingWeight * ChildBlendingWeight;
-
-				FCameraShakeUpdateResult ChildResult;
-
-				Pattern->UpdateShakePattern(ChildParams, ChildResult);
-
-				if (IsChildPatternFinished(PatternState, Pattern))
-				{
-					// This pattern just ended now. Reset its state and move on to the next.
-					ChildStates[Index] = FCameraShakeState();
-					continue;
-				}
-			
 				// Apply this result and pass it on to the next child pattern.
 				FCameraShakeApplyResultParams ApplyParams;
-				ApplyParams.Scale = ChildParams.GetTotalScale();
+				ApplyParams.Scale = Params.GetTotalScale();
 				ApplyParams.PlaySpace = ShakeInstance->GetPlaySpace();
 				ApplyParams.UserPlaySpaceMatrix = ShakeInstance->GetUserPlaySpaceMatrix();
+				ApplyParams.CameraManager = ShakeInstance->GetCameraManager();
 				// This applies the current pattern's update to the parameters we'll pass
 				// to the next one.
 				UCameraShakeBase::ApplyResult(ApplyParams, ChildResult, ChildParams.POV);
@@ -123,43 +94,21 @@ void UCompositeCameraShakePattern::ScrubShakePatternImpl(const FCameraShakeScrub
 	UCameraShakeBase* ShakeInstance = GetShakeInstance();
 	checkf(ShakeInstance, TEXT("Running a shake pattern without an outer shake instance"));
 
-	// Scrub each of our children states.
 	FCameraShakeScrubParams ChildParams(Params);
 
-	for (uint32 Index = 0, Num = ChildPatterns.Num(); Index < Num; ++Index)
+	for (UCameraShakePattern* Pattern : ChildPatterns)
 	{
-		if (Index < (uint32)ChildStates.Num())
+		if (Pattern != nullptr) // Don't check for IsFinished here, we might scrub anywhere.
 		{
-			FCameraShakeState& PatternState = ChildStates[Index];
-			if (!PatternState.IsActive())
-			{
-				// This pattern was finished for some time already.
-				continue;
-			}
+			// Let the child pattern run on the current result, with its own blending weight.
+			FCameraShakeUpdateResult ChildResult;
+			Pattern->ScrubShakePattern(Params, ChildResult);
 
-			UCameraShakePattern* Pattern = ChildPatterns[Index];
-			if (Pattern != nullptr)
+			if (!Pattern->IsFinished())
 			{
-				// Scrub the child state.
-				float ChildBlendingWeight = PatternState.Scrub(Params.AbsoluteTime);
-
-				// Let the child pattern run on the current result, with its own blending weight.
-				ChildParams.BlendingWeight = Params.BlendingWeight * ChildBlendingWeight;
-
-				FCameraShakeUpdateResult ChildResult;
-
-				Pattern->ScrubShakePattern(ChildParams, ChildResult);
-
-				if (IsChildPatternFinished(PatternState, Pattern))
-				{
-					// This pattern just ended now. Reset its state and move on to the next.
-					ChildStates[Index] = FCameraShakeState();
-					continue;
-				}
-			
 				// Apply this result and pass it on to the next child pattern.
 				FCameraShakeApplyResultParams ApplyParams;
-				ApplyParams.Scale = ChildParams.GetTotalScale();
+				ApplyParams.Scale = Params.GetTotalScale();
 				ApplyParams.PlaySpace = ShakeInstance->GetPlaySpace();
 				ApplyParams.UserPlaySpaceMatrix = ShakeInstance->GetUserPlaySpaceMatrix();
 				// This applies the current pattern's scrubbing to the parameters we'll pass
@@ -177,37 +126,14 @@ void UCompositeCameraShakePattern::ScrubShakePatternImpl(const FCameraShakeScrub
 	OutResult.Flags = ECameraShakeUpdateResultFlags::ApplyAsAbsolute;
 }
 
-bool UCompositeCameraShakePattern::IsChildPatternFinished(const FCameraShakeState& ChildState, const UCameraShakePattern* ChildPattern) const
-{
-	if (ChildState.IsActive())
-	{
-		if (ChildState.HasDuration())
-		{
-			return ChildState.GetElapsedTime() >= ChildState.GetDuration();
-		}
-		else if (ChildPattern)
-		{
-			return ChildPattern->IsFinished();
-		}
-	}
-	return true;
-}
-
 bool UCompositeCameraShakePattern::IsFinishedImpl() const
 {
 	// We're not finished if any of our child patterns is not finished.
-	for (uint32 Index = 0, Num = ChildPatterns.Num(); Index < Num; ++Index)
+	for (UCameraShakePattern* Pattern : ChildPatterns)
 	{
-		const UCameraShakePattern* Pattern = ChildPatterns[Index];
-		if (Index < (uint32)ChildStates.Num())
+		if (Pattern != nullptr && !Pattern->IsFinished())
 		{
-			const FCameraShakeState& PatternState = ChildStates[Index];
-
-			const bool bIsChildFinished = IsChildPatternFinished(PatternState, Pattern);
-			if (!bIsChildFinished)
-			{
-				return false;
-			}
+			return false;
 		}
 	}
 	return true;
@@ -216,22 +142,11 @@ bool UCompositeCameraShakePattern::IsFinishedImpl() const
 void UCompositeCameraShakePattern::StopShakePatternImpl(const FCameraShakeStopParams& Params)
 {
 	// Stop all our children.
-	for (uint32 Index = 0, Num = ChildPatterns.Num(); Index < Num; ++Index)
+	for (UCameraShakePattern* Pattern : ChildPatterns)
 	{
-		UCameraShakePattern* Pattern = ChildPatterns[Index];
-		if (Index < (uint32)ChildStates.Num())
+		if (Pattern != nullptr)
 		{
-			FCameraShakeState& PatternState = ChildStates[Index];
-
-			if (PatternState.IsActive())
-			{
-				PatternState.Stop(Params.bImmediately);
-
-				if (Pattern != nullptr)
-				{
-					Pattern->StopShakePattern(Params);
-				}
-			}
+			Pattern->StopShakePattern(Params);
 		}
 	}
 }
@@ -246,8 +161,5 @@ void UCompositeCameraShakePattern::TeardownShakePatternImpl()
 			Pattern->TeardownShakePattern();
 		}
 	}
-
-	ChildStates.Reset();
 }
 
-
diff --git a/Engine/Plugins/Cameras/GameplayCameras/Source/GameplayCameras/Private/LegacyCameraShake.cpp b/Engine/Plugins/Cameras/GameplayCameras/Source/GameplayCameras/Private/LegacyCameraShake.cpp
--- a/Engine/Plugins/Cameras/GameplayCameras/Source/GameplayCameras/Private/LegacyCameraShake.cpp
+++ b/Engine/Plugins/Cameras/GameplayCameras/Source/GameplayCameras/Private/LegacyCameraShake.cpp
@@ -79,8 +79,6 @@ ULegacyCameraShake::ULegacyCameraShake(const FObjectInitializer& ObjectInitializ
 
 void ULegacyCameraShake::DoStopShake(bool bImmediately)
 {
-	APlayerCameraManager* CameraOwner = GetCameraManager();
-
 	if (bImmediately)
 	{
 		// stop oscillation
@@ -168,7 +166,6 @@ void ULegacyCameraShake::DoStartShake(const FCameraShakeStartParams& Params)
 	}
 
 	// init cameraanim shakes
-	APlayerCameraManager* CameraOwner = GetCameraManager();
 	if (AnimSequence != nullptr)
 	{
 		if (SequenceShakePattern == nullptr)
@@ -184,11 +181,6 @@ void ULegacyCameraShake::DoStartShake(const FCameraShakeStartParams& Params)
 		SequenceShakePattern->BlendOutTime = AnimBlendOutTime;
 		SequenceShakePattern->RandomSegmentDuration = RandomAnimSegmentDuration;
 		SequenceShakePattern->bRandomSegment = bRandomAnimSegment;
-		
-		// Initialize our state tracker for the sequence shake pattern.
-		FCameraShakeInfo SequenceShakeInfo;
-		SequenceShakePattern->GetShakePatternInfo(SequenceShakeInfo);
-		SequenceShakeState.Initialize(SequenceShakeInfo);
 
 		// Start the sequence shake pattern.
 		SequenceShakePattern->StartShakePattern(Params);
@@ -202,7 +194,6 @@ void ULegacyCameraShake::DoStartShake(const FCameraShakeStartParams& Params)
 void ULegacyCameraShake::DoUpdateShake(const FCameraShakeUpdateParams& Params, FCameraShakeUpdateResult& OutResult)
 {
 	const float DeltaTime = Params.DeltaTime;
-	const float BaseShakeScale = Params.GetTotalScale();
 
 	// update oscillation times... only decrease the time remaining if we're not infinite
 	if (OscillatorTimeRemaining > 0.f)
@@ -306,35 +297,24 @@ void ULegacyCameraShake::DoUpdateShake(const FCameraShakeUpdateParams& Params, F
 	}
 
 	// Update the sequence animation if there's one.
-	if (SequenceShakePattern != nullptr)
+	if (SequenceShakePattern != nullptr && !SequenceShakePattern->IsFinished())
 	{
-		const float ChildBlendWeight = SequenceShakeState.Update(Params.DeltaTime);
-		if (SequenceShakeState.IsActive())
-		{
-			FCameraShakeUpdateParams ChildParams(Params);
-			ChildParams.BlendingWeight = Params.BlendingWeight * ChildBlendWeight;
-
-			FCameraShakeUpdateResult ChildResult;
-
-			SequenceShakePattern->UpdateShakePattern(ChildParams, ChildResult);
-
-			// The sequence shake pattern returns a local, additive, unscaled result. So we should be able to
-			// just combine the two results directly.
-			check(ChildResult.Flags == ECameraShakeUpdateResultFlags::Default);
-			ApplyScale(ChildParams.BlendingWeight, ChildResult);
-			OutResult.Location += ChildResult.Location;
-			OutResult.Rotation += ChildResult.Rotation;
-			OutResult.FOV += ChildResult.FOV;
-			// We don't have anything else animating post-process settings so we can stomp them.
-			OutResult.PostProcessSettings = ChildResult.PostProcessSettings;
-			OutResult.PostProcessBlendWeight = ChildResult.PostProcessBlendWeight;
-		}
+		FCameraShakeUpdateResult ChildResult;
+		SequenceShakePattern->UpdateShakePattern(Params, ChildResult);
+
+		// The sequence shake pattern returns a local, additive result. So we should be able to
+		// just combine the two results directly.
+		OutResult.Location += ChildResult.Location;
+		OutResult.Rotation += ChildResult.Rotation;
+		OutResult.FOV += ChildResult.FOV;
+		// We don't have anything else animating post-process settings so we can stomp them.
+		OutResult.PostProcessSettings = ChildResult.PostProcessSettings;
+		OutResult.PostProcessBlendWeight = ChildResult.PostProcessBlendWeight;
 	}
 
-	// Apply the playspace and the scaling so we have an absolute result we can pass to the legacy blueprint API.
+	// Apply the scaling, limits, and playspace so we have an absolute result we can pass to the legacy blueprint API.
 	check(OutResult.Flags == ECameraShakeUpdateResultFlags::Default);
-	const float CurShakeScale = Params.ShakeScale * Params.DynamicScale;
-	ApplyScale(CurShakeScale, OutResult);
+	ApplyScale(Params.GetTotalScale(), OutResult);
 	ApplyLimits(Params.POV, OutResult);
 	ApplyPlaySpace(Params, OutResult);
 	check(EnumHasAnyFlags(OutResult.Flags, ECameraShakeUpdateResultFlags::ApplyAsAbsolute));
@@ -353,7 +333,6 @@ void ULegacyCameraShake::DoUpdateShake(const FCameraShakeUpdateParams& Params, F
 		OutResult.FOV = InOutPOV.FOV;
 	}
 
-
 	UE_LOG(LogLegacyCameraShake, Verbose, TEXT("ULegacyCameraShake::DoUpdateShake %s Finished: %i Duration: %f Remaining: %f"), *GetNameSafe(this), bOscillationFinished, OscillationDuration, OscillatorTimeRemaining);
 }
 
@@ -401,8 +380,7 @@ void ULegacyCameraShake::DoScrubShake(const FCameraShakeScrubParams& Params, FCa
 bool ULegacyCameraShake::DoGetIsFinished() const
 {
 	return ((OscillatorTimeRemaining <= 0.f) &&									// oscillator is finished
-		((SequenceShakePattern == nullptr) ||									// other anim is finished
-			SequenceShakeState.GetElapsedTime() >= SequenceShakeState.GetDuration()) &&
+		((SequenceShakePattern == nullptr) || SequenceShakePattern->IsFinished()) && // other anim is finished
 		ReceiveIsFinished()														// BP thinks it's finished
 		);
 }
diff --git a/Engine/Plugins/Cameras/GameplayCameras/Source/GameplayCameras/Private/PerlinNoiseCameraShakePattern.cpp b/Engine/Plugins/Cameras/GameplayCameras/Source/GameplayCameras/Private/PerlinNoiseCameraShakePattern.cpp
--- a/Engine/Plugins/Cameras/GameplayCameras/Source/GameplayCameras/Private/PerlinNoiseCameraShakePattern.cpp
+++ b/Engine/Plugins/Cameras/GameplayCameras/Source/GameplayCameras/Private/PerlinNoiseCameraShakePattern.cpp
@@ -25,6 +25,8 @@ UPerlinNoiseCameraShakePattern::UPerlinNoiseCameraShakePattern(const FObjectInit
 
 void UPerlinNoiseCameraShakePattern::StartShakePatternImpl(const FCameraShakeStartParams& Params)
 {
+	Super::StartShakePatternImpl(Params);
+
 	if (!Params.bIsRestarting)
 	{
 		// All offsets are random. This is because the core perlin noise implementation
@@ -43,6 +45,9 @@ void UPerlinNoiseCameraShakePattern::StartShakePatternImpl(const FCameraShakeSta
 void UPerlinNoiseCameraShakePattern::UpdateShakePatternImpl(const FCameraShakeUpdateParams& Params, FCameraShakeUpdateResult& OutResult)
 {
 	UpdatePerlinNoise(Params.DeltaTime, OutResult);
+
+	const float BlendWeight = State.Update(Params.DeltaTime);
+	OutResult.ApplyScale(BlendWeight);
 }
 
 void UPerlinNoiseCameraShakePattern::ScrubShakePatternImpl(const FCameraShakeScrubParams& Params, FCameraShakeUpdateResult& OutResult)
@@ -53,6 +58,9 @@ void UPerlinNoiseCameraShakePattern::ScrubShakePatternImpl(const FCameraShakeScr
 	CurrentFOVOffset = InitialFOVOffset;
 
 	UpdatePerlinNoise(Params.AbsoluteTime, OutResult);
+
+	const float BlendWeight = State.Scrub(Params.AbsoluteTime);
+	OutResult.ApplyScale(BlendWeight);
 }
 
 void UPerlinNoiseCameraShakePattern::UpdatePerlinNoise(float DeltaTime, FCameraShakeUpdateResult& OutResult)
diff --git a/Engine/Plugins/Cameras/GameplayCameras/Source/GameplayCameras/Private/SimpleCameraShakePattern.cpp b/Engine/Plugins/Cameras/GameplayCameras/Source/GameplayCameras/Private/SimpleCameraShakePattern.cpp
--- a/Engine/Plugins/Cameras/GameplayCameras/Source/GameplayCameras/Private/SimpleCameraShakePattern.cpp
+++ b/Engine/Plugins/Cameras/GameplayCameras/Source/GameplayCameras/Private/SimpleCameraShakePattern.cpp
@@ -19,4 +19,23 @@ void USimpleCameraShakePattern::GetShakePatternInfoImpl(FCameraShakeInfo& OutInf
 	OutInfo.BlendOut = BlendOutTime;
 }
 
+void USimpleCameraShakePattern::StartShakePatternImpl(const FCameraShakeStartParams& Params)
+{
+	State.Start(this);
+}
+
+bool USimpleCameraShakePattern::IsFinishedImpl() const
+{
+	return !State.IsPlaying();
+}
+
+void USimpleCameraShakePattern::StopShakePatternImpl(const FCameraShakeStopParams& Params)
+{
+	State.Stop(Params.bImmediately);
+}
+
+void USimpleCameraShakePattern::TeardownShakePatternImpl()
+{
+	State = FCameraShakeState();
+}
 
diff --git a/Engine/Plugins/Cameras/GameplayCameras/Source/GameplayCameras/Private/Tests/CameraShakeStateTests.cpp b/Engine/Plugins/Cameras/GameplayCameras/Source/GameplayCameras/Private/Tests/CameraShakeStateTests.cpp
--- a/Engine/Plugins/Cameras/GameplayCameras/Source/GameplayCameras/Private/Tests/CameraShakeStateTests.cpp
+++ b/Engine/Plugins/Cameras/GameplayCameras/Source/GameplayCameras/Private/Tests/CameraShakeStateTests.cpp
@@ -16,7 +16,7 @@ bool FCameraShakeStateUpdate::RunTest(const FString& Parameters)
 	Info.BlendOut = 0.2f;
 
 	FCameraShakeState State;
-	State.Initialize(Info);
+	State.Start(Info);
 	UTEST_EQUAL("Update 1", State.Update(0.1f), 0.5f);
 	UTEST_EQUAL("Update 2", State.Update(0.1f), 1.f);
 	UTEST_EQUAL("Update 3", State.Update(0.6f), 1.f);
@@ -36,9 +36,9 @@ bool FCameraShakeStateRestart::RunTest(const FString& Parameters)
 		Info.Duration = 1.f;
 
 		FCameraShakeState State;
-		State.Initialize(Info);
+		State.Start(Info);
 		State.Update(0.5f);
-		State.Initialize(Info);
+		State.Start(Info);
 		UTEST_EQUAL("ElapsedTime", State.GetElapsedTime(), 0.f);
 	}
 
@@ -50,14 +50,16 @@ bool FCameraShakeStateRestart::RunTest(const FString& Parameters)
 		Info.BlendOut = 0.3f;
 
 		FCameraShakeState State;
-		State.Initialize(Info);
+		State.Start(Info);
 		UTEST_EQUAL("Update", State.Update(0.85f), 0.5f);
 		// We were half-way into the blend-out, so we should be half-way
 		// into the blend-in as we restart. And the duration would be extended
 		// by that lead-in time so that the shake still lasts 1 second overall.
-		State.Initialize(Info);
-		UTEST_EQUAL("ElapsedTime", State.GetElapsedTime(), 0.1f);
-		UTEST_EQUAL("Duration", State.GetDuration(), 1.1f);
+		State.Start(Info);
+		UTEST_TRUE("BlendingIn", State.IsBlendingIn());
+		UTEST_EQUAL("BlendingInTime", State.GetCurrentBlendInTime(), 0.1f);
+		UTEST_EQUAL("ElapsedTime", State.GetElapsedTime(), 0.f);
+		UTEST_EQUAL("Duration", State.GetDuration(), 1.f);
 	}
 
 	return true;
@@ -74,25 +76,25 @@ bool FCameraShakeStateScrub::RunTest(const FString& Parameters)
 	Info.BlendOut = 0.2f;
 
 	FCameraShakeState State;
-	State.Initialize(Info);
+	State.Start(Info);
 
 	UTEST_EQUAL("Scrub 1", State.Scrub(0.1f), 0.5f);
-	UTEST_TRUE("Scrub 1", State.IsActive());
+	UTEST_TRUE("Scrub 1", State.IsPlaying());
 
 	UTEST_EQUAL("Scrub 2", State.Scrub(0.4f), 1.f);
-	UTEST_TRUE("Scrub 2", State.IsActive());
+	UTEST_TRUE("Scrub 2", State.IsPlaying());
 
 	UTEST_EQUAL("Scrub 3", State.Scrub(-1.f), 0.f);
-	UTEST_FALSE("Scrub 3", State.IsActive());
+	UTEST_FALSE("Scrub 3", State.IsPlaying());
 
 	UTEST_EQUAL("Scrub 4", State.Scrub(0.2f), 1.0f);
-	UTEST_TRUE("Scrub 4", State.IsActive());
+	UTEST_TRUE("Scrub 4", State.IsPlaying());
 
 	UTEST_EQUAL("Scrub 5", State.Scrub(2.0f), 0.0f);
-	UTEST_FALSE("Scrub 5", State.IsActive());
+	UTEST_FALSE("Scrub 5", State.IsPlaying());
 
 	UTEST_EQUAL("Scrub 6", State.Scrub(0.9f), 0.5f);
-	UTEST_TRUE("Scrub 6", State.IsActive());
+	UTEST_TRUE("Scrub 6", State.IsPlaying());
 
 	return true;
 }
diff --git a/Engine/Plugins/Cameras/GameplayCameras/Source/GameplayCameras/Private/Tests/CameraShakeTestObjects.cpp b/Engine/Plugins/Cameras/GameplayCameras/Source/GameplayCameras/Private/Tests/CameraShakeTestObjects.cpp
--- a/Engine/Plugins/Cameras/GameplayCameras/Source/GameplayCameras/Private/Tests/CameraShakeTestObjects.cpp
+++ b/Engine/Plugins/Cameras/GameplayCameras/Source/GameplayCameras/Private/Tests/CameraShakeTestObjects.cpp
@@ -13,6 +13,9 @@ void UConstantCameraShakePattern::UpdateShakePatternImpl(const FCameraShakeUpdat
 {
 	OutResult.Location = LocationOffset;
 	OutResult.Rotation = RotationOffset;
+
+	const float BlendWeight = State.Update(Params.DeltaTime);
+	OutResult.ApplyScale(BlendWeight);
 }
 
 
diff --git a/Engine/Plugins/Cameras/GameplayCameras/Source/GameplayCameras/Private/Tests/CameraShakeTests.cpp b/Engine/Plugins/Cameras/GameplayCameras/Source/GameplayCameras/Private/Tests/CameraShakeTests.cpp
--- a/Engine/Plugins/Cameras/GameplayCameras/Source/GameplayCameras/Private/Tests/CameraShakeTests.cpp
+++ b/Engine/Plugins/Cameras/GameplayCameras/Source/GameplayCameras/Private/Tests/CameraShakeTests.cpp
@@ -305,4 +305,76 @@ bool FCompositeCameraShakeRunTest::RunTest(const FString& Parameters)
 	return true;
 }
 
+IMPLEMENT_SIMPLE_AUTOMATION_TEST(FCameraShakeImmediatelyStopInfiniteTest,
+	"System.Engine.Cameras.ImmediatelyStopInfiniteCameraShake",
+	EAutomationTestFlags::EditorContext | EAutomationTestFlags::EngineFilter)
+bool FCameraShakeImmediatelyStopInfiniteTest::RunTest(const FString& Parameters)
+{
+	FMinimalViewInfo ViewInfo;
+	ViewInfo.Location = FVector::ZeroVector;
+	ViewInfo.Rotation = FRotator::ZeroRotator;
+	auto TestShake = UTestCameraShake::CreateWithPattern<UConstantCameraShakePattern>();
+	TestShake.Pattern->Duration = -1.f; // Negative means infinite duration
+	TestShake.Pattern->LocationOffset = { 10, 0, 0 };
+
+	TestShake.Shake->StartShake(nullptr, 1.f, ECameraShakePlaySpace::CameraLocal);
+	UTEST_FALSE("Infinite shake is running", TestShake.Shake->IsFinished());
+	UTEST_TRUE("Infinite shake is active", TestShake.Shake->IsActive());
+
+	ViewInfo.Location = FVector::ZeroVector;
+	TestShake.Shake->UpdateAndApplyCameraShake(99.f, 1.f, ViewInfo);
+	UTEST_FALSE("Infinite shake is still running", TestShake.Shake->IsFinished());
+	UTEST_TRUE("Infinite shake is still active", TestShake.Shake->IsActive());
+	UTEST_EQUAL("Location offset", ViewInfo.Location, FVector(10, 0, 0));
+	
+	ViewInfo.Location = FVector::ZeroVector;
+	TestShake.Shake->UpdateAndApplyCameraShake(99.f, 1.f, ViewInfo);
+	UTEST_FALSE("Infinite shake is still running", TestShake.Shake->IsFinished());
+	UTEST_TRUE("Infinite shake is still active", TestShake.Shake->IsActive());
+	UTEST_EQUAL("Location offset", ViewInfo.Location, FVector(10, 0, 0));
+
+	TestShake.Shake->StopShake(true);
+	UTEST_TRUE("Infinite shake is finished", TestShake.Shake->IsFinished());
+
+	return true;
+}
+
+IMPLEMENT_SIMPLE_AUTOMATION_TEST(FCameraShakeBlendOutInfiniteTest,
+	"System.Engine.Cameras.BlendOutInfiniteCameraShake",
+	EAutomationTestFlags::EditorContext | EAutomationTestFlags::EngineFilter)
+bool FCameraShakeBlendOutInfiniteTest::RunTest(const FString& Parameters)
+{
+	FMinimalViewInfo ViewInfo;
+	ViewInfo.Location = FVector::ZeroVector;
+	ViewInfo.Rotation = FRotator::ZeroRotator;
+	auto TestShake = UTestCameraShake::CreateWithPattern<UConstantCameraShakePattern>();
+	TestShake.Pattern->Duration = -1.f; // Negative means infinite duration
+	TestShake.Pattern->LocationOffset = { 10, 0, 0 };
+	TestShake.Pattern->BlendOutTime = 1.f; // One second blend out
+
+	TestShake.Shake->StartShake(nullptr, 1.f, ECameraShakePlaySpace::CameraLocal);
+	UTEST_FALSE("Infinite shake is running", TestShake.Shake->IsFinished());
+	UTEST_TRUE("Infinite shake is active", TestShake.Shake->IsActive());
+
+	ViewInfo.Location = FVector::ZeroVector;
+	TestShake.Shake->UpdateAndApplyCameraShake(99.f, 1.f, ViewInfo);
+	UTEST_FALSE("Infinite shake is still running", TestShake.Shake->IsFinished());
+	UTEST_TRUE("Infinite shake is still active", TestShake.Shake->IsActive());
+	UTEST_EQUAL("Location offset", ViewInfo.Location, FVector(10, 0, 0));
+	
+	ViewInfo.Location = FVector::ZeroVector;
+	TestShake.Shake->StopShake(false);
+	TestShake.Shake->UpdateAndApplyCameraShake(0.5f, 1.f, ViewInfo);
+	UTEST_FALSE("Infinite shake is still running", TestShake.Shake->IsFinished());
+	UTEST_TRUE("Infinite shake is still active", TestShake.Shake->IsActive());
+	UTEST_EQUAL("Location offset", ViewInfo.Location, FVector(5, 0, 0));
+
+	ViewInfo.Location = FVector::ZeroVector;
+	TestShake.Shake->UpdateAndApplyCameraShake(0.5f, 1.f, ViewInfo);
+	UTEST_TRUE("Infinite shake is finished", TestShake.Shake->IsFinished());
+	UTEST_EQUAL("Location offset", ViewInfo.Location, FVector(0, 0, 0));
+
+	return true;
+}
+
 #undef LOCTEXT_NAMESPACE
diff --git a/Engine/Plugins/Cameras/GameplayCameras/Source/GameplayCameras/Private/WaveOscillatorCameraShakePattern.cpp b/Engine/Plugins/Cameras/GameplayCameras/Source/GameplayCameras/Private/WaveOscillatorCameraShakePattern.cpp
--- a/Engine/Plugins/Cameras/GameplayCameras/Source/GameplayCameras/Private/WaveOscillatorCameraShakePattern.cpp
+++ b/Engine/Plugins/Cameras/GameplayCameras/Source/GameplayCameras/Private/WaveOscillatorCameraShakePattern.cpp
@@ -33,6 +33,8 @@ UWaveOscillatorCameraShakePattern::UWaveOscillatorCameraShakePattern(const FObje
 
 void UWaveOscillatorCameraShakePattern::StartShakePatternImpl(const FCameraShakeStartParams& Params)
 {
+	Super::StartShakePatternImpl(Params);
+
 	if (!Params.bIsRestarting)
 	{
 		X.Initialize(InitialLocationOffset.X);
@@ -56,6 +58,9 @@ void UWaveOscillatorCameraShakePattern::StartShakePatternImpl(const FCameraShake
 void UWaveOscillatorCameraShakePattern::UpdateShakePatternImpl(const FCameraShakeUpdateParams& Params, FCameraShakeUpdateResult& OutResult)
 {
 	UpdateOscillators(Params.DeltaTime, OutResult);
+
+	const float BlendWeight = State.Update(Params.DeltaTime);
+	OutResult.ApplyScale(BlendWeight);
 }
 
 void UWaveOscillatorCameraShakePattern::ScrubShakePatternImpl(const FCameraShakeScrubParams& Params, FCameraShakeUpdateResult& OutResult)
@@ -66,6 +71,9 @@ void UWaveOscillatorCameraShakePattern::ScrubShakePatternImpl(const FCameraShake
 	CurrentFOVOffset = InitialFOVOffset;
 	
 	UpdateOscillators(Params.AbsoluteTime, OutResult);
+
+	const float BlendWeight = State.Scrub(Params.AbsoluteTime);
+	OutResult.ApplyScale(BlendWeight);
 }
 
 void UWaveOscillatorCameraShakePattern::UpdateOscillators(float DeltaTime, FCameraShakeUpdateResult& OutResult)
diff --git a/Engine/Plugins/Cameras/GameplayCameras/Source/GameplayCameras/Public/CompositeCameraShakePattern.h b/Engine/Plugins/Cameras/GameplayCameras/Source/GameplayCameras/Public/CompositeCameraShakePattern.h
--- a/Engine/Plugins/Cameras/GameplayCameras/Source/GameplayCameras/Public/CompositeCameraShakePattern.h
+++ b/Engine/Plugins/Cameras/GameplayCameras/Source/GameplayCameras/Public/CompositeCameraShakePattern.h
@@ -41,13 +41,5 @@ private:
 	virtual bool IsFinishedImpl() const override;
 	virtual void StopShakePatternImpl(const FCameraShakeStopParams& Params) override;
 	virtual void TeardownShakePatternImpl() override;
-
-private:
-
-	bool IsChildPatternFinished(const FCameraShakeState& ChildState, const UCameraShakePattern* ChildPattern) const;
-
-private:
-
-	/** Play state of the children patterns */
-	TArray<FCameraShakeState> ChildStates;
 };
+
diff --git a/Engine/Plugins/Cameras/GameplayCameras/Source/GameplayCameras/Public/LegacyCameraShake.h b/Engine/Plugins/Cameras/GameplayCameras/Source/GameplayCameras/Public/LegacyCameraShake.h
--- a/Engine/Plugins/Cameras/GameplayCameras/Source/GameplayCameras/Public/LegacyCameraShake.h
+++ b/Engine/Plugins/Cameras/GameplayCameras/Source/GameplayCameras/Public/LegacyCameraShake.h
@@ -276,10 +276,7 @@ protected:
 
 	/** Sequence shake pattern for when using a sequence instead of a camera anim */
 	UPROPERTY(Instanced)
-		TObjectPtr<class USequenceCameraShakePattern> SequenceShakePattern;
-
-	/** State tracking for the sequence shake pattern */
-	FCameraShakeState SequenceShakeState;
+	TObjectPtr<class USequenceCameraShakePattern> SequenceShakePattern;
 
 private:
 
diff --git a/Engine/Plugins/Cameras/GameplayCameras/Source/GameplayCameras/Public/SimpleCameraShakePattern.h b/Engine/Plugins/Cameras/GameplayCameras/Source/GameplayCameras/Public/SimpleCameraShakePattern.h
--- a/Engine/Plugins/Cameras/GameplayCameras/Source/GameplayCameras/Public/SimpleCameraShakePattern.h
+++ b/Engine/Plugins/Cameras/GameplayCameras/Source/GameplayCameras/Public/SimpleCameraShakePattern.h
@@ -31,8 +31,18 @@ public:
 	UPROPERTY(EditAnywhere, Category=Timing)
 	float BlendOutTime = 0.2f;
 
-private:
+protected:
 
 	// UCameraShakePattern interface
 	virtual void GetShakePatternInfoImpl(FCameraShakeInfo& OutInfo) const override;
+	virtual void StartShakePatternImpl(const FCameraShakeStartParams& Params) override;
+	virtual bool IsFinishedImpl() const override;
+	virtual void StopShakePatternImpl(const FCameraShakeStopParams& Params) override;
+	virtual void TeardownShakePatternImpl()  override;
+
+protected:
+
+	/** The ongoing state for this shake */
+	FCameraShakeState State;
 };
+
diff --git a/Engine/Plugins/MovieScene/TemplateSequence/Source/TemplateSequence/Private/SequenceCameraShake.cpp b/Engine/Plugins/MovieScene/TemplateSequence/Source/TemplateSequence/Private/SequenceCameraShake.cpp
--- a/Engine/Plugins/MovieScene/TemplateSequence/Source/TemplateSequence/Private/SequenceCameraShake.cpp
+++ b/Engine/Plugins/MovieScene/TemplateSequence/Source/TemplateSequence/Private/SequenceCameraShake.cpp
@@ -5,6 +5,7 @@
 #include "CameraAnimationSequencePlayer.h"
 #include "EntitySystem/MovieSceneEntityManager.h"
 #include "MovieScene.h"
+#include "MovieSceneFwd.h"
 
 #include UE_INLINE_GENERATED_CPP_BY_NAME(SequenceCameraShake)
 
@@ -69,6 +70,9 @@ void USequenceCameraShakePattern::StartShakePatternImpl(const FCameraShakeStartP
 	// Initialize it and start playing.
 	Player->Initialize(Sequence);
 	Player->Play(bRandomSegment, bRandomSegment);
+
+	// Initialize our state.
+	State.Start(this);
 }
 
 void USequenceCameraShakePattern::UpdateShakePatternImpl(const FCameraShakeUpdateParams& Params, FCameraShakeUpdateResult& OutResult)
@@ -78,6 +82,9 @@ void USequenceCameraShakePattern::UpdateShakePatternImpl(const FCameraShakeUpdat
 
 	const FFrameTime NewPosition = CurrentPosition + Params.DeltaTime * PlayRate * DisplayRate;
 	UpdateCamera(NewPosition, Params.POV, OutResult);
+
+	const float BlendWeight = State.Update(Params.DeltaTime);
+	OutResult.ApplyScale(BlendWeight);
 }
 
 void USequenceCameraShakePattern::ScrubShakePatternImpl(const FCameraShakeScrubParams& Params, FCameraShakeUpdateResult& OutResult)
@@ -87,10 +94,18 @@ void USequenceCameraShakePattern::ScrubShakePatternImpl(const FCameraShakeScrubP
 		const FFrameRate InputRate = Player->GetInputRate();
 		const FFrameTime NewPosition = Params.AbsoluteTime * PlayRate * InputRate;
 		UpdateCamera(NewPosition, Params.POV, OutResult);
+		
+		const float BlendWeight = State.Scrub(Params.AbsoluteTime);
+		OutResult.ApplyScale(BlendWeight);
 	}
 	Player->EndScrubbing();
 }
 
+bool USequenceCameraShakePattern::IsFinishedImpl() const
+{
+	return (Player == nullptr || Player->GetPlaybackStatus() == EMovieScenePlayerStatus::Stopped);
+}
+
 void USequenceCameraShakePattern::StopShakePatternImpl(const FCameraShakeStopParams& Params)
 {
 	using namespace UE::MovieScene;
@@ -115,6 +130,9 @@ void USequenceCameraShakePattern::StopShakePatternImpl(const FCameraShakeStopPar
 				Player->Jump(BlendOutStartFrame);
 			}
 		}
+
+		// Stop tracking time.
+		State.Stop(Params.bImmediately);
 	}
 }
 
@@ -125,6 +143,9 @@ void USequenceCameraShakePattern::TeardownShakePatternImpl()
 	// Stop if we had reached the end of the animation and the sequence needs finishing.
 	// If the shake had been stopped explicitly, this basically won't do anything.
 	Player->Stop();
+
+	// Reset our time tracking.
+	State = FCameraShakeState();
 }
 
 void USequenceCameraShakePattern::UpdateCamera(FFrameTime NewPosition, const FMinimalViewInfo& InPOV, FCameraShakeUpdateResult& OutResult)
diff --git a/Engine/Plugins/MovieScene/TemplateSequence/Source/TemplateSequence/Public/SequenceCameraShake.h b/Engine/Plugins/MovieScene/TemplateSequence/Source/TemplateSequence/Public/SequenceCameraShake.h
--- a/Engine/Plugins/MovieScene/TemplateSequence/Source/TemplateSequence/Public/SequenceCameraShake.h
+++ b/Engine/Plugins/MovieScene/TemplateSequence/Source/TemplateSequence/Public/SequenceCameraShake.h
@@ -66,6 +66,7 @@ private:
 	virtual void StartShakePatternImpl(const FCameraShakeStartParams& Params) override;
 	virtual void UpdateShakePatternImpl(const FCameraShakeUpdateParams& Params, FCameraShakeUpdateResult& OutResult) override;
 	virtual void ScrubShakePatternImpl(const FCameraShakeScrubParams& Params, FCameraShakeUpdateResult& OutResult) override;
+	virtual bool IsFinishedImpl() const override;
 	virtual void StopShakePatternImpl(const FCameraShakeStopParams& Params) override;
 	virtual void TeardownShakePatternImpl() override;
 
@@ -80,6 +81,9 @@ private:
 	/** Standin for the camera actor and components */
 	UPROPERTY(Instanced, Transient)
 	TObjectPtr<UCameraAnimationSequenceCameraStandIn> CameraStandIn;
+
+	/** State tracking */
+	FCameraShakeState State;
 };
 
 #if UE_ENABLE_INCLUDE_ORDER_DEPRECATED_IN_5_2
diff --git a/Engine/Source/Runtime/Engine/Classes/Camera/CameraShakeBase.h b/Engine/Source/Runtime/Engine/Classes/Camera/CameraShakeBase.h
--- a/Engine/Source/Runtime/Engine/Classes/Camera/CameraShakeBase.h
+++ b/Engine/Source/Runtime/Engine/Classes/Camera/CameraShakeBase.h
@@ -9,6 +9,7 @@
 #include "UObject/ScriptMacros.h"
 #include "Camera/CameraTypes.h"
 #include "Engine/Scene.h"
+
 #include "CameraShakeBase.generated.h"
 
 class APlayerCameraManager;
@@ -51,18 +52,15 @@ struct ENGINE_API FCameraShakeUpdateParams
 	/** The dynamic scale being passed down from the camera manger for the next update */
 	UPROPERTY(EditAnywhere, BlueprintReadWrite, Category=CameraShake)
 	float DynamicScale = 1.f;
-	/** The auto-computed blend in/out scale, when blending is handled by externally (see UCameraShakeBase::GetShakeInfo) */
-	UPROPERTY(EditAnywhere, BlueprintReadWrite, Category=CameraShake)
-	float BlendingWeight = 1.f;
 
 	/** The current view that this camera shake should modify */
 	UPROPERTY(EditAnywhere, BlueprintReadWrite, Category=CameraShake)
 	FMinimalViewInfo POV;
 
-	/** The total scale to apply to the camera shake during the current update. Equals ShakeScale * DynamicScale * BlendingWeight */
+	/** The total scale to apply to the camera shake during the current update. Equals ShakeScale * DynamicScale */
 	float GetTotalScale() const
 	{
-		return FMath::Max(ShakeScale * DynamicScale * BlendingWeight, 0.f);
+		return FMath::Max(ShakeScale * DynamicScale, 0.f);
 	}
 };
 
@@ -94,18 +92,15 @@ struct ENGINE_API FCameraShakeScrubParams
 	/** The dynamic scale being passed down from the camera manger for the next update */
 	UPROPERTY(EditAnywhere, BlueprintReadWrite, Category=CameraShake)
 	float DynamicScale = 1.f;
-	/** The auto-computed blend in/out scale, when blending is handled by externally (see UCameraShakeBase::GetShakeInfo) */
-	UPROPERTY(EditAnywhere, BlueprintReadWrite, Category=CameraShake)
-	float BlendingWeight = 1.f;
 
 	/** The current view that this camera shake should modify */
 	UPROPERTY(EditAnywhere, BlueprintReadWrite, Category=CameraShake)
 	FMinimalViewInfo POV;
 
-	/** The total scale to apply to the camera shake during the current update. Equals ShakeScale * DynamicScale * BlendingWeight */
+	/** The total scale to apply to the camera shake during the current update. Equals ShakeScale * DynamicScale */
 	float GetTotalScale() const
 	{
-		return FMath::Max(ShakeScale * DynamicScale * BlendingWeight, 0.f);
+		return FMath::Max(ShakeScale * DynamicScale, 0.f);
 	}
 };
 
@@ -157,6 +152,9 @@ struct ENGINE_API FCameraShakeUpdateResult
 
 	/** Flags for how the base class should handle the result */
 	ECameraShakeUpdateResultFlags Flags;
+
+	/** Apply the given scale to the result (only if it is "relative") */
+	void ApplyScale(float InScale);
 };
 
 /**
@@ -269,17 +267,20 @@ struct ENGINE_API FCameraShakeApplyResultParams
  */
 struct ENGINE_API FCameraShakeState
 {
-	FCameraShakeState()
-		: ElapsedTime(0.f)
-		, bIsActive(false)
-		, bHasBlendIn(false)
-		, bHasBlendOut(false)
-	{}
+	/**
+	 * Create a new camera shake state
+	 */
+	FCameraShakeState();
 
 	/**
-	 * Initialize the state with a shake's info.
+	 * Initialize the state with a shake's info and start playing.
+	 */
+	void Start(const FCameraShakeInfo& InShakeInfo);
+	
+	/**
+	 * Initialize the state with a shake's info and start playing.
 	 */
-	void Initialize(const FCameraShakeInfo& InShakeInfo);
+	void Start(const UCameraShakePattern* InShakePattern);
 
 	/**
 	 * Updates the state with a delta time.
@@ -312,21 +313,31 @@ struct ENGINE_API FCameraShakeState
 	/**
 	 * Marks the shake has having been stopped.
 	 *
-	 * If the state isn't managed (i.e. it doesn't have any fixed duration information), this doesn't
-	 * do anything and returns false.
-	 *
-	 * If the state is managed, this puts the current elapsed time at the end of the shake's duration
-	 * (if we need to stop immediately), or at the start of the shake's blend-out, if any (if we don't
-	 * stop immediately).
+	 * This renders the shake inactive (if we need to stop immediately), or starts the shake's blend-out,
+	 * if any (if we don't stop immediately). If no duration or blending information is available (i.e. if
+	 * the shake duration is "Custom"), stopping non-immediately does nothing: the sub-class is expected
+	 * to handle it.
 	 */
-	bool Stop(bool bImmediately);
+	void Stop(bool bImmediately);
 
-	/** Returns whether the shake is active */
-	bool IsActive() const { return bIsActive; }
+	/** Returns whether the shake is playing */
+	bool IsPlaying() const { return bIsPlaying; }
+
+	/** Returns whether the shake is blending in */
+	bool IsBlendingIn() const { return bIsBlendingIn; }
+
+	/** Returns whether the shake is blending out */
+	bool IsBlendingOut() const { return bIsBlendingOut; }
 
 	/** Returns the elapsed time of the shake's current run */
 	float GetElapsedTime() const { return ElapsedTime; }
 
+	/** Returns the current time into the blend in (only valid if IsBlendingIn() returns true) */
+	float GetCurrentBlendInTime() const { return CurrentBlendInTime; }
+
+	/** Returns the current time into the blend in (only valid if IsBlendingOut() returns true) */
+	float GetCurrentBlendOutTime() const { return CurrentBlendOutTime; }
+
 	/** Returns the current shake info */
 	const FCameraShakeInfo& GetShakeInfo() const { return ShakeInfo; }
 
@@ -339,13 +350,35 @@ struct ENGINE_API FCameraShakeState
 	/** Helper method to get GetShakeInfo().Duration.IsInifnite() */
 	bool IsInfinite() const { return ShakeInfo.Duration.IsInfinite(); }
 
+public:
+
+	UE_DEPRECATED(5.3, "Please use Start")
+	void Initialize(const FCameraShakeInfo& InShakeInfo) { Start(InShakeInfo); }
+
+	UE_DEPRECATED(5.3, "Please use IsPlaying")
+	bool IsActive() const { return bIsPlaying; }
+
 private:
 
+	void InitializePlaying();
+
+private:
+
+	// Information about the shake/shake pattern we're managing
 	FCameraShakeInfo ShakeInfo;
 
+	// Running state
 	float ElapsedTime;
-	bool bIsActive : 1;
 
+	float CurrentBlendInTime;
+	float CurrentBlendOutTime;
+
+	bool bIsBlendingIn : 1;
+	bool bIsBlendingOut : 1;
+
+	bool bIsPlaying : 1;
+
+	// Cached values for blending information
 	bool bHasBlendIn : 1;
 	bool bHasBlendOut : 1;
 };
@@ -427,6 +460,9 @@ public:
 	/** Uses the given result parameters to apply the given result to the given input view info */
 	static void ApplyResult(const FCameraShakeApplyResultParams& ApplyParams, const FCameraShakeUpdateResult& InResult, FMinimalViewInfo& InOutPOV);
 
+	/** Applies all the appropriate auto-scaling to the current shake offset (only if the result is "relative") */
+	static void ApplyScale(const FCameraShakeUpdateParams& Params, FCameraShakeUpdateResult& InOutResult);
+
 	/** Applies the given scale to the current shake offset (only if the result is "relative") */
 	static void ApplyScale(float Scale, FCameraShakeUpdateResult& InOutResult);
 
@@ -441,6 +477,7 @@ public:
 	static void ApplyPlaySpace(ECameraShakePlaySpace PlaySpace, FMatrix UserPlaySpaceMatrix, const FMinimalViewInfo& InPOV, FCameraShakeUpdateResult& InOutResult);
 
 public:
+
 	/** 
 	 *  If true to only allow a single instance of this shake class to play at any given time.
 	 *  Subsequent attempts to play this shake will simply restart the timer.
@@ -479,14 +516,7 @@ public:
 	 *
 	 * A camera shake is active between the calls to StartShake and TeardownShake.
 	 */
-	bool IsActive() const { return State.IsActive(); }
-
-	/**
-	 * Returns the elapsed time of the current state.
-	 *
-	 * A camera shake is active between the calls to StartShake and TeardownShake.
-	 */
-	float GetElapsedTime() const { return State.IsActive() ? State.GetElapsedTime() : 0.0f; }
+	bool IsActive() const { return bIsActive; }
 
 	/** Starts this camera shake with the given parameters */
 	void StartShake(APlayerCameraManager* Camera, float Scale, ECameraShakePlaySpace InPlaySpace, FRotator UserPlaySpaceRot = FRotator::ZeroRotator);
@@ -506,6 +536,9 @@ public:
 	/** Tears down this camera shake before destruction or recycling */
 	void TeardownShake();
 
+	UE_DEPRECATED(5.3, "Elapsed time doesn't exist anymore, get the information from the root shake pattern")
+	float GetElapsedTime() const { return 0.f; }
+
 public:
 
 	/** Gets the current camera manager. Will be null if the shake isn't active. */
@@ -520,9 +553,6 @@ public:
 
 protected:
 
-	/** Applies all the appropriate auto-scaling to the current shake offset (only if the result is "relative") */
-	void ApplyScale(const FCameraShakeUpdateParams& Params, FCameraShakeUpdateResult& InOutResult) const;
-
 	/**
 	 * Modifies the current shake offset to be oriented in the current shake's play space (only if the result is "relative")
 	 *
@@ -546,8 +576,8 @@ private:
 	/** Matrix defining a custom play space, used when PlaySpace is UserDefined. Only valid when the shake is active. */
 	FMatrix UserPlaySpaceMatrix;
 
-	/** Root shake pattern state. Only valid when the shake is active. */
-	FCameraShakeState State;
+	/** Whether this shake is active, which is true between StartShake and TeardownShake */
+	bool bIsActive;
 };
 
 /**
diff --git a/Engine/Source/Runtime/Engine/Private/Camera/CameraModifier_CameraShake.cpp b/Engine/Source/Runtime/Engine/Private/Camera/CameraModifier_CameraShake.cpp
--- a/Engine/Source/Runtime/Engine/Private/Camera/CameraModifier_CameraShake.cpp
+++ b/Engine/Source/Runtime/Engine/Private/Camera/CameraModifier_CameraShake.cpp
@@ -502,9 +502,8 @@ void UCameraModifier_CameraShake::DisplayDebug(UCanvas* Canvas, const FDebugDisp
 			const FString DurationString = !ShakeInfo.ShakeInstance->GetCameraShakeDuration().IsInfinite() ? 
 				FString::SanitizeFloat(ShakeInfo.ShakeInstance->GetCameraShakeDuration().Get()) : TEXT("Infinite");
 			Canvas->DrawText(DrawFont,
-					FString::Printf(TEXT("[%d] %s Source:%s Duration: %s Elapsed: %f"), 
-						i, *GetNameSafe(ShakeInfo.ShakeInstance), *GetNameSafe(ShakeInfo.ShakeSource.Get()), 
-						*DurationString, ShakeInfo.ShakeInstance->GetElapsedTime()),
+					FString::Printf(TEXT("[%d] %s Source:%s"), 
+						i, *GetNameSafe(ShakeInfo.ShakeInstance), *GetNameSafe(ShakeInfo.ShakeSource.Get())), 
 					Indentation * YL, (LineNumber++) * YL);
 		}
 	}
diff --git a/Engine/Source/Runtime/Engine/Private/Camera/CameraShakeBase.cpp b/Engine/Source/Runtime/Engine/Private/Camera/CameraShakeBase.cpp
--- a/Engine/Source/Runtime/Engine/Private/Camera/CameraShakeBase.cpp
+++ b/Engine/Source/Runtime/Engine/Private/Camera/CameraShakeBase.cpp
@@ -6,6 +6,8 @@
 #include "Engine/Engine.h"
 #include "IXRTrackingSystem.h"
 #include "Math/RotationMatrix.h"
+#include "Misc/AssertionMacros.h"
+#include "Misc/CoreMiscDefines.h"
 #include "Stats/Stats.h"
 
 #include UE_INLINE_GENERATED_CPP_BY_NAME(CameraShakeBase)
@@ -23,6 +25,7 @@ UCameraShakeBase::UCameraShakeBase(const FObjectInitializer& ObjectInitializer)
 	, bSingleInstance(false)
 	, ShakeScale(1.f)
 	, PlaySpace(ECameraShakePlaySpace::CameraLocal)
+	, bIsActive(false)
 {
 	RootShakePattern = static_cast<UCameraShakePattern*>(ObjectInitializer.CreateDefaultSubobject(
 			this, 
@@ -41,85 +44,134 @@ FCameraShakeDuration UCameraShakeBase::GetCameraShakeDuration() const
 	return TempInfo.Duration;
 }
 
-void FCameraShakeState::Initialize(const FCameraShakeInfo& InShakeInfo)
+FCameraShakeState::FCameraShakeState()
+	: ElapsedTime(0.f)
+	, CurrentBlendInTime(0.f)
+	, CurrentBlendOutTime(0.f)
+	, bIsBlendingIn(false)
+	, bIsBlendingOut(false)
+	, bIsPlaying(false)
+	, bHasBlendIn(false)
+	, bHasBlendOut(false)
 {
-	ShakeInfo = InShakeInfo;
+}
+
+void FCameraShakeState::Start(const FCameraShakeInfo& InShakeInfo)
+{
+	const FCameraShakeState PrevState(*this);
 
 	// Cache a few things about the shake.
+	ShakeInfo = InShakeInfo;
 	bHasBlendIn = ShakeInfo.BlendIn > 0.f;
 	bHasBlendOut = ShakeInfo.BlendOut > 0.f;
 
 	// Initialize our running state.
-	const bool bIsRestarting = bIsActive;
-	if (!bIsRestarting)
-	{
-		// Set the active state.
-		ElapsedTime = 0.f;
-		bIsActive = true;
-	}
-	else
+	InitializePlaying();
+
+	if (PrevState.bIsPlaying)
 	{
 		// Single instance shake is being restarted... let's see if we need to
-		// reverse a blend out into a blend in.
-		if (ShakeInfo.Duration.IsFixed() && bHasBlendIn && bHasBlendOut)
+		// reverse a blend-out into a blend-in.
+		if (bHasBlendIn && PrevState.bHasBlendOut && PrevState.bIsBlendingOut)
 		{
-			const float BlendOutStartTime = ShakeInfo.Duration.Get() - ShakeInfo.BlendOut;
-			if (ElapsedTime > BlendOutStartTime)
-			{
-				// We had started blending out... let's start at an equivalent weight into the blend in.
-				const float BlendOutCurrentTime = ElapsedTime - BlendOutStartTime;
-				ElapsedTime = ShakeInfo.BlendIn * (1.f - BlendOutCurrentTime / ShakeInfo.BlendOut);
-				// Because this means we are shortening the shake (by the amount that we start into the
-				// blend in, instead of starting from zero), we need to lengthen the shake to make it
-				// last the same duration as it's supposed to.
-				ShakeInfo.Duration = FCameraShakeDuration(ShakeInfo.Duration.Get() + ElapsedTime);
-			}
-			else
-			{
-				// We had not started blending out, so we were at 100%. Let's go back to the beginning
-				// but skip the blend in time.
-				ElapsedTime = 0.f;
-				bHasBlendIn = false;
-				ShakeInfo.BlendIn = 0.f;
-			}
+			// We had started blending out... let's start at an equivalent weight into the blend in.
+			CurrentBlendInTime = ShakeInfo.BlendIn * (1.f - PrevState.CurrentBlendOutTime / PrevState.ShakeInfo.BlendOut);
 		}
-		else
+		else if (bHasBlendIn && !PrevState.bHasBlendOut)
 		{
-			// We either don't have blending, or our shake pattern is doing custom stuff.
-			ElapsedTime = 0.f;
+			// We have no blend out, so we were still at 100%. But we don't want to suddenly drop 
+			// to 0% at the beginning of the blend in, so we skip the blend-in.
+			bIsBlendingIn = false;
 		}
 	}
 }
 
+void FCameraShakeState::Start(const UCameraShakePattern* InShakePattern)
+{
+	check(InShakePattern);
+	FCameraShakeInfo Info;
+	InShakePattern->GetShakePatternInfo(Info);
+	Start(Info);
+}
+
+void FCameraShakeState::InitializePlaying()
+{
+	ElapsedTime = 0.f;
+	CurrentBlendInTime = 0.f;
+	CurrentBlendOutTime = 0.f;
+	bIsBlendingIn = bHasBlendIn;
+	bIsBlendingOut = false;
+	bIsPlaying = true;
+}
+
 float FCameraShakeState::Update(float DeltaTime)
 {
-	// If we have a fixed duration for our shake, we can do all the time-keeping stuff ourselves.
+	// If we have duration information for our shake, we can do all the time-keeping stuff ourselves.
 	// This includes figuring out if the shake is finished, and what kind of blend in/out weight
 	// we should apply.
 	float BlendingWeight = 1.f;
+
+	// Advance shake and blending times.
+	ElapsedTime += DeltaTime;
+	if (bIsBlendingIn)
+	{
+		CurrentBlendInTime += DeltaTime;
+	}
+	if (bIsBlendingOut)
+	{
+		CurrentBlendOutTime += DeltaTime;
+	}
+
+	// Advance progress into the shake.
 	if (HasDuration())
 	{
-		// Advance progress into the shake.
 		const float ShakeDuration = ShakeInfo.Duration.Get();
-		ElapsedTime = ElapsedTime + DeltaTime;
 		if (ElapsedTime < 0.f || ElapsedTime >= ShakeDuration)
 		{
-			// The shake has ended, or hasn't started yet.
-			bIsActive = false;
+			// The shake has ended, or hasn't started yet (which can happen if we update backwards)
+			bIsPlaying = false;
 			return 0.f;
 		}
 
-		// Blending in?
-		if (bHasBlendIn && ElapsedTime < ShakeInfo.BlendIn)
+		const float DurationRemaining = (ShakeDuration - ElapsedTime);
+		if (bHasBlendOut && !bIsBlendingOut && DurationRemaining < ShakeInfo.BlendOut)
 		{
-			BlendingWeight *= (ElapsedTime / ShakeInfo.BlendIn);
+			// We started blending out.
+			bIsBlendingOut = true;
+			CurrentBlendOutTime = (ShakeInfo.BlendOut - DurationRemaining);
 		}
+	}
 
-		// Blending out?
-		const float DurationRemaining = (ShakeDuration - ElapsedTime);
-		if (bHasBlendOut && DurationRemaining < ShakeInfo.BlendOut)
+	// Compute blend-in and blend-out weight.
+	if (bHasBlendIn && bIsBlendingIn)
+	{
+		if (CurrentBlendInTime < ShakeInfo.BlendIn)
+		{
+			BlendingWeight *= (CurrentBlendInTime / ShakeInfo.BlendIn);
+		}
+		else
+		{
+			// Finished blending in!
+			bIsBlendingIn = false;
+			CurrentBlendInTime = ShakeInfo.BlendIn;
+		}
+	}
+	if (bHasBlendOut && bIsBlendingOut)
+	{
+		if (CurrentBlendOutTime < ShakeInfo.BlendOut)
+		{
+			BlendingWeight *= (1.f - CurrentBlendOutTime / ShakeInfo.BlendOut);
+		}
+		else
 		{
-			BlendingWeight *= (DurationRemaining / ShakeInfo.BlendOut);
+			// Finished blending out!
+			bIsBlendingOut = false;
+			CurrentBlendOutTime = ShakeInfo.BlendOut;
+			// We also end the shake itself. In most cases we would have hit the similar case
+			// above already, but if we have an infinite shake we have no duration to reach the end
+			// of so we only finish here.
+			bIsPlaying = false;
+			return 0.f;
 		}
 	}
 	return BlendingWeight;
@@ -127,55 +179,49 @@ float FCameraShakeState::Update(float DeltaTime)
 
 float FCameraShakeState::Scrub(float AbsoluteTime)
 {
-	float BlendingWeight = 1.f;
-	if (HasDuration())
-	{
-		// Reset the state to active, at the beginning, and update from there.
-		bIsActive = true;
-		ElapsedTime = 0.f;
-		return Update(AbsoluteTime);
-	}
-	return BlendingWeight;
+	// Reset the state to active, at the beginning, and update from there.
+	InitializePlaying();
+	return Update(AbsoluteTime);
 }
 
-bool FCameraShakeState::Stop(bool bImmediately)
+void FCameraShakeState::Stop(bool bImmediately)
 {
-	if (HasDuration())
+	// For stopping immediately, we don't do anything besides render the shake inactive.
+	if (bImmediately || !bHasBlendOut)
 	{
-		// If we have duration information, we can set our time-keeping accordingly to stop the shake.
-		// For stopping immediately, we just go to the end right away.
-		// For stopping with a "graceful" blend-out:
-		// - If we are already blending out, let's keep doing that and not change anything.
-		// - If we are not, move to the start of the blend out.
-		const float ShakeDuration = ShakeInfo.Duration.Get();
-		if (bImmediately || !bHasBlendOut)
-		{
-			ElapsedTime = ShakeDuration;
-		}
-		else
-		{
-			const float BlendOutStartTime = ShakeDuration - ShakeInfo.BlendOut;
-			if (ElapsedTime < BlendOutStartTime)
-			{
-				ElapsedTime = BlendOutStartTime;
-			}
-		}
-		return true;
+		bIsPlaying = false;
+	}
+	// For stopping with a "graceful" blend-out:
+	// - If we are already blending out, let's keep doing that and not change anything.
+	// - If we are not, let's start blending out.
+	else if (bHasBlendOut && !bIsBlendingOut)
+	{
+		bIsBlendingOut = true;
+		CurrentBlendOutTime = 0.f;
 	}
-	return false;
 }
 
 FCameraShakeUpdateParams FCameraShakeScrubParams::ToUpdateParams() const
 {
-	FCameraShakeUpdateParams UpdateParams;
+	FCameraShakeUpdateParams UpdateParams(POV);
 	UpdateParams.DeltaTime = AbsoluteTime;
 	UpdateParams.ShakeScale = ShakeScale;
 	UpdateParams.DynamicScale = DynamicScale;
-	UpdateParams.BlendingWeight = BlendingWeight;
-	UpdateParams.POV = POV;
 	return UpdateParams;
 }
 
+void FCameraShakeUpdateResult::ApplyScale(float InScale)
+{
+	if (ensureMsgf(!EnumHasAnyFlags(Flags, ECameraShakeUpdateResultFlags::ApplyAsAbsolute),
+			TEXT("Can't scale absolute shake result")))
+	{
+		Location *= InScale;
+		Rotation *= InScale;
+		FOV *= InScale;
+		PostProcessBlendWeight *= InScale;
+	}
+}
+
 void UCameraShakeBase::GetCameraShakeBlendTimes(float& OutBlendIn, float& OutBlendOut) const
 {
 	FCameraShakeInfo TempInfo;
@@ -186,7 +232,7 @@ void UCameraShakeBase::GetCameraShakeBlendTimes(float& OutBlendIn, float& OutBle
 
 void UCameraShakeBase::SetRootShakePattern(UCameraShakePattern* InPattern)
 {
-	if (ensureMsgf(!State.IsActive(), TEXT("Can't change the root shake pattern while the shake is running!")))
+	if (ensureMsgf(!bIsActive, TEXT("Can't change the root shake pattern while the shake is running!")))
 	{
 		RootShakePattern = InPattern;
 	}
@@ -206,7 +252,7 @@ void UCameraShakeBase::StartShake(APlayerCameraManager* Camera, float Scale, ECa
 
 	// Check that we were correctly stopped before we are asked to play again.
 	// Note that single-instance shakes can be restarted while they're running.
-	checkf(!State.IsActive() || bSingleInstance, TEXT("Starting to play a shake that was already playing."));
+	checkf(!bIsActive || bSingleInstance, TEXT("Starting to play a shake that was already playing."));
 
 	// Remember the various settings for this run.
 	// Note that the camera manager can be null, for example in unit tests.
@@ -216,11 +262,8 @@ void UCameraShakeBase::StartShake(APlayerCameraManager* Camera, float Scale, ECa
 	UserPlaySpaceMatrix = (InPlaySpace == ECameraShakePlaySpace::UserDefined) ? 
 		FRotationMatrix(UserPlaySpaceRot) : FRotationMatrix::Identity;
 
-	// Acquire info about the shake we're running, and initialize our transient state.
-	const bool bIsRestarting = State.IsActive();
-	FCameraShakeInfo ActiveInfo;
-	GetShakeInfo(ActiveInfo);
-	State.Initialize(ActiveInfo);
+	const bool bIsRestarting = bIsActive;
+	bIsActive = true;
 
 	// Let the root pattern initialize itself.
 	if (RootShakePattern)
@@ -236,122 +279,93 @@ void UCameraShakeBase::UpdateAndApplyCameraShake(float DeltaTime, float Alpha, F
 	SCOPE_CYCLE_COUNTER(STAT_UpdateShake);
 	SCOPE_CYCLE_UOBJECT(This, this);
 
-	checkf(State.IsActive(), TEXT("Updating a camera shake that wasn't started with a call to StartShake!"));
+	checkf(bIsActive, TEXT("Updating a camera shake that wasn't started with a call to StartShake!"));
 
-	// Update our state, and early out if we have gone over the duration of the shake.
-	float BlendingWeight = State.Update(DeltaTime);
-	if (!State.IsActive())
+	if (RootShakePattern)
 	{
-		return;
-	}
+		// Make the sub-class do the actual work.
+		FCameraShakeUpdateParams Params(InOutPOV);
+		Params.DeltaTime = DeltaTime;
+		Params.ShakeScale = ShakeScale;
+		Params.DynamicScale = Alpha;
 
-	// Make the sub-class do the actual work.
-	FCameraShakeUpdateParams Params(InOutPOV);
-	Params.DeltaTime = DeltaTime;
-	Params.ShakeScale = ShakeScale;
-	Params.DynamicScale = Alpha;
-	Params.BlendingWeight = BlendingWeight;
-	
-	// Result object is initialized with zero values since the default flags make us handle it
-	// as an additive offset.
-	FCameraShakeUpdateResult Result;
+		// Result object is initialized with zero values since the default flags make us handle it
+		// as an additive offset.
+		FCameraShakeUpdateResult Result;
 
-	if (RootShakePattern)
-	{
 		RootShakePattern->UpdateShakePattern(Params, Result);
-	}
 
-	// Apply the result to the given view info.
-	FCameraShakeApplyResultParams ApplyParams;
-	ApplyParams.Scale = Params.GetTotalScale();
-	ApplyParams.PlaySpace = PlaySpace;
-	ApplyParams.UserPlaySpaceMatrix = UserPlaySpaceMatrix;
-	ApplyParams.CameraManager = CameraManager;
-	ApplyResult(ApplyParams, Result, InOutPOV);
+		if (!RootShakePattern->IsFinished())
+		{
+			// Apply the result to the given view info.
+			FCameraShakeApplyResultParams ApplyParams;
+			ApplyParams.Scale = Params.GetTotalScale();
+			ApplyParams.PlaySpace = PlaySpace;
+			ApplyParams.UserPlaySpaceMatrix = UserPlaySpaceMatrix;
+			ApplyParams.CameraManager = CameraManager;
+			ApplyResult(ApplyParams, Result, InOutPOV);
+		}
+	}
 }
 
 void UCameraShakeBase::ScrubAndApplyCameraShake(float AbsoluteTime, float Alpha, FMinimalViewInfo& InOutPOV)
 {
 	// This code is similar to the above UpdateAndApplyCameraShake method, but calls the scrub method
 	// on the state manager and root pattern instead of the update method.
-	
+
 	SCOPE_CYCLE_COUNTER(STAT_UpdateShake);
 	SCOPE_CYCLE_UOBJECT(This, this);
 
-	checkf(State.IsActive(), TEXT("Updating a camera shake that wasn't started with a call to StartShake!"));
+	checkf(bIsActive, TEXT("Updating a camera shake that wasn't started with a call to StartShake!"));
 
-	float BlendingWeight = State.Scrub(AbsoluteTime);
-	if (!State.IsActive())
+	if (RootShakePattern)
 	{
-		return;
-	}
+		// Make the sub-class do the actual work.
+		FCameraShakeScrubParams Params(InOutPOV);
+		Params.AbsoluteTime = AbsoluteTime;
+		Params.ShakeScale = ShakeScale;
+		Params.DynamicScale = Alpha;
 
-	// Make the sub-class do the actual work.
-	FCameraShakeScrubParams Params(InOutPOV);
-	Params.AbsoluteTime = AbsoluteTime;
-	Params.ShakeScale = ShakeScale;
-	Params.DynamicScale = Alpha;
-	Params.BlendingWeight = BlendingWeight;
+		// Result object is initialized with zero values since the default flags make us handle it
+		// as an additive offset.
+		FCameraShakeUpdateResult Result;
 
-	// Result object is initialized with zero values since the default flags make us handle it
-	// as an additive offset.
-	FCameraShakeUpdateResult Result;
-
-	if (RootShakePattern)
-	{
 		RootShakePattern->ScrubShakePattern(Params, Result);
-	}
 
-	// Apply the result to the given view info.
-	FCameraShakeApplyResultParams ApplyParams;
-	ApplyParams.Scale = Params.GetTotalScale();
-	ApplyParams.PlaySpace = PlaySpace;
-	ApplyParams.UserPlaySpaceMatrix = UserPlaySpaceMatrix;
-	ApplyParams.CameraManager = CameraManager;
-	ApplyResult(ApplyParams, Result, InOutPOV);
+		if (!RootShakePattern->IsFinished())
+		{
+			// Apply the result to the given view info.
+			FCameraShakeApplyResultParams ApplyParams;
+			ApplyParams.Scale = Params.GetTotalScale();
+			ApplyParams.PlaySpace = PlaySpace;
+			ApplyParams.UserPlaySpaceMatrix = UserPlaySpaceMatrix;
+			ApplyParams.CameraManager = CameraManager;
+			ApplyResult(ApplyParams, Result, InOutPOV);
+		}
+	}
 }
 
 bool UCameraShakeBase::IsFinished() const
 {
-	if (State.IsActive())
+	if (RootShakePattern)
 	{
-		switch (State.GetShakeInfo().Duration.GetDurationType())
-		{
-			case ECameraShakeDurationType::Fixed:
-				// If we have duration information, we can simply figure out ourselves if
-				// we are finished.
-				return State.GetElapsedTime() >= State.GetDuration();
-
-			case ECameraShakeDurationType::Infinite:
-				return false;
-
-			case ECameraShakeDurationType::Custom:
-				if (RootShakePattern)
-				{
-					// Ask the root pattern whether it's finished.
-					return RootShakePattern->IsFinished();
-				}
-				else
-				{
-					// We have no root pattern, we don't have anything to do.
-					return true;
-				}
-		}
+		// Ask the root pattern whether it's finished.
+		return RootShakePattern->IsFinished();
+	}
+	else
+	{
+		// We have no root pattern, we don't have anything to do.
+		return true;
 	}
-	// We're not active, so we're finished.
-	return true;
 }
 
 void UCameraShakeBase::StopShake(bool bImmediately)
 {
-	if (!ensureMsgf(State.IsActive(), TEXT("Stopping a shake that wasn't active")))
+	if (!ensureMsgf(bIsActive, TEXT("Stopping a shake that wasn't active")))
 	{
 		return;
 	}
 
-	// Make our transient state as stopping or stopped.
-	State.Stop(bImmediately);
-
 	// Let the root pattern do any custom logic.
 	if (RootShakePattern)
 	{
@@ -368,7 +382,7 @@ void UCameraShakeBase::TeardownShake()
 		RootShakePattern->TeardownShakePattern();
 	}
 
-	State = FCameraShakeState();
+	bIsActive = false;
 }
 
 void UCameraShakeBase::ApplyResult(const FCameraShakeApplyResultParams& ApplyParams, const FCameraShakeUpdateResult& InResult, FMinimalViewInfo& InOutPOV)
@@ -414,17 +428,14 @@ void UCameraShakeBase::ApplyResult(const FCameraShakeApplyResultParams& ApplyPar
 	}
 }
 
-void UCameraShakeBase::ApplyScale(const FCameraShakeUpdateParams& Params, FCameraShakeUpdateResult& InOutResult) const
+void UCameraShakeBase::ApplyScale(const FCameraShakeUpdateParams& Params, FCameraShakeUpdateResult& InOutResult)
 {
-	return ApplyScale(Params.GetTotalScale(), InOutResult);
+	InOutResult.ApplyScale(Params.GetTotalScale());
 }
 
 void UCameraShakeBase::ApplyScale(float Scale, FCameraShakeUpdateResult& InOutResult)
 {
-	InOutResult.Location *= Scale;
-	InOutResult.Rotation *= Scale;
-	InOutResult.FOV *= Scale;
-	InOutResult.PostProcessBlendWeight *= Scale;
+	InOutResult.ApplyScale(Scale);
 }
 
 void UCameraShakeBase::ApplyLimits(const FMinimalViewInfo& InPOV, FCameraShakeUpdateResult& InOutResult)
@@ -511,4 +522,3 @@ void UCameraShakePattern::TeardownShakePattern()
 	TeardownShakePatternImpl();
 }
 
-
