RemoveSubObjectGroupMembership does not need to ensure if no connections are present.

The netconnection cleanup code will have removed all existing group relationships already

diff --git a/Engine/Source/Runtime/Engine/Private/Net/Iris/ReplicationSystem/ReplicationSystemUtil.cpp b/Engine/Source/Runtime/Engine/Private/Net/Iris/ReplicationSystem/ReplicationSystemUtil.cpp
--- a/Engine/Source/Runtime/Engine/Private/Net/Iris/ReplicationSystem/ReplicationSystemUtil.cpp
+++ b/Engine/Source/Runtime/Engine/Private/Net/Iris/ReplicationSystem/ReplicationSystemUtil.cpp
@@ -514,8 +514,7 @@ void FReplicationSystemUtil::RemoveSubObjectGroupMembership(const APlayerControl
 	}
 	
 	// We assume the player controller is tied to a single connection.
-	UNetConnection* Conn = PC->GetNetConnection();
-	if (ensureAlways(Conn))
+	if (UNetConnection* Conn = PC->GetNetConnection())
 	{
 		if (UReplicationSystem* ReplicationSystem = Conn->GetDriver() ? Conn->GetDriver()->GetReplicationSystem() : nullptr)
 		{
