Fix LWI instanced static mesh managers returning incorrect indices. Fixes issue where spawned actors were being associated with the incorrect LWI.

diff --git a/Engine/Source/Runtime/Engine/Private/LightWeightInstanceStaticMeshManager.cpp b/Engine/Source/Runtime/Engine/Private/LightWeightInstanceStaticMeshManager.cpp
--- a/Engine/Source/Runtime/Engine/Private/LightWeightInstanceStaticMeshManager.cpp
+++ b/Engine/Source/Runtime/Engine/Private/LightWeightInstanceStaticMeshManager.cpp
@@ -55,12 +55,22 @@ void ALightWeightInstanceStaticMeshManager::SetRepresentedClass(UClass* ActorCla
 
 int32 ALightWeightInstanceStaticMeshManager::ConvertCollisionIndexToLightWeightIndex(int32 InIndex) const
 {
-	return InIndex;//RenderingIndicesToDataIndices[InIndex];
+	if (ensureMsgf(RenderingIndicesToDataIndices.IsValidIndex(InIndex), TEXT("Invalid index [ %d ]"), InIndex))
+	{
+		return RenderingIndicesToDataIndices[InIndex];
+	}
+
+	return InIndex;
 }
 
 int32 ALightWeightInstanceStaticMeshManager::ConvertLightWeightIndexToCollisionIndex(int32 InIndex) const
 {
-	return InIndex;//DataIndicesToRenderingIndices[InIndex];
+	if (ensureMsgf(DataIndicesToRenderingIndices.IsValidIndex(InIndex), TEXT("Invalid index [ %d ]"), InIndex))
+	{
+		return DataIndicesToRenderingIndices[InIndex];
+	}
+
+	return InIndex;
 }
 
 void ALightWeightInstanceStaticMeshManager::AddNewInstanceAt(FLWIData* InitData, int32 Index)
