Make light weight instanced static mesh manager use a UObject delegate for timer delegate rather than capturing this - prevent dangling pointer if object is destroyed before timer delegate is fired.

diff --git a/Engine/Source/Runtime/Engine/Private/LightWeightInstanceStaticMeshManager.cpp b/Engine/Source/Runtime/Engine/Private/LightWeightInstanceStaticMeshManager.cpp
--- a/Engine/Source/Runtime/Engine/Private/LightWeightInstanceStaticMeshManager.cpp
+++ b/Engine/Source/Runtime/Engine/Private/LightWeightInstanceStaticMeshManager.cpp
@@ -117,12 +117,12 @@ void ALightWeightInstanceStaticMeshManager::RemoveInstanceFromRendering(int32 Da
 #if WITH_EDITOR
 			if (GEditor)
 			{
-				GEditor->GetTimerManager()->SetTimerForNextTick([this]() { PostRemoveInstanceFromRendering(); });
+				GEditor->GetTimerManager()->SetTimerForNextTick(this, &ALightWeightInstanceStaticMeshManager::PostRemoveInstanceFromRendering);
 			}
 			else
 #endif // WITH_EDITOR
 			{
-				GetWorld()->GetTimerManager().SetTimerForNextTick([this]() { PostRemoveInstanceFromRendering(); });
+				GetWorld()->GetTimerManager().SetTimerForNextTick(this, &ALightWeightInstanceStaticMeshManager::PostRemoveInstanceFromRendering);
 			}
 		}
 
