Add 'debug stall' console command to block main thread for 10 seconds.

diff --git a/Engine/Source/Runtime/Engine/Private/UnrealEngine.cpp b/Engine/Source/Runtime/Engine/Private/UnrealEngine.cpp
--- a/Engine/Source/Runtime/Engine/Private/UnrealEngine.cpp
+++ b/Engine/Source/Runtime/Engine/Private/UnrealEngine.cpp
@@ -10401,6 +10401,18 @@ bool UEngine::PerformError(const TCHAR* Cmd, FOutputDevice& Ar)
 		CauseCrtError();
 		return true;
 	}
+	else if (FParse::Command(&Cmd, TEXT("STALL")))
+	{
+		float Seconds = FCString::Atof(Cmd);
+		if (Seconds == 0.0f)
+		{
+			Seconds = 10.0f;
+		}
+		SCOPE_CYCLE_COUNTER(STAT_IntentionalHitch);
+		Ar.Logf(ELogVerbosity::Display, TEXT("Stalling for %.0f seconds."), Seconds);
+		FPlatformProcess::Sleep(Seconds);
+		return true;
+	}
 	else if (FParse::Command(&Cmd, TEXT("HITCH")))
 	{
 		float Seconds = FCString::Atof(Cmd) / 1000.0f;
