Fixed Strata material preview on nodes to account for the sub strata tree from the previewed node only.
this fixed material function input / output, reroute nodes and everything in between that is not directly connected to FrontMaterial.

diff --git a/Engine/Source/Editor/MaterialEditor/Private/MaterialEditor.cpp b/Engine/Source/Editor/MaterialEditor/Private/MaterialEditor.cpp
--- a/Engine/Source/Editor/MaterialEditor/Private/MaterialEditor.cpp
+++ b/Engine/Source/Editor/MaterialEditor/Private/MaterialEditor.cpp
@@ -311,11 +311,10 @@ int32 FMatExpressionPreview::CompilePropertyAndSetMaterialProperty(EMaterialProp
 	}
 	else if (Property == MP_FrontMaterial)
 	{
-		// Try to fetch the material interface to compile the front material.
-		UMaterial* ExprMat = Expression.IsValid() ? Expression->Material : nullptr;
-		FMaterialRenderProxy* MatProxy = ExprMat ? ExprMat->GetRenderProxy() : nullptr;
-		UMaterialInterface* MatInterface = MatProxy ? MatProxy->GetMaterialInterface() : nullptr;
-		return MatInterface ? MatInterface->CompileProperty(Compiler, MP_FrontMaterial) : Compiler->StrataCreateAndRegisterNullMaterial();
+		// No need to compile the front material: when previewing a node, the FrontMaterial is plugged into the emissive color.
+		// Then CompilePreview is called, and this is where we convert the strata material to a single color for preview.
+		// That single color is then scheduled to be output thanks to setting the compiler as SME_MaterialPreview. 
+		return Compiler->StrataCreateAndRegisterNullMaterial();
 	}
 	else
 	{
diff --git a/Engine/Source/Editor/MaterialEditor/Private/MaterialEditor.h b/Engine/Source/Editor/MaterialEditor/Private/MaterialEditor.h
--- a/Engine/Source/Editor/MaterialEditor/Private/MaterialEditor.h
+++ b/Engine/Source/Editor/MaterialEditor/Private/MaterialEditor.h
@@ -117,6 +117,8 @@ public:
 	/** Entry point for compiling a specific material property.  This must call SetMaterialProperty. */
 	virtual int32 CompilePropertyAndSetMaterialProperty(EMaterialProperty Property, FMaterialCompiler* Compiler, EShaderFrequency OverrideShaderFrequency, bool bUsePreviousFrameTime) const override;
 
+	virtual UMaterialExpression* GetMaterialGraphNodePreviewExpression() const override { return Expression.Get(); }
+
 	virtual EMaterialDomain GetMaterialDomain() const override { return MD_Surface; }
 	virtual FString GetMaterialUsageDescription() const override { return FString::Printf(TEXT("FMatExpressionPreview %s"), Expression.IsValid() ? *Expression->GetName() : TEXT("NULL")); }
 	virtual bool IsTwoSided() const override { return false; }
diff --git a/Engine/Source/Runtime/Engine/Classes/Materials/MaterialExpressionFunctionOutput.h b/Engine/Source/Runtime/Engine/Classes/Materials/MaterialExpressionFunctionOutput.h
--- a/Engine/Source/Runtime/Engine/Classes/Materials/MaterialExpressionFunctionOutput.h
+++ b/Engine/Source/Runtime/Engine/Classes/Materials/MaterialExpressionFunctionOutput.h
@@ -55,6 +55,7 @@ class UMaterialExpressionFunctionOutput : public UMaterialExpression
 	//~ Begin UMaterialExpression Interface
 #if WITH_EDITOR
 	virtual int32 Compile(class FMaterialCompiler* Compiler, int32 OutputIndex) override;
+	virtual int32 CompilePreview(class FMaterialCompiler* Compiler, int32 OutputIndex) override;
 	virtual void GetCaption(TArray<FString>& OutCaptions) const override;
 	virtual FName GetInputName(int32 InputIndex) const override
 	{
diff --git a/Engine/Source/Runtime/Engine/Private/Materials/HLSLMaterialTranslator.cpp b/Engine/Source/Runtime/Engine/Private/Materials/HLSLMaterialTranslator.cpp
--- a/Engine/Source/Runtime/Engine/Private/Materials/HLSLMaterialTranslator.cpp
+++ b/Engine/Source/Runtime/Engine/Private/Materials/HLSLMaterialTranslator.cpp
@@ -836,8 +836,31 @@ bool FHLSLMaterialTranslator::Translate()
 		// Process the strata tree representing the material topology.
 		//
 		const bool bStrataEnabled = Strata::IsStrataEnabled();
-		FStrataMaterialInput* FrontMaterialInput = Material->GetMaterialInterface() ? &Material->GetMaterialInterface()->GetMaterial()->GetEditorOnlyData()->FrontMaterial : nullptr;
-		UMaterialExpression* FrontMaterialExpr = FrontMaterialInput ? FrontMaterialInput->GetTracedInput().Expression : nullptr;
+		UMaterialExpression* FrontMaterialExpr = nullptr;
+		int32 FrontMaterialOutputIndex = INDEX_NONE;
+
+		UMaterialExpression* ExpressionToPreview = Material->GetMaterialGraphNodePreviewExpression();
+		if (ExpressionToPreview)
+		{
+			// If this is a compilation used for a preview of a node from am aterial graph, then we need ot use that node as the root of the Strata tree.
+			// Then the resulting StrataData will be converted to a color using StrataCompilePreview translator function when compiling the emissive color channel the node has been short circuited into.
+			const uint32 ExpressionPreviewOutputIndex = 0;
+			if (ExpressionToPreview && ExpressionToPreview->IsResultStrataMaterial(ExpressionPreviewOutputIndex))
+			{
+				FrontMaterialExpr = ExpressionToPreview;
+				FrontMaterialOutputIndex = ExpressionPreviewOutputIndex;
+			}
+			else
+			{
+				FrontMaterialExpr = nullptr;	// This is not a strata input that is connected there so we cannot create a Strata tree.
+			}
+		}
+		else
+		{
+			FStrataMaterialInput* FrontMaterialInput = Material->GetMaterialInterface() ? &Material->GetMaterialInterface()->GetMaterial()->GetEditorOnlyData()->FrontMaterial : nullptr;
+			FrontMaterialExpr = FrontMaterialInput ? FrontMaterialInput->GetTracedInput().Expression : nullptr;
+			FrontMaterialOutputIndex = FrontMaterialInput ? FrontMaterialInput->OutputIndex : INDEX_NONE;
+		}
 		if (bStrataEnabled && FrontMaterialExpr)
 		{
 			// Temp code chunk scope (e.g.needed for the creation of static booleans from static switch parameter node, see UMaterialExpressionStaticSwitch::GetEffectiveInput).
@@ -857,7 +880,7 @@ bool FHLSLMaterialTranslator::Translate()
 			
 				FExpressionInput* SurfaceThickness = Material->IsThinSurface() && Material->GetMaterialInterface() ? &Material->GetMaterialInterface()->GetMaterial()->GetEditorOnlyData()->SurfaceThickness : nullptr;
 				StrataThicknessStackPush(nullptr, SurfaceThickness);
-				FrontMaterialExpr->StrataGenerateMaterialTopologyTree(this, nullptr, FrontMaterialInput->OutputIndex); 
+				FrontMaterialExpr->StrataGenerateMaterialTopologyTree(this, nullptr, FrontMaterialOutputIndex);
 				StrataThicknessStackPop();
 
 				check(StrataCtx.StrataThicknessStack.Num() == 0);
diff --git a/Engine/Source/Runtime/Engine/Private/Materials/MaterialExpressions.cpp b/Engine/Source/Runtime/Engine/Private/Materials/MaterialExpressions.cpp
--- a/Engine/Source/Runtime/Engine/Private/Materials/MaterialExpressions.cpp
+++ b/Engine/Source/Runtime/Engine/Private/Materials/MaterialExpressions.cpp
@@ -2033,6 +2033,7 @@ void UMaterialExpression::UpdateParameterGuid(bool bForceGeneration, bool bAllow
 
 void UMaterialExpression::ConnectToPreviewMaterial(UMaterial* InMaterial, int32 OutputIndex)
 {
+	// This is used when a node is right clicked and "Start previewing node" is used.
 	if (InMaterial && OutputIndex >= 0 && OutputIndex < Outputs.Num())
 	{
 		if (Strata::IsStrataEnabled())
@@ -16892,6 +16893,15 @@ int32 UMaterialExpressionFunctionInput::Compile(class FMaterialCompiler* Compile
 
 int32 UMaterialExpressionFunctionInput::CompilePreview(class FMaterialCompiler* Compiler, int32 OutputIndex)
 {
+	if (InputType == FunctionInput_Substrate)
+	{
+		// Compile the StrataData output.
+		int32 StrataDataCodeChunk = Compile(Compiler, OutputIndex);
+		// Convert the StrataData to a preview color.
+		int32 PreviewCodeChunk = Compiler->StrataCompilePreview(StrataDataCodeChunk);
+		return PreviewCodeChunk;
+	}
+
 	// Compile the preview value, outputting a float type
 	return Compiler->ValidCast(CompilePreviewValue(Compiler), MCT_Float3);
 }
@@ -17149,6 +17159,23 @@ int32 UMaterialExpressionFunctionOutput::Compile(class FMaterialCompiler* Compil
 	}
 	return A.Compile(Compiler);
 }
+
+int32 UMaterialExpressionFunctionOutput::CompilePreview(class FMaterialCompiler* Compiler, int32 OutputIndex)
+{
+	FExpressionInput ATraced = A.GetTracedInput();
+	UMaterialExpression* AExpression = ATraced.Expression;
+	if (AExpression && AExpression->IsResultStrataMaterial(ATraced.OutputIndex))
+	{
+		// Compile the StrataData output.
+		int32 StrataDataCodeChunk = Compile(Compiler, ATraced.OutputIndex);
+		// Convert the StrataData to a preview color.
+		int32 PreviewCodeChunk = Compiler->StrataCompilePreview(StrataDataCodeChunk);
+		return PreviewCodeChunk;
+	}
+
+	// Compile the preview value, outputting a float type
+	return Compile(Compiler, OutputIndex);
+}
 #endif // WITH_EDITOR
 
 void UMaterialExpressionFunctionOutput::ConditionallyGenerateId(bool bForce)
diff --git a/Engine/Source/Runtime/Engine/Public/MaterialShared.h b/Engine/Source/Runtime/Engine/Public/MaterialShared.h
--- a/Engine/Source/Runtime/Engine/Public/MaterialShared.h
+++ b/Engine/Source/Runtime/Engine/Public/MaterialShared.h
@@ -2342,6 +2342,13 @@ protected:
 
 #if WITH_EDITOR
 	void SetCompilingShaderMap(FMaterialShaderMap* InMaterialShaderMap);
+
+	/**
+	 * The returned value is not const because the expression is used to build the Strata tree and this requires to execute multiple such as link function caller, compile expression, etc.
+	 * See HLSLTranslator.cpp for the single use case.
+	 * @return Nullptr if this is not a material used to generate a preview of a node of a material graph.
+	 */
+	virtual UMaterialExpression* GetMaterialGraphNodePreviewExpression() const { return nullptr; }
 #endif
 
 private:
