Tiny ISM component change required to push LWIxMass forward.

diff --git a/Engine/Source/Runtime/Engine/Classes/Components/InstancedStaticMeshComponent.h b/Engine/Source/Runtime/Engine/Classes/Components/InstancedStaticMeshComponent.h
--- a/Engine/Source/Runtime/Engine/Classes/Components/InstancedStaticMeshComponent.h
+++ b/Engine/Source/Runtime/Engine/Classes/Components/InstancedStaticMeshComponent.h
@@ -464,7 +464,7 @@ protected:
 	int32 AddInstanceInternal(int32 InstanceIndex, FInstancedStaticMeshInstanceData* InNewInstanceData, const FTransform& InstanceTransform, bool bWorldSpace);
 
 	/** Internal implementation of AddInstances */
-	TArray<int32> AddInstancesInternal(const TArray<FTransform>& InstanceTransforms, bool bShouldReturnIndices, bool bWorldSpace);
+	TArray<int32> AddInstancesInternal(TConstArrayView<FTransform> InstanceTransforms, bool bShouldReturnIndices, bool bWorldSpace);
 
 	/** Internal version of RemoveInstance */	
 	bool RemoveInstanceInternal(int32 InstanceIndex, bool InstanceAlreadyRemoved);
diff --git a/Engine/Source/Runtime/Engine/Private/InstancedStaticMesh.cpp b/Engine/Source/Runtime/Engine/Private/InstancedStaticMesh.cpp
--- a/Engine/Source/Runtime/Engine/Private/InstancedStaticMesh.cpp
+++ b/Engine/Source/Runtime/Engine/Private/InstancedStaticMesh.cpp
@@ -3288,7 +3288,7 @@ int32 UInstancedStaticMeshComponent::AddInstance(const FTransform& InstanceTrans
 	return AddInstanceInternal(PerInstanceSMData.Num(), nullptr, InstanceTransform, bWorldSpace);
 }
 
-TArray<int32> UInstancedStaticMeshComponent::AddInstancesInternal(const TArray<FTransform>& InstanceTransforms, bool bShouldReturnIndices, bool bWorldSpace)
+TArray<int32> UInstancedStaticMeshComponent::AddInstancesInternal(TConstArrayView<FTransform> InstanceTransforms, bool bShouldReturnIndices, bool bWorldSpace)
 {
 	const int32 Count = InstanceTransforms.Num();
 
