World Partition Class Descriptors: don't create class descriptors for newly created Blueprints and fix how we handle loaded ones.

diff --git a/Engine/Source/Runtime/Engine/Private/BlueprintGeneratedClass.cpp b/Engine/Source/Runtime/Engine/Private/BlueprintGeneratedClass.cpp
--- a/Engine/Source/Runtime/Engine/Private/BlueprintGeneratedClass.cpp
+++ b/Engine/Source/Runtime/Engine/Private/BlueprintGeneratedClass.cpp
@@ -255,7 +255,13 @@ void UBlueprintGeneratedClass::GetAssetRegistryTags(TArray<FAssetRegistryTag>& O
 #if WITH_EDITOR
 	if (AActor* BlueprintCDO = Cast<AActor>(ClassDefaultObject))
 	{
-		FWorldPartitionActorDescUtils::AppendAssetDataTagsFromActor(BlueprintCDO, OutTags);
+		if (UPackage* BlueprintCDOPackage = BlueprintCDO->GetPackage())
+		{
+			if (!FPackageName::IsTempPackage(BlueprintCDOPackage->GetName()) && !BlueprintCDOPackage->HasAnyPackageFlags(PKG_PlayInEditor))
+			{
+				FWorldPartitionActorDescUtils::AppendAssetDataTagsFromActor(BlueprintCDO, OutTags);
+			}
+		}
 	}
 #endif
 }
diff --git a/Engine/Source/Runtime/Engine/Private/WorldPartition/WorldPartitionClassDescRegistry.cpp b/Engine/Source/Runtime/Engine/Private/WorldPartition/WorldPartitionClassDescRegistry.cpp
--- a/Engine/Source/Runtime/Engine/Private/WorldPartition/WorldPartitionClassDescRegistry.cpp
+++ b/Engine/Source/Runtime/Engine/Private/WorldPartition/WorldPartitionClassDescRegistry.cpp
@@ -7,6 +7,7 @@
 #include "Algo/ForEach.h"
 #include "Algo/RemoveIf.h"
 #include "Algo/Transform.h"
+#include "Misc/PackageName.h"
 #include "UObject/CoreRedirects.h"
 #include "UObject/ObjectSaveContext.h"
 #include "UObject/UObjectIterator.h"
@@ -20,6 +21,32 @@
 #include "WorldPartition/WorldPartitionActorDesc.h"
 #include "WorldPartition/WorldPartitionActorDescUtils.h"
 
+/*
+ * Class descriptors are specific properties that we keep available for unloaded actor classes. We want to keep in sync class descriptors and their
+ * corresponding Blueprint generated class, to have the exact same behavior as delta serialization of Actors from their class CDOs. The different
+ * scenarios for creating and updating class descriptors from the class CDO are as follow:
+ *
+ *	[Creating a new Blueprint]
+ *	  - Don't create a class descriptor in that case, since we can't save an actor instance until we save the Blueprint.
+ *	  - The Blueprint generated class is not even created at this point anyways.
+ *	
+ *	[Saving a new Blueprint]
+ *	  - Create the class descriptor from the newly generated Blueprint class.
+ *	
+ *	[Modifying an existing Blueprint]
+ *	  - Don't update the class descriptor, as actor instances of this class will use the existing (outdated) generated Blueprint class for delta
+ *		serialization.
+ *
+ *	[Compiling an existing Blueprint]
+ *	  - Create or update the class descriptor, but only if it already exists on disk.
+ *
+ *	[Saving a existing Blueprint]
+ *	  - Create or update the class descriptor.
+ *
+ *	[Loading a Blueprint]
+ *	  - Create or update the class descriptor from the assset-registry data.
+ */
+
 /*
  * We have to deal with some different realities between non-cooked and cooked editor builds: for non-cooked editor builds, UBlueprint are considered assets
  * and will be the opnly thing visible from the asset registry standpoint. On the other hand, for cooked editor builds, UBlueprint are not considered assets
@@ -79,7 +106,6 @@ void FWorldPartitionClassDescRegistry::Initialize()
 	}
 
 	IAssetRegistry& AssetRegistry = FModuleManager::LoadModuleChecked<FAssetRegistryModule>(TEXT("AssetRegistry")).Get();
-	AssetRegistry.OnInMemoryAssetCreated().AddRaw(this, &FWorldPartitionClassDescRegistry::OnAssetCreated);
 	AssetRegistry.OnAssetRemoved().AddRaw(this, &FWorldPartitionClassDescRegistry::OnAssetRemoved);
 	AssetRegistry.OnAssetRenamed().AddRaw(this, &FWorldPartitionClassDescRegistry::OnAssetRenamed);
 
@@ -98,7 +124,6 @@ void FWorldPartitionClassDescRegistry::Uninitialize()
 	if (!IsEngineExitRequested())
 	{
 		IAssetRegistry& AssetRegistry = FModuleManager::LoadModuleChecked<FAssetRegistryModule>(TEXT("AssetRegistry")).Get();
-		AssetRegistry.OnInMemoryAssetCreated().RemoveAll(this);
 		AssetRegistry.OnAssetRemoved().RemoveAll(this);
 		AssetRegistry.OnAssetRenamed().RemoveAll(this);
 	}
@@ -126,7 +151,11 @@ void FWorldPartitionClassDescRegistry::PrefetchClassDescs(const TArray<FTopLevel
 	TArray<FTopLevelAssetPath> ClassPaths;
 	ClassPaths.Reserve(InClassPaths.Num());
 	Algo::Transform(InClassPaths, ClassPaths, [this](const FTopLevelAssetPath& ClassPath) { return RedirectClassPath(ClassPath); });
-	ClassPaths.SetNum(Algo::RemoveIf(ClassPaths, [this](const FTopLevelAssetPath& ClassPath) { return ClassByPath.Contains(ClassPath) || ClassPath.ToString().StartsWith(TEXT("/Script/")); }));
+	ClassPaths.SetNum(Algo::RemoveIf(ClassPaths, [this](const FTopLevelAssetPath& ClassPath)
+	{
+		const FString ClasPathString = ClassPath.ToString();
+		return ClassByPath.Contains(ClassPath) || FPackageName::IsScriptPackage(ClasPathString) || FPackageName::IsMemoryPackage(ClasPathString) || FPackageName::IsTempPackage(ClasPathString);
+	}));
 
 	if (ClassPaths.IsEmpty())
 	{
@@ -354,15 +383,16 @@ void FWorldPartitionClassDescRegistry::RegisterClassDescriptorFromActorClass(con
 	RegisterClassDescriptor(NewActorDesc);
 }
 
-void FWorldPartitionClassDescRegistry::OnAssetCreated(UObject* InAssetCreated)
-{
-	OnAssetLoaded(InAssetCreated);
-}
-
 void FWorldPartitionClassDescRegistry::OnAssetLoaded(UObject* InAssetLoaded)
 {
+	check(InAssetLoaded);
 	check(IsInitialized());
 
+	if (InAssetLoaded->GetPackage()->HasAnyPackageFlags(PKG_PlayInEditor))
+	{
+		return;
+	}
+
 	UBlueprintGeneratedClass* BlueprintGeneratedClass = Cast<UBlueprintGeneratedClass>(InAssetLoaded);
 
 	if (!BlueprintGeneratedClass)
@@ -373,38 +403,21 @@ void FWorldPartitionClassDescRegistry::OnAssetLoaded(UObject* InAssetLoaded)
 		}
 	}
 
-	if (BlueprintGeneratedClass && BlueprintGeneratedClass->IsChildOf<AActor>())
+	if (!BlueprintGeneratedClass || !BlueprintGeneratedClass->IsChildOf<AActor>())
 	{
-		if (!ClassByPath.Contains(FTopLevelAssetPath(BlueprintGeneratedClass->GetPathName())))
-		{
-			ValidateInternalState();
-			RegisterClass(BlueprintGeneratedClass);
-			ValidateInternalState();
-		}
+		return;
 	}
+
+	ValidateInternalState();
+	PrefetchClassDescs({ FTopLevelAssetPath(BlueprintGeneratedClass->GetPathName()) });
+	ValidateInternalState();
 }
 
 void FWorldPartitionClassDescRegistry::OnObjectPreSave(UObject* InObject, FObjectPreSaveContext InSaveContext)
 {
 	if (!InSaveContext.IsProceduralSave() && !(InSaveContext.GetSaveFlags() & SAVE_FromAutosave))
 	{
-		ValidateInternalState();
-
-		if (AActor* Actor = Cast<AActor>(InObject))
-		{
-			if (!Actor->HasAnyFlags(RF_ArchetypeObject | RF_ClassDefaultObject))
-			{
-				// We are saving an actor
-				if (UBlueprintGeneratedClass* BlueprintGeneratedClass = Cast<UBlueprintGeneratedClass>(Actor->GetClass()))
-				{
-					if (UBlueprint* Blueprint = Cast<UBlueprint>(BlueprintGeneratedClass->ClassGeneratedBy))
-					{
-						UpdateClassDescriptor(Blueprint, true);
-					}
-				}
-			}
-		}
-		else if (UBlueprint* Blueprint = Cast<UBlueprint>(InObject))
+		if (UBlueprint* Blueprint = Cast<UBlueprint>(InObject))
 		{
 			// We are saving a blueprint
 			if (UBlueprintGeneratedClass* BlueprintGeneratedClass = Cast<UBlueprintGeneratedClass>(Blueprint->GeneratedClass))
@@ -415,8 +428,6 @@ void FWorldPartitionClassDescRegistry::OnObjectPreSave(UObject* InObject, FObjec
 				}
 			}
 		}
-
-		ValidateInternalState();
 	}
 }
 
@@ -427,9 +438,7 @@ void FWorldPartitionClassDescRegistry::OnObjectPropertyChanged(UObject* InObject
 		// The generated class is invalid in some situations, like renaming a blueprint, etc.
 		if (Blueprint->GeneratedClass && Blueprint->GeneratedClass->IsChildOf<AActor>())
 		{
-			ValidateInternalState();
 			UpdateClassDescriptor(InObject, true);
-			ValidateInternalState();
 		}
 	}
 }
@@ -560,6 +569,8 @@ void FWorldPartitionClassDescRegistry::UpdateClassDescriptor(UObject* InObject,
 	{
 		const FTopLevelAssetPath PreviousParentClassPath = ParentClassMap.FindChecked(ClassPath);
 
+		ValidateInternalState();
+
 		if (!*ExistingClassDesc)
 		{
 			if (!bOnlyIfExists)
@@ -580,6 +591,8 @@ void FWorldPartitionClassDescRegistry::UpdateClassDescriptor(UObject* InObject,
 				ParentClassMap.Add(ClassPath, CurrentParentClassPath);
 			}
 		}
+
+		ValidateInternalState();
 	}
 }
 
diff --git a/Engine/Source/Runtime/Engine/Public/WorldPartition/WorldPartitionClassDescRegistry.h b/Engine/Source/Runtime/Engine/Public/WorldPartition/WorldPartitionClassDescRegistry.h
--- a/Engine/Source/Runtime/Engine/Public/WorldPartition/WorldPartitionClassDescRegistry.h
+++ b/Engine/Source/Runtime/Engine/Public/WorldPartition/WorldPartitionClassDescRegistry.h
@@ -43,7 +43,6 @@ private:
 	const FWorldPartitionActorDesc* GetClassDescDefaultForActor(const FTopLevelAssetPath& InClassPath) const;
 	const FWorldPartitionActorDesc* GetClassDescDefaultForClass(const FTopLevelAssetPath& InClassPath) const;
 
-	void OnAssetCreated(UObject* InAssetCreated);
 	void OnAssetLoaded(UObject* InAssetLoaded);
 	void OnObjectPreSave(UObject* InObject, FObjectPreSaveContext InSaveContext);
 	void OnObjectPropertyChanged(UObject* InObject, FPropertyChangedEvent& InPropertyChangedEvent);
