fix slow NearlyEqual Transform test

#rnx

diff --git a/Engine/Source/Runtime/Engine/Private/Tests/EngineAutomationTests.cpp b/Engine/Source/Runtime/Engine/Private/Tests/EngineAutomationTests.cpp
--- a/Engine/Source/Runtime/Engine/Private/Tests/EngineAutomationTests.cpp
+++ b/Engine/Source/Runtime/Engine/Private/Tests/EngineAutomationTests.cpp
@@ -568,7 +568,7 @@ const float FAutomationNearlyEqualTest::PositiveHalfDifference((1.e-4f) / 2.f);
 const FString FAutomationNearlyEqualTest::TestFailMessage(TEXT("Total amount of errors is not equal to the expected amount"));
 const uint32 FAutomationNearlyEqualTest::MaxFVectorMutationBitMask(0b111);
 const uint32 FAutomationNearlyEqualTest::MaxFRotatorMutationBitMask(0b111);
-const uint32 FAutomationNearlyEqualTest::MaxFTransformMutationBitMask(0b111111111);
+const uint32 FAutomationNearlyEqualTest::MaxFTransformMutationBitMask(0b111);
 
 const float FAutomationNearlyEqualTest::BaseFloatValue(0.f);
 const float FAutomationNearlyEqualTest::ActualFloatValue(BaseFloatValue);
