FWorldPartitionActorDescUtils::FixupRedirectedAssetPath: Fix warning when passing in full SoftObjectPath (Private DataLayers)

diff --git a/Engine/Source/Runtime/Engine/Private/WorldPartition/WorldPartitionActorDescUtils.cpp b/Engine/Source/Runtime/Engine/Private/WorldPartition/WorldPartitionActorDescUtils.cpp
--- a/Engine/Source/Runtime/Engine/Private/WorldPartition/WorldPartitionActorDescUtils.cpp
+++ b/Engine/Source/Runtime/Engine/Private/WorldPartition/WorldPartitionActorDescUtils.cpp
@@ -169,7 +169,8 @@ bool FWorldPartitionActorDescUtils::FixupRedirectedAssetPath(FName& InOutAssetPa
 	const FAssetData* AssetData;
 	IAssetRegistry& AssetRegistry = FModuleManager::LoadModuleChecked<FAssetRegistryModule>(TEXT("AssetRegistry")).Get();
 
-	FTopLevelAssetPath AssetPath = FTopLevelAssetPath(InOutAssetPath.ToString());
+	FSoftObjectPath SoftObjectPath(InOutAssetPath.ToString());
+	FTopLevelAssetPath AssetPath = SoftObjectPath.GetAssetPath();
 
 	for (;;)
 	{
@@ -210,7 +211,8 @@ bool FWorldPartitionActorDescUtils::FixupRedirectedAssetPath(FName& InOutAssetPa
 		AssetPath = FTopLevelAssetPath(DestinationObjectPath);
 	}
 
-	InOutAssetPath = FName(AssetPath.ToString());
+	SoftObjectPath.SetPath(AssetPath, SoftObjectPath.GetSubPathString());	
+	InOutAssetPath = FName(*SoftObjectPath.ToString());
 	return true;
 }
 
