Matinee removal follow-up: remove PostInterpChange, as nobody calls it anymore

diff --git a/Engine/Source/Runtime/CoreUObject/Public/UObject/Object.h b/Engine/Source/Runtime/CoreUObject/Public/UObject/Object.h
--- a/Engine/Source/Runtime/CoreUObject/Public/UObject/Object.h
+++ b/Engine/Source/Runtime/CoreUObject/Public/UObject/Object.h
@@ -409,11 +409,7 @@ public:
 	/** After a critical error, perform any mission-critical cleanup, such as restoring the video mode orreleasing hardware resources. */
 	virtual void ShutdownAfterError() {}
 
-	/** 
-	 * This is called when property is modified by InterpPropertyTracks
-	 *
-	 * @param PropertyThatChanged	Property that changed
-	 */
+	UE_DEPRECATED(5.3, "Matinee has been removed in 5.0, nothing is calling PostInterpChange anymore")
 	virtual void PostInterpChange(FProperty* PropertyThatChanged) {}
 
 #if WITH_EDITOR
diff --git a/Engine/Source/Runtime/Engine/Classes/Atmosphere/AtmosphericFogComponent.h b/Engine/Source/Runtime/Engine/Classes/Atmosphere/AtmosphericFogComponent.h
--- a/Engine/Source/Runtime/Engine/Classes/Atmosphere/AtmosphericFogComponent.h
+++ b/Engine/Source/Runtime/Engine/Classes/Atmosphere/AtmosphericFogComponent.h
@@ -82,8 +82,6 @@ public:
 	
 	//~ Begin UObject Interface. 
 	virtual bool IsPostLoadThreadSafe() const override;
-
-	virtual void PostInterpChange(FProperty* PropertyThatChanged) override;
 	virtual void Serialize(FArchive& Ar) override;
 	//~ End UObject Interface;
 
diff --git a/Engine/Source/Runtime/Engine/Classes/Components/ExponentialHeightFogComponent.h b/Engine/Source/Runtime/Engine/Classes/Components/ExponentialHeightFogComponent.h
--- a/Engine/Source/Runtime/Engine/Classes/Components/ExponentialHeightFogComponent.h
+++ b/Engine/Source/Runtime/Engine/Classes/Components/ExponentialHeightFogComponent.h
@@ -288,7 +288,6 @@ public:
 	virtual bool CanEditChange(const FProperty* InProperty) const override;
 	virtual void PostEditChangeProperty(FPropertyChangedEvent& PropertyChangedEvent) override;
 #endif // WITH_EDITOR
-	virtual void PostInterpChange(FProperty* PropertyThatChanged) override;
 	virtual void Serialize(FArchive& Ar) override;
 	//~ End UObject Interface
 };
diff --git a/Engine/Source/Runtime/Engine/Classes/Components/LightComponent.h b/Engine/Source/Runtime/Engine/Classes/Components/LightComponent.h
--- a/Engine/Source/Runtime/Engine/Classes/Components/LightComponent.h
+++ b/Engine/Source/Runtime/Engine/Classes/Components/LightComponent.h
@@ -479,13 +479,6 @@ public:
 	/** Get the color temperature in the working color space. */
 	FLinearColor GetColorTemperature() const;
 
-	/** 
-	 * Called when property is modified by InterpPropertyTracks
-	 *
-	 * @param PropertyThatChanged	Property that changed
-	 */
-	virtual void PostInterpChange(FProperty* PropertyThatChanged) override;
-
 	/** 
 	 * Iterates over ALL stationary light components in the target world and assigns their preview shadowmap channel, and updates light icons accordingly.
 	 * Also handles assignment after a lighting build, so that the same algorithm is used for previewing and static lighting.
diff --git a/Engine/Source/Runtime/Engine/Classes/Components/LocalLightComponent.h b/Engine/Source/Runtime/Engine/Classes/Components/LocalLightComponent.h
--- a/Engine/Source/Runtime/Engine/Classes/Components/LocalLightComponent.h
+++ b/Engine/Source/Runtime/Engine/Classes/Components/LocalLightComponent.h
@@ -83,13 +83,6 @@ public:
 #endif // WITH_EDITOR
 	//~ End UObject Interface
 
-	/** 
-	 * This is called when property is modified by InterpPropertyTracks
-	 *
-	 * @param PropertyThatChanged	Property that changed
-	 */
-	virtual void PostInterpChange(FProperty* PropertyThatChanged) override;
-
 private:
 
 	/** Pushes the value of radius to the rendering thread. */
diff --git a/Engine/Source/Runtime/Engine/Classes/Components/PointLightComponent.h b/Engine/Source/Runtime/Engine/Classes/Components/PointLightComponent.h
--- a/Engine/Source/Runtime/Engine/Classes/Components/PointLightComponent.h
+++ b/Engine/Source/Runtime/Engine/Classes/Components/PointLightComponent.h
@@ -94,13 +94,6 @@ public:
 	virtual void PostEditChangeProperty(FPropertyChangedEvent& PropertyChangedEvent) override;
 #endif // WITH_EDITOR
 	//~ End UObject Interface
-
-	/** 
-	 * This is called when property is modified by InterpPropertyTracks
-	 *
-	 * @param PropertyThatChanged	Property that changed
-	 */
-	virtual void PostInterpChange(FProperty* PropertyThatChanged) override;
 };
 
 
diff --git a/Engine/Source/Runtime/Engine/Classes/Components/SceneComponent.h b/Engine/Source/Runtime/Engine/Classes/Components/SceneComponent.h
--- a/Engine/Source/Runtime/Engine/Classes/Components/SceneComponent.h
+++ b/Engine/Source/Runtime/Engine/Classes/Components/SceneComponent.h
@@ -931,7 +931,6 @@ public:
 	//~ End ActorComponent Interface
 
 	//~ Begin UObject Interface
-	virtual void PostInterpChange(FProperty* PropertyThatChanged) override;
 	virtual void BeginDestroy() override;
 	virtual bool IsPostLoadThreadSafe() const override;
 	virtual void PreNetReceive() override;
diff --git a/Engine/Source/Runtime/Engine/Classes/Components/SkyAtmosphereComponent.h b/Engine/Source/Runtime/Engine/Classes/Components/SkyAtmosphereComponent.h
--- a/Engine/Source/Runtime/Engine/Classes/Components/SkyAtmosphereComponent.h
+++ b/Engine/Source/Runtime/Engine/Classes/Components/SkyAtmosphereComponent.h
@@ -223,7 +223,6 @@ protected:
 public:
 
 	//~ Begin UObject Interface. 
-	virtual void PostInterpChange(FProperty* PropertyThatChanged) override;
 	virtual void Serialize(FArchive& Ar) override;
 #if WITH_EDITOR
 	virtual void PostEditChangeProperty(FPropertyChangedEvent& PropertyChangedEvent) override;
diff --git a/Engine/Source/Runtime/Engine/Classes/Components/SkyLightComponent.h b/Engine/Source/Runtime/Engine/Classes/Components/SkyLightComponent.h
--- a/Engine/Source/Runtime/Engine/Classes/Components/SkyLightComponent.h
+++ b/Engine/Source/Runtime/Engine/Classes/Components/SkyLightComponent.h
@@ -216,7 +216,6 @@ class ENGINE_API USkyLightComponent : public ULightComponentBase
 	//~ Begin UObject Interface
 	virtual void PostInitProperties() override;
 	virtual void PostLoad() override;
-	virtual void PostInterpChange(FProperty* PropertyThatChanged) override;
 #if WITH_EDITOR
 	virtual void PreEditChange(FProperty* PropertyAboutToChange) override;
 	virtual void PostEditChangeProperty(FPropertyChangedEvent& PropertyChangedEvent) override;
diff --git a/Engine/Source/Runtime/Engine/Classes/Components/VectorFieldComponent.h b/Engine/Source/Runtime/Engine/Classes/Components/VectorFieldComponent.h
--- a/Engine/Source/Runtime/Engine/Classes/Components/VectorFieldComponent.h
+++ b/Engine/Source/Runtime/Engine/Classes/Components/VectorFieldComponent.h
@@ -64,7 +64,6 @@ public:
 	//~ End UActorComponent Interface.
 
 	//~ Begin UObject Interface.
-	virtual void PostInterpChange(FProperty* PropertyThatChanged) override;
 #if WITH_EDITOR
 	virtual void PostEditChangeProperty(FPropertyChangedEvent& PropertyChangedEvent) override;
 #endif // WITH_EDITOR
diff --git a/Engine/Source/Runtime/Engine/Classes/Components/VolumetricCloudComponent.h b/Engine/Source/Runtime/Engine/Classes/Components/VolumetricCloudComponent.h
--- a/Engine/Source/Runtime/Engine/Classes/Components/VolumetricCloudComponent.h
+++ b/Engine/Source/Runtime/Engine/Classes/Components/VolumetricCloudComponent.h
@@ -186,7 +186,6 @@ protected:
 public:
 
 	//~ Begin UObject Interface. 
-	virtual void PostInterpChange(FProperty* PropertyThatChanged) override;
 	virtual void Serialize(FArchive& Ar) override;
 #if WITH_EDITOR
 	virtual void PostEditChangeProperty(FPropertyChangedEvent& PropertyChangedEvent) override;
diff --git a/Engine/Source/Runtime/Engine/Private/Atmosphere/AtmosphericFogComponent.cpp b/Engine/Source/Runtime/Engine/Private/Atmosphere/AtmosphericFogComponent.cpp
--- a/Engine/Source/Runtime/Engine/Private/Atmosphere/AtmosphericFogComponent.cpp
+++ b/Engine/Source/Runtime/Engine/Private/Atmosphere/AtmosphericFogComponent.cpp
@@ -27,13 +27,6 @@ bool UAtmosphericFogComponent::IsPostLoadThreadSafe() const
 	return true;
 }
 
-void UAtmosphericFogComponent::PostInterpChange(FProperty* PropertyThatChanged)
-{
-	Super::PostInterpChange(PropertyThatChanged);
-
-	MarkRenderStateDirty();
-}
-
 UAtmosphericFogComponent::~UAtmosphericFogComponent()
 {
 }
diff --git a/Engine/Source/Runtime/Engine/Private/Components/HeightFogComponent.cpp b/Engine/Source/Runtime/Engine/Private/Components/HeightFogComponent.cpp
--- a/Engine/Source/Runtime/Engine/Private/Components/HeightFogComponent.cpp
+++ b/Engine/Source/Runtime/Engine/Private/Components/HeightFogComponent.cpp
@@ -131,13 +131,6 @@ void UExponentialHeightFogComponent::PostEditChangeProperty(FPropertyChangedEven
 }
 #endif // WITH_EDITOR
 
-void UExponentialHeightFogComponent::PostInterpChange(FProperty* PropertyThatChanged)
-{
-	Super::PostInterpChange(PropertyThatChanged);
-
-	MarkRenderStateDirty();
-}
-
 void UExponentialHeightFogComponent::SetFogDensity(float Value)
 {
 	if(FogDensity != Value)
diff --git a/Engine/Source/Runtime/Engine/Private/Components/LightComponent.cpp b/Engine/Source/Runtime/Engine/Private/Components/LightComponent.cpp
--- a/Engine/Source/Runtime/Engine/Private/Components/LightComponent.cpp
+++ b/Engine/Source/Runtime/Engine/Private/Components/LightComponent.cpp
@@ -1408,42 +1408,6 @@ void ULightComponent::SetMaterial(int32 ElementIndex, UMaterialInterface* InMate
 	}
 }
 
-/** 
-* This is called when property is modified by InterpPropertyTracks
-*
-* @param PropertyThatChanged	Property that changed
-*/
-void ULightComponent::PostInterpChange(FProperty* PropertyThatChanged)
-{
-	static FName LightColorName(TEXT("LightColor"));
-	static FName IntensityName(TEXT("Intensity"));
-	static FName BrightnessName(TEXT("Brightness"));
-	static FName IndirectLightingIntensityName(TEXT("IndirectLightingIntensity"));
-	static FName VolumetricScatteringIntensityName(TEXT("VolumetricScatteringIntensity"));
-	static FName TemperatureName(TEXT("Temperature"));
-
-	FName PropertyName = PropertyThatChanged->GetFName();
-	if (PropertyName == LightColorName
-		|| PropertyName == IntensityName
-		|| PropertyName == BrightnessName
-		|| PropertyName == IndirectLightingIntensityName
-		|| PropertyName == TemperatureName
-		|| PropertyName == VolumetricScatteringIntensityName)
-	{
-		// Old brightness tracks will animate the deprecated value
-		if (PropertyName == BrightnessName)
-		{
-			Intensity = Brightness_DEPRECATED;
-		}
-
-		UpdateColorAndBrightness();
-	}
-	else
-	{
-		Super::PostInterpChange(PropertyThatChanged);
-	}
-}
-
 /** Stores a light and a channel it has been assigned to. */
 struct FLightAndChannel
 {
diff --git a/Engine/Source/Runtime/Engine/Private/Components/LocalLightComponent.cpp b/Engine/Source/Runtime/Engine/Private/Components/LocalLightComponent.cpp
--- a/Engine/Source/Runtime/Engine/Private/Components/LocalLightComponent.cpp
+++ b/Engine/Source/Runtime/Engine/Private/Components/LocalLightComponent.cpp
@@ -129,29 +129,6 @@ void ULocalLightComponent::PostEditChangeProperty(FPropertyChangedEvent& Propert
 }
 #endif // WITH_EDITOR
 
-void ULocalLightComponent::PostInterpChange(FProperty* PropertyThatChanged)
-{
-	static FName RadiusName(TEXT("Radius"));
-	static FName AttenuationRadiusName(TEXT("AttenuationRadius"));
-	FName PropertyName = PropertyThatChanged->GetFName();
-
-	if (PropertyName == RadiusName
-		|| PropertyName == AttenuationRadiusName)
-	{
-		// Old radius tracks will animate the deprecated value
-		if (PropertyName == RadiusName)
-		{
-			AttenuationRadius = Radius_DEPRECATED;
-		}
-
-		PushRadiusToRenderThread();
-	}
-	else
-	{
-		Super::PostInterpChange(PropertyThatChanged);
-	}
-}
-
 void ULocalLightComponent::PushRadiusToRenderThread()
 {
 	if (CastShadows)
diff --git a/Engine/Source/Runtime/Engine/Private/Components/PointLightComponent.cpp b/Engine/Source/Runtime/Engine/Private/Components/PointLightComponent.cpp
--- a/Engine/Source/Runtime/Engine/Private/Components/PointLightComponent.cpp
+++ b/Engine/Source/Runtime/Engine/Private/Components/PointLightComponent.cpp
@@ -333,18 +333,3 @@ void UPointLightComponent::PostEditChangeProperty(FPropertyChangedEvent& Propert
 }
 #endif // WITH_EDITOR
 
-void UPointLightComponent::PostInterpChange(FProperty* PropertyThatChanged)
-{
-	static FName LightFalloffExponentName(TEXT("LightFalloffExponent"));
-	FName PropertyName = PropertyThatChanged->GetFName();
-
-	if (PropertyName == LightFalloffExponentName)
-	{
-		MarkRenderStateDirty();
-	}
-	else
-	{
-		Super::PostInterpChange(PropertyThatChanged);
-	}
-}
-
diff --git a/Engine/Source/Runtime/Engine/Private/Components/SceneComponent.cpp b/Engine/Source/Runtime/Engine/Private/Components/SceneComponent.cpp
--- a/Engine/Source/Runtime/Engine/Private/Components/SceneComponent.cpp
+++ b/Engine/Source/Runtime/Engine/Private/Components/SceneComponent.cpp
@@ -2526,16 +2526,6 @@ void USceneComponent::UpdateChildTransforms(EUpdateTransformFlags UpdateTransfor
 	}
 }
 
-void USceneComponent::PostInterpChange(FProperty* PropertyThatChanged)
-{
-	Super::PostInterpChange(PropertyThatChanged);
-
-	if (PropertyThatChanged->GetFName() == GetRelativeScale3DPropertyName())
-	{
-		UpdateComponentToWorld();
-	}
-}
-
 TArray<FName> USceneComponent::GetAllSocketNames() const
 {
 	TArray<FComponentSocketDescription> SocketList;
diff --git a/Engine/Source/Runtime/Engine/Private/Components/SkyAtmosphereComponent.cpp b/Engine/Source/Runtime/Engine/Private/Components/SkyAtmosphereComponent.cpp
--- a/Engine/Source/Runtime/Engine/Private/Components/SkyAtmosphereComponent.cpp
+++ b/Engine/Source/Runtime/Engine/Private/Components/SkyAtmosphereComponent.cpp
@@ -274,12 +274,6 @@ void USkyAtmosphereComponent::PostEditChangeProperty(FPropertyChangedEvent& Prop
 
 #endif // WITH_EDITOR
 
-void USkyAtmosphereComponent::PostInterpChange(FProperty* PropertyThatChanged)
-{
-	Super::PostInterpChange(PropertyThatChanged);
-	MarkRenderStateDirty();
-}
-
 void USkyAtmosphereComponent::Serialize(FArchive& Ar)
 {
 	Super::Serialize(Ar);
diff --git a/Engine/Source/Runtime/Engine/Private/Components/SkyLightComponent.cpp b/Engine/Source/Runtime/Engine/Private/Components/SkyLightComponent.cpp
--- a/Engine/Source/Runtime/Engine/Private/Components/SkyLightComponent.cpp
+++ b/Engine/Source/Runtime/Engine/Private/Components/SkyLightComponent.cpp
@@ -494,32 +494,6 @@ void USkyLightComponent::UpdateOcclusionRenderingStateFast()
 
 }
 
-/** 
-* This is called when property is modified by InterpPropertyTracks
-*
-* @param PropertyThatChanged	Property that changed
-*/
-void USkyLightComponent::PostInterpChange(FProperty* PropertyThatChanged)
-{
-	static FName LightColorName(TEXT("LightColor"));
-	static FName IntensityName(TEXT("Intensity"));
-	static FName IndirectLightingIntensityName(TEXT("IndirectLightingIntensity"));
-	static FName VolumetricScatteringIntensityName(TEXT("VolumetricScatteringIntensity"));
-
-	FName PropertyName = PropertyThatChanged->GetFName();
-	if (PropertyName == LightColorName
-		|| PropertyName == IntensityName
-		|| PropertyName == IndirectLightingIntensityName
-		|| PropertyName == VolumetricScatteringIntensityName)
-	{
-		UpdateLimitedRenderingStateFast();
-	}
-	else
-	{
-		Super::PostInterpChange(PropertyThatChanged);
-	}
-}
-
 void USkyLightComponent::DestroyRenderState_Concurrent()
 {
 	Super::DestroyRenderState_Concurrent();
diff --git a/Engine/Source/Runtime/Engine/Private/Components/VolumetricCloudComponent.cpp b/Engine/Source/Runtime/Engine/Private/Components/VolumetricCloudComponent.cpp
--- a/Engine/Source/Runtime/Engine/Private/Components/VolumetricCloudComponent.cpp
+++ b/Engine/Source/Runtime/Engine/Private/Components/VolumetricCloudComponent.cpp
@@ -120,12 +120,6 @@ void UVolumetricCloudComponent::PostEditChangeProperty(FPropertyChangedEvent& Pr
 
 #endif // WITH_EDITOR
 
-void UVolumetricCloudComponent::PostInterpChange(FProperty* PropertyThatChanged)
-{
-	// This is called when property is modified by InterpPropertyTracks
-	Super::PostInterpChange(PropertyThatChanged);
-}
-
 void UVolumetricCloudComponent::Serialize(FArchive& Ar)
 {
 	Super::Serialize(Ar);
diff --git a/Engine/Source/Runtime/Engine/Private/VectorField.cpp b/Engine/Source/Runtime/Engine/Private/VectorField.cpp
--- a/Engine/Source/Runtime/Engine/Private/VectorField.cpp
+++ b/Engine/Source/Runtime/Engine/Private/VectorField.cpp
@@ -904,19 +904,6 @@ void UVectorFieldComponent::SetIntensity(float NewIntensity)
 }
 
 
-void UVectorFieldComponent::PostInterpChange(FProperty* PropertyThatChanged)
-{
-	static const FName IntensityPropertyName(TEXT("Intensity"));
-
-	if (FXSystem && PropertyThatChanged
-		&& PropertyThatChanged->GetFName() == IntensityPropertyName)
-	{
-		FXSystem->UpdateVectorField(this);
-	}
-
-	Super::PostInterpChange(PropertyThatChanged);
-}
-
 #if WITH_EDITOR
 void UVectorFieldComponent::PostEditChangeProperty(FPropertyChangedEvent& PropertyChangedEvent)
 {
diff --git a/Engine/Source/Runtime/Renderer/Private/SparseVolumeTexture/SparseVolumeTextureViewerComponent.cpp b/Engine/Source/Runtime/Renderer/Private/SparseVolumeTexture/SparseVolumeTextureViewerComponent.cpp
--- a/Engine/Source/Runtime/Renderer/Private/SparseVolumeTexture/SparseVolumeTextureViewerComponent.cpp
+++ b/Engine/Source/Runtime/Renderer/Private/SparseVolumeTexture/SparseVolumeTextureViewerComponent.cpp
@@ -64,12 +64,6 @@ void USparseVolumeTextureViewerComponent::PostEditChangeProperty(FPropertyChange
 
 #endif // WITH_EDITOR
 
-void USparseVolumeTextureViewerComponent::PostInterpChange(FProperty* PropertyThatChanged)
-{
-	// This is called when property is modified by InterpPropertyTracks
-	Super::PostInterpChange(PropertyThatChanged);
-}
-
 void USparseVolumeTextureViewerComponent::Serialize(FArchive& Ar)
 {
 	Super::Serialize(Ar);
diff --git a/Engine/Source/Runtime/Renderer/Private/SparseVolumeTexture/SparseVolumeTextureViewerComponent.h b/Engine/Source/Runtime/Renderer/Private/SparseVolumeTexture/SparseVolumeTextureViewerComponent.h
--- a/Engine/Source/Runtime/Renderer/Private/SparseVolumeTexture/SparseVolumeTextureViewerComponent.h
+++ b/Engine/Source/Runtime/Renderer/Private/SparseVolumeTexture/SparseVolumeTextureViewerComponent.h
@@ -73,7 +73,6 @@ protected:
 public:
 
 	//~ Begin UObject Interface. 
-	virtual void PostInterpChange(FProperty* PropertyThatChanged) override;
 	virtual void Serialize(FArchive& Ar) override;
 #if WITH_EDITOR
 	virtual void PostEditChangeProperty(FPropertyChangedEvent& PropertyChangedEvent) override;
