[EngineSharedPCH] Sorted includes and updated to include SlateSharedPCH.h to make it easier to manage.

diff --git a/Engine/Source/Runtime/Engine/Public/EngineSharedPCH.h b/Engine/Source/Runtime/Engine/Public/EngineSharedPCH.h
--- a/Engine/Source/Runtime/Engine/Public/EngineSharedPCH.h
+++ b/Engine/Source/Runtime/Engine/Public/EngineSharedPCH.h
@@ -2,430 +2,68 @@
 
 #pragma once
 
+#include "SlateSharedPCH.h"
+
 // From Core:
-#include "Containers/ContainersFwd.h"
-#include "CoreTypes.h"
-#include "Misc/Exec.h"
-#include "Misc/AssertionMacros.h"
-#include "HAL/PlatformMisc.h"
-#include "GenericPlatform/GenericPlatformMisc.h"
-#include "CoreFwd.h"
-#include "UObject/UObjectHierarchyFwd.h"
-#include "HAL/PlatformCrt.h"
-#include "Containers/Array.h"
-#include "HAL/UnrealMemory.h"
-#include "Templates/IsPointer.h"
-#include "HAL/PlatformMemory.h"
-#include "GenericPlatform/GenericPlatformMemory.h"
-#include "HAL/MemoryBase.h"
-#include "Misc/OutputDevice.h"
-#include "Logging/LogVerbosity.h"
-#include "Misc/VarArgs.h"
-#include "HAL/PlatformAtomics.h"
-#include "GenericPlatform/GenericPlatformAtomics.h"
-#include "Templates/UnrealTypeTraits.h"
-#include "Templates/AndOrNot.h"
-#include "Templates/IsArithmetic.h"
-#include "Templates/IsPODType.h"
-#include "Templates/IsTriviallyCopyConstructible.h"
-#include "Templates/UnrealTemplate.h"
-#include "Templates/EnableIf.h"
-#include "Templates/RemoveReference.h"
-#include "Templates/TypeCompatibleBytes.h"
-#include "Templates/ChooseClass.h"
-#include "Templates/IntegralConstant.h"
-#include "Templates/IsClass.h"
-#include "Traits/IsContiguousContainer.h"
-#include "Containers/ContainerAllocationPolicies.h"
-#include "HAL/PlatformMath.h"
-#include "GenericPlatform/GenericPlatformMath.h"
-#include "Templates/MemoryOps.h"
-#include "Templates/IsTriviallyCopyAssignable.h"
-#include "Templates/IsTriviallyDestructible.h"
-#include "Math/NumericLimits.h"
-#include "Serialization/Archive.h"
-#include "Templates/IsEnumClass.h"
-#include "HAL/PlatformProperties.h"
-#include "GenericPlatform/GenericPlatformProperties.h"
-#include "Misc/Compression.h"
-#include "Misc/EngineVersionBase.h"
-#include "Internationalization/TextNamespaceFwd.h"
-#include "Templates/Less.h"
-#include "Templates/Sorting.h"
-#include "Containers/UnrealString.h"
-#include "Misc/CString.h"
-#include "Misc/Char.h"
-#include "HAL/PlatformString.h"
-#include "GenericPlatform/GenericPlatformStricmp.h"
-#include "GenericPlatform/GenericPlatformString.h"
-#include "Misc/Crc.h"
-#include "Math/UnrealMathUtility.h"
-#include "Containers/Map.h"
-#include "Misc/StructBuilder.h"
-#include "Templates/AlignmentTemplates.h"
-#include "Templates/Function.h"
-#include "Templates/Decay.h"
-#include "Templates/Invoke.h"
-#include "Templates/PointerIsConvertibleFromTo.h"
-#include "Containers/Set.h"
-#include "Templates/TypeHash.h"
-#include "Containers/SparseArray.h"
-#include "Containers/ScriptArray.h"
-#include "Containers/BitArray.h"
-#include "Algo/Reverse.h"
-#include "Math/IntPoint.h"
-#include "Logging/LogMacros.h"
-#include "Logging/LogCategory.h"
-#include "UObject/NameTypes.h"
-#include "HAL/CriticalSection.h"
-#include "Misc/Timespan.h"
-#include "Containers/StringConv.h"
-#include "UObject/UnrealNames.h"
-#include "Templates/SharedPointer.h"
-#include "CoreGlobals.h"
-#include "HAL/PlatformTLS.h"
-#include "GenericPlatform/GenericPlatformTLS.h"
-#include "Delegates/Delegate.h"
-#include "UObject/WeakObjectPtrTemplates.h"
-#include "Delegates/MulticastDelegateBase.h"
-#include "Delegates/IDelegateInstance.h"
-#include "Delegates/DelegateSettings.h"
-#include "Delegates/DelegateBase.h"
-#include "Delegates/IntegerSequence.h"
-#include "Templates/Tuple.h"
-#include "UObject/ScriptDelegates.h"
-#include "Misc/Optional.h"
-#include "Misc/Parse.h"
-#include "Misc/DateTime.h"
-#include "CoreMinimal.h"
-#include "Math/UnrealMath.h"
-#include "Templates/IsFloatingPoint.h"
-#include "Templates/IsIntegral.h"
-#include "Templates/IsSigned.h"
-#include "Templates/Greater.h"
-#include "Math/ColorList.h"
-#include "Math/Color.h"
-#include "Math/IntRect.h"
-#include "Math/Vector2D.h"
-#include "Math/Edge.h"
-#include "Math/Vector.h"
-#include "Misc/ByteSwap.h"
-#include "Internationalization/Text.h"
-#include "Containers/EnumAsByte.h"
-#include "Internationalization/CulturePointer.h"
-#include "Internationalization/TextLocalizationManager.h"
-#include "Templates/UniquePtr.h"
-#include "Templates/IsArray.h"
-#include "Templates/RemoveExtent.h"
-#include "Internationalization/Internationalization.h"
-#include "Templates/UniqueObj.h"
-#include "Math/IntVector.h"
-#include "Math/CapsuleShape.h"
-#include "Math/RangeSet.h"
-#include "Math/Range.h"
-#include "Math/RangeBound.h"
-#include "Math/Box2D.h"
-#include "Math/BoxSphereBounds.h"
-#include "Math/Sphere.h"
-#include "Math/Box.h"
-#include "Math/OrientedBox.h"
-#include "Math/Interval.h"
-#include "Math/RotationAboutPointMatrix.h"
-#include "Math/Matrix.h"
-#include "Math/Vector4.h"
-#include "Math/Plane.h"
-#include "UObject/ObjectVersion.h"
-#include "Math/Rotator.h"
-#include "Math/VectorRegister.h"
-#include "Math/Axis.h"
-#include "Math/RotationTranslationMatrix.h"
-#include "Math/ScaleRotationTranslationMatrix.h"
-#include "Math/RotationMatrix.h"
-#include "Math/PerspectiveMatrix.h"
-#include "Math/OrthoMatrix.h"
-#include "Math/TranslationMatrix.h"
-#include "Math/QuatRotationTranslationMatrix.h"
-#include "Math/Quat.h"
-#include "Math/InverseRotationMatrix.h"
-#include "Math/ScaleMatrix.h"
-#include "Math/MirrorMatrix.h"
-#include "Math/ClipProjectionMatrix.h"
-#include "Math/InterpCurve.h"
-#include "Math/TwoVectors.h"
-#include "Math/InterpCurvePoint.h"
-#include "Math/CurveEdInterface.h"
-#include "Math/Float16Color.h"
-#include "Math/Float16.h"
-#include "Math/Float32.h"
-#include "Math/Vector2DHalf.h"
-#include "Math/Transform.h"
-#include "Math/ConvexHull2d.h"
-#include "HAL/ThreadSingleton.h"
-#include "HAL/TlsAutoCleanup.h"
-#include "HAL/ThreadSafeCounter.h"
-#include "HAL/PlatformTime.h"
-#include "GenericPlatform/GenericPlatformTime.h"
-#include "Stats/Stats.h"
-#include "Containers/LockFreeList.h"
-#include "Misc/NoopCounter.h"
-#include "Containers/ChunkedArray.h"
-#include "Containers/IndirectArray.h"
-#include "Misc/EnumClassFlags.h"
-#include "Misc/Guid.h"
-#include "Misc/OutputDeviceRedirector.h"
-#include "ProfilingDebugging/ResourceSize.h"
-#include "Math/RandomStream.h"
-#include "Misc/CoreMisc.h"
-#include "Containers/List.h"
-#include "Modules/ModuleInterface.h"
-#include "Misc/Attribute.h"
-#include "HAL/IConsoleManager.h"
-#include "Serialization/BitReader.h"
-#include "Misc/NetworkGuid.h"
-#include "Serialization/BitWriter.h"
-#include "Containers/LockFreeFixedSizeAllocator.h"
-#include "Misc/MemStack.h"
-#include "HAL/Event.h"
-#include "Templates/RefCounting.h"
-#include "Async/TaskGraphInterfaces.h"
-#include "GenericPlatform/ICursor.h"
-#include "GenericPlatform/GenericApplicationMessageHandler.h"
-#include "GenericPlatform/GenericApplication.h"
-#include "GenericPlatform/GenericWindowDefinition.h"
-#include "GenericPlatform/GenericWindow.h"
-#include "Math/TransformCalculus.h"
-#include "Math/TransformCalculus2D.h"
-#include "Templates/ScopedCallback.h"
-#include "HAL/PlatformProcess.h"
-#include "GenericPlatform/GenericPlatformProcess.h"
-#include "Misc/CoreStats.h"
-#include "GenericPlatform/GenericPlatformAffinity.h"
-#include "Misc/IQueuedWork.h"
-#include "Misc/QueuedThreadPool.h"
-#include "Async/AsyncWork.h"
-#include "Serialization/BufferReader.h"
-#include "Containers/StaticArray.h"
-#include "Misc/SecureHash.h"
-#include "Misc/Paths.h"
-#include "Modules/ModuleManager.h"
-#include "Modules/Boilerplate/ModuleBoilerplate.h"
 #include "Async/Future.h"
-#include "Serialization/ArchiveProxy.h"
-#include "Math/SHMath.h"
-#include "Misc/ScopedEvent.h"
-#include "UObject/DebugSerializationFlags.h"
-#include "Misc/BufferedOutputDevice.h"
-#include "Containers/ResourceArray.h"
+#include "Async/ParallelFor.h"
 #include "Containers/DynamicRHIResourceArray.h"
-#include "Serialization/CustomVersion.h"
-#include "Misc/CommandLine.h"
-#include "Misc/OutputDeviceError.h"
-#include "Misc/ObjectThumbnail.h"
-#include "Misc/App.h"
-#include "Misc/ConfigCacheIni.h"
-#include "Misc/EngineVersion.h"
-#include "Internationalization/GatherableTextData.h"
-#include "Internationalization/InternationalizationMetadata.h"
-#include "Misc/ScopeLock.h"
-#include "Containers/Queue.h"
-#include "Features/IModularFeature.h"
-#include "HAL/ThreadSafeBool.h"
+#include "Containers/ResourceArray.h"
 #include "Features/IModularFeatures.h"
-#include "Serialization/MemoryWriter.h"
-#include "Serialization/MemoryArchive.h"
-#include "Logging/TokenizedMessage.h"
+#include "GenericPlatform/ICursor.h"
 #include "GenericPlatform/IInputInterface.h"
-#include "Containers/ArrayView.h"
-#include "GenericPlatform/GenericPlatformFile.h"
+#include "HAL/PlatformAffinity.h"
+#include "HAL/RunnableThread.h"
+#include "HAL/ThreadSafeBool.h"
+#include "Internationalization/GatherableTextData.h"
+#include "Internationalization/InternationalizationMetadata.h"
+#include "Math/SHMath.h"
+#include "Misc/AutomationTest.h"
+#include "Misc/BufferedOutputDevice.h"
 #include "Misc/CompilationResult.h"
-#include "HAL/FileManager.h"
+#include "Misc/NetworkGuid.h"
+#include "Misc/ScopedEvent.h"
+#include "Misc/SingleThreadRunnable.h"
 #include "ProfilingDebugging/Histogram.h"
-#include "Misc/FeedbackContext.h"
-#include "Misc/SlowTask.h"
-#include "Misc/SlowTaskStack.h"
-#include "HAL/Runnable.h"
-#include "Containers/Ticker.h"
-#include "Misc/AutomationTest.h"
+#include "ProfilingDebugging/ProfilingHelpers.h"
+#include "Serialization/ArchiveProxy.h"
+#include "Serialization/BitReader.h"
+#include "Serialization/BitWriter.h"
+#include "Stats/StatsMisc.h"
+#include "Templates/ScopedCallback.h"
 #include "Templates/ValueOrError.h"
+#include "UObject/DebugSerializationFlags.h"
 #include "UObject/PropertyPortFlags.h"
-#include "Stats/StatsMisc.h"
-#include "Misc/FileHelper.h"
-#include "Async/ParallelFor.h"
-#include "Misc/SingleThreadRunnable.h"
-#include "ProfilingDebugging/ProfilingHelpers.h"
-#include "HAL/RunnableThread.h"
-#include "HAL/PlatformAffinity.h"
 
 // From Json:
-#include "Serialization/JsonTypes.h"
-#include "Policies/PrettyJsonPrintPolicy.h"
-#include "Policies/JsonPrintPolicy.h"
-#include "Serialization/JsonWriter.h"
 #include "Dom/JsonObject.h"
-#include "JsonGlobals.h"
 #include "Dom/JsonValue.h"
+#include "JsonGlobals.h"
+#include "Policies/CondensedJsonPrintPolicy.h"
 #include "Serialization/JsonReader.h"
 #include "Serialization/JsonSerializer.h"
 #include "Serialization/JsonSerializerMacros.h"
-#include "Policies/CondensedJsonPrintPolicy.h"
 
 // From CoreUObject:
-#include "UObject/ObjectMacros.h"
-#include "UObject/Script.h"
-#include "UObject/UObjectBase.h"
-#include "UObject/UObjectGlobals.h"
-#include "UObject/UObjectArray.h"
-#include "UObject/Object.h"
-#include "UObject/UObjectBaseUtility.h"
-#include "UObject/UObjectMarks.h"
-#include "UObject/Class.h"
-#include "UObject/GarbageCollection.h"
-#include "UObject/GeneratedCppIncludes.h"
-#include "Serialization/ArchiveUObject.h"
-#include "UObject/CoreNative.h"
-#include "UObject/WeakObjectPtr.h"
+#include "Online/CoreOnline.h"
+#include "Serialization/BulkData.h"
 #include "Templates/Casts.h"
-#include "UObject/CoreNetTypes.h"
-#include "UObject/LazyObjectPtr.h"
-#include "UObject/PersistentObjectPtr.h"
-#include "UObject/SoftObjectPtr.h"
 #include "Templates/SubclassOf.h"
-#include "UObject/ScriptInterface.h"
-#include "UObject/UnrealType.h"
-#include "UObject/PropertyTag.h"
-#include "Serialization/SerializedPropertyScope.h"
 #include "UObject/CoreNet.h"
-#include "UObject/Stack.h"
-#include "UObject/ScriptMacros.h"
-#include "UObject/Interface.h"
-#include "Misc/WorldCompositionUtility.h"
-#include "UObject/GCObject.h"
-#include "Serialization/BulkData.h"
-#include "UObject/Package.h"
-#include "Online/CoreOnline.h"
+#include "UObject/Linker.h"
 #include "UObject/LinkerLoad.h"
+#include "UObject/ObjectRedirector.h"
 #include "UObject/ObjectResource.h"
-#include "UObject/Linker.h"
 #include "UObject/PackageFileSummary.h"
-#include "Misc/PackageName.h"
-#include "UObject/ObjectRedirector.h"
+#include "UObject/PersistentObjectPtr.h"
+#include "UObject/ScriptMacros.h"
+#include "UObject/Stack.h"
+#include "UObject/StructOnScope.h"
 #include "UObject/TextProperty.h"
-#include "UObject/UObjectIterator.h"
 #include "UObject/UObjectAnnotation.h"
 #include "UObject/UObjectHash.h"
+#include "UObject/UObjectIterator.h"
 #include "UObject/UObjectThreadContext.h"
-#include "UObject/StructOnScope.h"
-
-// From InputCore:
-#include "InputCoreTypes.h"
-
-// From SlateCore:
-#include "Styling/SlateColor.h"
-#include "Styling/WidgetStyle.h"
-#include "Types/SlateEnums.h"
-#include "Layout/Visibility.h"
-#include "Rendering/SlateRenderTransform.h"
-#include "Layout/SlateRect.h"
-#include "Rendering/SlateLayoutTransform.h"
-#include "Layout/Geometry.h"
-#include "Layout/PaintGeometry.h"
-#include "Input/Events.h"
-#include "Input/ReplyBase.h"
-#include "Input/CursorReply.h"
-#include "Input/Reply.h"
-#include "Input/DragAndDrop.h"
-#include "Input/PopupMethodReply.h"
-#include "Input/NavigationReply.h"
-#include "Widgets/SWidget.h"
-#include "Types/ISlateMetaData.h"
-#include "Layout/ArrangedWidget.h"
-#include "Types/WidgetActiveTimerDelegate.h"
-#include "Layout/LayoutGeometry.h"
-#include "Layout/Margin.h"
-#include "Widgets/SNullWidget.h"
-#include "Widgets/DeclarativeSyntaxSupport.h"
-#include "SlotBase.h"
-#include "Layout/Children.h"
-#include "Widgets/SPanel.h"
-#include "Widgets/SOverlay.h"
-#include "Styling/SlateBrush.h"
-#include "Types/SlateStructs.h"
-#include "Fonts/CompositeFont.h"
-#include "Fonts/SlateFontInfo.h"
-#include "Sound/SlateSound.h"
-#include "Brushes/SlateNoResource.h"
-#include "Styling/ISlateStyle.h"
-#include "Styling/StyleDefaults.h"
-#include "Styling/CoreStyle.h"
-#include "Styling/SlateWidgetStyle.h"
-#include "Styling/SlateTypes.h"
-#include "Widgets/SCompoundWidget.h"
-#include "Styling/SlateWidgetStyleAsset.h"
-#include "Styling/SlateWidgetStyleContainerBase.h"
-#include "Styling/SlateWidgetStyleContainerInterface.h"
-#include "Rendering/RenderingCommon.h"
-#include "Widgets/SBoxPanel.h"
-#include "Textures/SlateIcon.h"
-#include "Animation/CurveSequence.h"
-#include "Animation/CurveHandle.h"
-#include "Widgets/SWindow.h"
-#include "SlateGlobals.h"
-#include "Textures/SlateShaderResource.h"
-#include "Fonts/ShapedTextFwd.h"
-#include "Types/PaintArgs.h"
-#include "Textures/SlateTextureData.h"
-#include "Rendering/DrawElements.h"
-#include "Rendering/ShaderResourceManager.h"
-#include "Widgets/SLeafWidget.h"
-#include "Brushes/SlateDynamicImageBrush.h"
-#include "Rendering/SlateRenderer.h"
-#include "Widgets/IToolTip.h"
-#include "Layout/ArrangedChildren.h"
-#include "Application/SlateWindowHelper.h"
-#include "Application/SlateApplicationBase.h"
-#include "Application/ThrottleManager.h"
-#include "Layout/WidgetPath.h"
-#include "Types/SlateConstants.h"
-
-// From Slate:
-#include "SlateFwd.h"
-#include "Framework/Commands/InputChord.h"
-#include "Framework/SlateDelegates.h"
-#include "Widgets/Input/IVirtualKeyboardEntry.h"
-#include "Widgets/Layout/SBorder.h"
-#include "Framework/Text/TextLayout.h"
-#include "Framework/Text/TextRunRenderer.h"
-#include "Framework/Text/TextLineHighlight.h"
-#include "Framework/Text/IRun.h"
-#include "Framework/Text/ShapedTextCacheFwd.h"
-#include "Widgets/Text/STextBlock.h"
-#include "Framework/Application/IMenu.h"
-#include "Framework/Commands/UIAction.h"
-#include "Framework/Commands/UICommandInfo.h"
-#include "Framework/MultiBox/MultiBoxExtender.h"
-#include "Framework/Commands/UICommandList.h"
-#include "Framework/Application/SlateApplication.h"
-#include "Framework/Application/MenuStack.h"
-#include "Widgets/Layout/SBox.h"
-#include "Widgets/Layout/SScrollBar.h"
-#include "Framework/Layout/InertialScrollManager.h"
-#include "Framework/Layout/Overscroll.h"
-#include "Framework/MultiBox/MultiBoxBuilder.h"
-#include "Framework/MultiBox/MultiBoxDefs.h"
-#include "Framework/Docking/TabManager.h"
-#include "Framework/Docking/WorkspaceItem.h"
-#include "Widgets/Views/STableViewBase.h"
-#include "Framework/Layout/IScrollableWidget.h"
-#include "Framework/Views/ITypedTableView.h"
-#include "Widgets/Views/STableRow.h"
-#include "Widgets/Views/SExpanderArrow.h"
-#include "Widgets/Views/SHeaderRow.h"
-#include "Framework/Views/TableViewTypeTraits.h"
-#include "Widgets/Views/SListView.h"
-#include "Framework/Commands/Commands.h"
-#include "Framework/Commands/InputBindingManager.h"
-#include "Widgets/Input/SButton.h"
-#include "Widgets/Input/SCheckBox.h"
 
 // From RHI:
 #include "RHIDefinitions.h"
@@ -455,173 +93,171 @@
 #include "PacketHandler.h"
 
 // From Engine:
-#include "Engine/EngineTypes.h"
-#include "Engine/NetSerialization.h"
-#include "EngineLogs.h"
-#include "EdGraph/EdGraphNode.h"
-#include "EdGraph/EdGraphPin.h"
-#include "Engine/EngineBaseTypes.h"
-#include "Interfaces/Interface_AssetUserData.h"
-#include "Components/ActorComponent.h"
+#include "AI/Navigation/NavAgentInterface.h"
+#include "AI/Navigation/NavigationTypes.h"
+#include "AI/Navigation/NavQueryFilter.h"
+#include "AI/Navigation/NavRelevantInterface.h"
+#include "AlphaBlend.h"
+#include "Animation/AnimationAsset.h"
+#include "Animation/AnimBlueprint.h"
+#include "Animation/AnimBlueprintGeneratedClass.h"
+#include "Animation/AnimClassInterface.h"
+#include "Animation/AnimCompositeBase.h"
+#include "Animation/AnimCurveTypes.h"
+#include "Animation/AnimInstance.h"
+#include "Animation/AnimLinkableElement.h"
+#include "Animation/AnimMontage.h"
+#include "Animation/AnimNodeBase.h"
+#include "Animation/AnimNotifyQueue.h"
+#include "Animation/AnimSequence.h"
+#include "Animation/AnimSequenceBase.h"
+#include "Animation/AnimStateMachineTypes.h"
+#include "Animation/AnimStats.h"
+#include "Animation/AnimTypes.h"
+#include "Animation/PreviewAssetAttachComponent.h"
+#include "Animation/Skeleton.h"
+#include "Animation/SmartName.h"
+#include "AnimInterpFilter.h"
+#include "Audio.h"
+#include "BatchedElements.h"
+#include "BlendableManager.h"
+#include "BlueprintUtilities.h"
+#include "BoneContainer.h"
+#include "BonePose.h"
+#include "Camera/CameraShakeBase.h"
+#include "Camera/CameraTypes.h"
+#include "Camera/PlayerCameraManager.h"
+#include "ClothSimData.h"
+#include "CollisionQueryParams.h"
 #include "ComponentInstanceDataCache.h"
+#include "Components.h"
+#include "Components/ActorComponent.h"
+#include "Components/InputComponent.h"
+#include "Components/MeshComponent.h"
+#include "Components/PrimitiveComponent.h"
 #include "Components/SceneComponent.h"
-#include "GameFramework/Actor.h"
-#include "Engine/Level.h"
-#include "Engine/MaterialMerging.h"
+#include "Components/SkeletalMeshComponent.h"
+#include "Components/SkinnedMeshComponent.h"
+#include "Components/StaticMeshComponent.h"
+#include "ConvexVolume.h"
+#include "Curves/CurveBase.h"
+#include "Curves/CurveFloat.h"
+#include "Curves/CurveOwnerInterface.h"
+#include "Curves/IndexedCurve.h"
+#include "Curves/KeyHandle.h"
+#include "Curves/RichCurve.h"
+#include "CustomBoneIndexArray.h"
+#include "DebugViewModeHelpers.h"
+#include "EdGraph/EdGraph.h"
+#include "EdGraph/EdGraphNode.h"
+#include "EdGraph/EdGraphPin.h"
+#include "EdGraph/EdGraphSchema.h"
+#include "Engine/BlendableInterface.h"
 #include "Engine/Blueprint.h"
 #include "Engine/BlueprintCore.h"
-#include "EngineDefines.h"
-#include "AI/Navigation/NavigationTypes.h"
-#include "AI/Navigation/NavQueryFilter.h"
-#include "CollisionQueryParams.h"
-#include "Engine/LatentActionManager.h"
-#include "AI/Navigation/NavAgentInterface.h"
-#include "GameFramework/Pawn.h"
-#include "Engine/PendingNetGame.h"
-#include "WorldCollision.h"
-#include "Engine/World.h"
+#include "Engine/BlueprintGeneratedClass.h"
+#include "Engine/Brush.h"
+#include "Engine/Channel.h"
+#include "Engine/ChildConnection.h"
+#include "Engine/DebugDisplayProperty.h"
+#include "Engine/DeveloperSettings.h"
+#include "Engine/Engine.h"
+#include "Engine/EngineBaseTypes.h"
+#include "Engine/EngineTypes.h"
 #include "Engine/GameInstance.h"
-#include "PixelFormat.h"
-#include "SceneTypes.h"
-#include "Engine/BlendableInterface.h"
-#include "HitProxies.h"
-#include "Engine/Scene.h"
-#include "UnrealClient.h"
-#include "Curves/KeyHandle.h"
-#include "ShowFlags.h"
 #include "Engine/GameViewportClient.h"
-#include "Engine/ScriptViewportClient.h"
-#include "Engine/ViewportSplitScreen.h"
-#include "Engine/TitleSafeZone.h"
 #include "Engine/GameViewportDelegates.h"
-#include "Engine/DebugDisplayProperty.h"
-#include "Curves/IndexedCurve.h"
-#include "ConvexVolume.h"
-#include "SceneView.h"
-#include "SceneInterface.h"
-#include "FinalPostProcessSettings.h"
-#include "BlendableManager.h"
-#include "DebugViewModeHelpers.h"
-#include "Curves/RichCurve.h"
-#include "Engine/Engine.h"
-#include "PrimitiveUniformShaderParameters.h"
-#include "PrimitiveViewRelevance.h"
-#include "PrimitiveSceneProxy.h"
-#include "Materials/MaterialInterface.h"
-#include "MaterialShared.h"
-#include "StaticParameterSet.h"
-#include "Components.h"
+#include "Engine/LatentActionManager.h"
+#include "Engine/Level.h"
+#include "Engine/LocalPlayer.h"
+#include "Engine/MaterialMerging.h"
+#include "Engine/MemberReference.h"
 #include "Engine/MeshMerging.h"
-#include "Engine/TextureDefines.h"
+#include "Engine/NetConnection.h"
+#include "Engine/NetDriver.h"
+#include "Engine/NetSerialization.h"
+#include "Engine/PendingNetGame.h"
+#include "Engine/Player.h"
+#include "Engine/PoseWatch.h"
+#include "Engine/Scene.h"
+#include "Engine/ScriptViewportClient.h"
+#include "Engine/SkeletalMesh.h"
+#include "Engine/StaticMesh.h"
 #include "Engine/Texture.h"
-#include "Interfaces/Interface_CollisionDataProvider.h"
-#include "TextureResource.h"
 #include "Engine/Texture2D.h"
-#include "Engine/Brush.h"
-#include "StaticBoundShaderState.h"
-#include "Curves/CurveFloat.h"
-#include "Curves/CurveBase.h"
-#include "Curves/CurveOwnerInterface.h"
-#include "SceneUtils.h"
-#include "MeshBatch.h"
-#include "BatchedElements.h"
-#include "EngineGlobals.h"
-#include "SceneManagement.h"
+#include "Engine/TextureDefines.h"
 #include "Engine/TextureLightProfile.h"
-#include "PhysicsEngine/BodyInstance.h"
-#include "PhysxUserData.h"
 #include "Engine/TextureStreamingTypes.h"
-#include "Components/PrimitiveComponent.h"
-#include "AI/Navigation/NavRelevantInterface.h"
-#include "RawIndexBuffer.h"
-#include "LocalVertexFactory.h"
-#include "Engine/StaticMesh.h"
-#include "ReferenceSkeleton.h"
-#include "GPUSkinPublicDefs.h"
-#include "Components/MeshComponent.h"
-#include "Animation/AnimTypes.h"
-#include "Animation/AnimLinkableElement.h"
-#include "BoneContainer.h"
-#include "Animation/PreviewAssetAttachComponent.h"
-#include "Engine/SkeletalMesh.h"
-#include "Animation/AnimationAsset.h"
-#include "AnimInterpFilter.h"
-#include "EdGraph/EdGraphSchema.h"
-#include "Animation/Skeleton.h"
-#include "Animation/SmartName.h"
-#include "Animation/AnimCurveTypes.h"
-#include "Animation/AnimSequenceBase.h"
-#include "Engine/BlueprintGeneratedClass.h"
-#include "Animation/AnimSequence.h"
-#include "BonePose.h"
-#include "CustomBoneIndexArray.h"
-#include "Animation/AnimStats.h"
-#include "Math/GenericOctreePublic.h"
+#include "Engine/TitleSafeZone.h"
+#include "Engine/ViewportSplitScreen.h"
+#include "Engine/World.h"
+#include "EngineDefines.h"
+#include "EngineGlobals.h"
+#include "EngineLogs.h"
+#include "EngineStats.h"
+#include "FinalPostProcessSettings.h"
+#include "GameFramework/Actor.h"
+#include "GameFramework/Controller.h"
+#include "GameFramework/DamageType.h"
+#include "GameFramework/ForceFeedbackEffect.h"
+#include "GameFramework/Info.h"
 #include "GameFramework/OnlineReplStructs.h"
-#include "Audio.h"
-#include "Sound/SoundAttenuation.h"
+#include "GameFramework/Pawn.h"
+#include "GameFramework/PlayerController.h"
+#include "GameFramework/PlayerMuteList.h"
+#include "GameFramework/Volume.h"
+#include "GameFramework/WorldSettings.h"
+#include "GlobalShader.h"
+#include "GPUSkinPublicDefs.h"
+#include "HitProxies.h"
 #include "IAudioExtensionPlugin.h"
-#include "Tickable.h"
-#include "TimerManager.h"
+#include "Interfaces/Interface_AssetUserData.h"
+#include "Interfaces/Interface_CollisionDataProvider.h"
+#include "Kismet/BlueprintFunctionLibrary.h"
+#include "LocalVertexFactory.h"
 #include "MaterialExpressionIO.h"
+#include "Materials/Material.h"
 #include "Materials/MaterialExpression.h"
-#include "AlphaBlend.h"
-#include "BlueprintUtilities.h"
-#include "EdGraph/EdGraph.h"
 #include "Materials/MaterialExpressionMaterialFunctionCall.h"
-#include "Materials/Material.h"
 #include "Materials/MaterialFunction.h"
-#include "Animation/AnimBlueprint.h"
-#include "Animation/AnimBlueprintGeneratedClass.h"
-#include "Engine/PoseWatch.h"
-#include "Animation/AnimStateMachineTypes.h"
-#include "Animation/AnimClassInterface.h"
-#include "Model.h"
-#include "GlobalShader.h"
-#include "Camera/CameraTypes.h"
-#include "Engine/MemberReference.h"
-#include "GameFramework/Volume.h"
-#include "GameFramework/Info.h"
-#include "Sound/AudioVolume.h"
 #include "Materials/MaterialInstance.h"
 #include "Materials/MaterialInstanceBasePropertyOverrides.h"
-#include "UnrealEngine.h"
-#include "GameFramework/Controller.h"
-#include "Camera/CameraShakeBase.h"
-#include "GameFramework/WorldSettings.h"
-#include "GameFramework/DamageType.h"
-#include "Components/InputComponent.h"
-#include "GameFramework/PlayerController.h"
-#include "GameFramework/PlayerMuteList.h"
-#include "Camera/PlayerCameraManager.h"
-#include "GameFramework/ForceFeedbackEffect.h"
 #include "Materials/MaterialInstanceDynamic.h"
-#include "Engine/Player.h"
-#include "Animation/AnimNodeBase.h"
-#include "Kismet/BlueprintFunctionLibrary.h"
-#include "EngineStats.h"
+#include "Materials/MaterialInterface.h"
+#include "MaterialShaderType.h"
+#include "MaterialShared.h"
 #include "Math/GenericOctree.h"
-#include "Components/SkinnedMeshComponent.h"
-#include "Engine/NetDriver.h"
-#include "Components/StaticMeshComponent.h"
-#include "Engine/Channel.h"
-#include "Engine/NetConnection.h"
+#include "Math/GenericOctreePublic.h"
+#include "MeshBatch.h"
+#include "MeshMaterialShaderType.h"
+#include "Model.h"
 #include "Net/DataBunch.h"
-#include "Components/SkeletalMeshComponent.h"
-#include "ClothSimData.h"
+#include "PhysicsEngine/BodyInstance.h"
+#include "PhysxUserData.h"
+#include "PixelFormat.h"
+#include "PrimitiveSceneProxy.h"
+#include "PrimitiveUniformShaderParameters.h"
+#include "PrimitiveViewRelevance.h"
+#include "RawIndexBuffer.h"
+#include "ReferenceSkeleton.h"
+#include "SceneInterface.h"
+#include "SceneManagement.h"
+#include "SceneTypes.h"
+#include "SceneUtils.h"
+#include "SceneView.h"
+#include "ShowFlags.h"
 #include "SingleAnimationPlayData.h"
-#include "Animation/AnimCompositeBase.h"
-#include "MaterialShaderType.h"
-#include "Animation/AnimMontage.h"
-#include "Animation/AnimInstance.h"
-#include "Animation/AnimNotifyQueue.h"
-#include "MeshMaterialShaderType.h"
-#include "Engine/LocalPlayer.h"
-#include "Engine/ChildConnection.h"
-#include "Engine/DeveloperSettings.h"
+#include "Sound/AudioVolume.h"
+#include "Sound/SoundAttenuation.h"
 #include "Sound/SoundBase.h"
 #include "Sound/SoundConcurrency.h"
-#include "Sound/SoundWave.h"
 #include "Sound/SoundGroups.h"
-
-
+#include "Sound/SoundWave.h"
+#include "StaticBoundShaderState.h"
+#include "StaticParameterSet.h"
+#include "TextureResource.h"
+#include "Tickable.h"
+#include "TimerManager.h"
+#include "UnrealClient.h"
+#include "UnrealEngine.h"
+#include "WorldCollision.h"
