Use UWPLevelStreamingDynamic StreamingWorld (WorldPartition's OuterWorld) instead of OwningWorld when calling LoadActors.
Fix an issue with Instanced WorldPartitions where the actors folders would try to resolve against the top level World instead of their own world.

diff --git a/Engine/Source/Runtime/Engine/Private/WorldPartition/WorldPartitionLevelStreamingDynamic.cpp b/Engine/Source/Runtime/Engine/Private/WorldPartition/WorldPartitionLevelStreamingDynamic.cpp
--- a/Engine/Source/Runtime/Engine/Private/WorldPartition/WorldPartitionLevelStreamingDynamic.cpp
+++ b/Engine/Source/Runtime/Engine/Private/WorldPartition/WorldPartitionLevelStreamingDynamic.cpp
@@ -414,7 +414,7 @@ bool UWorldPartitionLevelStreamingDynamic::IssueLoadRequests()
 	// Load saved actors
 	if (ChildPackagesToLoad.Num())
 	{
-		FWorldPartitionLevelHelper::LoadActors(World, RuntimeLevel, ChildPackagesToLoad, PackageReferencer, FinalizeLoading, World->IsGameWorld(), MoveTemp(InstancingContext));
+		FWorldPartitionLevelHelper::LoadActors(GetStreamingWorld(), RuntimeLevel, ChildPackagesToLoad, PackageReferencer, FinalizeLoading, World->IsGameWorld(), MoveTemp(InstancingContext));
 	}
 	else
 	{
