Add WasAnimNotifyStateActive to FCachedAnimStateData interface

#tests PIE
#rnx

diff --git a/Engine/Source/Runtime/Engine/Private/Animation/CachedAnimData.cpp b/Engine/Source/Runtime/Engine/Private/Animation/CachedAnimData.cpp
--- a/Engine/Source/Runtime/Engine/Private/Animation/CachedAnimData.cpp
+++ b/Engine/Source/Runtime/Engine/Private/Animation/CachedAnimData.cpp
@@ -66,6 +66,11 @@ bool FCachedAnimStateData::IsActiveState(class UAnimInstance& InAnimInstance) co
 	return IsValid(InAnimInstance) ? (InAnimInstance.GetCurrentStateName(MachineIndex) == StateName) : false;
 }
 
+bool FCachedAnimStateData::WasAnimNotifyStateActive(UAnimInstance& InAnimInstance, TSubclassOf<UAnimNotifyState> AnimNotifyStateType) const
+{
+	return IsValid(InAnimInstance) ? (InAnimInstance.WasAnimNotifyStateActiveInSourceState(MachineIndex, StateIndex, AnimNotifyStateType)) : false;
+}
+
 bool FCachedAnimStateArray::IsValid(UAnimInstance& InAnimInstance) const
 {
 	// Make sure the setup validates our assumptions.
diff --git a/Engine/Source/Runtime/Engine/Public/Animation/CachedAnimData.h b/Engine/Source/Runtime/Engine/Public/Animation/CachedAnimData.h
--- a/Engine/Source/Runtime/Engine/Public/Animation/CachedAnimData.h
+++ b/Engine/Source/Runtime/Engine/Public/Animation/CachedAnimData.h
@@ -52,6 +52,8 @@ struct ENGINE_API FCachedAnimStateData
 	/** Is State active? */
 	bool IsActiveState(class UAnimInstance& InAnimInstance) const;
 
+	bool WasAnimNotifyStateActive(UAnimInstance& InAnimInstance, TSubclassOf<UAnimNotifyState> AnimNotifyStateType) const;
+		
 private:
 	mutable int32 MachineIndex;
 	mutable int32 StateIndex;
