Fix for static analysis error.

diff --git a/Engine/Source/Runtime/Engine/Private/Materials/HLSLMaterialTranslator.cpp b/Engine/Source/Runtime/Engine/Private/Materials/HLSLMaterialTranslator.cpp
--- a/Engine/Source/Runtime/Engine/Private/Materials/HLSLMaterialTranslator.cpp
+++ b/Engine/Source/Runtime/Engine/Private/Materials/HLSLMaterialTranslator.cpp
@@ -8082,15 +8082,19 @@ FLinearColor FHLSLMaterialTranslator::GetTypeMaskedValue(EMaterialValueType Type
 	case MCT_Float1:
 	case MCT_LWCScalar:
 		OutValue = FLinearColor(ConstValue.R, 0.0f, 0.0f, 0.0f);
+		break;
 	case MCT_Float2:
 	case MCT_LWCVector2:
 		OutValue = FLinearColor(ConstValue.R, ConstValue.G, 0.0f, 0.0f);
+		break;
 	case MCT_Float3:
 	case MCT_LWCVector3:
 		OutValue = FLinearColor(ConstValue.R, ConstValue.G, ConstValue.B, 0.0f);
+		break;
 	case MCT_Float4:
 	case MCT_LWCVector4:
 		OutValue = ConstValue;
+		break;
 	default:
 		bValid = false;
 	}
@@ -8853,7 +8857,7 @@ int32 FHLSLMaterialTranslator::Normalize(int32 X)
 	{
 		FVector4 Vector(GetTypeMaskedValue(ResultType, ValueX, nullptr));
 		FLinearColor Result(Vector.GetSafeNormal());
-		return ConstResultValue(ResultType, Result);
+		return ConstResultValue(ResultType, Result);		
 	}
 
 	if (GetParameterUniformExpression(X))
