- Ensure CPU has enough space to read the last element

diff --git a/Engine/Source/Runtime/Engine/Private/VectorField.cpp b/Engine/Source/Runtime/Engine/Private/VectorField.cpp
--- a/Engine/Source/Runtime/Engine/Private/VectorField.cpp
+++ b/Engine/Source/Runtime/Engine/Private/VectorField.cpp
@@ -431,10 +431,10 @@ void UVectorFieldStatic::UpdateCPUData(bool bDiscardData)
 		FMemoryWriter Ar(CPUData);
 
 #if VECTOR_FIELD_DATA_AS_HALF
-		// because of vector implementations in VectorLoadHalf we want to make sure that our buffer
-		// is padded out to support reading the last element
-		constexpr int32 DestComponentCount = 3;
-		CPUData.Reset(Align(SampleCount * DestComponentCount * sizeof(FFloat16), sizeof(FVector4f)));
+		// Ensure we have enough space in the buffer to read the last element
+		constexpr int32 SampleTypeSize = sizeof(FFloat16) * 3;
+		constexpr int32 SampleReadSize = sizeof(FFloat16) * 4;
+		CPUData.Reset((SampleCount * SampleTypeSize) - SampleTypeSize + SampleReadSize);
 
 		for (size_t SampleIt = 0; SampleIt < SampleCount; ++SampleIt)
 		{
@@ -443,7 +443,10 @@ void UVectorFieldStatic::UpdateCPUData(bool bDiscardData)
 			Ar << Ptr[SampleIt].B;
 		}
 #else
-		CPUData.Reset(Align(SampleCount * sizeof(FVector3f), sizeof(FVector4f)));
+		// Ensure we have enough space in the buffer to read the last element
+		constexpr int32 SampleTypeSize = sizeof(FVector3f);
+		constexpr int32 SampleReadSize = sizeof(FVector4f);
+		CPUData.Reset((SampleCount * SampleTypeSize) - SampleTypeSize + SampleReadSize);
 
 		for (size_t SampleIt = 0; SampleIt < SampleCount; ++SampleIt)
 		{
