Fix Regression: When PackageNameToLoad is NAME_None

#rnx

diff --git a/Engine/Source/Runtime/Engine/Private/LevelStreaming.cpp b/Engine/Source/Runtime/Engine/Private/LevelStreaming.cpp
--- a/Engine/Source/Runtime/Engine/Private/LevelStreaming.cpp
+++ b/Engine/Source/Runtime/Engine/Private/LevelStreaming.cpp
@@ -2197,7 +2197,7 @@ bool ULevelStreaming::IsValidStreamingLevel() const
 		} 
 		
 		// Handle unloaded instanced package
-		if (WorldPackageName != PackageNameToLoad)
+		if (PackageNameToLoad != NAME_None && WorldPackageName != PackageNameToLoad)
 		{
 			WorldPackageName = PackageNameToLoad;
 		}
