WorldPartitionTests: Make tests EditorContext only (was crashing when running in -game)

#rnx

diff --git a/Engine/Source/Runtime/Engine/Private/Tests/WorldPartition/WorldPartitionHandleTests.cpp b/Engine/Source/Runtime/Engine/Private/Tests/WorldPartition/WorldPartitionHandleTests.cpp
--- a/Engine/Source/Runtime/Engine/Private/Tests/WorldPartition/WorldPartitionHandleTests.cpp
+++ b/Engine/Source/Runtime/Engine/Private/Tests/WorldPartition/WorldPartitionHandleTests.cpp
@@ -32,7 +32,7 @@ struct FWorldPartitionActorDescUnitTestAcccessor
 
 namespace WorldPartitionTests
 {
-	IMPLEMENT_SIMPLE_AUTOMATION_TEST(FWorldPartitionSoftRefTest, TEST_NAME_ROOT ".Handle", EAutomationTestFlags::ApplicationContextMask | EAutomationTestFlags::EngineFilter)
+	IMPLEMENT_SIMPLE_AUTOMATION_TEST(FWorldPartitionSoftRefTest, TEST_NAME_ROOT ".Handle", EAutomationTestFlags::EditorContext | EAutomationTestFlags::EngineFilter)
 
 #if WITH_EDITOR
 	template <class LoadingContextType>
diff --git a/Engine/Source/Runtime/Engine/Private/Tests/WorldPartition/WorldPartitionPathConversionTests.cpp b/Engine/Source/Runtime/Engine/Private/Tests/WorldPartition/WorldPartitionPathConversionTests.cpp
--- a/Engine/Source/Runtime/Engine/Private/Tests/WorldPartition/WorldPartitionPathConversionTests.cpp
+++ b/Engine/Source/Runtime/Engine/Private/Tests/WorldPartition/WorldPartitionPathConversionTests.cpp
@@ -14,7 +14,7 @@
 
 namespace WorldPartitionTests
 {
-	IMPLEMENT_SIMPLE_AUTOMATION_TEST(FWorldPartitionPathConversionTest, TEST_NAME_ROOT ".PathConversion", EAutomationTestFlags::ApplicationContextMask | EAutomationTestFlags::EngineFilter)
+	IMPLEMENT_SIMPLE_AUTOMATION_TEST(FWorldPartitionPathConversionTest, TEST_NAME_ROOT ".PathConversion", EAutomationTestFlags::EditorContext | EAutomationTestFlags::EngineFilter)
 
 	bool FWorldPartitionPathConversionTest::RunTest(const FString& Parameters)
 	{
diff --git a/Engine/Source/Runtime/Engine/Private/Tests/WorldPartition/WorldPartitionStreamingGenerationTests.cpp b/Engine/Source/Runtime/Engine/Private/Tests/WorldPartition/WorldPartitionStreamingGenerationTests.cpp
--- a/Engine/Source/Runtime/Engine/Private/Tests/WorldPartition/WorldPartitionStreamingGenerationTests.cpp
+++ b/Engine/Source/Runtime/Engine/Private/Tests/WorldPartition/WorldPartitionStreamingGenerationTests.cpp
@@ -18,7 +18,7 @@
 
 namespace WorldPartitionTests
 {
-	IMPLEMENT_SIMPLE_AUTOMATION_TEST(FWorldPartitionStreamingGenerationTest, TEST_NAME_ROOT ".StreamingGeneration", EAutomationTestFlags::ApplicationContextMask | EAutomationTestFlags::EngineFilter)
+	IMPLEMENT_SIMPLE_AUTOMATION_TEST(FWorldPartitionStreamingGenerationTest, TEST_NAME_ROOT ".StreamingGeneration", EAutomationTestFlags::EditorContext | EAutomationTestFlags::EngineFilter)
 
 	bool FWorldPartitionStreamingGenerationTest::RunTest(const FString& Parameters)
 	{
