Remove unnecessary expression validation from HLSLMaterialTranslator.

#rnx

diff --git a/Engine/Source/Runtime/Engine/Private/Materials/HLSLMaterialTranslator.cpp b/Engine/Source/Runtime/Engine/Private/Materials/HLSLMaterialTranslator.cpp
--- a/Engine/Source/Runtime/Engine/Private/Materials/HLSLMaterialTranslator.cpp
+++ b/Engine/Source/Runtime/Engine/Private/Materials/HLSLMaterialTranslator.cpp
@@ -12798,82 +12798,8 @@ int32 FHLSLMaterialTranslator::AccessMaterialAttribute(int32 CodeIndex, const FG
 		*AttributeName);
 }
 
-static bool IsMaterialPathExemptFromRestrictiveMode(FStringBuilderBase& PathName)
-{
-	const FStringView PathNameView = PathName;
-
-	// Anything currently being edited is NOT exempt from restrictive mode
-	const FStringView EngineTransientPath = TEXT("/Engine/Transient.");
-	if (PathNameView.StartsWith(EngineTransientPath))
-	{
-		PathName.RemoveAt(0, EngineTransientPath.Len());
-		return false;
-	}
-
-	// Anything else in /Engine/ and /Game/ and other selected built-in folders is exempt
-	if (PathNameView.StartsWith(TEXT("/Engine/")) ||
-		PathNameView.StartsWith(TEXT("/Game/")) ||
-		PathNameView.StartsWith(TEXT("/Landmass/")) ||
-		PathNameView.StartsWith(TEXT("/Water/")))
-	{
-		return true;
-	}
-
-	return false;
-}
-
-static bool IsMaterialExpressionExemptFromRestrictiveMode(const UMaterialExpression* Expression, FStringBuilderBase& PathName, EShaderPlatform ShaderPlatform)
-{
-	static bool bNoPreviewPlatforms = FParse::Param(FCommandLine::Get(), TEXT("NoPreviewPlatforms"));
-	if (!bNoPreviewPlatforms && FDataDrivenShaderPlatformInfo::GetIsPreviewPlatform(ShaderPlatform))
-	{
-		// Preview shader platforms do not need to have their expressions validated.
-		// The validation for the main editor platform will prevent disallowed expressions
-		// from being used in the material editor.
-		return true;
-	}
-
-	if (Expression)
-	{
-		if (Expression->Material)
-		{
-			Expression->Material->GetPathName(nullptr, PathName);
-			return IsMaterialPathExemptFromRestrictiveMode(PathName);
-		}
-		if (Expression->Function)
-		{
-			Expression->Function->GetPathName(nullptr, PathName);
-			return IsMaterialPathExemptFromRestrictiveMode(PathName);
-		}
-	}
-	return false;
-}
-
-bool FHLSLMaterialTranslator::ValidateMaterialExpressionPermission(const UMaterialExpression* Expression)
-{
-	FNameBuilder MaterialPathName;
-	if (!IsExpressionClassPermitted(Expression->GetClass()) && !IsMaterialExpressionExemptFromRestrictiveMode(Expression, MaterialPathName, Platform))
-	{
-		if (MaterialPathName.Len() > 0)
-		{
-			Errorf(TEXT("Material expression %s not permitted in %s"), *Expression->GetCreationName().ToString(), *MaterialPathName);
-		}
-		else
-		{
-			Errorf(TEXT("Material expression %s not permitted"), *Expression->GetCreationName().ToString());
-		}
-		return false;
-	}
-	return true;
-}
-
 int32 FHLSLMaterialTranslator::CustomExpression( class UMaterialExpressionCustom* Custom, int32 OutputIndex, TArray<int32>& CompiledInputs )
 {
-	if (!ValidateMaterialExpressionPermission(Custom))
-	{
-		return INDEX_NONE;
-	}
-
 	const FMaterialCustomExpressionEntry* CustomEntry = nullptr;
 	for (const FMaterialCustomExpressionEntry& Entry : CustomExpressions)
 	{
diff --git a/Engine/Source/Runtime/Engine/Private/Materials/HLSLMaterialTranslator.h b/Engine/Source/Runtime/Engine/Private/Materials/HLSLMaterialTranslator.h
--- a/Engine/Source/Runtime/Engine/Private/Materials/HLSLMaterialTranslator.h
+++ b/Engine/Source/Runtime/Engine/Private/Materials/HLSLMaterialTranslator.h
@@ -757,8 +757,6 @@ protected:
 
 	FString StrataGetCastParameterCode(int32 Index, EMaterialValueType DestType);
 
-	bool ValidateMaterialExpressionPermission(const UMaterialExpression* Expression);
-
 	// FMaterialCompiler interface.
 
 	/** 
