Validated AttachParent before attempting to detach garbage objects, otherwise fall back to old diaster recovery logic, make checks more explicit

diff --git a/Engine/Source/Runtime/Engine/Private/Components/SceneComponent.cpp b/Engine/Source/Runtime/Engine/Private/Components/SceneComponent.cpp
--- a/Engine/Source/Runtime/Engine/Private/Components/SceneComponent.cpp
+++ b/Engine/Source/Runtime/Engine/Private/Components/SceneComponent.cpp
@@ -1209,29 +1209,40 @@ PRAGMA_ENABLE_DEPRECATION_WARNINGS
 			{
 				if (USceneComponent* Child = AttachChildren.Last())
 				{
-					if (Child->GetAttachParent())
+					USceneComponent* ChildAttachParentWas = Child->GetAttachParent();
+					if (ChildAttachParentWas)
 					{
-						// If the child is also being destroyed during GC, don't reattach it to anything
-						if (Child->HasAnyInternalFlags(SkipFlags))
-						{
-							Child->DetachFromComponent(FDetachmentTransformRules::KeepWorldTransform);
-						}
-						else if (Child->GetAttachParent() == this)
+						if (ChildAttachParentWas == this)
 						{
 							bool bNeedsDetach = true;
-							USceneComponent* NewParent = GetAttachParent();
-							// Walk up the hierarchy until we find a valid parent which is not marked for destruction by gameplay or GC 
-							while (NewParent && NewParent->HasAnyInternalFlags(SkipFlags))
-							{
-								NewParent = NewParent->GetAttachParent();
-							}
-							if (NewParent)
+							// If the child is also being destroyed during GC, don't reattach it to anything
+							if (!Child->HasAnyInternalFlags(SkipFlags))
 							{
-								bNeedsDetach = (Child->AttachToComponent(NewParent, FAttachmentTransformRules::KeepWorldTransform) == false);
+								// child is alive, attempt to reattach it to some living parent:
+								USceneComponent* NewParent = GetAttachParent();
+								// Walk up the hierarchy until we find a valid parent which is not marked for destruction by gameplay or GC 
+								while (NewParent && NewParent->HasAnyInternalFlags(SkipFlags))
+								{
+									NewParent = NewParent->GetAttachParent();
+								}
+								if(NewParent)
+								{
+									// if we reattach to a new parent, we won't need to detach, in that case AttachChildren should decrement
+									// from the AttachToComponent call
+									const bool bAttachedToNewParent = Child->AttachToComponent(NewParent, FAttachmentTransformRules::KeepWorldTransform);
+									if(bAttachedToNewParent)
+									{
+										checkf(ChildCount > AttachChildren.Num(), TEXT("AttachChildren count did not decrease while reattaching '%s', likely caused by OnAttachmentChanged introducing new children, which could lead to an infinite loop."), *Child->GetName());
+									}
+									bNeedsDetach = !bAttachedToNewParent;
+								}
 							}
+
 							if (bNeedsDetach)
 							{
+								checkf(ChildAttachParentWas->GetAttachChildren().Contains(Child), TEXT("Did not find '%s', in Child->AttachParent->AttachChildren. Child: '%s', Child->AttachParent: '%s'"), *GetFullName(), *Child->GetFullName(), *Child->GetAttachParent()->GetFullName());
 								Child->DetachFromComponent(FDetachmentTransformRules::KeepWorldTransform);
+								checkf(ChildCount > AttachChildren.Num(), TEXT("AttachChildren count increased while detaching '%s', likely caused by OnAttachmentChanged introducing new children, which could lead to an infinite loop."), *Child->GetName());
 							}
 						}
 						else
@@ -1259,7 +1270,6 @@ PRAGMA_ENABLE_DEPRECATION_WARNINGS
 						}
 						AttachChildren.Pop(false);
 					}
-					checkf(ChildCount > AttachChildren.Num(), TEXT("AttachChildren count increased while detaching '%s', likely caused by OnAttachmentChanged introducing new children, which could lead to an infinite loop."), *Child->GetName());
 				}
 				else
 				{
