World Partition Resave Actors Commandlet:
- Added FWorldPartitionActorDesc::ShouldResave that returns true if resaving this actor will have an impact on streaming generation. Before class descriptors, properties changed on a Blueprint that would affect streaming generation weren't taken into account until the actors affected by the change were resaved.
- Dump full actor descriptors description when they need to be resaved.

diff --git a/Engine/Source/Editor/UnrealEd/Private/WorldPartition/WorldPartitionResaveActorsBuilder.cpp b/Engine/Source/Editor/UnrealEd/Private/WorldPartition/WorldPartitionResaveActorsBuilder.cpp
--- a/Engine/Source/Editor/UnrealEd/Private/WorldPartition/WorldPartitionResaveActorsBuilder.cpp
+++ b/Engine/Source/Editor/UnrealEd/Private/WorldPartition/WorldPartitionResaveActorsBuilder.cpp
@@ -456,15 +456,15 @@ bool UWorldPartitionResaveActorsBuilder::RunInternal(UWorld* World, const FCellI
 				{
 					TUniquePtr<FWorldPartitionActorDesc> NewActorDesc = Actor->CreateActorDesc();
 
-					if (!ActorDesc->IsResaveNeeded() && ActorDesc->Equals(NewActorDesc.Get()))
+					if (!ActorDesc->IsResaveNeeded() && !ActorDesc->ShouldResave(NewActorDesc.Get()))
 					{
 						return true;
 					}
 
 					if (bDiffDirtyActorDescs)
 					{
-						DirtyActorDescsOld.Add(ActorDesc->ToString());
-						DirtyActorDescsNew.Add(NewActorDesc->ToString());
+						DirtyActorDescsOld.Add(ActorDesc->ToString(FWorldPartitionActorDesc::EToStringMode::Full));
+						DirtyActorDescsNew.Add(NewActorDesc->ToString(FWorldPartitionActorDesc::EToStringMode::Full));
 					}
 
 					UE_LOG(LogWorldPartitionResaveActorsBuilder, Log, TEXT("Package %s needs to be resaved."), *Package->GetName());
@@ -480,7 +480,7 @@ bool UWorldPartitionResaveActorsBuilder::RunInternal(UWorld* World, const FCellI
 		}, ForEachActorWithLoadingParams);
 	}
 
-	if (bDiffDirtyActorDescs)
+	if (bDiffDirtyActorDescs && DirtyActorDescsOld.Num())
 	{
 		auto WriteTempFile = [](const TArray<FString>& Lines, FString& TempFileName)
 		{
diff --git a/Engine/Source/Runtime/Engine/Private/WorldPartition/WorldPartitionActorDesc.cpp b/Engine/Source/Runtime/Engine/Private/WorldPartition/WorldPartitionActorDesc.cpp
--- a/Engine/Source/Runtime/Engine/Private/WorldPartition/WorldPartitionActorDesc.cpp
+++ b/Engine/Source/Runtime/Engine/Private/WorldPartition/WorldPartitionActorDesc.cpp
@@ -260,6 +260,41 @@ bool FWorldPartitionActorDesc::Equals(const FWorldPartitionActorDesc* Other) con
 		Properties == Other->Properties;
 }
 
+bool FWorldPartitionActorDesc::ShouldResave(const FWorldPartitionActorDesc* Other) const
+{
+	check(Guid == Other->Guid);
+	check(ActorPackage == Other->ActorPackage);
+	check(ActorPath == Other->ActorPath);
+
+	// Tolerate up to 5% for bounds change
+	const float BoundsChangeTolerance = 0.05f;
+	const FBox ThisBounds = GetRuntimeBounds();
+	const FBox OtherBounds = Other->GetRuntimeBounds();
+	const FVector BoundsMinChangeRatio = ThisBounds.Min / OtherBounds.Min;
+	const FVector BoundsMaxChangeRatio = ThisBounds.Max / OtherBounds.Max;
+	if (!BoundsMinChangeRatio.Equals(FVector::OneVector, BoundsChangeTolerance) || !BoundsMaxChangeRatio.Equals(FVector::OneVector, BoundsChangeTolerance))
+	{
+		return true;
+	}
+
+	if (RuntimeGrid != Other->RuntimeGrid ||
+		bIsSpatiallyLoaded != Other->bIsSpatiallyLoaded ||
+		bActorIsEditorOnly != Other->bActorIsEditorOnly ||
+		bActorIsRuntimeOnly != Other->bActorIsRuntimeOnly ||
+		HLODLayer != Other->HLODLayer ||
+		ParentActor != Other->ParentActor ||
+		ContentBundleGuid != Other->ContentBundleGuid ||
+		!CompareUnsortedArrays(DataLayers, Other->DataLayers) ||
+		!CompareUnsortedArrays(References, Other->References))
+	{
+		return true;
+	}
+
+	// If the actor descriptor says the actor is HLOD relebant but in reality it's not, this will incur a loading time penalty during HLOD generation
+	// but will not affect the final result, as the value from the loaded actor will be used instead, so don't consider this as affecting streaming generation.
+	return (bActorIsHLODRelevant == Other->bActorIsHLODRelevant) || (bActorIsHLODRelevant && !Other->bActorIsHLODRelevant);
+}
+
 void FWorldPartitionActorDesc::SerializeTo(TArray<uint8>& OutData)
 {
 	// Serialize to archive and gather custom versions
diff --git a/Engine/Source/Runtime/Engine/Public/WorldPartition/WorldPartitionActorDesc.h b/Engine/Source/Runtime/Engine/Public/WorldPartition/WorldPartitionActorDesc.h
--- a/Engine/Source/Runtime/Engine/Public/WorldPartition/WorldPartitionActorDesc.h
+++ b/Engine/Source/Runtime/Engine/Public/WorldPartition/WorldPartitionActorDesc.h
@@ -250,6 +250,12 @@ public:
 
 	virtual bool Equals(const FWorldPartitionActorDesc* Other) const;
 
+	/**
+	 * Returns true if resaving this actor will have an impact on streaming generation. Before class descriptors, properties changed on a Blueprint 
+	 * that would affect streaming generation weren't taken into account until the actors affected by the change were resaved.
+	 */
+	virtual bool ShouldResave(const FWorldPartitionActorDesc* Other) const;
+
 	void SerializeTo(TArray<uint8>& OutData);
 
 	void TransformInstance(const FString& From, const FString& To);
