Do not allow Animation blueprints to place Input Event nodes in their event graphs. This behavior won't work because the animation blueprint doesn't have a valid Player Controller and will never actually be recieving input.

The correct way for animation blueprints to receive input data would be to set a property value on the animation blueprint and set it in an actual Player Controller (or some other input enabled actor) blueprint.

This will not produce any content warnings or errors. Any existing placed input nodes will behave the same as they would have before, which is to do nothing.

diff --git a/Engine/Source/Editor/UnrealEd/Classes/Settings/AnimBlueprintSettings.h b/Engine/Source/Editor/UnrealEd/Classes/Settings/AnimBlueprintSettings.h
--- a/Engine/Source/Editor/UnrealEd/Classes/Settings/AnimBlueprintSettings.h
+++ b/Engine/Source/Editor/UnrealEd/Classes/Settings/AnimBlueprintSettings.h
@@ -34,6 +34,18 @@ public:
 	/** Whether to allow restrict which base function overrides can created/displayed in animation blueprints */
 	UPROPERTY()
 	bool bRestrictBaseFunctionOverrides = false;
+	
+	/**
+	* Whether to allow input events to be created/displayed in animation blueprints.
+	* 
+	* You used to be able to place input event nodes in anim graphs and they would just not work, sometimes causing an ensure.
+	* 
+	* If needed, you can enable this legacy behavior by setting this property to true (allowing input events to be placed in a graph).
+	* 
+	* Default value is false.
+	*/
+	UPROPERTY(Config)
+	bool bSupportInputEventsForBackwardsCompatibility = false;
 
 	/** The set of allowed base functions if restricted */
 	UPROPERTY()
diff --git a/Engine/Source/Runtime/Engine/Classes/Animation/AnimBlueprint.h b/Engine/Source/Runtime/Engine/Classes/Animation/AnimBlueprint.h
--- a/Engine/Source/Runtime/Engine/Classes/Animation/AnimBlueprint.h
+++ b/Engine/Source/Runtime/Engine/Classes/Animation/AnimBlueprint.h
@@ -185,6 +185,7 @@ class ENGINE_API UAnimBlueprint : public UBlueprint, public IInterface_PreviewMe
 	virtual bool SupportsEventGraphs() const override;
 	virtual bool SupportsDelegates() const override;
 	virtual bool SupportsMacros() const override;
+	virtual bool SupportsInputEvents() const override;
 	virtual bool AllowFunctionOverride(const UFunction* const InFunction) const override;
 	virtual void PostEditChangeProperty(FPropertyChangedEvent& PropertyChangedEvent) override;
 	
diff --git a/Engine/Source/Runtime/Engine/Private/AnimBlueprint.cpp b/Engine/Source/Runtime/Engine/Private/AnimBlueprint.cpp
--- a/Engine/Source/Runtime/Engine/Private/AnimBlueprint.cpp
+++ b/Engine/Source/Runtime/Engine/Private/AnimBlueprint.cpp
@@ -240,6 +240,21 @@ bool UAnimBlueprint::SupportsMacros() const
 	return GetDefault<UAnimBlueprintSettings>()->bAllowMacros;
 }
 
+bool UAnimBlueprint::SupportsInputEvents() const
+{
+	// Animation blueprints don't really support input events. You used to be able to place
+	// the input event nodes in the anim graphs and they would just not work, causing an ensure.
+	// To keep backwards compatibility we will allow users to enable them in case
+	// their project somehow relies on it in some custom node extension, but
+	// for new projects it should not be allowed. 
+	if (GetDefault<UAnimBlueprintSettings>()->bSupportInputEventsForBackwardsCompatibility)
+	{
+		return Super::SupportsInputEvents();
+	}
+	
+	return false;	
+}
+
 bool UAnimBlueprint::AllowFunctionOverride(const UFunction* const InFunction) const
 {
 	check(InFunction);
