Add UserExposedCaption override for material function assets. This is be able to change the exposed UI name without having have asset file name restrictions (no space) or to rely on perforce renaming and redictor assets.

diff --git a/Engine/Source/Editor/UnrealEd/Private/MaterialGraphSchema.cpp b/Engine/Source/Editor/UnrealEd/Private/MaterialGraphSchema.cpp
--- a/Engine/Source/Editor/UnrealEd/Private/MaterialGraphSchema.cpp
+++ b/Engine/Source/Editor/UnrealEd/Private/MaterialGraphSchema.cpp
@@ -1019,13 +1019,23 @@ void UMaterialGraphSchema::GetMaterialFunctionActions(FGraphActionMenuBuilder& A
 					}
 				}
 
-				// Extract the object name from the path
 				FString FunctionName = FunctionPathName;
-				int32 PeriodIndex = FunctionPathName.Find(TEXT("."), ESearchCase::CaseSensitive, ESearchDir::FromEnd);
 
-				if (PeriodIndex != INDEX_NONE)
+				const FString UserExposedCaption = AssetData.GetTagValueRef<FString>("UserExposedCaption");
+				if (!UserExposedCaption.IsEmpty())
 				{
-					FunctionName = FunctionPathName.Right(FunctionPathName.Len() - PeriodIndex - 1);
+					// If the UI user exposed name name is not empty, use it directly
+					FunctionName = UserExposedCaption;
+				}
+				else
+				{
+					// Extract the object name from the path
+					int32 PeriodIndex = FunctionPathName.Find(TEXT("."), ESearchCase::CaseSensitive, ESearchDir::FromEnd);
+
+					if (PeriodIndex != INDEX_NONE)
+					{
+						FunctionName = FunctionPathName.Right(FunctionPathName.Len() - PeriodIndex - 1);
+					}
 				}
 
 				// For each category the function should belong to...
diff --git a/Engine/Source/Runtime/Engine/Classes/Materials/MaterialFunction.h b/Engine/Source/Runtime/Engine/Classes/Materials/MaterialFunction.h
--- a/Engine/Source/Runtime/Engine/Classes/Materials/MaterialFunction.h
+++ b/Engine/Source/Runtime/Engine/Classes/Materials/MaterialFunction.h
@@ -51,6 +51,10 @@ class UMaterialFunction : public UMaterialFunctionInterface
 	UPROPERTY(EditAnywhere, Category=MaterialFunction, AssetRegistrySearchable)
 	FString Description;
 
+	/** Name of the function to be displayed on the node within the material editor instead of the asset name. */
+	UPROPERTY(EditAnywhere, Category=MaterialFunction, AssetRegistrySearchable)
+	FString UserExposedCaption;
+
 	/** Whether to list this function in the material function library, which is a window in the material editor that lists categorized functions. */
 	UPROPERTY(EditAnywhere, Category=MaterialFunction, AssetRegistrySearchable)
 	uint8 bExposeToLibrary:1;
@@ -155,6 +159,9 @@ public:
 
 	/** Returns an array of the functions that this function is dependent on, directly or indirectly. */
 	ENGINE_API virtual void GetDependentFunctions(TArray<UMaterialFunctionInterface*>& DependentFunctions) const override;
+
+	/** Returns If returns an empty string, use the default class name for the material function. Otherwise, the string will be the name shown when the function is exposed to users in the material graph as a node, or from the contextual menu when searching for nodes. */
+	virtual FString GetUserExposedCaption() const override { return UserExposedCaption; }
 #endif // WITH_EDITORONLY_DATA
 
 #if WITH_EDITOR
diff --git a/Engine/Source/Runtime/Engine/Classes/Materials/MaterialFunctionInterface.h b/Engine/Source/Runtime/Engine/Classes/Materials/MaterialFunctionInterface.h
--- a/Engine/Source/Runtime/Engine/Classes/Materials/MaterialFunctionInterface.h
+++ b/Engine/Source/Runtime/Engine/Classes/Materials/MaterialFunctionInterface.h
@@ -147,6 +147,9 @@ public:
 	/** Returns an array of the functions that this function is dependent on, directly or indirectly. */
 	ENGINE_API virtual void GetDependentFunctions(TArray<UMaterialFunctionInterface*>& DependentFunctions) const
 		PURE_VIRTUAL(UMaterialFunctionInterface::GetDependentFunctions,);
+
+	/** Returns If returns an empty string, use the default class name for the material function. Otherwise, the string will be the name shown when the function is exposed to users in the material graph as a node, or from the contextual menu when searching for nodes. */
+	virtual FString GetUserExposedCaption() const { return TEXT(""); }
 #endif // WITH_EDITORONLY_DATA
 
 #if WITH_EDITOR
diff --git a/Engine/Source/Runtime/Engine/Private/Materials/MaterialExpressions.cpp b/Engine/Source/Runtime/Engine/Private/Materials/MaterialExpressions.cpp
--- a/Engine/Source/Runtime/Engine/Private/Materials/MaterialExpressions.cpp
+++ b/Engine/Source/Runtime/Engine/Private/Materials/MaterialExpressions.cpp
@@ -16162,7 +16162,22 @@ int32 UMaterialExpressionMaterialFunctionCall::Compile(class FMaterialCompiler*
 
 void UMaterialExpressionMaterialFunctionCall::GetCaption(TArray<FString>& OutCaptions) const
 {
-	OutCaptions.Add(MaterialFunction ? MaterialFunction->GetName() : TEXT("Unspecified Function"));
+	if (MaterialFunction)
+	{
+		FString UserExposedCaption = MaterialFunction->GetUserExposedCaption();
+		if (!UserExposedCaption.IsEmpty())
+		{
+			OutCaptions.Add(UserExposedCaption);
+		}
+		else
+		{
+			OutCaptions.Add(MaterialFunction->GetName());
+		}
+	}
+	else
+	{
+		OutCaptions.Add(TEXT("Unspecified Function"));
+	}
 }
 
 const TArray<FExpressionInput*> UMaterialExpressionMaterialFunctionCall::GetInputs()
