Speculative fix for seen once crashes in typed element debug ID printing.

We shouldn't be crashing before the reference tracing is even able to output, but this looks like a garbage handle due to something getting cleaned up that's still got references.

diff --git a/Engine/Source/Runtime/Engine/Public/Elements/Actor/ActorElementData.h b/Engine/Source/Runtime/Engine/Public/Elements/Actor/ActorElementData.h
--- a/Engine/Source/Runtime/Engine/Public/Elements/Actor/ActorElementData.h
+++ b/Engine/Source/Runtime/Engine/Public/Elements/Actor/ActorElementData.h
@@ -22,7 +22,7 @@ template <>
 inline FString GetTypedElementDebugId<FActorElementData>(const FActorElementData& InElementData)
 {
 	UObject* Object = (UObject*)InElementData.Actor;
-	return Object
+	return IsValid(Object)
 		? Object->GetFullName()
 		: TEXT("null");
 }
