Adding blueprint callable getter to allow access to primitive components within the cluster union.

diff --git a/Engine/Source/Runtime/Engine/Classes/PhysicsEngine/ClusterUnionComponent.h b/Engine/Source/Runtime/Engine/Classes/PhysicsEngine/ClusterUnionComponent.h
--- a/Engine/Source/Runtime/Engine/Classes/PhysicsEngine/ClusterUnionComponent.h
+++ b/Engine/Source/Runtime/Engine/Classes/PhysicsEngine/ClusterUnionComponent.h
@@ -137,6 +137,9 @@ public:
 	UFUNCTION(BlueprintCallable, Category = "Cluster Union")
 	void RemoveComponentFromCluster(UPrimitiveComponent* InComponent);
 
+	UFUNCTION(BlueprintCallable, Category = "Cluster Union")
+	TArray<UPrimitiveComponent*> GetPrimitiveComponents();
+
 	UFUNCTION(BlueprintCallable, Category = "Cluster Union")
 	void SetIsAnchored(bool bIsAnchored);
 
diff --git a/Engine/Source/Runtime/Engine/Private/PhysicsEngine/ClusterUnionComponent.cpp b/Engine/Source/Runtime/Engine/Private/PhysicsEngine/ClusterUnionComponent.cpp
--- a/Engine/Source/Runtime/Engine/Private/PhysicsEngine/ClusterUnionComponent.cpp
+++ b/Engine/Source/Runtime/Engine/Private/PhysicsEngine/ClusterUnionComponent.cpp
@@ -143,6 +143,17 @@ void UClusterUnionComponent::RemoveComponentFromCluster(UPrimitiveComponent* InC
 	PhysicsProxy->RemovePhysicsObjects_External(PhysicsObjectsToRemove);
 }
 
+TArray<UPrimitiveComponent*> UClusterUnionComponent::GetPrimitiveComponents()
+{
+	TArray<UPrimitiveComponent*> PrimitiveComponents;
+	for (auto Iter = ComponentToPhysicsObjects.CreateIterator(); Iter; ++Iter)
+	{
+		PrimitiveComponents.Add(Iter.Key().ResolveObjectPtr());
+	}
+
+	return PrimitiveComponents;
+}
+
 void UClusterUnionComponent::SetIsAnchored(bool bIsAnchored)
 {
 	if (!PhysicsProxy)
