Make LWI manager clean up references to spawned actors when they are destroyed. Fix dangling pointers when actors are destroyed.

diff --git a/Engine/Source/Runtime/Engine/Classes/GameFramework/LightWeightInstanceManager.h b/Engine/Source/Runtime/Engine/Classes/GameFramework/LightWeightInstanceManager.h
--- a/Engine/Source/Runtime/Engine/Classes/GameFramework/LightWeightInstanceManager.h
+++ b/Engine/Source/Runtime/Engine/Classes/GameFramework/LightWeightInstanceManager.h
@@ -132,6 +132,8 @@ public:
 		return Cast<I>(FetchActorFromHandle(Handle));
 	}
 
+	AActor* FindActorForInstanceIndex(const int32 InstanceIndex);
+
 protected:
 	// Creates an actor to replace the instance specified by Handle
 	AActor* ConvertInstanceToActor(const FActorInstanceHandle& Handle);
@@ -185,6 +187,10 @@ public:
 
 	virtual void GetLifetimeReplicatedProps(TArray<FLifetimeProperty>& OutLifetimeProps) const override;
 
+private:
+	UFUNCTION()
+	void OnSpawnedActorDestroyed(AActor* DestroyedActor);
+
 protected:
 
 	FString BaseInstanceName;
@@ -227,7 +233,8 @@ protected:
 	//
 
 	// keep track of which instances are currently represented by an actor
-	TMap<int32, AActor*> Actors;
+	UPROPERTY(Transient)
+	TMap<int32, TObjectPtr<AActor>> Actors;
 
 	// list of indices that we are no longer using
 	UPROPERTY(Replicated)
diff --git a/Engine/Source/Runtime/Engine/Private/Engine/ActorInstanceHandle.cpp b/Engine/Source/Runtime/Engine/Private/Engine/ActorInstanceHandle.cpp
--- a/Engine/Source/Runtime/Engine/Private/Engine/ActorInstanceHandle.cpp
+++ b/Engine/Source/Runtime/Engine/Private/Engine/ActorInstanceHandle.cpp
@@ -41,10 +41,7 @@ FActorInstanceHandle::FActorInstanceHandle(ALightWeightInstanceManager* InManage
 	if (Manager.IsValid())
 	{
 		InstanceIndex = Manager->ConvertCollisionIndexToLightWeightIndex(InInstanceIndex);
-		if (AActor* const* FoundActor = Manager->Actors.Find(InstanceIndex))
-		{
-			Actor = *FoundActor;
-		}
+		Actor = Manager->FindActorForInstanceIndex(InstanceIndex);
 
 		UWorld* World = Manager->GetWorld();
 		if (ensure(World))
diff --git a/Engine/Source/Runtime/Engine/Private/LightWeightInstanceManager.cpp b/Engine/Source/Runtime/Engine/Private/LightWeightInstanceManager.cpp
--- a/Engine/Source/Runtime/Engine/Private/LightWeightInstanceManager.cpp
+++ b/Engine/Source/Runtime/Engine/Private/LightWeightInstanceManager.cpp
@@ -51,7 +51,7 @@ AActor* ALightWeightInstanceManager::FetchActorFromHandle(const FActorInstanceHa
 	if (ensure(!Handle.Actor.IsValid()))
 	{
 		// check if we already have an actor for this handle
-		AActor* const* FoundActor = nullptr;
+		TObjectPtr<AActor>* FoundActor = nullptr;
 		if ((FoundActor = Actors.Find(Handle.GetInstanceIndex())) == nullptr)
 		{
 			// spawn a new actor
@@ -90,6 +90,7 @@ AActor* ALightWeightInstanceManager::ConvertInstanceToActor(const FActorInstance
 		SetSpawnParameters(SpawnParams);
 		NewActor = GetLevel()->GetWorld()->SpawnActor<AActor>(GetActorClassToSpawn(Handle), GetTransform(Handle), SpawnParams);
 		check(NewActor);
+		NewActor->OnDestroyed.AddUniqueDynamic(this, &ALightWeightInstanceManager::OnSpawnedActorDestroyed);
 
 		Handle.Actor = NewActor;
 		Actors.Add(Handle.GetInstanceIndex(), NewActor);
@@ -100,9 +101,17 @@ AActor* ALightWeightInstanceManager::ConvertInstanceToActor(const FActorInstance
 	return NewActor;
 }
 
+void ALightWeightInstanceManager::OnSpawnedActorDestroyed(AActor* DestroyedActor)
+{
+	check(DestroyedActor);
+	const int32 DestroyedActorInstanceIndex = FindIndexForActor(DestroyedActor);
+	check(DestroyedActorInstanceIndex != INDEX_NONE);
+	Actors.Remove(DestroyedActorInstanceIndex);
+}
+
 int32 ALightWeightInstanceManager::FindIndexForActor(const AActor* InActor) const
 {
-	for (const TPair<int32, AActor*>& ActorPair : Actors)
+	for (const TPair<int32, TObjectPtr<AActor>>& ActorPair : Actors)
 	{
 		if (ActorPair.Value == InActor)
 		{
@@ -226,11 +235,16 @@ bool ALightWeightInstanceManager::FindActorForHandle(const FActorInstanceHandle&
 		return true;
 	}
 
-	AActor* const* FoundActor = Actors.Find(Handle.GetInstanceIndex());
+	const TObjectPtr<AActor>* FoundActor = Actors.Find(Handle.GetInstanceIndex());
 	Handle.Actor = FoundActor ? *FoundActor : nullptr;
 	return Handle.Actor != nullptr;
 }
 
+AActor* ALightWeightInstanceManager::FindActorForInstanceIndex(const int32 InstanceIndex)
+{
+	return Actors.FindRef(InstanceIndex);
+}
+
 FVector ALightWeightInstanceManager::GetLocation(const FActorInstanceHandle& Handle) const
 {
 	return GetTransform(Handle).GetTranslation();
@@ -340,7 +354,7 @@ void ALightWeightInstanceManager::RemoveInstance(const int32 Index)
 		ValidIndices[Index] = false;
 
 		// destroy the associated actor if one existed
-		if (AActor** FoundActor = Actors.Find(Index))
+		if (TObjectPtr<AActor>* FoundActor = Actors.Find(Index))
 		{
 			AActor* ActorToDestroy = *FoundActor;
 			ActorToDestroy->Destroy();
