LevelInstance: Disable partial loading until we can investigate issues

#rnx

diff --git a/Engine/Source/Runtime/Engine/Private/LevelInstance/LevelInstanceLevelStreaming.cpp b/Engine/Source/Runtime/Engine/Private/LevelInstance/LevelInstanceLevelStreaming.cpp
--- a/Engine/Source/Runtime/Engine/Private/LevelInstance/LevelInstanceLevelStreaming.cpp
+++ b/Engine/Source/Runtime/Engine/Private/LevelInstance/LevelInstanceLevelStreaming.cpp
@@ -27,7 +27,7 @@
 #include "Misc/ScopeExit.h"
 #include "UObject/LinkerLoad.h"
 
-static bool GDisableLevelInstanceEditorPartialLoading = false;
+static bool GDisableLevelInstanceEditorPartialLoading = true;
 FAutoConsoleVariableRef CVarDisableLevelInstanceEditorPartialLoading(
 	TEXT("wp.Editor.DisableLevelInstanceEditorPartialLoading"),
 	GDisableLevelInstanceEditorPartialLoading,
