Fix return value for constant expressions in FHLSLMaterialTranslator::Max.

diff --git a/Engine/Source/Runtime/Engine/Private/Materials/HLSLMaterialTranslator.cpp b/Engine/Source/Runtime/Engine/Private/Materials/HLSLMaterialTranslator.cpp
--- a/Engine/Source/Runtime/Engine/Private/Materials/HLSLMaterialTranslator.cpp
+++ b/Engine/Source/Runtime/Engine/Private/Materials/HLSLMaterialTranslator.cpp
@@ -8046,11 +8046,7 @@ bool FHLSLMaterialTranslator::CastConstantType(FLinearColor SourceValue, EMateri
 		}
 		else
 		{
-			NumComponents = FMath::Min(NumSourceComponents, NumDestComponents);
-			if (NumComponents != NumDestComponents)
-			{
-				OutResult = GetTypeMaskedValue(DestType, SourceValue, nullptr);
-			}
+			OutResult = GetTypeMaskedValue(DestType, SourceValue, nullptr);
 			return true;
 		}
 
@@ -9307,8 +9303,7 @@ int32 FHLSLMaterialTranslator::Max(int32 A,int32 B)
 			Result.B = FMath::Max(ValueA.B, ValueB.B);
 			Result.A = FMath::Max(ValueA.A, ValueB.A);
 
-			EMaterialValueType ResultType = GetArithmeticResultType(A, B);
-			return ConstResultValue(ResultType, Result);
+			return ConstResultValue(TypeA, Result);
 		}
 
 		return AddUniformExpression(new FMaterialUniformExpressionMax(ExpressionA, ExpressionB),GetParameterType(A),TEXT("max(%s,%s)"),*GetParameterCode(A),*CoerceParameter(B,GetParameterType(A)));
