Add UStaticMesh::IsDataValid() that verifies that an uncooked StaticMesh Asset has valid a LOD0 SourceModel MeshDescription/BulkData, this is not guaranteed due to outstanding issue UE-159769. Add FStaticMeshSourceModel::IsSourceModelInitialized() to support this, as other functions all have checks that could be hit if some of the state is invalid.

diff --git a/Engine/Source/Runtime/Engine/Classes/Engine/StaticMesh.h b/Engine/Source/Runtime/Engine/Classes/Engine/StaticMesh.h
--- a/Engine/Source/Runtime/Engine/Classes/Engine/StaticMesh.h
+++ b/Engine/Source/Runtime/Engine/Classes/Engine/StaticMesh.h
@@ -1550,6 +1550,7 @@ public:
 	ENGINE_API virtual void PreEditChange(FProperty* PropertyAboutToChange) override;
 	ENGINE_API virtual void PostEditChangeProperty(FPropertyChangedEvent& PropertyChangedEvent) override;
 	ENGINE_API virtual void PostEditUndo() override;
+	virtual EDataValidationResult IsDataValid(class FDataValidationContext& Context) override;
 	ENGINE_API virtual void GetAssetRegistryTagMetadata(TMap<FName, FAssetRegistryTagMetadata>& OutMetadata) const override;
 	
 	ENGINE_API virtual void WillNeverCacheCookedPlatformDataAgain() override;
diff --git a/Engine/Source/Runtime/Engine/Classes/Engine/StaticMeshSourceData.h b/Engine/Source/Runtime/Engine/Classes/Engine/StaticMeshSourceData.h
--- a/Engine/Source/Runtime/Engine/Classes/Engine/StaticMeshSourceData.h
+++ b/Engine/Source/Runtime/Engine/Classes/Engine/StaticMeshSourceData.h
@@ -90,6 +90,12 @@ public:
 	 */
 	ENGINE_API UStaticMesh* GetStaticMeshOwner() const;
 
+	/**
+	 * Verify that this SourceModel has been fully initialized. It may still not contain valid data ( use IsMeshDescriptionValid() for that )
+	 * Currently this checks if StaticMeshDescriptionBulkData is allocated, which the other functions below all require via checks.
+	 */
+	ENGINE_API bool IsSourceModelInitialized() const;
+
 	/**
 	 * Create a new MeshDescription object
 	 */
diff --git a/Engine/Source/Runtime/Engine/Private/StaticMesh.cpp b/Engine/Source/Runtime/Engine/Private/StaticMesh.cpp
--- a/Engine/Source/Runtime/Engine/Private/StaticMesh.cpp
+++ b/Engine/Source/Runtime/Engine/Private/StaticMesh.cpp
@@ -55,6 +55,7 @@
 #include "IMeshReductionInterfaces.h"
 #include "StaticMeshCompiler.h"
 #include "ObjectCacheContext.h"
+#include "Misc/DataValidation.h"
 #include "Engine/Texture2D.h"
 
 #include "DerivedDataCache.h"
@@ -4046,6 +4047,32 @@ void UStaticMesh::PostEditUndo()
 	Super::PostEditUndo();
 }
 
+
+#if WITH_EDITOR
+EDataValidationResult UStaticMesh::IsDataValid(class FDataValidationContext& Context)
+{
+	EDataValidationResult ValidationResult = Super::IsDataValid(Context);
+
+	// a cooked static mesh asset is probably not going to have valid SourceModels (?)
+	if (GetPackage()->HasAnyPackageFlags(PKG_Cooked) == false)
+	{
+		if (GetSourceModels().Num() == 0)
+		{
+			Context.AddError(LOCTEXT("StaticMeshValidation_NoSourceModel", "This Static Mesh Asset has no Source Models. This asset is not repairable, the Asset is corrupted and must be deleted."));
+			ValidationResult = EDataValidationResult::Invalid;
+		}
+		else if (GetSourceModels()[0].IsSourceModelInitialized() == false)
+		{
+			Context.AddError(LOCTEXT("StaticMeshValidation_UninitializedLOD0", "This Static Mesh Asset has no LOD0 Source Model mesh. This asset is not repairable, the Asset is corrupted and must be deleted."));
+			ValidationResult = EDataValidationResult::Invalid;
+		}
+	}
+
+	return ValidationResult;
+}
+#endif
+
+
 void UStaticMesh::SetLODGroup(FName NewGroup, bool bRebuildImmediately, bool bAllowModify)
 {
 #if WITH_EDITORONLY_DATA
diff --git a/Engine/Source/Runtime/Engine/Private/StaticMeshSourceData.cpp b/Engine/Source/Runtime/Engine/Private/StaticMeshSourceData.cpp
--- a/Engine/Source/Runtime/Engine/Private/StaticMeshSourceData.cpp
+++ b/Engine/Source/Runtime/Engine/Private/StaticMeshSourceData.cpp
@@ -192,6 +192,23 @@ UStaticMesh* FStaticMeshSourceModel::GetStaticMeshOwner() const
 	return Cast<UStaticMesh>(StaticMeshDescriptionBulkData->GetOuter());
 }
 
+bool FStaticMeshSourceModel::IsSourceModelInitialized() const
+{
+	FScopeLock StaticMeshDescriptionBulkDataLock(&StaticMeshDescriptionBulkDataCS);
+
+	if (StaticMeshDescriptionBulkData == nullptr)
+	{
+		return false;
+	}
+	if (StaticMeshDescriptionBulkData->IsBulkDataValid() == false)
+	{
+		return false;
+	}
+
+	return true;
+}
+
+
 bool FStaticMeshSourceModel::LoadMeshDescription(FMeshDescription& OutMeshDescription) const
 {
 	TRACE_CPUPROFILER_EVENT_SCOPE(FStaticMeshSourceModel::LoadMeshDescription);
