Export new UStaticMesh API function

diff --git a/Engine/Source/Runtime/Engine/Classes/Engine/StaticMesh.h b/Engine/Source/Runtime/Engine/Classes/Engine/StaticMesh.h
--- a/Engine/Source/Runtime/Engine/Classes/Engine/StaticMesh.h
+++ b/Engine/Source/Runtime/Engine/Classes/Engine/StaticMesh.h
@@ -1550,7 +1550,7 @@ public:
 	ENGINE_API virtual void PreEditChange(FProperty* PropertyAboutToChange) override;
 	ENGINE_API virtual void PostEditChangeProperty(FPropertyChangedEvent& PropertyChangedEvent) override;
 	ENGINE_API virtual void PostEditUndo() override;
-	virtual EDataValidationResult IsDataValid(class FDataValidationContext& Context) override;
+	ENGINE_API virtual EDataValidationResult IsDataValid(class FDataValidationContext& Context) override;
 	ENGINE_API virtual void GetAssetRegistryTagMetadata(TMap<FName, FAssetRegistryTagMetadata>& OutMetadata) const override;
 	
 	ENGINE_API virtual void WillNeverCacheCookedPlatformDataAgain() override;
