Fix issues with LWC conversions on const operations.

diff --git a/Engine/Source/Runtime/Engine/Private/Materials/HLSLMaterialTranslator.cpp b/Engine/Source/Runtime/Engine/Private/Materials/HLSLMaterialTranslator.cpp
--- a/Engine/Source/Runtime/Engine/Private/Materials/HLSLMaterialTranslator.cpp
+++ b/Engine/Source/Runtime/Engine/Private/Materials/HLSLMaterialTranslator.cpp
@@ -8120,27 +8120,54 @@ bool FHLSLMaterialTranslator::GetConstMaskedParameterValue(EMaterialValueType Ty
 	return bValid;
 }
 
+int32 FHLSLMaterialTranslator::LWCCastIfNeccessary(EMaterialValueType ResultType, int32 ResultCode)
+{
+	if (!IsNumericType(ResultType))
+	{
+		return Error(TEXT("LWCCastIfNeccessary Cannot cast non numeric type"));
+	}
+
+	EMaterialValueType CurrentType = GetType(ResultCode);
+	if(CurrentType == ResultType)
+	{
+		return ResultCode;
+	}
+
+	EMaterialValueType ResultNonLWCType = MakeNonLWCType(ResultType);
+	EMaterialValueType CurrentNonLWCType = MakeNonLWCType(CurrentType);
+
+	if (ResultNonLWCType == CurrentNonLWCType || ResultNonLWCType == MCT_Float || CurrentNonLWCType == MCT_Float)
+	{
+		return ForceCast(ResultCode, ResultType);
+	}
+	return Error(TEXT("LWCCastIfNeccessary can only cast between LWC and Non LWC types.  Current result has a different number of components than requested type"));
+}
+
 int32 FHLSLMaterialTranslator::ConstResultValue(EMaterialValueType Type, FLinearColor ConstantValue)
 {
-	switch (Type)
+	EMaterialValueType NonLWCType = MakeNonLWCType(Type);
+	int32 Result;
+	switch (NonLWCType)
 	{
 	case MCT_Float:
 	case MCT_Float1:
-	case MCT_LWCScalar:
-		return Constant(ConstantValue.R);
+		Result = Constant(ConstantValue.R);
+		break;
 	case MCT_Float2:
-	case MCT_LWCVector2:
-		return Constant2(ConstantValue.R, ConstantValue.G);
+		Result = Constant2(ConstantValue.R, ConstantValue.G);
+		break;
 	case MCT_Float3:
-	case MCT_LWCVector3:
-		return Constant3(ConstantValue.R, ConstantValue.G, ConstantValue.B);
+		Result = Constant3(ConstantValue.R, ConstantValue.G, ConstantValue.B);
+		break;
 	case MCT_Float4:
-	case MCT_LWCVector4:
-		return Constant4(ConstantValue.R, ConstantValue.G, ConstantValue.B, ConstantValue.A);
+		Result = Constant4(ConstantValue.R, ConstantValue.G, ConstantValue.B, ConstantValue.A);
+		break;
 	default:
 		// Error would already be output by GetArithmeticResultType in this case, just return an invalid index
 		return INDEX_NONE;
 	}
+
+	return LWCCastIfNeccessary(Type, Result);
 }
 
 int32 FHLSLMaterialTranslator::ConstResultValue(EMaterialValueType Type, float ConstantValue)
@@ -8208,15 +8235,6 @@ int32 FHLSLMaterialTranslator::Add(int32 A, int32 B)
 	const EMaterialValueType ResultType = GetArithmeticResultType(A, B);
 	if (ExpressionA && ExpressionB)
 	{
-		if (ExpressionA->IsConstant() && ExpressionB->IsConstant())
-		{
-			FLinearColor ValueA, ValueB;
-			GetConstParameterValue(ExpressionA, ValueA);
-			GetConstParameterValue(ExpressionB, ValueB);
-
-			FLinearColor ConstantValue = ValueA + ValueB;
-			return ConstResultValue(ResultType, ConstantValue);
-		}
 		return AddUniformExpression(new FMaterialUniformExpressionFoldedMath(GetParameterUniformExpression(A), GetParameterUniformExpression(B), FMO_Add), ResultType, TEXT("(%s + %s)"), *GetParameterCode(A), *GetParameterCode(B));
 	}
 	else
@@ -8251,15 +8269,6 @@ int32 FHLSLMaterialTranslator::Sub(int32 A, int32 B)
 	const EMaterialValueType ResultType = GetArithmeticResultType(A, B);
 	if (ExpressionA && ExpressionB)
 	{
-		if (ExpressionA->IsConstant() && ExpressionB->IsConstant())
-		{
-			FLinearColor ValueA, ValueB;
-			GetConstParameterValue(ExpressionA, ValueA);
-			GetConstParameterValue(ExpressionB, ValueB);
-
-			FLinearColor ConstantValue = ValueA - ValueB;
-			return ConstResultValue(ResultType, ConstantValue);
-		}
 		return AddUniformExpression(new FMaterialUniformExpressionFoldedMath(GetParameterUniformExpression(A), GetParameterUniformExpression(B), FMO_Sub), ResultType, TEXT("(%s - %s)"), *GetParameterCode(A), *GetParameterCode(B));
 	}
 	else
@@ -8323,7 +8332,6 @@ int32 FHLSLMaterialTranslator::Mul(int32 A,int32 B)
 	}
 	else if(IsExpressionConstantValue(A, 1.0f))
 	{
-		
 		ConstOneReturnValue = B;
 	}
 	else if (IsExpressionConstantValue(B, 1.0f))
@@ -8331,7 +8339,6 @@ int32 FHLSLMaterialTranslator::Mul(int32 A,int32 B)
 		ConstOneReturnValue = A;
 	}
 
-
 	if(ConstOneReturnValue != INDEX_NONE)
 	{
 		int32 Return = ConstOneReturnValue;
@@ -8340,7 +8347,7 @@ int32 FHLSLMaterialTranslator::Mul(int32 A,int32 B)
 		{
 			Return = AppendVector(Return, ConstOneReturnValue);
 		}
-		return Return;
+		return LWCCastIfNeccessary(ResultType, Return);
 	}
 
 
@@ -8394,7 +8401,7 @@ int32 FHLSLMaterialTranslator::Div(int32 A, int32 B)
 		{
 			Return = AppendVector(Return, A);
 		}
-		return Return;
+		return LWCCastIfNeccessary(ResultType, Return);
 	}
 
 	FMaterialUniformExpression* ExpressionA = GetParameterUniformExpression(A);
@@ -8441,43 +8448,19 @@ int32 FHLSLMaterialTranslator::Dot(int32 A,int32 B)
 		return INDEX_NONE;
 	}
 
+	EMaterialValueType TypeA = GetParameterType(A);
+	EMaterialValueType TypeB = GetParameterType(B);
+
 	if (IsExpressionConstantValue(A, 0.0f) || IsExpressionConstantValue(B, 0.0f))
 	{
-		return Constant(0.0f);
+		EMaterialValueType ResultType = (IsLWCType(TypeA) || IsLWCType(TypeB)) ? MCT_LWCScalar : MCT_Float;
+		return LWCCastIfNeccessary(ResultType, Constant(0.0f));	
 	}
 
 	FMaterialUniformExpression* ExpressionA = GetParameterUniformExpression(A);
 	FMaterialUniformExpression* ExpressionB = GetParameterUniformExpression(B);
-	EMaterialValueType TypeA = GetParameterType(A);
-	EMaterialValueType TypeB = GetParameterType(B);
 	if(ExpressionA && ExpressionB)
 	{
-		if (ExpressionA->IsConstant() && ExpressionB->IsConstant())
-		{
-			FLinearColor ValueA, ValueB;
-			GetConstMaskedParameterValue(TypeA, ExpressionA, ValueA);
-			GetConstMaskedParameterValue(TypeB, ExpressionB, ValueB);
-
-			// Promote scalar (or truncate the bigger type)
-			if (TypeA == MCT_Float || (TypeB != MCT_Float && GetNumComponents(TypeA) > GetNumComponents(TypeB)))
-			{
-				if(!CoerceConstantType(ValueA, TypeA, TypeB, ValueA))
-				{
-					return INDEX_NONE;
-				}
-			}
-			else
-			{
-				if(!CoerceConstantType(ValueB, TypeB, TypeA, ValueB))
-				{
-					return INDEX_NONE;
-				}
-			}
-
-			FVector4 VectorA(ValueA), VectorB(ValueB);
-			return ConstResultValue(MCT_Float, Dot4(VectorA, VectorB));
-		}
-
 		if (TypeA == MCT_Float && TypeB == MCT_Float)
 		{
 			return AddUniformExpression(new FMaterialUniformExpressionFoldedMath(ExpressionA,ExpressionB,FMO_Mul),MCT_Float,TEXT("(%s * %s)"),*GetParameterCode(A),*GetParameterCode(B));
@@ -8510,18 +8493,18 @@ int32 FHLSLMaterialTranslator::Dot(int32 A,int32 B)
 		}
 		else
 		{
-		// Promote scalar (or truncate the bigger type)
-		if (TypeA == MCT_Float || (TypeB != MCT_Float && GetNumComponents(TypeA) > GetNumComponents(TypeB)))
-		{
-			return AddCodeChunk(MCT_Float,TEXT("dot(%s, %s)"), *CoerceParameter(A, TypeB), *GetParameterCode(B));
-		}
-		else
-		{
-			return AddCodeChunk(MCT_Float,TEXT("dot(%s, %s)"), *GetParameterCode(A), *CoerceParameter(B, TypeA));
+			// Promote scalar (or truncate the bigger type)
+			if (TypeA == MCT_Float || (TypeB != MCT_Float && GetNumComponents(TypeA) > GetNumComponents(TypeB)))
+			{
+				return AddCodeChunk(MCT_Float,TEXT("dot(%s, %s)"), *CoerceParameter(A, TypeB), *GetParameterCode(B));
+			}
+			else
+			{
+				return AddCodeChunk(MCT_Float,TEXT("dot(%s, %s)"), *GetParameterCode(A), *CoerceParameter(B, TypeA));
+			}
 		}
 	}
 }
-}
 
 int32 FHLSLMaterialTranslator::Cross(int32 A,int32 B)
 {
@@ -8531,7 +8514,8 @@ int32 FHLSLMaterialTranslator::Cross(int32 A,int32 B)
 	}
 	EMaterialValueType ResultType = GetArithmeticResultType(A, B);
 
-	bool bIsValid = ResultType != MCT_Float2 && (ResultType & MCT_Float) != 0;
+	int ComponentCount = GetNumComponents(ResultType);
+	bool bIsValid = ComponentCount != 2 && ComponentCount != 1;
 
 	if(bIsValid && (IsExpressionConstantValue(A, 0.0f) || IsExpressionConstantValue(B, 0.0f)))
 	{
@@ -8547,26 +8531,6 @@ int32 FHLSLMaterialTranslator::Cross(int32 A,int32 B)
 			return Errorf(TEXT("Cross product requires 3-component vector input."));
 		}
 
-		if (ExpressionA->IsConstant() && ExpressionB->IsConstant())
-		{
-			EMaterialValueType TypeA = GetParameterType(A);
-			EMaterialValueType TypeB = GetParameterType(B);
-			FLinearColor ValueA, ValueB;
-			GetConstMaskedParameterValue(TypeA, ExpressionA, ValueA);
-			GetConstMaskedParameterValue(TypeB, ExpressionB, ValueB);
-
-			if(!CoerceConstantType(ValueA, TypeA, MCT_Float3, ValueA) || !CoerceConstantType(ValueB, TypeB, MCT_Float3, ValueB))
-			{
-				return INDEX_NONE;
-			}
-
-			FVector VectorA(ValueA);
-			FVector VectorB(ValueB);
-
-			FLinearColor Result(VectorA.Cross(VectorB));
-			return ConstResultValue(MCT_Float3, Result);
-		}
-
 		return AddUniformExpression(new FMaterialUniformExpressionFoldedMath(GetParameterUniformExpression(A),GetParameterUniformExpression(B),FMO_Cross,ResultType),MCT_Float3,TEXT("cross(%s,%s)"),*GetParameterCode(A),*GetParameterCode(B));
 	}
 	else if(IsAnalyticDerivEnabled())
@@ -8614,45 +8578,6 @@ int32 FHLSLMaterialTranslator::Power(int32 Base,int32 Exponent)
 		{
 			ConstantResult = Constant(0.0f);
 		}
-		else if(BaseExpression && BaseExpression->IsConstant())
-		{
-			EMaterialValueType BaseType = GetParameterType(Base);
-			EMaterialValueType ExponentType = GetParameterType(Exponent);
-			if (BaseType == ExponentType)
-			{
-				auto PositiveClampedPow = [](float BaseComponent, float ExponentComponent)
-				{
-					return (BaseComponent <= 0.0f) ? 0.0f : pow(BaseComponent, ExponentComponent);
-				};
-
-				FLinearColor BaseValue;
-				FLinearColor ExponentValue;
-				
-				GetConstParameterValue(BaseExpression, BaseValue);
-				GetConstParameterValue(ExponentExpression, ExponentValue);
-
-				//See note below on the issues with the two paths below for non const results.
-				// Currently mimicking that functionallity, but this shoul get changed when that does
-				FLinearColor Result;
-				FLinearColor ExponentConverted = ExponentValue;
-
-				if(!IsAnalyticDerivEnabled())
-				{
-					if (!CoerceConstantType(ExponentValue, MCT_Float, BaseType, ExponentConverted))
-					{
-						return INDEX_NONE;
-					}
-				} 
-
-				Result = FLinearColor(
-					PositiveClampedPow(BaseValue.R, ExponentConverted.R),
-					PositiveClampedPow(BaseValue.G, ExponentConverted.G),
-					PositiveClampedPow(BaseValue.B, ExponentConverted.B),
-					PositiveClampedPow(BaseValue.A, ExponentConverted.A));
-
-				return ConstResultValue(BaseType, Result);
-			}
-		}
 
 		if(ConstantResult != INDEX_NONE)
 		{
@@ -8670,7 +8595,7 @@ int32 FHLSLMaterialTranslator::Power(int32 Base,int32 Exponent)
 				Return = ComponentMask(Return, true, NumChannels > 1, NumChannels > 2, NumChannels > 3);
 			}
 
-			return Return;
+			return LWCCastIfNeccessary(ConstResultType, Return);
 		}
 	}
 
@@ -8696,22 +8621,6 @@ int32 FHLSLMaterialTranslator::Logarithm2(int32 X)
 		return INDEX_NONE;
 	}
 
-	FLinearColor ValueX;
-	FMaterialUniformExpression* ExpressionX = GetParameterUniformExpression(X);
-
-	if (GetConstParameterValue(ExpressionX, ValueX))
-	{
-		EMaterialValueType ResultType = GetParameterType(X);
-
-		FLinearColor Result(
-			FMath::Log2(ValueX.R),
-			FMath::Log2(ValueX.G),
-			FMath::Log2(ValueX.B),
-			FMath::Log2(ValueX.A));
-
-		return ConstResultValue(ResultType, Result);
-	}
-
 	if(GetParameterUniformExpression(X))
 	{
 		return AddUniformExpression(new FMaterialUniformExpressionLogarithm2(GetParameterUniformExpression(X)),GetParameterType(X),TEXT("log2(%s)"),*GetParameterCode(X));
@@ -8733,20 +8642,6 @@ int32 FHLSLMaterialTranslator::Logarithm10(int32 X)
 		return INDEX_NONE;
 	}
 
-	FLinearColor ValueX;
-	FMaterialUniformExpression* ExpressionX = GetParameterUniformExpression(X);
-	if (GetConstParameterValue(ExpressionX, ValueX))
-	{
-		EMaterialValueType ResultType = GetParameterType(X);
-		FLinearColor Result(
-			FMath::LogX(10, ValueX.R),
-			FMath::LogX(10, ValueX.G),
-			FMath::LogX(10, ValueX.B),
-			FMath::LogX(10, ValueX.A));
-
-		return ConstResultValue(ResultType, Result);
-	}
-		
 	if(GetParameterUniformExpression(X))
 	{
 		return AddUniformExpression(new FMaterialUniformExpressionLogarithm10(GetParameterUniformExpression(X)),GetParameterType(X),TEXT("log10(%s)"),*GetParameterCode(X));
@@ -8768,20 +8663,6 @@ int32 FHLSLMaterialTranslator::SquareRoot(int32 X)
 		return INDEX_NONE;
 	}
 
-	FLinearColor ValueX;
-	FMaterialUniformExpression* ExpressionX = GetParameterUniformExpression(X);
-	if (GetConstParameterValue(ExpressionX, ValueX))
-	{
-		EMaterialValueType ResultType = GetParameterType(X);
-		FLinearColor Result(
-			FMath::Sqrt(ValueX.R),
-			FMath::Sqrt(ValueX.G),
-			FMath::Sqrt(ValueX.B),
-			FMath::Sqrt(ValueX.A));
-
-		return ConstResultValue(ResultType, Result);
-	}
-
 	if(GetParameterUniformExpression(X))
 	{
 		return AddUniformExpression(new FMaterialUniformExpressionSquareRoot(GetParameterUniformExpression(X)),GetParameterType(X),TEXT("sqrt(%s)"),*GetParameterCode(X));
@@ -8806,21 +8687,6 @@ int32 FHLSLMaterialTranslator::Length(int32 X)
 		return INDEX_NONE;
 	}
 
-	FLinearColor ValueX;
-	FMaterialUniformExpression* ExpressionX = GetParameterUniformExpression(X);
-	if (GetConstParameterValue(ExpressionX, ValueX))
-	{
-		FVector4 Vector(ValueX);
-		float FloatResult = Vector.Size();
-		FLinearColor Result(
-			FloatResult,
-			FloatResult,
-			FloatResult,
-			FloatResult);
-
-		return ConstResultValue(MCT_Float, Result);
-	}
-
 	if (GetParameterUniformExpression(X))
 	{
 		return AddUniformExpression(new FMaterialUniformExpressionLength(GetParameterUniformExpression(X), GetParameterType(X)), MCT_Float, TEXT("length(%s)"), *GetParameterCode(X));
@@ -8847,15 +8713,6 @@ int32 FHLSLMaterialTranslator::Normalize(int32 X)
 
 	const EMaterialValueType ResultType = MakeNonLWCType(GetParameterType(X));
 
-	FLinearColor ValueX;
-	FMaterialUniformExpression* ExpressionX = GetParameterUniformExpression(X);
-	if (GetConstParameterValue(ExpressionX, ValueX))
-	{
-		FVector4 Vector(GetTypeMaskedValue(ResultType, ValueX, nullptr));
-		FLinearColor Result(Vector.GetSafeNormal());
-		return ConstResultValue(ResultType, Result);		
-	}
-
 	if (GetParameterUniformExpression(X))
 	{
 		return AddUniformExpression(new FMaterialUniformExpressionNormalize(GetParameterUniformExpression(X)), ResultType, TEXT("normalize(%s)"), *GetParameterCode(X));
diff --git a/Engine/Source/Runtime/Engine/Private/Materials/HLSLMaterialTranslator.h b/Engine/Source/Runtime/Engine/Private/Materials/HLSLMaterialTranslator.h
--- a/Engine/Source/Runtime/Engine/Private/Materials/HLSLMaterialTranslator.h
+++ b/Engine/Source/Runtime/Engine/Private/Materials/HLSLMaterialTranslator.h
@@ -740,6 +740,7 @@ protected:
 	// Returns false if the case fails
 	bool CoerceConstantType(FLinearColor SourceValue, EMaterialValueType SourceType, EMaterialValueType DestinationType, FLinearColor& OutResult);
 	bool CastConstantType(FLinearColor SourceValue, EMaterialValueType SourceType, EMaterialValueType DestinationType, EMaterialCastFlags Flags, FLinearColor& OutResult);
+	int32 LWCCastIfNeccessary(EMaterialValueType ResultType, int32 ResultCode);
 
 	// Returns a constant of the specified typ
 	int32 ConstResultValue(EMaterialValueType Type, FLinearColor ConstantValue);
