make UClusterUnionComponent::ComponentOverlapComponentWithResultImpl respect ignored actors and components

diff --git a/Engine/Source/Runtime/Engine/Private/PhysicsEngine/ClusterUnionComponent.cpp b/Engine/Source/Runtime/Engine/Private/PhysicsEngine/ClusterUnionComponent.cpp
--- a/Engine/Source/Runtime/Engine/Private/PhysicsEngine/ClusterUnionComponent.cpp
+++ b/Engine/Source/Runtime/Engine/Private/PhysicsEngine/ClusterUnionComponent.cpp
@@ -651,10 +651,38 @@ bool UClusterUnionComponent::OverlapComponentWithResult(const FVector& Pos, cons
 bool UClusterUnionComponent::ComponentOverlapComponentWithResultImpl(const class UPrimitiveComponent* const PrimComp, const FVector& Pos, const FQuat& Rot, const FCollisionQueryParams& Params, TArray<FOverlapResult>& OutOverlap) const
 {
 	bool bHasOverlap = false;
+
+	TSet<uint32> IgnoredActors;
+	IgnoredActors.Reserve(Params.GetIgnoredActors().Num());
+	for (uint32 Id : Params.GetIgnoredActors())
+	{
+		IgnoredActors.Add(Id);
+	}
+
+	TSet<uint32> IgnoredComponents;
+	IgnoredComponents.Reserve(Params.GetIgnoredComponents().Num());
+	for (uint32 Id : Params.GetIgnoredActors())
+	{
+		IgnoredComponents.Add(Id);
+	}
+
 	for (const TPair<TObjectKey<UPrimitiveComponent>, FClusteredComponentData>& Kvp : ComponentToPhysicsObjects)
 	{
 		if (UPrimitiveComponent* Component = Kvp.Key.ResolveObjectPtr())
 		{
+			if (IgnoredComponents.Contains(Component->GetUniqueID()))
+			{
+				continue;
+			}
+
+			if (AActor* Owner = Component->GetOwner())
+			{
+				if (IgnoredActors.Contains(Owner->GetUniqueID()))
+				{
+					continue;
+				}
+			}
+
 			TArray<FOverlapResult> SubOverlaps;
 			if (Component->ComponentOverlapComponentWithResult(PrimComp, Pos, Rot, Params, SubOverlaps))
 			{
