Fix missing emissive input on hair material with Substrate.

diff --git a/Engine/Shaders/Private/Strata/StrataExport.ush b/Engine/Shaders/Private/Strata/StrataExport.ush
--- a/Engine/Shaders/Private/Strata/StrataExport.ush
+++ b/Engine/Shaders/Private/Strata/StrataExport.ush
@@ -228,6 +228,7 @@ void PackStrataOut(
 						bIsOnlySlab = false;
 						bIsSimpleMaterial = false;
 						bIsSingleMaterial = false;
+						EmissiveColor += BSDF_GETEMISSIVE(BSDF) * BSDF.LuminanceWeightV;
 
 						TopLayerTotalWeight += BSDF.TopLayerDataWeight;
 						TopLayerData.Roughness += BSDF.TopLayerDataWeight * FullyRough;
diff --git a/Engine/Source/Runtime/Engine/Classes/Materials/Material.h b/Engine/Source/Runtime/Engine/Classes/Materials/Material.h
--- a/Engine/Source/Runtime/Engine/Classes/Materials/Material.h
+++ b/Engine/Source/Runtime/Engine/Classes/Materials/Material.h
@@ -462,12 +462,6 @@ class UMaterial : public UMaterialInterface
 	UPROPERTY(EditAnywhere, BlueprintReadOnly, Category=Nanite, meta = (EditInline, ShowOnlyInnerProperties))
 	FMaterialOverrideNanite NaniteOverrideMaterial;
 
-	/**
-	 * Cached connected inputs
-	 */
-	UPROPERTY()
-	uint32 CachedConnectedInputs;
-
 private:
 	/** Determines how inputs are combined to create the material's final color. */
 	UPROPERTY(EditAnywhere, Category=Material, AssetRegistrySearchable)
diff --git a/Engine/Source/Runtime/Engine/Private/Materials/Material.cpp b/Engine/Source/Runtime/Engine/Private/Materials/Material.cpp
--- a/Engine/Source/Runtime/Engine/Private/Materials/Material.cpp
+++ b/Engine/Source/Runtime/Engine/Private/Materials/Material.cpp
@@ -4811,8 +4811,30 @@ void UMaterial::RebuildShadingModelField()
 		check(EditorOnly->FrontMaterial.Expression);
 		if (EditorOnly->FrontMaterial.Expression->IsResultStrataMaterial(EditorOnly->FrontMaterial.OutputIndex))
 		{
+			// Mask of all input collected by Strata BSDF nodes
+			static const uint32 ConnectionMask = 
+				  (1u << MP_BaseColor)
+				| (1u << MP_Metallic)
+				| (1u << MP_Specular)
+				| (1u << MP_Roughness)
+				| (1u << MP_Anisotropy)
+				| (1u << MP_EmissiveColor)
+				| (1u << MP_Normal)
+				| (1u << MP_Tangent)
+				| (1u << MP_SubsurfaceColor)
+				| (1u << MP_CustomData0)
+				| (1u << MP_CustomData1)
+				| (1u << MP_Opacity)
+				| (1u << MP_ShadingModel)
+				| (1u << MP_DiffuseColor)
+				| (1u << MP_SpecularColor);
+
 			EditorOnly->FrontMaterial.Expression->GatherStrataMaterialInfo(StrataMaterialInfo, EditorOnly->FrontMaterial.OutputIndex);
-			this->CachedConnectedInputs = StrataMaterialInfo.GetPropertyConnected();
+
+			// Override the cached expression data with collected connection from StrataMaterialInfo, but preserve all other input (e.g., refraction)
+			check(this->CachedExpressionData);
+			this->CachedExpressionData->PropertyConnectedBitmask &= ~ConnectionMask;
+			this->CachedExpressionData->PropertyConnectedBitmask |= StrataMaterialInfo.GetPropertyConnected();
 
 			if (StrataMaterialInfo.GetStrataTreeOutOfStackDepthOccurred())
 			{
diff --git a/Engine/Source/Runtime/Engine/Private/Materials/MaterialShared.cpp b/Engine/Source/Runtime/Engine/Private/Materials/MaterialShared.cpp
--- a/Engine/Source/Runtime/Engine/Private/Materials/MaterialShared.cpp
+++ b/Engine/Source/Runtime/Engine/Private/Materials/MaterialShared.cpp
@@ -1971,7 +1971,7 @@ bool FMaterialResource::HasMaterialPropertyConnected(EMaterialProperty In) const
 			return Material->HasAmbientOcclusionConnected();
 		}
 		// Strata material traversal is cached as this is an expensive operation
-		return FStrataMaterialInfo::HasPropertyConnected(Material->CachedConnectedInputs, In);
+		return FStrataMaterialInfo::HasPropertyConnected(Material->GetCachedExpressionData().PropertyConnectedBitmask, In);
 	}
 	else
 	{
