GetRenderingDetailMode: Fix bad clamp

diff --git a/Engine/Source/Runtime/Engine/Private/KismetSystemLibrary.cpp b/Engine/Source/Runtime/Engine/Private/KismetSystemLibrary.cpp
--- a/Engine/Source/Runtime/Engine/Private/KismetSystemLibrary.cpp
+++ b/Engine/Source/Runtime/Engine/Private/KismetSystemLibrary.cpp
@@ -2562,7 +2562,7 @@ int32 UKismetSystemLibrary::GetRenderingDetailMode()
 	static const IConsoleVariable* CVar = IConsoleManager::Get().FindConsoleVariable(TEXT("r.DetailMode"));
 
 	// clamp range
-	int32 Ret = FMath::Clamp(CVar->GetInt(), 0, 2);
+	int32 Ret = FMath::Clamp(CVar->GetInt(), 0, DM_MAX - 1);
 
 	return Ret;
 }
