UWorldPartitionRuntimeLevelStreamingCell::GetOrCreateLevelStreaming check temporary disabled while it's investigated why it's being hit

diff --git a/Engine/Source/Runtime/Engine/Private/WorldPartition/WorldPartitionRuntimeLevelStreamingCell.cpp b/Engine/Source/Runtime/Engine/Private/WorldPartition/WorldPartitionRuntimeLevelStreamingCell.cpp
--- a/Engine/Source/Runtime/Engine/Private/WorldPartition/WorldPartitionRuntimeLevelStreamingCell.cpp
+++ b/Engine/Source/Runtime/Engine/Private/WorldPartition/WorldPartitionRuntimeLevelStreamingCell.cpp
@@ -367,7 +367,7 @@ UWorldPartitionLevelStreamingDynamic* UWorldPartitionRuntimeLevelStreamingCell::
 		// 
 		// and the LevelStreaming's WorldAsset TSoftObjectPtr<UWorld> should be remapped as: 
 		// /Game/SomePath/MainWorld/_Generated_/1AEHCD0PRR98XVM12PU4N1D8X_LevelInstance_1
-		check(!WorldPartition->GetTypedOuter<ULevel>()->IsInstancedLevel() || (LevelStreaming->PackageNameToLoad != LevelStreaming->GetWorldAssetPackageName()));
+		//check(!WorldPartition->GetTypedOuter<ULevel>()->IsInstancedLevel() || (LevelStreaming->PackageNameToLoad != LevelStreaming->GetWorldAssetPackageName()));
 	}
 #endif
 
