World Partition: Support reinit of main world partition

diff --git a/Engine/Source/Runtime/Engine/Private/WorldPartition/WorldPartition.cpp b/Engine/Source/Runtime/Engine/Private/WorldPartition/WorldPartition.cpp
--- a/Engine/Source/Runtime/Engine/Private/WorldPartition/WorldPartition.cpp
+++ b/Engine/Source/Runtime/Engine/Private/WorldPartition/WorldPartition.cpp
@@ -656,6 +656,11 @@ void UWorldPartition::Uninitialize()
 
 		InitState = EWorldPartitionInitState::Uninitializing;
 
+		if (IsMainWorldPartition())
+		{
+			AWorldPartitionReplay::Uninitialize(World);
+		}
+
 		UnregisterDelegates();
 		
 		// Unload all loaded cells
diff --git a/Engine/Source/Runtime/Engine/Private/WorldPartition/WorldPartitionReplay.cpp b/Engine/Source/Runtime/Engine/Private/WorldPartition/WorldPartitionReplay.cpp
--- a/Engine/Source/Runtime/Engine/Private/WorldPartition/WorldPartitionReplay.cpp
+++ b/Engine/Source/Runtime/Engine/Private/WorldPartition/WorldPartitionReplay.cpp
@@ -109,6 +109,20 @@ void AWorldPartitionReplay::Initialize(UWorld* World)
 	}
 }
 
+void AWorldPartitionReplay::Uninitialize(UWorld* World)
+{
+	if (World->IsGameWorld())
+	{
+		UWorldPartition* WorldPartition = World->GetWorldPartition();
+		if (ensure(WorldPartition && WorldPartition->Replay))
+		{
+			World->DestroyActor(WorldPartition->Replay);
+			WorldPartition->Replay->Rename(nullptr, GetTransientPackage(), REN_DontCreateRedirectors | REN_NonTransactional);
+			WorldPartition->Replay = nullptr;
+		}
+	}
+}
+
 bool AWorldPartitionReplay::IsPlaybackEnabled(UWorld* World)
 {
 	if (GUseReplayStreamingSources && World->IsPlayingReplay())
diff --git a/Engine/Source/Runtime/Engine/Public/WorldPartition/WorldPartitionReplay.h b/Engine/Source/Runtime/Engine/Public/WorldPartition/WorldPartitionReplay.h
--- a/Engine/Source/Runtime/Engine/Public/WorldPartition/WorldPartitionReplay.h
+++ b/Engine/Source/Runtime/Engine/Public/WorldPartition/WorldPartitionReplay.h
@@ -51,6 +51,7 @@ class ENGINE_API AWorldPartitionReplay : public AActor
 
 public:
 	static void Initialize(UWorld* World);
+	static void Uninitialize(UWorld* World);
 	static bool IsPlaybackEnabled(UWorld* World);
 	static bool IsRecordingEnabled(UWorld* World);
 
