Add a gamepad and touch default primary device type

#rnx

diff --git a/Engine/Source/Runtime/Engine/Classes/GameFramework/InputSettings.h b/Engine/Source/Runtime/Engine/Classes/GameFramework/InputSettings.h
--- a/Engine/Source/Runtime/Engine/Classes/GameFramework/InputSettings.h
+++ b/Engine/Source/Runtime/Engine/Classes/GameFramework/InputSettings.h
@@ -444,7 +444,14 @@ struct ENGINE_API FHardwareDeviceIdentifier
 	/** Hardware device ID that represents a keyboard and mouse. This is what will be set when an Input Event's FKey is not a gamepad key. */
 	static FHardwareDeviceIdentifier DefaultKeyboardAndMouse;
 
+	/** Hardware device ID that represents a default, generic, gamepad. */
+	static FHardwareDeviceIdentifier DefaultGamepad;
+
+	/** Hardware device id that represents a default, generic, mobile touch input (tablet, phone, etc) */
+	static FHardwareDeviceIdentifier DefaultMobileTouch;
+
 	bool operator==(const FHardwareDeviceIdentifier& Other) const;
+	bool operator!=(const FHardwareDeviceIdentifier& Other) const;
 
 	ENGINE_API friend uint32 GetTypeHash(const FHardwareDeviceIdentifier& InDevice);
 	ENGINE_API friend FArchive& operator<<(FArchive& Ar, FHardwareDeviceIdentifier& InDevice);
diff --git a/Engine/Source/Runtime/Engine/Private/UserInterface/InputSettings.cpp b/Engine/Source/Runtime/Engine/Private/UserInterface/InputSettings.cpp
--- a/Engine/Source/Runtime/Engine/Private/UserInterface/InputSettings.cpp
+++ b/Engine/Source/Runtime/Engine/Private/UserInterface/InputSettings.cpp
@@ -522,6 +522,22 @@ FHardwareDeviceIdentifier FHardwareDeviceIdentifier::DefaultKeyboardAndMouse =
 		/* .SupportedFeaturesMask= */ EHardwareDeviceSupportedFeatures::Type::Keypress | EHardwareDeviceSupportedFeatures::Type::Pointer
 	};
 
+FHardwareDeviceIdentifier FHardwareDeviceIdentifier::DefaultGamepad = 
+	{
+	/* .InputClassName= */ TEXT("DefaultGamepad"),
+	/* .HardwareDeviceIdentifier= */ TEXT("Gamepad"),
+	/* .PrimaryDeviceType= */ EHardwareDevicePrimaryType::Gamepad,
+	/* .SupportedFeaturesMask= */  EHardwareDeviceSupportedFeatures::Type::Gamepad
+};
+
+FHardwareDeviceIdentifier FHardwareDeviceIdentifier::DefaultMobileTouch = 
+	{
+	/* .InputClassName= */ TEXT("DefaultMobileTouch"),
+	/* .HardwareDeviceIdentifier= */ TEXT("MobileTouch"),
+	/* .PrimaryDeviceType= */ EHardwareDevicePrimaryType::Touch,
+	/* .SupportedFeaturesMask= */ EHardwareDeviceSupportedFeatures::Type::Touch
+};
+
 // Default to the invalid hardware device identifier
 FHardwareDeviceIdentifier::FHardwareDeviceIdentifier()
 	: InputClassName(Invalid.InputClassName)
@@ -585,6 +601,11 @@ bool FHardwareDeviceIdentifier::operator==(const FHardwareDeviceIdentifier& Othe
 		Other.PrimaryDeviceType == PrimaryDeviceType;	
 }
 
+bool FHardwareDeviceIdentifier::operator!=(const FHardwareDeviceIdentifier& Other) const
+{
+	return !(Other == *this);
+}
+
 //////////////////////////////////////////////////////////////////
 // UInputPlatformSettings
 
@@ -594,7 +615,14 @@ UInputPlatformSettings::UInputPlatformSettings()
 	, MaxTriggerVibrationTriggerPosition(9)
 	, MaxTriggerVibrationFrequency(255)
 	, MaxTriggerVibrationAmplitude(8)
-{ }
+{
+	// Add the default invalid and KBM hardware device ID's here so that they will
+	// be found if no custom devices exist when you call GetHardwareDeviceForClassName
+	HardwareDevices.AddUnique(FHardwareDeviceIdentifier::Invalid);
+	HardwareDevices.AddUnique(FHardwareDeviceIdentifier::DefaultKeyboardAndMouse);
+	HardwareDevices.AddUnique(FHardwareDeviceIdentifier::DefaultGamepad);
+	HardwareDevices.AddUnique(FHardwareDeviceIdentifier::DefaultMobileTouch);
+}
 
 UInputPlatformSettings* UInputPlatformSettings::Get()
 {
