Add validation for skeletal mesh and for anim sequence.

#rnx

diff --git a/Engine/Source/Runtime/Engine/Classes/Animation/AnimSequence.h b/Engine/Source/Runtime/Engine/Classes/Animation/AnimSequence.h
--- a/Engine/Source/Runtime/Engine/Classes/Animation/AnimSequence.h
+++ b/Engine/Source/Runtime/Engine/Classes/Animation/AnimSequence.h
@@ -390,6 +390,7 @@ public:
 	virtual bool IsCachedCookedPlatformDataLoaded(const ITargetPlatform* TargetPlatform) override;
 	virtual void WillNeverCacheCookedPlatformDataAgain() override;
 	virtual void ClearAllCachedCookedPlatformData() override;
+	virtual EDataValidationResult IsDataValid(TArray<FText>& ValidationErrors) override;
 #endif // WITH_EDITOR
 	virtual void BeginDestroy() override;
 	virtual bool IsReadyForFinishDestroy() override;
diff --git a/Engine/Source/Runtime/Engine/Classes/Engine/SkeletalMesh.h b/Engine/Source/Runtime/Engine/Classes/Engine/SkeletalMesh.h
--- a/Engine/Source/Runtime/Engine/Classes/Engine/SkeletalMesh.h
+++ b/Engine/Source/Runtime/Engine/Classes/Engine/SkeletalMesh.h
@@ -2428,6 +2428,8 @@ public:
 
 	void UpdateGenerateUpToData();
 
+	virtual EDataValidationResult IsDataValid(class FDataValidationContext& Context) override;
+
 #endif // WITH_EDITOR
 	virtual void BeginDestroy() override;
 	virtual bool IsReadyForFinishDestroy() override;
diff --git a/Engine/Source/Runtime/Engine/Private/Animation/AnimSequence.cpp b/Engine/Source/Runtime/Engine/Private/Animation/AnimSequence.cpp
--- a/Engine/Source/Runtime/Engine/Private/Animation/AnimSequence.cpp
+++ b/Engine/Source/Runtime/Engine/Private/Animation/AnimSequence.cpp
@@ -943,6 +943,21 @@ void ShowResaveMessage(const UAnimSequence* Sequence)
 	}
 }
 
+EDataValidationResult UAnimSequence::IsDataValid(TArray<FText>& ValidationErrors)
+{
+	EDataValidationResult ValidationResult = Super::IsDataValid(ValidationErrors);
+	// Do not validate cooked anim sequence
+	if (GetPackage()->HasAnyPackageFlags(PKG_Cooked) == false)
+	{
+		if (!GetSkeleton())
+		{
+			ValidationErrors.Add(LOCTEXT("AnimSequenceValidation_NoSkeleton", "This anim sequence asset has no Skeleton. Anim sequence asset need a valid skeleton."));
+			ValidationResult = EDataValidationResult::Invalid;
+		}
+	}
+	return ValidationResult;
+}
+
 #endif // WITH_EDITOR
 void UAnimSequence::BeginDestroy()
 {
@@ -1087,6 +1102,7 @@ bool UAnimSequence::IsCachedCookedPlatformDataLoaded(const ITargetPlatform* Targ
 		}
 	}
 
+
 	return false;
 }
 #endif // WITH_EDITOR
diff --git a/Engine/Source/Runtime/Engine/Private/SkeletalMesh.cpp b/Engine/Source/Runtime/Engine/Private/SkeletalMesh.cpp
--- a/Engine/Source/Runtime/Engine/Private/SkeletalMesh.cpp
+++ b/Engine/Source/Runtime/Engine/Private/SkeletalMesh.cpp
@@ -64,6 +64,7 @@
 #include "DerivedDataCacheInterface.h"
 #include "Engine/RendererSettings.h"
 #include "SkeletalDebugRendering.h"
+#include "Misc/DataValidation.h"
 #else
 #include "Interfaces/ITargetPlatform.h"
 #endif // #if WITH_EDITOR
@@ -1459,6 +1460,64 @@ void USkeletalMesh::UpdateGenerateUpToData()
 	}
 }
 
+EDataValidationResult USkeletalMesh::IsDataValid(class FDataValidationContext& Context)
+{
+	EDataValidationResult ValidationResult = Super::IsDataValid(Context);
+	// Do not validate a cooked skeletal mesh asset.
+	if (GetPackage()->HasAnyPackageFlags(PKG_Cooked) == false)
+	{
+		if (!GetSkeleton())
+		{
+			//We must have a valid skeleton
+			Context.AddError(LOCTEXT("SkeletalMeshValidation_NoSkeleton", "This skeletal mesh asset has no Skeleton. Skeletal mesh asset need a valid skeleton."));
+			ValidationResult = EDataValidationResult::Invalid;
+		}
+		else
+		{
+			//Validate if the skeleton is compatible with this skeletal mesh
+			if (!GetSkeleton()->IsCompatibleMesh(this))
+			{
+				Context.AddError(LOCTEXT("SkeletalMeshValidation_IncompatibleSkeleton", "This skeletal mesh asset has an incompatible Skeleton. Assign a compatible skeleton to this skeletal mesh asset."));
+				ValidationResult = EDataValidationResult::Invalid;
+			}
+		}
+		
+		const int32 NumRealBones = GetRefSkeleton().GetRawBoneNum();
+		const TArray<FTransform>& RawRefBonePose = GetRefSkeleton().GetRawRefBonePose();
+
+		// Precompute the Mesh.RefBasesInverse.
+		for (int32 BoneIndex = 0; BoneIndex < NumRealBones; BoneIndex++)
+		{
+			//Validate skeleton bone index
+			if (!GetRefSkeleton().IsValidRawIndex(BoneIndex))
+			{
+				Context.AddError(LOCTEXT("SkeletalMeshValidation_InvalidBoneIndex", "This skeletal mesh asset has invalid bone index. Asset is corrupted and must be re-create"));
+				ValidationResult = EDataValidationResult::Invalid;
+			}
+
+			//Validate Parent bone index
+			if (BoneIndex > 0)
+			{
+				int32 Parent = GetRefSkeleton().GetRawParentIndex(BoneIndex);
+				if (!GetRefSkeleton().IsValidRawIndex(Parent))
+				{
+					Context.AddError(LOCTEXT("SkeletalMeshValidation_InvalidParentBoneIndex", "This skeletal mesh asset has invalid parent bone index. Asset is corrupted and must be re-create"));
+					ValidationResult = EDataValidationResult::Invalid;
+				}
+			}
+
+			//Validate transform do not contains nan
+			FMatrix RefPoseMatrix = GetRefPoseMatrix(BoneIndex);
+			if (RefPoseMatrix.ContainsNaN())
+			{
+				Context.AddError(LOCTEXT("SkeletalMeshValidation_PoseMatrixContainNan", "This skeletal mesh asset has NAN (invalid float number) value in the pose matrix. Asset is corrupted and must be re-create"));
+				ValidationResult = EDataValidationResult::Invalid;
+			}
+		}
+	}
+	return ValidationResult;
+}
+
 #endif // WITH_EDITOR
 
 void USkeletalMesh::BeginDestroy()
@@ -4404,16 +4463,9 @@ void USkeletalMesh::ValidateBoneWeights(const ITargetPlatform* TargetPlatform)
 			const FSkeletalMeshLODModel& ImportLODModel = GetImportedModel()->LODModels[LODIndex];
 
 			int32 MaxBoneInfluences = ImportLODModel.GetMaxBoneInfluences();
-
-			for (int32 SectionIndex = 0; SectionIndex < ImportLODModel.Sections.Num(); ++SectionIndex)
+			if (MaxBoneInfluences > 12 )
 			{
-				const FSkelMeshSection & Section = ImportLODModel.Sections[SectionIndex];
-
-				int MaxBoneInfluencesSection = Section.MaxBoneInfluences;
-				if (MaxBoneInfluences > 12)
-				{
-					UE_LOG(LogSkeletalMesh, Warning, TEXT("Mesh: %s has more than 12 max bone influences, it has: %d"), *GetFullName(), MaxBoneInfluencesSection);
-				}
+				UE_LOG(LogSkeletalMesh, Warning, TEXT("Mesh: %s has more than 12 max bone influences, it has: %d"), *GetFullName(), MaxBoneInfluences);
 			}
 		}
 	}
