LogTextureUpload better log of built texture info

diff --git a/Engine/Source/Runtime/Engine/Private/TextureDerivedDataTask.cpp b/Engine/Source/Runtime/Engine/Private/TextureDerivedDataTask.cpp
--- a/Engine/Source/Runtime/Engine/Private/TextureDerivedDataTask.cpp
+++ b/Engine/Source/Runtime/Engine/Private/TextureDerivedDataTask.cpp
@@ -550,10 +550,16 @@ static void DDC1_StoreClassicTextureInDerivedData(
 		// CompressedDataSize can exceed int32 ; eg. 16k x 16k x RGBA16F == 2 GB
 		// DDC1 should be 64-bit safe now
 
-		// CompressedImage sizes were padded up to multiple of 4 for d3d, no longer
-		UE_LOG(LogTextureUpload, Verbose, TEXT("Compressed Mip %d PF=%d : %dx%dx%d : %d ; up4 %dx%d=%d"),
+		// CompressedImage sizes were padded up to multiple of 4 for d3d, no longer ; log the align-up-to-4 sizes for debugging?
+
+		// log what the TextureFormat built :
+		//	(todo: change to "LogTexture" instead of "LogTextureUpload" and remove the align-up-to-4 debug logs)
+		UE_LOG(LogTextureUpload, Verbose, TEXT("Built texture: %s Compressed Mip %d PF=%d=%s : %dx%dx%d : %lld ; up4 %dx%d=%d"),
+			*TexturePathName,
 			MipIndex, (int)CompressedImage.PixelFormat,
-			CompressedImage.SizeX, CompressedImage.SizeY, CompressedImage.SizeZ, (int)CompressedDataSize,
+			GetPixelFormatString((EPixelFormat)CompressedImage.PixelFormat),
+			CompressedImage.SizeX, CompressedImage.SizeY, CompressedImage.SizeZ, 
+			CompressedDataSize,
 			(CompressedImage.SizeX + 3) & (~3),
 			(CompressedImage.SizeY + 3) & (~3),
 			((CompressedImage.SizeX + 3) & (~3)) * ((CompressedImage.SizeY + 3) & (~3)));
@@ -637,6 +643,7 @@ static void DDC1_BuildTexture(
 		return;
 	}
 
+	// this logs the "Building textures: " message :
 	FTextureStatusMessageContext StatusMessage(
 		ComposeTextureBuildText(TexturePathName, TextureData, InBuildSettingsPerLayer[0], (ETextureEncodeSpeed)InBuildSettingsPerLayer[0].RepresentsEncodeSpeedNoSend, RequiredMemoryEstimate, bForVirtualTextureStreamingBuild)
 		);
