Lightweight Instance Manager grid size can now be overidden per manager class
* A GetGridSize() virtual method was added
* Exposed two helper methods for grid coords/bounds retrieval: ConvertPositionToCoord() / ConvertPositionToGridBounds()

diff --git a/Engine/Source/Runtime/Engine/Classes/GameFramework/LightWeightInstanceManager.h b/Engine/Source/Runtime/Engine/Classes/GameFramework/LightWeightInstanceManager.h
--- a/Engine/Source/Runtime/Engine/Classes/GameFramework/LightWeightInstanceManager.h
+++ b/Engine/Source/Runtime/Engine/Classes/GameFramework/LightWeightInstanceManager.h
@@ -105,6 +105,15 @@ public:
 	// Returns the index used by collision and rendering that is associated with the instance referred to by InIndex
 	virtual int32 ConvertLightWeightIndexToCollisionIndex(int32 InIndex) const;
 
+	// LWI grid size to use for this manager.
+	virtual int32 GetGridSize() const;
+
+	// Helper that converts a position (world space) into a coordinate for the LWI grid.
+	FInt32Vector3 ConvertPositionToCoord(const FVector& InPosition) const;
+
+	// Helper that retrieve the world space bounds of the LWI grid cell encompassing the provided coordinate.
+	FBox ConvertPositionToGridBounds(const FVector& InPosition) const;
+
 	template<typename U>
 	bool IsInterfaceSupported() const
 	{
diff --git a/Engine/Source/Runtime/Engine/Classes/GameFramework/LightWeightInstanceSubsystem.h b/Engine/Source/Runtime/Engine/Classes/GameFramework/LightWeightInstanceSubsystem.h
--- a/Engine/Source/Runtime/Engine/Classes/GameFramework/LightWeightInstanceSubsystem.h
+++ b/Engine/Source/Runtime/Engine/Classes/GameFramework/LightWeightInstanceSubsystem.h
@@ -100,7 +100,8 @@ struct ENGINE_API FLightWeightInstanceSubsystem
 	}
 
 	// Helper that converts a position (world space) into a coordinate for the LWI grid.
-	static FInt32Vector3 ConvertPositionToCoord(const FVector & InPosition);
+	UE_DEPRECATED(5.3, "Use LWI Managers version of ConvertPositionToCoord()")
+	static FInt32Vector3 ConvertPositionToCoord(const FVector& InPosition);
 
 	// Add a manager to the subsystem, thread safe.
 	bool AddManager(ALightWeightInstanceManager* Manager);
@@ -126,10 +127,6 @@ private:
 	UPROPERTY()
 	TArray<ALightWeightInstanceManager*> LWInstanceManagers;
 
-	// CVar variable that is the size of the grid managers are placed into.
-	static int32 LWIGridSize;
-	static FAutoConsoleVariableRef CVarLWIGridSize;
-
 	// Mutex to make sure we don't change the LWInstanceManagers array while reading/writing it.
 	mutable FRWLock LWIManagersRWLock;
 
diff --git a/Engine/Source/Runtime/Engine/Private/LightWeightInstanceManager.cpp b/Engine/Source/Runtime/Engine/Private/LightWeightInstanceManager.cpp
--- a/Engine/Source/Runtime/Engine/Private/LightWeightInstanceManager.cpp
+++ b/Engine/Source/Runtime/Engine/Private/LightWeightInstanceManager.cpp
@@ -8,6 +8,17 @@
 
 #include UE_INLINE_GENERATED_CPP_BY_NAME(LightWeightInstanceManager)
 
+// CVar variable that is the size of the grid managers are placed into.
+static int32 LWIGridSize = -1;
+static FAutoConsoleVariableRef CVarLWIGridSize
+(
+	TEXT("LWI.Editor.GridSize"),
+	LWIGridSize,
+	TEXT("Sets the size of a grid that LWI managers will be generated with."),
+	ECVF_Default
+);
+
+
 UActorInstanceHandleInterface::UActorInstanceHandleInterface(const FObjectInitializer& ObjectInitializer)
 	:Super(ObjectInitializer)
 {
@@ -201,6 +212,42 @@ int32 ALightWeightInstanceManager::ConvertLightWeightIndexToCollisionIndex(int32
 	return InIndex;
 }
 
+int32 ALightWeightInstanceManager::GetGridSize() const
+{
+	return LWIGridSize;	
+}
+
+FInt32Vector3 ALightWeightInstanceManager::ConvertPositionToCoord(const FVector& InPosition) const
+{
+	float GridSize = GetGridSize();
+	if (GridSize > 0)
+	{
+		return FInt32Vector3(
+			FMath::FloorToInt(InPosition.X / GridSize),
+			FMath::FloorToInt(InPosition.Y / GridSize),
+			FMath::FloorToInt(InPosition.Z / GridSize)
+		);
+	}
+
+	return FInt32Vector3::ZeroValue;
+}
+
+FBox ALightWeightInstanceManager::ConvertPositionToGridBounds(const FVector& InPosition) const
+{
+	float GridSize = GetGridSize();
+	if (GridSize > 0)
+	{
+		const FInt32Vector3 GridCoord = ConvertPositionToCoord(InPosition);
+
+		FBox GridBounds;
+		GridBounds.Min = FVector(GridCoord * GridSize);
+		GridBounds.Max = FVector(GridCoord * GridSize + FInt32Vector3(GridSize));
+		return GridBounds;
+	}
+
+	return FBox(InPosition, InPosition);
+}
+
 void ALightWeightInstanceManager::SetSpawnParameters(FActorSpawnParameters& SpawnParams)
 {
 	SpawnParams.OverrideLevel = GetLevel();
diff --git a/Engine/Source/Runtime/Engine/Private/LightWeightInstanceSubsystem.cpp b/Engine/Source/Runtime/Engine/Private/LightWeightInstanceSubsystem.cpp
--- a/Engine/Source/Runtime/Engine/Private/LightWeightInstanceSubsystem.cpp
+++ b/Engine/Source/Runtime/Engine/Private/LightWeightInstanceSubsystem.cpp
@@ -10,16 +10,6 @@ DEFINE_LOG_CATEGORY(LogLightWeightInstance);
 TSharedPtr<FLightWeightInstanceSubsystem> FLightWeightInstanceSubsystem::LWISubsystem;
 FCriticalSection FLightWeightInstanceSubsystem::GetFunctionCS;
 
-// CVar variable that is the size of the grid managers are placed into.
-int32 FLightWeightInstanceSubsystem::LWIGridSize = -1;
-FAutoConsoleVariableRef FLightWeightInstanceSubsystem::CVarLWIGridSize
-(
-TEXT("LWI.Editor.GridSize"),
-FLightWeightInstanceSubsystem::LWIGridSize,
-TEXT("Sets the size of a grid that LWI managers will be generated with."),
-ECVF_Cheat
-);
-
 FLightWeightInstanceSubsystem::FLightWeightInstanceSubsystem()
 {
 #if WITH_EDITOR
@@ -86,21 +76,6 @@ bool FLightWeightInstanceSubsystem::RemoveManager(ALightWeightInstanceManager* M
 	return Get().LWInstanceManagers.Remove(Manager) ? true : false;
 }
 
-FInt32Vector3 FLightWeightInstanceSubsystem::ConvertPositionToCoord(const FVector & InPosition)
-{
-	const int32 GridSize = LWIGridSize;
-	if (GridSize > 0)
-	{
-		return FInt32Vector3(
-			FMath::FloorToInt(InPosition.X / (float)GridSize),
-			FMath::FloorToInt(InPosition.Y / (float)GridSize),
-			FMath::FloorToInt(InPosition.Z / (float)GridSize)
-			);
-	}
-
-	return FInt32Vector3::ZeroValue;
-}
-
 ALightWeightInstanceManager* FLightWeightInstanceSubsystem::FindLightWeightInstanceManager(const FActorInstanceHandle& Handle) const
 {
 	if (Handle.Manager.IsValid())
@@ -110,27 +85,29 @@ ALightWeightInstanceManager* FLightWeightInstanceSubsystem::FindLightWeightInsta
 
 	if (Handle.Actor.IsValid())
 	{
-		const FInt32Vector3 GridCoord = ConvertPositionToCoord(Handle.GetLocation());
-
 		FReadScopeLock Lock(LWIManagersRWLock);
 		// see if we already have a match
 		for (ALightWeightInstanceManager* LWInstance : LWInstanceManagers)
 		{
-			const FInt32Vector3 ManagerGridCoord = ConvertPositionToCoord(LWInstance->GetActorLocation());
-			if (Handle.Actor->GetClass() == LWInstance->GetRepresentedClass() && ManagerGridCoord == GridCoord)
+			if (Handle.Actor->GetClass() == LWInstance->GetRepresentedClass())
 			{
+				const FInt32Vector3 GridCoord = LWInstance->ConvertPositionToCoord(Handle.GetLocation());
+				const FInt32Vector3 ManagerGridCoord = LWInstance->ConvertPositionToCoord(LWInstance->GetActorLocation());
+				if (ManagerGridCoord == GridCoord)
+				{
 #if WITH_EDITOR
-				// make sure the data layers match
-				TArray<const UDataLayerInstance*> ManagerLayers = LWInstance->GetDataLayerInstances();
-				const UDataLayerInstance* ManagerLayer = ManagerLayers.Num() > 0 ? ManagerLayers[0] : nullptr;
+					// make sure the data layers match
+					TArray<const UDataLayerInstance*> ManagerLayers = LWInstance->GetDataLayerInstances();
+					const UDataLayerInstance* ManagerLayer = ManagerLayers.Num() > 0 ? ManagerLayers[0] : nullptr;
 
-				TArray<const UDataLayerInstance*> ActorLayers = LWInstance->GetDataLayerInstances();
-				const UDataLayerInstance* ActorLayer = ActorLayers.Num() > 0 ? ActorLayers[0] : nullptr;
+					TArray<const UDataLayerInstance*> ActorLayers = LWInstance->GetDataLayerInstances();
+					const UDataLayerInstance* ActorLayer = ActorLayers.Num() > 0 ? ActorLayers[0] : nullptr;
 
-				if (ManagerLayer == ActorLayer)
+					if (ManagerLayer == ActorLayer)
 #endif // WITH_EDITOR
-				{
-					return LWInstance;
+					{
+						return LWInstance;
+					}
 				}
 			}
 		}
@@ -155,22 +132,25 @@ ALightWeightInstanceManager* FLightWeightInstanceSubsystem::FindOrAddLightWeight
 
 ALightWeightInstanceManager* FLightWeightInstanceSubsystem::FindLightWeightInstanceManager(UClass& ActorClass, UWorld& World, const FVector& InPos, const UDataLayerInstance* DataLayer) const
 {
-	const FInt32Vector3 GridCoord = ConvertPositionToCoord(InPos);
 	FReadScopeLock Lock(LWIManagersRWLock);
 
 	for (ALightWeightInstanceManager* InstanceManager : LWInstanceManagers)
 	{
 		if (IsValid(InstanceManager))
 		{
-			const FInt32Vector3 ManagerGridCoord = ConvertPositionToCoord(InstanceManager->GetActorLocation());
-			
-			if (InstanceManager->GetRepresentedClass() == &ActorClass && ManagerGridCoord == GridCoord)
+			if (InstanceManager->GetRepresentedClass() == &ActorClass)
 			{
-	#if WITH_EDITOR
-				if (!DataLayer || (InstanceManager->SupportsDataLayer() && InstanceManager->ContainsDataLayer(DataLayer)))
-	#endif // WITH_EDITOR
+				const FInt32Vector3 GridCoord = InstanceManager->ConvertPositionToCoord(InPos);
+				const FInt32Vector3 ManagerGridCoord = InstanceManager->ConvertPositionToCoord(InstanceManager->GetActorLocation());
+
+				if (ManagerGridCoord == GridCoord)
 				{
-					return InstanceManager;
+#if WITH_EDITOR
+					if (!DataLayer || (InstanceManager->SupportsDataLayer() && InstanceManager->ContainsDataLayer(DataLayer)))
+#endif // WITH_EDITOR
+					{
+						return InstanceManager;
+					}
 				}
 			}
 		}
@@ -199,11 +179,14 @@ ALightWeightInstanceManager* FLightWeightInstanceSubsystem::FindOrAddLightWeight
 	SpawnParams.ObjectFlags = RF_Transactional;
 
 	FTransform ManagerTransform(FTransform::Identity);
-	if (LWIGridSize > 0)
+
+	ALightWeightInstanceManager* InstanceManagerCDO = CastChecked<ALightWeightInstanceManager>(BestMatchingClass->GetDefaultObject(true));
+	int32 GridSize = InstanceManagerCDO->GetGridSize();
+	if (GridSize > 0)
 	{
-		FVector ManagerLocation(ConvertPositionToCoord(InPos));
-		ManagerLocation *= LWIGridSize;
-		ManagerLocation += FVector((float)LWIGridSize * 0.5f);
+		FVector ManagerLocation(InstanceManagerCDO->ConvertPositionToCoord(InPos));
+		ManagerLocation *= GridSize;
+		ManagerLocation += FVector((float)GridSize * 0.5f);
 		ManagerTransform.SetLocation(ManagerLocation);
 	}
 
