Fix for COTF server crashes.

- In some scenarios the client/game incorrectly requests a Material/VFType/ShaderType combination that shouldn't actually be compiled.
- These requests go to the COTF server, which then tries to compile combinations that aren't supported.  In some cases the VFType + ShaderType doesn't even compile.  For example, PositionOnlyDepthShaders are only compiled for LocalVF and trying to compile for a different VF type would result in a shader compiler error (the shader code doesn't even compile due to missing types).
- In TryGetShaders we now determine if we should cache the shader first by asking the Material, VFType and the ShaderType.  If we _shouldn't_ cache the shader TryGetShaders returns false, which results in the game falling back to the WorldGridMaterial.  If it fails for WorldGridMaterial it will not add a draw call.

todo: In the future we should log the results and fixup the offending rendering code which is trying to render with a shader that doesn't exist.

diff --git a/Engine/Source/Runtime/Engine/Private/Materials/MaterialShared.cpp b/Engine/Source/Runtime/Engine/Private/Materials/MaterialShared.cpp
--- a/Engine/Source/Runtime/Engine/Private/Materials/MaterialShared.cpp
+++ b/Engine/Source/Runtime/Engine/Private/Materials/MaterialShared.cpp
@@ -22,6 +22,7 @@
 #include "Materials/MaterialExpressionBreakMaterialAttributes.h"
 #include "Materials/MaterialExpressionRerouteBase.h"
 #include "ShaderCompiler.h"
+#include "MeshMaterialShader.h"
 #include "MeshMaterialShaderType.h"
 #include "RendererInterface.h"
 #include "Materials/HLSLMaterialTranslator.h"
@@ -3369,6 +3370,30 @@ bool FMaterial::TryGetShaders(const FMaterialShaderTypes& InTypes, const FVertex
 	TArray<FShaderCommonCompileJobPtr> CompileJobs;
 	bool bMissingShader = false;
 
+	auto ShouldCacheShaderType = [&ShaderPlatform, this, &PermutationFlags](const FShaderType* ShaderType, const FVertexFactoryType* VertexFactoryType) -> bool {
+		// Check to see if the FMaterial should cache these types.
+		if (!ShouldCache(ShaderPlatform, ShaderType, VertexFactoryType))
+		{
+			return false;
+		}
+
+		// if we are just a MaterialShaderType (not associated with a mesh)
+		if (const FMaterialShaderType* MaterialShader = ShaderType->GetMaterialShaderType())
+		{
+			return MaterialShader->ShouldCompilePermutation(ShaderPlatform, this, kUniqueShaderPermutationId, PermutationFlags);
+		}
+
+		// if we are a MeshMaterialShader
+		if (const FMeshMaterialShaderType* MeshMaterialShader = ShaderType->GetMeshMaterialShaderType())
+		{
+			const bool bVFShouldCache = FMeshMaterialShaderType::ShouldCompileVertexFactoryPermutation(ShaderPlatform, this, VertexFactoryType, ShaderType, PermutationFlags);
+			const bool bShaderShouldCache = MeshMaterialShader->ShouldCompilePermutation(ShaderPlatform, this, VertexFactoryType, kUniqueShaderPermutationId, PermutationFlags);
+			return bVFShouldCache && bShaderShouldCache;
+		}
+		
+		return false;
+	};
+
 	if (InTypes.PipelineType &&
 		RHISupportsShaderPipelines(ShaderPlatform) &&
 		CVarShaderPipelines && CVarShaderPipelines->GetValueOnAnyThread(bIsInGameThread) != 0)
@@ -3403,6 +3428,13 @@ bool FMaterial::TryGetShaders(const FMaterialShaderTypes& InTypes, const FVertex
 				{
 					if (GODSCManager->IsHandlingRequests())
 					{
+				        for (const FShaderType* ShaderType : InTypes.PipelineType->GetStages())
+				        {
+					        if (!ShouldCacheShaderType(ShaderType, InVertexFactoryType))
+					        {
+						        return false;
+					        }
+				        }
 						const FString MaterialName = GetFullPath();
 						const FString VFTypeName(InVertexFactoryType ? InVertexFactoryType->GetName() : TEXT(""));
 						const FString PipelineName(InTypes.PipelineType->GetName());
@@ -3485,6 +3517,10 @@ bool FMaterial::TryGetShaders(const FMaterialShaderTypes& InTypes, const FVertex
 					{
 						if (GODSCManager->IsHandlingRequests())
 						{
+					        if (!ShouldCacheShaderType(ShaderType, InVertexFactoryType))
+					        {
+						        return false;
+					        }
 							const FString MaterialName = GetFullPath();
 							const FString VFTypeName(InVertexFactoryType ? InVertexFactoryType->GetName() : TEXT(""));
 							const FString PipelineName;
