Fix NaN crash in FAnimNode_LinkedInputPose::Evaluate_AnyThread.
FAnimInstanceProxy::RecalcRequiredBones calls CachedInputPose.SetBoneContainer, which leaves CachedInputPose with bone data allocated but unintialized.
This change adds a flag to know when CachedInputPose is properly initialized.  LinkedInputPose will return a ref pose otherwise.

diff --git a/Engine/Source/Runtime/Engine/Classes/Animation/AnimNode_LinkedInputPose.h b/Engine/Source/Runtime/Engine/Classes/Animation/AnimNode_LinkedInputPose.h
--- a/Engine/Source/Runtime/Engine/Classes/Animation/AnimNode_LinkedInputPose.h
+++ b/Engine/Source/Runtime/Engine/Classes/Animation/AnimNode_LinkedInputPose.h
@@ -45,6 +45,10 @@ struct ENGINE_API FAnimNode_LinkedInputPose : public FAnimNode_Base
 	FBlendedHeapCurve CachedInputCurve;
 	UE::Anim::FHeapAttributeContainer CachedAttributes;
 
+	// CachedInputPose can have bone data allocated but uninitialized.
+	// This can happen if an anim graph has an Input Pose node with nothing populating it (e.g. if it's played as the only animbp on an actor).
+	bool bIsCachedInputPoseInitialized = false;
+
 	// The node index of the currently-linked outer node
 	int32 OuterGraphNodeIndex;
 
diff --git a/Engine/Source/Runtime/Engine/Private/Animation/AnimInstanceProxy.cpp b/Engine/Source/Runtime/Engine/Private/Animation/AnimInstanceProxy.cpp
--- a/Engine/Source/Runtime/Engine/Private/Animation/AnimInstanceProxy.cpp
+++ b/Engine/Source/Runtime/Engine/Private/Animation/AnimInstanceProxy.cpp
@@ -1126,6 +1126,9 @@ void FAnimInstanceProxy::RecalcRequiredBones(USkeletalMeshComponent* Component,
 	if(DefaultLinkedInstanceInputNode)
 	{
 		DefaultLinkedInstanceInputNode->CachedInputPose.SetBoneContainer(&RequiredBones);
+		
+		// SetBoneContainer allocates space for bone data but leaves it uninitalized.
+		DefaultLinkedInstanceInputNode->bIsCachedInputPoseInitialized = false;
 	}
 
 	// When RequiredBones mapping has changed, AnimNodes need to update their bones caches. 
diff --git a/Engine/Source/Runtime/Engine/Private/Animation/AnimNode_LinkedInputPose.cpp b/Engine/Source/Runtime/Engine/Private/Animation/AnimNode_LinkedInputPose.cpp
--- a/Engine/Source/Runtime/Engine/Private/Animation/AnimNode_LinkedInputPose.cpp
+++ b/Engine/Source/Runtime/Engine/Private/Animation/AnimNode_LinkedInputPose.cpp
@@ -50,7 +50,7 @@ void FAnimNode_LinkedInputPose::Evaluate_AnyThread(FPoseContext& Output)
 		Output.Curve.MoveFrom(InputContext.Curve);
 		Output.CustomAttributes.MoveFrom(InputContext.CustomAttributes);
 	}
-	else if(CachedInputPose.IsValid() && ensure(Output.Pose.GetNumBones() == CachedInputPose.GetNumBones()))
+	else if(CachedInputPose.IsValid() && ensure(Output.Pose.GetNumBones() == CachedInputPose.GetNumBones()) && bIsCachedInputPoseInitialized)
 	{
 		Output.Pose.CopyBonesFrom(CachedInputPose);
 		Output.Curve.CopyFrom(CachedInputCurve);
diff --git a/Engine/Source/Runtime/Engine/Private/Components/SkeletalMeshComponent.cpp b/Engine/Source/Runtime/Engine/Private/Components/SkeletalMeshComponent.cpp
--- a/Engine/Source/Runtime/Engine/Private/Components/SkeletalMeshComponent.cpp
+++ b/Engine/Source/Runtime/Engine/Private/Components/SkeletalMeshComponent.cpp
@@ -989,6 +989,9 @@ bool USkeletalMeshComponent::InitializeAnimScriptInstance(bool bForceReinit, boo
 				if(FAnimNode_LinkedInputPose* InputNode = PostProcessAnimInstance->GetLinkedInputPoseNode())
 				{
 					InputNode->CachedInputPose.SetBoneContainer(&PostProcessAnimInstance->GetRequiredBones());
+
+					// SetBoneContainer allocates space for bone data but leaves it uninitalized.
+					InputNode->bIsCachedInputPoseInitialized = false;
 				}
 
 				bInitializedPostInstance = true;
@@ -2201,12 +2204,14 @@ void USkeletalMeshComponent::EvaluatePostProcessMeshInstance(TArray<FTransform>&
 				InputNode->CachedInputPose.CopyBonesFrom(InOutPose);
 				InputNode->CachedInputCurve.CopyFrom(OutCurve);
 				InputNode->CachedAttributes.CopyFrom(OutAttributes);
+				InputNode->bIsCachedInputPoseInitialized = true;
 			}
 			else
 			{
 				const FBoneContainer& RequiredBone = PostProcessAnimInstance->GetRequiredBonesOnAnyThread();
 				InputNode->CachedInputPose.ResetToRefPose(RequiredBone);
 				InputNode->CachedInputCurve.InitFrom(RequiredBone);
+				InputNode->bIsCachedInputPoseInitialized = true;
 			}
 		}
 
