ExternalMorphSets: zero morph weights to prevent bad offsets

diff --git a/Engine/Source/Runtime/Engine/Classes/Components/SkinnedMeshComponent.h b/Engine/Source/Runtime/Engine/Classes/Components/SkinnedMeshComponent.h
--- a/Engine/Source/Runtime/Engine/Classes/Components/SkinnedMeshComponent.h
+++ b/Engine/Source/Runtime/Engine/Classes/Components/SkinnedMeshComponent.h
@@ -1067,7 +1067,7 @@ public:
 	void ResizeExternalMorphTargetSets();
 
 	/** Refresh the external morph target weight buffers. This makes sure the amount of morph sets and number of weights are valid. */
-	void RefreshExternalMorphTargetWeights();
+	void RefreshExternalMorphTargetWeights(bool bZeroOldWeights = true);
 
 	/**
 	 *	Sets the value of the bForceWireframe flag and reattaches the component as necessary.
diff --git a/Engine/Source/Runtime/Engine/Private/Components/SkinnedMeshComponent.cpp b/Engine/Source/Runtime/Engine/Private/Components/SkinnedMeshComponent.cpp
--- a/Engine/Source/Runtime/Engine/Private/Components/SkinnedMeshComponent.cpp
+++ b/Engine/Source/Runtime/Engine/Private/Components/SkinnedMeshComponent.cpp
@@ -755,7 +755,7 @@ void USkinnedMeshComponent::ResizeExternalMorphTargetSets()
 	ExternalMorphSets.AddDefaulted(GetNumLODs());
 }
 
-void USkinnedMeshComponent::RefreshExternalMorphTargetWeights()
+void USkinnedMeshComponent::RefreshExternalMorphTargetWeights(bool bZeroOldWeights)
 {
 	// Clear the external weights if there is no skeletal mesh.
 	if (GetSkinnedAsset() == nullptr)
@@ -798,7 +798,16 @@ void USkinnedMeshComponent::RefreshExternalMorphTargetWeights()
 			check(MorphSetWeights != nullptr);
 
 			MorphSetWeights->Name = Item.Value->Name;
-			MorphSetWeights->Weights.SetNumZeroed(MorphBuffer.GetNumMorphs());
+
+			if (bZeroOldWeights)
+			{
+				MorphSetWeights->Weights.Reset();
+				MorphSetWeights->Weights.AddZeroed(MorphBuffer.GetNumMorphs());
+			}
+			else
+			{
+				MorphSetWeights->Weights.SetNumZeroed(MorphBuffer.GetNumMorphs());
+			}
 		}
 
 		ExternalWeights.UpdateNumActiveMorphTargets();
@@ -966,7 +975,8 @@ void USkinnedMeshComponent::CreateRenderState_Concurrent(FRegisterComponentConte
 					ActiveMorphTargets.Empty();
 				}
 
-				RefreshExternalMorphTargetWeights();
+				constexpr bool bZeroOldWeights = false;
+				RefreshExternalMorphTargetWeights(bZeroOldWeights);
 				MeshObject->Update(ModifiedLODLevel, this, ActiveMorphTargets, MorphTargetWeights, EPreviousBoneTransformUpdateMode::UpdatePrevious, GetExternalMorphWeights(ModifiedLODLevel));  // send to rendering thread
 			}
 		}
