Level Instance Partial Loading: force loading entire level instance when chanling level package and when the main world is not partitioned.

diff --git a/Engine/Source/Runtime/Engine/Private/LevelInstance/LevelInstanceLevelStreaming.cpp b/Engine/Source/Runtime/Engine/Private/LevelInstance/LevelInstanceLevelStreaming.cpp
--- a/Engine/Source/Runtime/Engine/Private/LevelInstance/LevelInstanceLevelStreaming.cpp
+++ b/Engine/Source/Runtime/Engine/Private/LevelInstance/LevelInstanceLevelStreaming.cpp
@@ -363,8 +363,9 @@ void ULevelStreamingLevelInstance::OnLevelLoadedChanged(ULevel* InLevel)
 		if (ULevelInstanceSubsystem* LevelInstanceSubsystem = GetWorld()->GetSubsystem<ULevelInstanceSubsystem>())
 		{
 			LevelInstanceSubsystem->RegisterLoadedLevelStreamingLevelInstance(this);
+
 #if WITH_EDITOR
-			if (!GDisableLevelInstanceEditorPartialLoading)
+			if (!GDisableLevelInstanceEditorPartialLoading && GetWorld()->IsPartitionedWorld())
 			{
 				if (UWorldPartition* WorldPartition = NewLoadedLevel->GetWorldPartition())
 				{
@@ -373,7 +374,7 @@ void ULevelStreamingLevelInstance::OnLevelLoadedChanged(ULevel* InLevel)
 					{
 						if (AActor* Actor = CastChecked<AActor>(LevelInstance))
 						{
-							if (!Actor->GetPackage()->HasAnyPackageFlags(PKG_NewlyCreated))
+							if (!Actor->GetPackage()->IsDirty())
 							{
 								WorldPartition->bForceEnableStreamingInEditor = LevelInstance->SupportsPartialEditorLoading();
 							}
