FLevelInstanceID:
- Make sure LevelInstance that are part of a streaming level loaded by ULevelStreamingDynamic::LoadLevelInstance/LoadLevelInstanceBySoftObjectPtr have a unique id by hashing the actor package short name (UDN report)
- Remove FGuid array and use ContainerID for comparison instead

diff --git a/Engine/Source/Runtime/Engine/Private/LevelInstance/LevelInstanceSubsystem.cpp b/Engine/Source/Runtime/Engine/Private/LevelInstance/LevelInstanceSubsystem.cpp
--- a/Engine/Source/Runtime/Engine/Private/LevelInstance/LevelInstanceSubsystem.cpp
+++ b/Engine/Source/Runtime/Engine/Private/LevelInstance/LevelInstanceSubsystem.cpp
@@ -113,33 +113,38 @@ ILevelInstanceInterface* ULevelInstanceSubsystem::GetLevelInstance(const FLevelI
 
 FLevelInstanceID::FLevelInstanceID(ULevelInstanceSubsystem* LevelInstanceSubsystem, ILevelInstanceInterface* LevelInstance)
 {
+	TArray<FGuid> Guids;
 	AActor* LevelInstanceActor = CastChecked<AActor>(LevelInstance);
-	LevelInstanceSubsystem->ForEachLevelInstanceAncestorsAndSelf(LevelInstanceActor, [this](const ILevelInstanceInterface* AncestorOrSelf)
+	LevelInstanceSubsystem->ForEachLevelInstanceAncestorsAndSelf(LevelInstanceActor, [&Guids](const ILevelInstanceInterface* AncestorOrSelf)
 	{
 		Guids.Add(AncestorOrSelf->GetLevelInstanceGuid());
 		return true;
 	});
 	check(!Guids.IsEmpty());
 	
-	uint64 NameHash = 0;
-	ActorName = LevelInstanceActor->GetFName();
-	// Add Actor Name to hash because with World Partition Embedding top level of Level Instance hierarchy can get stripped leaving us with clashing ids.
-	// When embedding actors we make sure their names are unique (they get suffixed with their parent container id)
-	// Only do it for actor with a stable name. Actor with an unstable name are dynamically spawned and should have an unique replicated GUID.
+	uint64 TmpHash = 0;
+		
 	if (LevelInstanceActor->IsNameStableForNetworking())
 	{
+		ActorName = LevelInstanceActor->GetFName();
 		FString NameStr = ActorName.ToString();
-		NameHash = CityHash64((const char*)*NameStr, NameStr.Len() * sizeof(TCHAR));
+		TmpHash = CityHash64((const char*)*NameStr, NameStr.Len() * sizeof(TCHAR));
+
+		if (UWorld* OuterWorld = LevelInstanceActor->GetTypedOuter<UWorld>())
+		{
+			PackageShortName = UWorld::RemovePIEPrefix(FPackageName::GetShortName(OuterWorld->GetPackage()));
+			TmpHash = CityHash64WithSeed((const char*)*PackageShortName, PackageShortName.Len() * sizeof(TCHAR), TmpHash);
+		}
 	}
 	
-	Hash = CityHash64WithSeed((const char*)Guids.GetData(), Guids.Num() * sizeof(FGuid), NameHash);
-
 	// Make sure to start main container id
 	ContainerID = FActorContainerID();
 	for (int32 GuidIndex = Guids.Num() - 1; GuidIndex >= 0; GuidIndex--)
 	{
 		ContainerID = FActorContainerID(ContainerID, Guids[GuidIndex]);
 	}
+
+	Hash = CityHash64WithSeed((const char*)&ContainerID, sizeof(ContainerID), TmpHash);	
 }
 
 FLevelInstanceID ULevelInstanceSubsystem::RegisterLevelInstance(ILevelInstanceInterface* LevelInstance)
diff --git a/Engine/Source/Runtime/Engine/Public/LevelInstance/LevelInstanceTypes.h b/Engine/Source/Runtime/Engine/Public/LevelInstance/LevelInstanceTypes.h
--- a/Engine/Source/Runtime/Engine/Public/LevelInstance/LevelInstanceTypes.h
+++ b/Engine/Source/Runtime/Engine/Public/LevelInstance/LevelInstanceTypes.h
@@ -27,10 +27,10 @@ struct FLevelInstanceID
 
 	inline bool operator==(const FLevelInstanceID& Other) const
 	{
-		return Hash == Other.Hash && Guids == Other.Guids && ActorName == Other.ActorName;
+		return Hash == Other.Hash && ContainerID == Other.ContainerID && ActorName == Other.ActorName && PackageShortName == Other.PackageShortName;
 	}
 
-	inline bool IsValid() const { return !Guids.IsEmpty(); }
+	inline bool IsValid() const { return !ContainerID.IsMainContainer(); }
 
 	inline uint64 GetHash() const { return Hash; }
 
@@ -39,8 +39,17 @@ struct FLevelInstanceID
 private:
 	uint64 Hash = 0;
 	FActorContainerID ContainerID;
-	TArray<FGuid> Guids;
+	
+	// Hashed only for Loaded LevelInstances 
+	// Spawned LevelInstances have a unique Guid which is enough.
+	
+	// Name allows distinguishing between LevelInstances with embedded parent LevelInstances because embedded actors have a guaranteed unique name (ContainerID suffix)
 	FName ActorName;
+	// PackageShortName allows distinguising between instanced LevelInstances of the same source level.
+	// - Loading /Game/Path/WorldA.WorldA as /Game/Path/WorldA_LevelInstance1.WorldA & /Game/Path/WorldA_LevelInstance".WorldA
+	//   with the source WorldA containing one of many LevelInstance actors. 
+	//   Those actors would end up with the same hash. We use PackageShortName (WorldA_LevelInstance1 & WorldA_LevelInstance2) to distinguish them.
+	FString PackageShortName;
 };
 
 UENUM()
