World Partition Class Descriptors Registry:
- Don't rely on OnObjectLoaded event to prefetch existing class descriptors to avoid parsing the Asset Registry in the early process of initializing the editor.
- Instead, lazy prefetch existing class descriptors only when required.
- Optimized internal state validation when running in Debug.
Fix editor boot time regression.

diff --git a/Engine/Source/Runtime/Engine/Private/Blueprint.cpp b/Engine/Source/Runtime/Engine/Private/Blueprint.cpp
--- a/Engine/Source/Runtime/Engine/Private/Blueprint.cpp
+++ b/Engine/Source/Runtime/Engine/Private/Blueprint.cpp
@@ -1072,6 +1072,14 @@ void UBlueprint::GetExtendedAssetRegistryTagsForSave(const ITargetPlatform* Targ
 		}
 	}
 
+	/*
+	 * Using GetExtendedAssetRegistryTagsForSave here because:
+	 *	- UBlueprint is an asset in editor builds only.
+	 *	- UBlueprintGeneratedClass is an asset in cooked builds only.
+	 *	- Extended tags are only present in editor builds.
+	 *
+	 * See UBlueprintGeneratedClass::GetAssetRegistryTags.
+	 */
 	if (AActor* BlueprintCDO = GeneratedClass ? Cast<AActor>(GeneratedClass->ClassDefaultObject) : nullptr)
 	{
 		FWorldPartitionActorDescUtils::AppendAssetDataTagsFromActor(BlueprintCDO, OutTags);
diff --git a/Engine/Source/Runtime/Engine/Private/BlueprintGeneratedClass.cpp b/Engine/Source/Runtime/Engine/Private/BlueprintGeneratedClass.cpp
--- a/Engine/Source/Runtime/Engine/Private/BlueprintGeneratedClass.cpp
+++ b/Engine/Source/Runtime/Engine/Private/BlueprintGeneratedClass.cpp
@@ -253,13 +253,24 @@ void UBlueprintGeneratedClass::GetAssetRegistryTags(TArray<FAssetRegistryTag>& O
 #endif //#if WITH_EDITORONLY_DATA
 
 #if WITH_EDITOR
-	if (AActor* BlueprintCDO = Cast<AActor>(ClassDefaultObject))
-	{
-		if (UPackage* BlueprintCDOPackage = BlueprintCDO->GetPackage())
-		{
-			if (!FPackageName::IsTempPackage(BlueprintCDOPackage->GetName()) && !BlueprintCDOPackage->HasAnyPackageFlags(PKG_PlayInEditor))
+	/*
+	 * Can't use GetExtendedAssetRegistryTagsForSave here because:
+	 *	- UBlueprint is an asset in editor builds only.
+	 *	- UBlueprintGeneratedClass is an asset in cooked builds only.
+	 *	- Extended tags are not present in cooked builds.
+	 *
+	 * See UBlueprint::GetExtendedAssetRegistryTagsForSave.
+	 */
+	if (GIsSavingPackage)
+	{
+		if (AActor* BlueprintCDO = Cast<AActor>(ClassDefaultObject))
+		{
+			if (UPackage* BlueprintCDOPackage = BlueprintCDO->GetPackage())
 			{
-				FWorldPartitionActorDescUtils::AppendAssetDataTagsFromActor(BlueprintCDO, OutTags);
+				if (!FPackageName::IsTempPackage(BlueprintCDOPackage->GetName()) && !BlueprintCDOPackage->HasAnyPackageFlags(PKG_PlayInEditor))
+				{
+					FWorldPartitionActorDescUtils::AppendAssetDataTagsFromActor(BlueprintCDO, OutTags);
+				}
 			}
 		}
 	}
diff --git a/Engine/Source/Runtime/Engine/Private/WorldPartition/WorldPartitionClassDescRegistry.cpp b/Engine/Source/Runtime/Engine/Private/WorldPartition/WorldPartitionClassDescRegistry.cpp
--- a/Engine/Source/Runtime/Engine/Private/WorldPartition/WorldPartitionClassDescRegistry.cpp
+++ b/Engine/Source/Runtime/Engine/Private/WorldPartition/WorldPartitionClassDescRegistry.cpp
@@ -42,9 +42,6 @@
  *
  *	[Saving a existing Blueprint]
  *	  - Create or update the class descriptor.
- *
- *	[Loading a Blueprint]
- *	  - Create or update the class descriptor from the assset-registry data.
  */
 
 /*
@@ -94,7 +91,6 @@ void FWorldPartitionClassDescRegistry::Initialize()
 
 	RegisterClasses();
 
-	FCoreUObjectDelegates::OnAssetLoaded.AddRaw(this, &FWorldPartitionClassDescRegistry::OnAssetLoaded);
 	FCoreUObjectDelegates::OnObjectPreSave.AddRaw(this, &FWorldPartitionClassDescRegistry::OnObjectPreSave);
 	FCoreUObjectDelegates::OnObjectPropertyChanged.AddRaw(this, &FWorldPartitionClassDescRegistry::OnObjectPropertyChanged);
 
@@ -117,7 +113,6 @@ void FWorldPartitionClassDescRegistry::Uninitialize()
 	check(IsInitialized());
 	FActorDescList::Empty();
 
-	FCoreUObjectDelegates::OnAssetLoaded.RemoveAll(this);
 	FCoreUObjectDelegates::OnObjectPreSave.RemoveAll(this);
 	FCoreUObjectDelegates::OnObjectPropertyChanged.RemoveAll(this);
 	
@@ -383,36 +378,6 @@ void FWorldPartitionClassDescRegistry::RegisterClassDescriptorFromActorClass(con
 	RegisterClassDescriptor(NewActorDesc);
 }
 
-void FWorldPartitionClassDescRegistry::OnAssetLoaded(UObject* InAssetLoaded)
-{
-	check(InAssetLoaded);
-	check(IsInitialized());
-
-	if (InAssetLoaded->GetPackage()->HasAnyPackageFlags(PKG_PlayInEditor))
-	{
-		return;
-	}
-
-	UBlueprintGeneratedClass* BlueprintGeneratedClass = Cast<UBlueprintGeneratedClass>(InAssetLoaded);
-
-	if (!BlueprintGeneratedClass)
-	{
-		if (UBlueprint* Blueprint = Cast<UBlueprint>(InAssetLoaded))
-		{
-			BlueprintGeneratedClass = Cast<UBlueprintGeneratedClass>(Blueprint->GeneratedClass);
-		}
-	}
-
-	if (!BlueprintGeneratedClass || !BlueprintGeneratedClass->IsChildOf<AActor>())
-	{
-		return;
-	}
-
-	ValidateInternalState();
-	PrefetchClassDescs({ FTopLevelAssetPath(BlueprintGeneratedClass->GetPathName()) });
-	ValidateInternalState();
-}
-
 void FWorldPartitionClassDescRegistry::OnObjectPreSave(UObject* InObject, FObjectPreSaveContext InSaveContext)
 {
 	if (!InSaveContext.IsProceduralSave() && !(InSaveContext.GetSaveFlags() & SAVE_FromAutosave))
@@ -565,49 +530,52 @@ void FWorldPartitionClassDescRegistry::UpdateClassDescriptor(UObject* InObject,
 	const FTopLevelAssetPath ClassPath(ActorCDO->GetClass()->GetPathName());
 	const FTopLevelAssetPath CurrentParentClassPath = FTopLevelAssetPath(ActorCDO->GetClass()->GetSuperClass()->GetPathName());
 
-	if (TUniquePtr<FWorldPartitionActorDesc>** ExistingClassDesc = ClassByPath.Find(ClassPath))
+	ValidateInternalState();
+
+	if (!ClassByPath.Contains(ClassPath))
 	{
-		const FTopLevelAssetPath PreviousParentClassPath = ParentClassMap.FindChecked(ClassPath);
+		PrefetchClassDescs({ ClassPath });
+	}
 
-		ValidateInternalState();
+	TUniquePtr<FWorldPartitionActorDesc>* ExistingClassDesc = ClassByPath.FindChecked(ClassPath);
 
-		if (!*ExistingClassDesc)
-		{
-			if (!bOnlyIfExists)
-			{
-				FWorldPartitionActorDesc* NewActorDesc = ActorCDO->CreateActorDesc().Release();
-				RegisterClassDescriptor(NewActorDesc);
+	ValidateInternalState();
 
-				ParentClassMap.Add(ClassPath, CurrentParentClassPath);
-			}
-		}
-		else
+	if (!ExistingClassDesc)
+	{
+		if (!bOnlyIfExists)
 		{
-			FWorldPartitionActorDescUtils::UpdateActorDescriptorFromActor(ActorCDO, **ExistingClassDesc);
+			FWorldPartitionActorDesc* NewActorDesc = ActorCDO->CreateActorDesc().Release();
+			RegisterClassDescriptor(NewActorDesc);
 
-			if (PreviousParentClassPath != CurrentParentClassPath)
-			{
-				// We are reparenting a blueprint, update our parent map
-				ParentClassMap.Add(ClassPath, CurrentParentClassPath);
-			}
+			ParentClassMap.Add(ClassPath, CurrentParentClassPath);
 		}
+	}
+	else
+	{
+		FWorldPartitionActorDescUtils::UpdateActorDescriptorFromActor(ActorCDO, *ExistingClassDesc);
 
-		ValidateInternalState();
+		const FTopLevelAssetPath PreviousParentClassPath = ParentClassMap.FindChecked(ClassPath);
+		if (PreviousParentClassPath != CurrentParentClassPath)
+		{
+			// We are reparenting a blueprint, update our parent map
+			ParentClassMap.Add(ClassPath, CurrentParentClassPath);
+		}
 	}
+
+	ValidateInternalState();
 }
 
 void FWorldPartitionClassDescRegistry::ValidateInternalState()
 {
 #if DO_GUARD_SLOW
+	check(ClassByPath.Num() == ParentClassMap.Num());
+	static const FTopLevelAssetPath ActorClassPath(TEXT("/Script/Engine.Actor"));
 	for (auto& [ClassPath, ParentClassPath] : ParentClassMap)
 	{
 		// Validate that the class parents chain is valid
-		static const FTopLevelAssetPath ActorClassPath(TEXT("/Script/Engine.Actor"));
-		for (FTopLevelAssetPath CurrentClassPath = ClassPath; CurrentClassPath != ActorClassPath; CurrentClassPath = ParentClassMap.FindChecked(CurrentClassPath))
-		{
-			check(CurrentClassPath.IsValid());
-			ClassByPath.FindChecked(CurrentClassPath);
-		}
+		check(ClassByPath.Contains(ClassPath));
+		check((ParentClassPath != ActorClassPath) || ParentClassMap.Contains(ParentClassPath));
 	}
 #endif
 }
@@ -627,4 +595,4 @@ FTopLevelAssetPath FWorldPartitionClassDescRegistry::RedirectClassPath(const FTo
 
 	return ResolvedClassPath;
 }
-#endif
\ No newline at end of file
+#endif
diff --git a/Engine/Source/Runtime/Engine/Public/WorldPartition/WorldPartitionClassDescRegistry.h b/Engine/Source/Runtime/Engine/Public/WorldPartition/WorldPartitionClassDescRegistry.h
--- a/Engine/Source/Runtime/Engine/Public/WorldPartition/WorldPartitionClassDescRegistry.h
+++ b/Engine/Source/Runtime/Engine/Public/WorldPartition/WorldPartitionClassDescRegistry.h
@@ -43,7 +43,6 @@ private:
 	const FWorldPartitionActorDesc* GetClassDescDefaultForActor(const FTopLevelAssetPath& InClassPath) const;
 	const FWorldPartitionActorDesc* GetClassDescDefaultForClass(const FTopLevelAssetPath& InClassPath) const;
 
-	void OnAssetLoaded(UObject* InAssetLoaded);
 	void OnObjectPreSave(UObject* InObject, FObjectPreSaveContext InSaveContext);
 	void OnObjectPropertyChanged(UObject* InObject, FPropertyChangedEvent& InPropertyChangedEvent);
 
