SwitchTargetPlatform nop comments only

diff --git a/Engine/Source/Runtime/Engine/Private/TextureLODSettings.cpp b/Engine/Source/Runtime/Engine/Private/TextureLODSettings.cpp
--- a/Engine/Source/Runtime/Engine/Private/TextureLODSettings.cpp
+++ b/Engine/Source/Runtime/Engine/Private/TextureLODSettings.cpp
@@ -112,12 +112,18 @@ UTextureLODSettings::UTextureLODSettings(const FObjectInitializer& ObjectInitial
  */
 TArray<FString> UTextureLODSettings::GetTextureGroupNames()
 {
+	// no need for this to be FString, could be TCHAR *
 	TArray<FString> TextureGroupNames;
 
+	// TEXTUREGROUP_MAX is not actually Max, it's count
+	TextureGroupNames.Reserve(TEXTUREGROUP_MAX);
+
 #define GROUPNAMES(g) new(TextureGroupNames) FString(TEXT(#g));
 	FOREACH_ENUM_TEXTUREGROUP(GROUPNAMES)
 #undef GROUPNAMES
 
+	check( TextureGroupNames.Num() == TEXTUREGROUP_MAX );
+
 	return TextureGroupNames;
 }
 
