Chaos optimization : remove allocation when looking for existing pair of colliding proxies in chaos scene HandleCollisionEvents

diff --git a/Engine/Source/Runtime/Core/Public/Containers/Map.h b/Engine/Source/Runtime/Core/Public/Containers/Map.h
--- a/Engine/Source/Runtime/Core/Public/Containers/Map.h
+++ b/Engine/Source/Runtime/Core/Public/Containers/Map.h
@@ -1414,6 +1414,26 @@ public:
 		}
 	}
 
+	/**
+	* Visit all values associated with the specified key 
+	* This is replacement to the array returning MultiFind when allocations is to be avoided
+	* if the visitor function return false, this will exit out from looping through the values
+	* Visitor Function signature must be compatible with : bool VisitorFunction(const ValueType&)
+	*
+	* @param Key The key to find associated values for.
+	* @param VisitorFunction visitor function 
+	*/
+	template<typename TVisitorFunction> void ForEachValue(KeyInitType Key, TVisitorFunction VisitorFunction) const
+	{
+		for (typename Super::ElementSetType::TConstKeyIterator It(Super::Pairs, Key); It; ++It)
+		{
+			if (!VisitorFunction(It->Value))
+			{
+				return;
+			}
+		}
+	}
+
 	/**
 	 * Finds all values associated with the specified key.
 	 *
diff --git a/Engine/Source/Runtime/Engine/Private/PhysicsEngine/Experimental/PhysScene_Chaos.cpp b/Engine/Source/Runtime/Engine/Private/PhysicsEngine/Experimental/PhysScene_Chaos.cpp
--- a/Engine/Source/Runtime/Engine/Private/PhysicsEngine/Experimental/PhysScene_Chaos.cpp
+++ b/Engine/Source/Runtime/Engine/Private/PhysicsEngine/Experimental/PhysScene_Chaos.cpp
@@ -790,23 +790,26 @@ const FBodyInstance* FPhysScene_Chaos::GetBodyInstanceFromProxyAndShape(IPhysics
 FCollisionNotifyInfo& FPhysScene_Chaos::GetPendingCollisionForContactPair(const void* P0, const void* P1, Chaos::FReal SolverTime, bool& bNewEntry)
 {
 	const FUniqueContactPairKey Key = { P0, P1 };
-	TArray<int32> ExistingPairs;
-	ContactPairToPendingNotifyMap.MultiFind(Key, ExistingPairs);
-	for(int32 ExistingIdx : ExistingPairs)
-	{
-		if(FMath::IsNearlyEqual(PendingCollisionNotifies[ExistingIdx].SolverTime, SolverTime))
+
+	int32 ExistingIndex = INDEX_NONE;
+	ContactPairToPendingNotifyMap.ForEachValue(Key, [this, SolverTime, &ExistingIndex](const int32 Value)
 		{
-			// we already have one for this pair
-			bNewEntry = false;
-			return PendingCollisionNotifies[ExistingIdx];	
-		}
-	}
+			if (FMath::IsNearlyEqual(PendingCollisionNotifies[Value].SolverTime, SolverTime))
+			{
+				ExistingIndex = Value;
+				return false; // early exit we have found what we were looking for 
+			}
+			return true;
+		});
 
-	// make a new entry
-	bNewEntry = true;
-	int32 NewIdx = PendingCollisionNotifies.AddZeroed();
-	ContactPairToPendingNotifyMap.Add(Key, NewIdx);
-	return PendingCollisionNotifies[NewIdx];
+	bNewEntry = (ExistingIndex == INDEX_NONE);
+	if (bNewEntry)
+	{
+		// make a new entry
+		ExistingIndex = PendingCollisionNotifies.AddZeroed();
+		ContactPairToPendingNotifyMap.Add(Key, ExistingIndex);
+	}
+	return PendingCollisionNotifies[ExistingIndex];
 }
 
 void FPhysScene_Chaos::HandleCollisionEvents(const Chaos::FCollisionEventData& Event)
