- Add option to skip zero dt updates in Cascade to fix PSA_Velocity aligned sprites for disappearing
- fx.Cascade.SkipZeroDeltaTime can be used to disable the zero dt skipping

#rnx

diff --git a/Engine/Source/Runtime/Engine/Private/Particles/ParticleComponents.cpp b/Engine/Source/Runtime/Engine/Private/Particles/ParticleComponents.cpp
--- a/Engine/Source/Runtime/Engine/Private/Particles/ParticleComponents.cpp
+++ b/Engine/Source/Runtime/Engine/Private/Particles/ParticleComponents.cpp
@@ -197,6 +197,14 @@ FAutoConsoleVariableRef CVarFXLWCTileRecache(
 	TEXT("Setting this value to 0 will remove this behavior but could introduce rendering & simulation artifacts.\n"),
 	ECVF_Default);
 
+static bool GFXSkipZeroDeltaTime = true;
+FAutoConsoleVariableRef CVarFXSkipZeroDeltaTime(
+	TEXT("fx.Cascade.SkipZeroDeltaTime"),
+	GFXSkipZeroDeltaTime,
+	TEXT("When enabled a delta tick time of nearly 0.0 will cause us to skip the component update.\n")
+	TEXT("This fixes issue like PSA_Velocity aligned sprites, but could cause issues with things that rely on accurate velocities (i.e. TSR)."),
+	ECVF_Default);
+
 /** Whether to allow particle systems to perform work. */
 ENGINE_API bool GIsAllowingParticles = true;
 
@@ -5397,8 +5405,13 @@ void UParticleSystemComponent::TickComponent(float DeltaTime, enum ELevelTick Ti
 	}
 
 	bForcedInActive = false;
+
 	DeltaTime *= CustomTimeDilation;
 	DeltaTimeTick = DeltaTime;
+	if (FMath::IsNearlyZero(DeltaTimeTick) && GFXSkipZeroDeltaTime)
+	{
+		return;
+	}
 
 	AccumTickTime += DeltaTime;
 
