Fix screenshot not identifying filename path characters on windows

diff --git a/Engine/Source/Runtime/Engine/Private/UnrealClient.cpp b/Engine/Source/Runtime/Engine/Private/UnrealClient.cpp
--- a/Engine/Source/Runtime/Engine/Private/UnrealClient.cpp
+++ b/Engine/Source/Runtime/Engine/Private/UnrealClient.cpp
@@ -379,7 +379,7 @@ void FScreenshotRequest::CreateViewportScreenShotFilename(FString& InOutFilename
 
 	//default to using the path that is given
 	InOutFilename = TypeName;
-	if (!TypeName.Contains(TEXT("/")))
+	if (!TypeName.Contains(TEXT("/")) && !TypeName.Contains(TEXT("\\")))
 	{
 		InOutFilename = GetDefault<UEngine>()->GameScreenshotSaveDirectory.Path / TypeName;
 	}
