fix many texture and image int64 issues
can now successfully import over 32K textures to VT

diff --git a/Engine/Plugins/Interchange/Runtime/Source/Import/Private/Texture/InterchangeTexturePayloadData.cpp b/Engine/Plugins/Interchange/Runtime/Source/Import/Private/Texture/InterchangeTexturePayloadData.cpp
--- a/Engine/Plugins/Interchange/Runtime/Source/Import/Private/Texture/InterchangeTexturePayloadData.cpp
+++ b/Engine/Plugins/Interchange/Runtime/Source/Import/Private/Texture/InterchangeTexturePayloadData.cpp
@@ -8,10 +8,9 @@
 #include "Misc/MessageDialog.h"
 #include "RHI.h"
 
-bool UE::Interchange::FImportImageHelper::IsImportResolutionValid(int32 Width, int32 Height, bool bAllowNonPowerOfTwo, FText* OutErrorMessage)
+bool UE::Interchange::FImportImageHelper::IsImportResolutionValid(int64 Width, int64 Height, bool bAllowNonPowerOfTwo, FText* OutErrorMessage)
 {
-	// VT res is currently limited by pixel count fitting in int32
-	const int64 MaximumSupportedVirtualTextureResolution = 32768;
+	// code dupe from: UTextureFactory::IsImportResolutionValid
 
 	// limit on current rendering RHI : == GetMax2DTextureDimension()
 	const int64 CurrentRHIMaxResolution = int64(1) << (GMaxTextureMipCount - 1);
@@ -45,26 +44,11 @@ bool UE::Interchange::FImportImageHelper::IsImportResolutionValid(int32 Width, i
 		bTextureTooLargeOrInvalidResolution = true;
 	}
 
-	// pixel count must fit in int32 :
-	//  mip surface could still be larger than 2 GB, that's allowed
-	//	eg. 16k RGBA float = 4 GB
-	if (Width * Height > MAX_int32)
-	{
-		bTextureTooLargeOrInvalidResolution = true;
-	}
-
-	if ((Width * Height) > FMath::Square(MaximumSupportedVirtualTextureResolution))
-	{
-		bTextureTooLargeOrInvalidResolution = true;
-	}
-
 	if (bTextureTooLargeOrInvalidResolution)
 	{
 		if (OutErrorMessage)
 		{
-			*OutErrorMessage = FText::Format(
-				NSLOCTEXT("Interchange", "Warning_TextureSizeTooLargeOrInvalid", "Texture is too large to import or it has an invalid resolution. The current maximum is {0} pixels"),
-				FText::AsNumber(FMath::Square(MaximumSupportedVirtualTextureResolution)));
+			*OutErrorMessage = NSLOCTEXT("Interchange", "Warning_TextureSizeTooLargeOrInvalid", "Texture is too large to import or it has an invalid resolution.");
 		}
 
 		return false;
@@ -94,7 +78,10 @@ bool UE::Interchange::FImportImageHelper::IsImportResolutionValid(int32 Width, i
 		const bool bIsPowerOfTwo = FMath::IsPowerOfTwo(Width) && FMath::IsPowerOfTwo(Height);
 		if (!bIsPowerOfTwo)
 		{
-			*OutErrorMessage = NSLOCTEXT("Interchange", "Warning_TextureNotAPowerOfTwo", "Cannot import texture with non-power of two dimensions");
+			if ( OutErrorMessage )
+			{
+				*OutErrorMessage = NSLOCTEXT("Interchange", "Warning_TextureNotAPowerOfTwo", "Cannot import texture with non-power of two dimensions");
+			}
 			return false;
 		}
 	}
diff --git a/Engine/Plugins/Interchange/Runtime/Source/Import/Public/Texture/InterchangeTexturePayloadData.h b/Engine/Plugins/Interchange/Runtime/Source/Import/Public/Texture/InterchangeTexturePayloadData.h
--- a/Engine/Plugins/Interchange/Runtime/Source/Import/Public/Texture/InterchangeTexturePayloadData.h
+++ b/Engine/Plugins/Interchange/Runtime/Source/Import/Public/Texture/InterchangeTexturePayloadData.h
@@ -58,7 +58,7 @@ namespace UE
 			 *
 			 * @return bool true if the given height/width represent a supported texture resolution, false if not
 			 */
-			static bool IsImportResolutionValid(int32 Width, int32 Height, bool bAllowNonPowerOfTwo, FText* OutErrorMessage = nullptr);
+			static bool IsImportResolutionValid(int64 Width, int64 Height, bool bAllowNonPowerOfTwo, FText* OutErrorMessage = nullptr);
 		};
 	}//ns Interchange
 }//ns UE
diff --git a/Engine/Source/Developer/TextureCompressor/Private/TextureCompressorModule.cpp b/Engine/Source/Developer/TextureCompressor/Private/TextureCompressorModule.cpp
--- a/Engine/Source/Developer/TextureCompressor/Private/TextureCompressorModule.cpp
+++ b/Engine/Source/Developer/TextureCompressor/Private/TextureCompressorModule.cpp
@@ -3657,8 +3657,6 @@ private:
 				check(SourceImage.GetBytesPerPixel() == sizeof(FLinearColor));
 				check(TargetImage.GetBytesPerPixel() == sizeof(FLinearColor));
 
-				const int32 SourceBytesPerLine = SourceImage.SizeX * SourceImage.GetBytesPerPixel();
-				const int32 DestBytesPerLine = TargetImage.SizeX * TargetImage.GetBytesPerPixel();
 				for (int32 SliceIndex = 0; SliceIndex < SourceImage.NumSlices; ++SliceIndex)
 				{
 					for (int32 Y = 0; Y < TargetTextureSizeY; ++Y)
diff --git a/Engine/Source/Editor/Blutility/Private/AsyncImageExport.cpp b/Engine/Source/Editor/Blutility/Private/AsyncImageExport.cpp
--- a/Engine/Source/Editor/Blutility/Private/AsyncImageExport.cpp
+++ b/Engine/Source/Editor/Blutility/Private/AsyncImageExport.cpp
@@ -92,25 +92,25 @@ void UAsyncImageExport::Activate()
 
 		check( TextureSource.IsValid() );
 
-		// todo : should use GetMipImage instead of GetMipData
+		// todo : should use GetMipImage instead of GetMipData ; use FImage and FImageUtils for image export
 		TArray64<uint8> TextureRawData;
 		verify( TextureSource.GetMipData(TextureRawData, 0) );
 
-		const int32 BytesPerPixel = TextureSource.GetBytesPerPixel();
+		const int64 BytesPerPixel = TextureSource.GetBytesPerPixel();
 		const ETextureSourceFormat SourceFormat = TextureSource.GetFormat();
 
-		const int32 Width = TextureSource.GetSizeX();
-		const int32 Height = TextureSource.GetSizeY();
+		const int64 Width = TextureSource.GetSizeX();
+		const int64 Height = TextureSource.GetSizeY();
 
 		TArray<FColor> OutPixels;
 		OutPixels.SetNumZeroed(Width * Height);
 
 		FColor Color(0, 0, 0, 0);
-		for (int Y = 0; Y < Height; ++Y)
+		for (int64 Y = 0; Y < Height; ++Y)
 		{
-			for (int X = 0; X < Width; ++X)
+			for (int64 X = 0; X < Width; ++X)
 			{
-				const int32 PixelByteOffset = (X + Y * Width) * BytesPerPixel;
+				const int64 PixelByteOffset = (X + Y * Width) * BytesPerPixel;
 				const uint8* PixelPtr = TextureRawData.GetData() + PixelByteOffset;
 
 				switch (SourceFormat)
diff --git a/Engine/Source/Editor/UnrealEd/Private/Factories/EditorFactories.cpp b/Engine/Source/Editor/UnrealEd/Private/Factories/EditorFactories.cpp
--- a/Engine/Source/Editor/UnrealEd/Private/Factories/EditorFactories.cpp
+++ b/Engine/Source/Editor/UnrealEd/Private/Factories/EditorFactories.cpp
@@ -3222,7 +3222,7 @@ bool UTextureFactory::ImportImage(const uint8* Buffer, int64 Length, FFeedbackCo
 
 UTexture* UTextureFactory::ImportTextureUDIM(UClass* Class, UObject* InParent, FName Name, EObjectFlags Flags, const TCHAR* Type, const TMap<int32, FString>& UDIMIndexToFile, FFeedbackContext* Warn)
 {
-	TArray<uint8> TextureData;
+	TArray64<uint8> TextureData;
 	TArray<FImportImage> SourceImages;
 	TArray<FTextureSourceBlock> SourceBlocks;
 	TArray<FString> SourceFileNames;
@@ -3792,7 +3792,7 @@ UObject* UTextureFactory::FactoryCreateBinary
 	UObject*			Context,
 	const TCHAR*		Type,
 	const uint8*&		Buffer,
-	const uint8*			BufferEnd,
+	const uint8*		BufferEnd,
 	FFeedbackContext*	Warn
 )
 {
@@ -4486,11 +4486,11 @@ void UTextureFactory::ApplyAutoImportSettings(UTexture* Texture)
 
 bool UTextureFactory::IsImportResolutionValid(int64 Width, int64 Height, bool bAllowNonPowerOfTwo, FFeedbackContext* Warn)
 {
+	// code dupe to:
+	//UE::Interchange::FImportImageHelper::IsImportResolutionValid
+
 	static const auto CVarVirtualTexturesEnabled = IConsoleManager::Get().FindTConsoleVariableDataInt(TEXT("r.VirtualTextures")); check(CVarVirtualTexturesEnabled);
 
-	// VT res is currently limited by pixel count fitting in int32
-	const int64 MaximumSupportedVirtualTextureResolution = 32768;
-	
 	// Get the non-VT size limit :
 	int64 MaximumSupportedResolutionNonVT = (int64)UTexture::GetMaximumDimensionOfNonVT();
 	
@@ -4503,10 +4503,7 @@ bool UTextureFactory::IsImportResolutionValid(int64 Width, int64 Height, bool bA
 	// No zero-size textures :
 	if (Width == 0 || Height == 0 )
 	{
-		Warn->Log(ELogVerbosity::Error, *FText::Format(
-			NSLOCTEXT("UnrealEd", "Warning_TextureSizeZero", "Texture has zero width or height"),
-			FText::AsNumber(FMath::Square(MaximumSupportedVirtualTextureResolution))
-			).ToString());
+		Warn->Log(ELogVerbosity::Error, NSLOCTEXT("UnrealEd", "Warning_TextureSizeZero", "Texture has zero width or height").ToString());
 
 		return false;
 	}
@@ -4515,42 +4512,15 @@ bool UTextureFactory::IsImportResolutionValid(int64 Width, int64 Height, bool bA
 	//  could be negative here if it was over 2G and int32 was used earlier
 	if (Width < 0 || Height < 0 || Width > MAX_int32 || Height > MAX_int32)
 	{
-		Warn->Log(ELogVerbosity::Error, *FText::Format(
-			NSLOCTEXT("UnrealEd", "Warning_TextureSizeTooLargeOrInvalid", "Texture is too large to import or it has an invalid resolution. The current maximum is {0} pixels"),
-			FText::AsNumber(FMath::Square(MaximumSupportedVirtualTextureResolution))
-			).ToString());
-
-		return false;
-	}
-
-	// pixel count must fit in int32 :
-	//  mip surface could still be larger than 2 GB, that's allowed
-	//	eg. 16k RGBA float = 4 GB
-	if ( Width * Height > MAX_int32)
-	{
-		Warn->Log(ELogVerbosity::Error, *FText::Format(
-			NSLOCTEXT("UnrealEd", "Warning_TextureSizeTooLargeOrInvalid", "Texture is too large to import or it has an invalid resolution. The current maximum is {0} pixels"),
-			FText::AsNumber(FMath::Square(MaximumSupportedVirtualTextureResolution))
-			).ToString());
+		Warn->Log(ELogVerbosity::Error, NSLOCTEXT("UnrealEd", "Warning_TextureSizeTooLargeOrInvalid", "Texture is has an invalid resolution.").ToString());
 
 		return false;
 	}
-	
-	if ( (Width * Height) > FMath::Square(MaximumSupportedVirtualTextureResolution))
-	{
-		Warn->Log(ELogVerbosity::Error, *FText::Format(
-			NSLOCTEXT("UnrealEd", "Warning_TextureSizeTooLargeOrInvalid", "Texture is too large to import or it has an invalid resolution. The current maximum is {0} pixels"),
-			FText::AsNumber(FMath::Square(MaximumSupportedVirtualTextureResolution))
-			).ToString());
-
-		return false;
-	}
-
+		
 	if ( Width > MaximumSupportedResolutionNonVT || Height > MaximumSupportedResolutionNonVT )
 	{
 
 		// we're larger than MaximumSupportedResolution
-		// but not larger than MaximumSupportedVirtualTextureResolution
 		// so this texture can still work, but only as VT
 		// prompt about this :
 
diff --git a/Engine/Source/Runtime/Engine/Classes/Engine/Texture.h b/Engine/Source/Runtime/Engine/Classes/Engine/Texture.h
--- a/Engine/Source/Runtime/Engine/Classes/Engine/Texture.h
+++ b/Engine/Source/Runtime/Engine/Classes/Engine/Texture.h
@@ -91,7 +91,7 @@ struct FTextureSource
 	ENGINE_API FTextureSource();
 
 	// should match ERawImageFormat::GetBytesPerPixel
-	ENGINE_API static int32 GetBytesPerPixel(ETextureSourceFormat Format);
+	ENGINE_API static int64 GetBytesPerPixel(ETextureSourceFormat Format);
 	// should match ERawImageFormat::IsHDR
 	FORCEINLINE static bool IsHDR(ETextureSourceFormat Format) { return UE::TextureDefines::IsHDR(Format); }
 	
@@ -293,7 +293,7 @@ struct FTextureSource
 	ENGINE_API int64 CalcMipSize(int32 BlockIndex, int32 LayerIndex, int32 MipIndex) const;
 
 	/** Computes the number of bytes per-pixel. */
-	ENGINE_API int32 GetBytesPerPixel(int32 LayerIndex = 0) const;
+	ENGINE_API int64 GetBytesPerPixel(int32 LayerIndex = 0) const;
 
 	/** Return true if the source XY size is power-of-2.  Does not check Z size for volumes.  */
 	ENGINE_API bool IsPowerOfTwo(int32 BlockIndex = 0) const;
diff --git a/Engine/Source/Runtime/Engine/Private/Texture.cpp b/Engine/Source/Runtime/Engine/Private/Texture.cpp
--- a/Engine/Source/Runtime/Engine/Private/Texture.cpp
+++ b/Engine/Source/Runtime/Engine/Private/Texture.cpp
@@ -1772,7 +1772,7 @@ FTextureSourceBlock::FTextureSourceBlock()
 {
 }
 
-int32 FTextureSource::GetBytesPerPixel(ETextureSourceFormat Format)
+int64 FTextureSource::GetBytesPerPixel(ETextureSourceFormat Format)
 {
 	ERawImageFormat::Type RawFormat = FImageCoreUtils::ConvertToRawImageFormat(Format);
 	return ERawImageFormat::GetBytesPerPixel(RawFormat);
@@ -2474,7 +2474,7 @@ int64 FTextureSource::CalcMipSize(int32 BlockIndex, int32 LayerIndex, int32 MipI
 	return MipSizeX * MipSizeY * MipSlices * BytesPerPixel;
 }
 
-int32 FTextureSource::GetBytesPerPixel(int32 LayerIndex) const
+int64 FTextureSource::GetBytesPerPixel(int32 LayerIndex) const
 {
 	return GetBytesPerPixel(GetFormat(LayerIndex));
 }
diff --git a/Engine/Source/Runtime/Engine/Private/TextureDerivedDataTask.cpp b/Engine/Source/Runtime/Engine/Private/TextureDerivedDataTask.cpp
--- a/Engine/Source/Runtime/Engine/Private/TextureDerivedDataTask.cpp
+++ b/Engine/Source/Runtime/Engine/Private/TextureDerivedDataTask.cpp
@@ -945,7 +945,7 @@ static int64 GetBuildRequiredMemoryEstimate(UTexture* InTexture,
 			
 			ERawImageFormat::Type IntermediateImageFormat = UE::TextureBuildUtilities::GetVirtualTextureBuildIntermediateFormat(BuildSettings);
 
-			int IntermediateBytesPerPixel = ERawImageFormat::GetBytesPerPixel(IntermediateImageFormat);
+			int64 IntermediateBytesPerPixel = ERawImageFormat::GetBytesPerPixel(IntermediateImageFormat);
 
 			VTIntermediateSizeBytes += TotalPixelsPerLayer * IntermediateBytesPerPixel;
 
diff --git a/Engine/Source/Runtime/Engine/Private/VT/VirtualTextureDataBuilder.cpp b/Engine/Source/Runtime/Engine/Private/VT/VirtualTextureDataBuilder.cpp
--- a/Engine/Source/Runtime/Engine/Private/VT/VirtualTextureDataBuilder.cpp
+++ b/Engine/Source/Runtime/Engine/Private/VT/VirtualTextureDataBuilder.cpp
@@ -34,8 +34,8 @@ static TAutoConsoleVariable<int32> CVarVTParallelTileCompression(
 struct FPixelDataRectangle
 {
 	ETextureSourceFormat Format;
-	int32 Width;
-	int32 Height;
+	int64 Width;
+	int64 Height;
 	uint8 *Data;
 
 	FPixelDataRectangle(ETextureSourceFormat SetFormat, int32 SetWidth, int32 SetHeight, uint8* SetData) :
@@ -55,9 +55,9 @@ struct FPixelDataRectangle
 		checkf(DestX >= 0 && DestX < Width, TEXT("Destination location out of bounds"));
 		checkf(DestY >= 0 && DestY < Height, TEXT("Destination location out of bounds"));
 
-		int32 PixelSize = FTextureSource::GetBytesPerPixel(Source.Format);
-		int32 SrcScanlineSize = Source.Width * PixelSize;
-		int32 DstScanlineSize = Width * PixelSize;
+		int64 PixelSize = FTextureSource::GetBytesPerPixel(Source.Format);
+		int64 SrcScanlineSize = Source.Width * PixelSize;
+		int64 DstScanlineSize = Width * PixelSize;
 
 		// Handle source position having negative coordinates in source image
 		if (SourceX < 0)
@@ -83,7 +83,7 @@ struct FPixelDataRectangle
 		// Copy the data a scan line at a time
 
 		uint8 *DstScanline = Data + DestX * PixelSize + DestY * DstScanlineSize;
-		const uint8 *SrcScanline = Source.Data + SourceX * PixelSize + (SIZE_T)SourceY * (SIZE_T)SrcScanlineSize;
+		const uint8 *SrcScanline = Source.Data + SourceX * PixelSize + SourceY * SrcScanlineSize;
 
 		for (int Y = 0; Y < ClampedHeight; Y++)
 		{
@@ -97,15 +97,23 @@ struct FPixelDataRectangle
 	{
 		switch (Mode)
 		{
-		case TA_Wrap:
-			return (x % Width + Width) % Width; // Make sure it's a proper module for negative numbers ....
-		case TA_Clamp:
-			return FMath::Max(FMath::Min(x, Width-1), 0);
-		case TA_Mirror:
-			int32 DoubleWidth = Width + Width;
-			int32 DoubleWrap = (x % DoubleWidth + DoubleWidth) % DoubleWidth;
-			return (DoubleWrap < Width) ? DoubleWrap : (Width-1) - (DoubleWrap - Width);
+			case TA_Wrap:
+			{
+				// Make sure it's a proper module for negative numbers ....
+				int x_Wrap = x % Width;
+				return x_Wrap + ((x_Wrap < 0) ? Width : 0);
+			}
+			case TA_Clamp:
+				return FMath::Max(FMath::Min(x, Width-1), 0);
+			case TA_Mirror:
+			{
+				int32 DoubleWidth = Width + Width;
+				int32 DoubleWrap = x % DoubleWidth;
+				DoubleWrap += ((DoubleWrap < 0) ? DoubleWidth : 0);
+				return (DoubleWrap < Width) ? DoubleWrap : ((Width-1) - (DoubleWrap - Width));
+			}
 		}
+		check(0);
 		return x;
 	}
 
@@ -122,7 +130,7 @@ struct FPixelDataRectangle
 		// Fast copy of regular pixels
 		CopyRectangle(DestX, DestY, Source, SourceX, SourceY, RectWidth, RectHeight);
 
-		size_t pixelSize = FTextureSource::GetBytesPerPixel(Format);
+		int64 pixelSize = FTextureSource::GetBytesPerPixel(Format);
 
 		// Special case the out of bounds pixels loop over all oob pixels and get the properly adjusted values
 		if (SourceX < 0 ||
@@ -183,13 +191,13 @@ struct FPixelDataRectangle
 		FMemory::Memzero(Data, FTextureSource::GetBytesPerPixel(Format) * Width * Height);
 	}
 
-	inline void SetPixel(int32 x, int32 y, void *Value, size_t PixelSize)
+	inline void SetPixel(int32 x, int32 y, void *Value, int64 PixelSize)
 	{
 		void *DestPixelData = GetPixel(x, y, PixelSize);
 		FMemory::Memcpy(DestPixelData, Value, PixelSize);
 	}
 
-	inline void *GetPixel(int32 x, int32 y, size_t PixelSize) const
+	inline void *GetPixel(int32 x, int32 y, int64 PixelSize) const
 	{
 		check(x >= 0);
 		check(y >= 0);
@@ -682,7 +690,6 @@ void FVirtualTextureDataBuilder::BuildTiles(const TArray<FVTSourceTileEntry>& Ti
 
 			CompressedFormat = (EPixelFormat)CompressedMip[0].PixelFormat;
 
-			const uint32 SizeRaw = CompressedMip[0].RawData.Num() * CompressedMip[0].RawData.GetTypeSize();
 			GeneratedData.TilePayload[TileIndex] = MoveTemp(CompressedMip[0].RawData);
 
 			// if the SourceBlocks we used has no more tiles that need to read from it, free it now?
@@ -727,7 +734,7 @@ void FVirtualTextureDataBuilder::PushDataToChunk(const TArray<FVTSourceTileEntry
 {
 	const int32 NumLayers = SourceLayers.Num();
 
-	uint32 TotalSize = sizeof(FVirtualTextureChunkHeader);
+	int64 TotalSize = sizeof(FVirtualTextureChunkHeader);
 	for (int32 Layer = 0; Layer < NumLayers; ++Layer)
 	{
 		TotalSize += LayerData[Layer].CodecPayload.Num();
@@ -737,12 +744,15 @@ void FVirtualTextureDataBuilder::PushDataToChunk(const TArray<FVTSourceTileEntry
 		}
 	}
 
+	// Built VT data structures use uint32 :
+	check(TotalSize <= MAX_uint32 );
+
 	FVirtualTextureDataChunk& Chunk = OutData.Chunks.AddDefaulted_GetRef();
 	Chunk.SizeInBytes = TotalSize;
 	FByteBulkData& BulkData = Chunk.BulkData;
 	BulkData.Lock(LOCK_READ_WRITE);
 	uint8* NewChunkData = (uint8*)BulkData.Realloc(TotalSize);
-	uint32 ChunkOffset = 0u;
+	int64 ChunkOffset = 0u;
 
 	// Header for the chunk
 	FVirtualTextureChunkHeader* Header = (FVirtualTextureChunkHeader*)NewChunkData;
diff --git a/Engine/Source/Runtime/Engine/Private/VolumeTexture.cpp b/Engine/Source/Runtime/Engine/Private/VolumeTexture.cpp
--- a/Engine/Source/Runtime/Engine/Private/VolumeTexture.cpp
+++ b/Engine/Source/Runtime/Engine/Private/VolumeTexture.cpp
@@ -21,6 +21,7 @@
 #include "UObject/Package.h"
 #include "UObject/StrongObjectPtr.h"
 #include "ImageCoreUtils.h"
+#include "GuardedInt.h"
 
 #include UE_INLINE_GENERATED_CPP_BY_NAME(VolumeTexture)
 
@@ -42,13 +43,8 @@ const int32 MAX_VOLUME_TEXTURE_DEPTH = 512;
 // Returns 0 if the product of A and B would overflow
 static int32 CheckedNonNegativeProduct(int32 A, int32 B)
 {
-	int64 Product = (int64)A * (int64)B;
-	// Either factor negative or product too large to fit in int32?
-	if (A < 0 || B < 0 || Product > 0x7fffffff)
-	{
-		return 0;
-	}
-	return Product;
+	FGuardedInt32 Product = FGuardedInt32(A) * FGuardedInt32(B);
+	return Product.Get(0);
 }
 
 //*****************************************************************************
@@ -251,7 +247,7 @@ ENGINE_API bool UVolumeTexture::UpdateSourceFromFunction(TFunction<void(int32, i
 	Modify(true);
 
 	Source2DTexture = nullptr;
-	const int32 FormatDataSize = ERawImageFormat::GetBytesPerPixel(FImageCoreUtils::ConvertToRawImageFormat(Format));
+	const int64 FormatDataSize = ERawImageFormat::GetBytesPerPixel(FImageCoreUtils::ConvertToRawImageFormat(Format));
 
 	// Allocate temp buffer used to fill texture
 	uint8* const NewData = (uint8*)FMemory::Malloc(SizeX * SizeY * SizeZ * FormatDataSize);
diff --git a/Engine/Source/Runtime/ImageCore/Private/ImageCore.cpp b/Engine/Source/Runtime/ImageCore/Private/ImageCore.cpp
--- a/Engine/Source/Runtime/ImageCore/Private/ImageCore.cpp
+++ b/Engine/Source/Runtime/ImageCore/Private/ImageCore.cpp
@@ -1190,9 +1190,9 @@ IMAGECORE_API const TCHAR * ERawImageFormat::GetName(Type Format)
 	}
 }
 
-IMAGECORE_API int32 ERawImageFormat::GetBytesPerPixel(Type Format)
+IMAGECORE_API int64 ERawImageFormat::GetBytesPerPixel(Type Format)
 {
-	int32 BytesPerPixel = 0;
+	int64 BytesPerPixel = 0;
 	switch (Format)
 	{
 	case ERawImageFormat::G8:
diff --git a/Engine/Source/Runtime/ImageCore/Public/ImageCore.h b/Engine/Source/Runtime/ImageCore/Public/ImageCore.h
--- a/Engine/Source/Runtime/ImageCore/Public/ImageCore.h
+++ b/Engine/Source/Runtime/ImageCore/Public/ImageCore.h
@@ -62,7 +62,7 @@ namespace ERawImageFormat
 	// RGBA -> G8 takes from the  *R* channel
 	// G8 -> RGBA replicated to gray
  
-	IMAGECORE_API int32 GetBytesPerPixel(Type Format);
+	IMAGECORE_API int64 GetBytesPerPixel(Type Format);
 	
 	IMAGECORE_API const TCHAR * GetName(Type Format);
 	
@@ -161,7 +161,7 @@ struct FImageInfo
 	 *
 	 * @return Bytes per pixel.
 	 */
-	FORCEINLINE int32 GetBytesPerPixel() const
+	FORCEINLINE int64 GetBytesPerPixel() const
 	{
 		return ERawImageFormat::GetBytesPerPixel(Format);
 	}
@@ -186,8 +186,8 @@ struct FImageInfo
 		return GetSliceNumPixels() * GetBytesPerPixel();
 	}
 	
-	FORCEINLINE int32 GetWidth()  const { return SizeX; }
-	FORCEINLINE int32 GetHeight() const { return SizeY; }
+	FORCEINLINE int64 GetWidth()  const { return SizeX; }
+	FORCEINLINE int64 GetHeight() const { return SizeY; }
 	
 	FORCEINLINE EGammaSpace GetGammaSpace() const
 	{
diff --git a/Engine/Source/Runtime/ImageWrapper/Private/Formats/HdrImageWrapper.cpp b/Engine/Source/Runtime/ImageWrapper/Private/Formats/HdrImageWrapper.cpp
--- a/Engine/Source/Runtime/ImageWrapper/Private/Formats/HdrImageWrapper.cpp
+++ b/Engine/Source/Runtime/ImageWrapper/Private/Formats/HdrImageWrapper.cpp
@@ -233,12 +233,12 @@ bool FHdrImageWrapper::GetRaw(const ERGBFormat InFormat, int32 InBitDepth, TArra
 	return true;
 }
 
-int32 FHdrImageWrapper::GetWidth() const
+int64 FHdrImageWrapper::GetWidth() const
 {
 	return Width;
 }
 
-int32 FHdrImageWrapper::GetHeight() const
+int64 FHdrImageWrapper::GetHeight() const
 {
 	return Height;
 }
diff --git a/Engine/Source/Runtime/ImageWrapper/Private/Formats/JpegImageWrapper.cpp b/Engine/Source/Runtime/ImageWrapper/Private/Formats/JpegImageWrapper.cpp
--- a/Engine/Source/Runtime/ImageWrapper/Private/Formats/JpegImageWrapper.cpp
+++ b/Engine/Source/Runtime/ImageWrapper/Private/Formats/JpegImageWrapper.cpp
@@ -259,11 +259,14 @@ void FJpegImageWrapper::Uncompress(const ERGBFormat InFormat, int32 InBitDepth)
 	check(CompressedData.Num());
 
 	int32 NumColors;
+	int32 jpegWidth,jpegHeight;
 	uint8* OutData = jpgd::decompress_jpeg_image_from_memory(
-		CompressedData.GetData(), CompressedData.Num(), &Width, &Height, &NumColors, Channels, (int)InFormat);
+		CompressedData.GetData(), CompressedData.Num(), &jpegWidth, &jpegHeight, &NumColors, Channels, (int)InFormat);
 
 	if (OutData)
 	{
+		Width = jpegWidth;
+		Height = jpegHeight;
 		RawData.Reset(Width * Height * Channels);
 		RawData.AddUninitialized(Width * Height * Channels);
 		FMemory::Memcpy(RawData.GetData(), OutData, RawData.Num());
diff --git a/Engine/Source/Runtime/ImageWrapper/Private/Formats/PngImageWrapper.cpp b/Engine/Source/Runtime/ImageWrapper/Private/Formats/PngImageWrapper.cpp
--- a/Engine/Source/Runtime/ImageWrapper/Private/Formats/PngImageWrapper.cpp
+++ b/Engine/Source/Runtime/ImageWrapper/Private/Formats/PngImageWrapper.cpp
@@ -378,7 +378,7 @@ void FPngImageWrapper::UncompressPNGData(const ERGBFormat InFormat, const int32
 				png_set_palette_to_rgb(png_ptr);
 			}
 
-			// really we should just call png_expand() here and remove all these conditionals
+			// @todo Oodle: really we should just call png_expand() here and remove all these conditionals
 
 			if (((ColorType & PNG_COLOR_MASK_COLOR) == 0) && BitDepth < 8)
 			{
@@ -429,6 +429,10 @@ void FPngImageWrapper::UncompressPNGData(const ERGBFormat InFormat, const int32
 			}
 #endif
 
+			// @todo Oodle : remove all these conversions
+			//	ImageWrappers should load images as they are in the file
+			//	FImage does conversions after loading
+
 			// Convert grayscale png to RGB if requested
 			if ((ColorType & PNG_COLOR_MASK_COLOR) == 0 &&
 				(InFormat == ERGBFormat::RGBA || InFormat == ERGBFormat::BGRA))
diff --git a/Engine/Source/Runtime/ImageWrapper/Private/ImageWrapperBase.cpp b/Engine/Source/Runtime/ImageWrapper/Private/ImageWrapperBase.cpp
--- a/Engine/Source/Runtime/ImageWrapper/Private/ImageWrapperBase.cpp
+++ b/Engine/Source/Runtime/ImageWrapper/Private/ImageWrapperBase.cpp
@@ -153,7 +153,7 @@ bool FImageWrapperBase::SetRaw(const void* InRawData, int64 InRawSize, const int
 }
 
 
-int IImageWrapper::GetRGBFormatBytesPerPel(ERGBFormat RGBFormat,int BitDepth)
+int64 IImageWrapper::GetRGBFormatBytesPerPel(ERGBFormat RGBFormat,int BitDepth)
 {
 	switch(RGBFormat)
 	{
diff --git a/Engine/Source/Runtime/ImageWrapper/Private/ImageWrapperBase.h b/Engine/Source/Runtime/ImageWrapper/Private/ImageWrapperBase.h
--- a/Engine/Source/Runtime/ImageWrapper/Private/ImageWrapperBase.h
+++ b/Engine/Source/Runtime/ImageWrapper/Private/ImageWrapperBase.h
@@ -98,14 +98,14 @@ public:
 		return Format;
 	}
 
-	virtual int32 GetHeight() const override
+	virtual int64 GetHeight() const override
 	{
 		return Height;
 	}
 
 	virtual bool GetRaw(const ERGBFormat InFormat, int32 InBitDepth, TArray64<uint8>& OutRawData) override;
 
-	virtual int32 GetWidth() const override
+	virtual int64 GetWidth() const override
 	{
 		return Width;
 	}
@@ -115,8 +115,8 @@ public:
 
 protected:
 
-	int GetBytesPerPel() const { return GetRGBFormatBytesPerPel(Format,BitDepth); }
-	int GetBytesPerRow() const { return Width * GetBytesPerPel(); }
+	int64 GetBytesPerPel() const { return GetRGBFormatBytesPerPel(Format,BitDepth); }
+	int64 GetBytesPerRow() const { return Width * GetBytesPerPel(); }
 
 	// For writers: after SetRaw(), call this to get an ImageView of the raw data that was set
 	//	can return false if the SetRaw does not map to an image format
@@ -132,9 +132,8 @@ protected:
 	int BitDepth;
 
 	/** Width/Height of the image data */
-	// @todo Oodle : change these to int64
-	int32 Width;
-	int32 Height;
+	int64 Width;
+	int64 Height;
 	
 	/** Last Error Message. */
 	FString LastError;
diff --git a/Engine/Source/Runtime/ImageWrapper/Public/Formats/HdrImageWrapper.h b/Engine/Source/Runtime/ImageWrapper/Public/Formats/HdrImageWrapper.h
--- a/Engine/Source/Runtime/ImageWrapper/Public/Formats/HdrImageWrapper.h
+++ b/Engine/Source/Runtime/ImageWrapper/Public/Formats/HdrImageWrapper.h
@@ -30,8 +30,8 @@ public:
 	virtual bool CanSetRawFormat(const ERGBFormat InFormat, const int32 InBitDepth) const override;
 	virtual ERawImageFormat::Type GetSupportedRawFormat(const ERawImageFormat::Type InFormat) const override;
 
-	virtual int32 GetWidth() const override;
-	virtual int32 GetHeight() const override;
+	virtual int64 GetWidth() const override;
+	virtual int64 GetHeight() const override;
 	virtual int32 GetBitDepth() const override;
 	virtual ERGBFormat GetFormat() const override;
 
@@ -72,9 +72,9 @@ private:
 	TArray64<uint8> RawDataHolder;
 
 	/** INDEX_NONE if not valid */
-	int32 Width = INDEX_NONE;
+	int64 Width = INDEX_NONE;
 	/** INDEX_NONE if not valid */
-	int32 Height = INDEX_NONE;
+	int64 Height = INDEX_NONE;
 
 	// Reported error
 	FText ErrorMessage;
diff --git a/Engine/Source/Runtime/ImageWrapper/Public/IImageWrapper.h b/Engine/Source/Runtime/ImageWrapper/Public/IImageWrapper.h
--- a/Engine/Source/Runtime/ImageWrapper/Public/IImageWrapper.h
+++ b/Engine/Source/Runtime/ImageWrapper/Public/IImageWrapper.h
@@ -280,7 +280,7 @@ public:
 	 * @return Image width.
 	 * @see GetHeight
 	 */
-	virtual int32 GetWidth() const = 0;
+	virtual int64 GetWidth() const = 0;
 
 	/** 
 	 * Gets the height of the image.
@@ -288,7 +288,7 @@ public:
 	 * @return Image height.
 	 * @see GetWidth
 	 */
-	virtual int32 GetHeight() const = 0;
+	virtual int64 GetHeight() const = 0;
 
 	/** 
 	 * Gets the bit depth of the image.
@@ -327,7 +327,7 @@ public:
 	IMAGEWRAPPER_API static void ConvertRawImageFormat(ERawImageFormat::Type RawFormat, ERGBFormat & OutFormat,int & OutBitDepth);
 	IMAGEWRAPPER_API static ERawImageFormat::Type ConvertRGBFormat(ERGBFormat RGBFormat,int BitDepth,bool * bIsExactMatch = nullptr);
 	
-	IMAGEWRAPPER_API static int GetRGBFormatBytesPerPel(ERGBFormat RGBFormat,int BitDepth);
+	IMAGEWRAPPER_API static int64 GetRGBFormatBytesPerPel(ERGBFormat RGBFormat,int BitDepth);
 
 	/* get the current image format, mapped to an ERawImageFormat
 	 * if ! *bIsExactMatch , conversion is needed
