Level instance Partial Loading: properly handle the case where an ancestor is not spatially loaded.

diff --git a/Engine/Source/Runtime/Engine/Private/LevelInstance/LevelInstanceInterface.cpp b/Engine/Source/Runtime/Engine/Private/LevelInstance/LevelInstanceInterface.cpp
--- a/Engine/Source/Runtime/Engine/Private/LevelInstance/LevelInstanceInterface.cpp
+++ b/Engine/Source/Runtime/Engine/Private/LevelInstance/LevelInstanceInterface.cpp
@@ -21,7 +21,36 @@ ULevelInstanceInterface::ULevelInstanceInterface(const FObjectInitializer& Objec
 #if WITH_EDITOR
 bool ILevelInstanceInterface::SupportsPartialEditorLoading() const
 {
-	return CastChecked<AActor>(this)->GetIsSpatiallyLoaded() && !IsEditing() && !HasParentEdit();
+	if (HasValidLevelInstanceID())
+	{
+		if (ULevelInstanceSubsystem* LevelInstanceSubsystem = GetLevelInstanceSubsystem())
+		{
+			bool bResult = true;
+			const AActor* Actor = CastChecked<AActor>(this);
+
+			LevelInstanceSubsystem->ForEachLevelInstanceAncestorsAndSelf(Actor, [LevelInstanceSubsystem, &bResult](const ILevelInstanceInterface* LevelInstance)
+			{
+				if (LevelInstanceSubsystem->IsEditingLevelInstance(LevelInstance))
+				{
+					bResult = false;
+					return false;
+				}
+
+				const AActor* LevelInstanceActor = CastChecked<AActor>(LevelInstance);
+				if (!LevelInstanceActor->GetIsSpatiallyLoaded())
+				{
+					bResult = false;
+					return false;
+				}
+
+				return true;
+			});
+
+			return bResult;
+		}
+	}
+
+	return false;
 }
 #endif
 
