[Backout] - CL24990850
Original CL Desc
-----------------------------------------------------------------
World Partition Class Descriptors Registry:
- Don't rely on OnObjectLoaded event to prefetch existing class descriptors to avoid parsing the Asset Registry in the early process of initializing the editor.
- Instead, lazy prefetch existing class descriptors only when required.
- Optimized internal state validation when running in Debug.
Fix editor boot time regression.

diff --git a/Engine/Source/Runtime/Engine/Private/Blueprint.cpp b/Engine/Source/Runtime/Engine/Private/Blueprint.cpp
--- a/Engine/Source/Runtime/Engine/Private/Blueprint.cpp
+++ b/Engine/Source/Runtime/Engine/Private/Blueprint.cpp
@@ -1072,14 +1072,6 @@ void UBlueprint::GetExtendedAssetRegistryTagsForSave(const ITargetPlatform* Targ
 		}
 	}
 
-	/*
-	 * Using GetExtendedAssetRegistryTagsForSave here because:
-	 *	- UBlueprint is an asset in editor builds only.
-	 *	- UBlueprintGeneratedClass is an asset in cooked builds only.
-	 *	- Extended tags are only present in editor builds.
-	 *
-	 * See UBlueprintGeneratedClass::GetAssetRegistryTags.
-	 */
 	if (AActor* BlueprintCDO = GeneratedClass ? Cast<AActor>(GeneratedClass->ClassDefaultObject) : nullptr)
 	{
 		FWorldPartitionActorDescUtils::AppendAssetDataTagsFromActor(BlueprintCDO, OutTags);
diff --git a/Engine/Source/Runtime/Engine/Private/BlueprintGeneratedClass.cpp b/Engine/Source/Runtime/Engine/Private/BlueprintGeneratedClass.cpp
--- a/Engine/Source/Runtime/Engine/Private/BlueprintGeneratedClass.cpp
+++ b/Engine/Source/Runtime/Engine/Private/BlueprintGeneratedClass.cpp
@@ -253,24 +253,13 @@ void UBlueprintGeneratedClass::GetAssetRegistryTags(TArray<FAssetRegistryTag>& O
 #endif //#if WITH_EDITORONLY_DATA
 
 #if WITH_EDITOR
-	/*
-	 * Can't use GetExtendedAssetRegistryTagsForSave here because:
-	 *	- UBlueprint is an asset in editor builds only.
-	 *	- UBlueprintGeneratedClass is an asset in cooked builds only.
-	 *	- Extended tags are not present in cooked builds.
-	 *
-	 * See UBlueprint::GetExtendedAssetRegistryTagsForSave.
-	 */
-	if (GIsSavingPackage)
-	{
-		if (AActor* BlueprintCDO = Cast<AActor>(ClassDefaultObject))
-		{
-			if (UPackage* BlueprintCDOPackage = BlueprintCDO->GetPackage())
+	if (AActor* BlueprintCDO = Cast<AActor>(ClassDefaultObject))
+	{
+		if (UPackage* BlueprintCDOPackage = BlueprintCDO->GetPackage())
+		{
+			if (!FPackageName::IsTempPackage(BlueprintCDOPackage->GetName()) && !BlueprintCDOPackage->HasAnyPackageFlags(PKG_PlayInEditor))
 			{
-				if (!FPackageName::IsTempPackage(BlueprintCDOPackage->GetName()) && !BlueprintCDOPackage->HasAnyPackageFlags(PKG_PlayInEditor))
-				{
-					FWorldPartitionActorDescUtils::AppendAssetDataTagsFromActor(BlueprintCDO, OutTags);
-				}
+				FWorldPartitionActorDescUtils::AppendAssetDataTagsFromActor(BlueprintCDO, OutTags);
 			}
 		}
 	}
diff --git a/Engine/Source/Runtime/Engine/Private/WorldPartition/WorldPartitionClassDescRegistry.cpp b/Engine/Source/Runtime/Engine/Private/WorldPartition/WorldPartitionClassDescRegistry.cpp
--- a/Engine/Source/Runtime/Engine/Private/WorldPartition/WorldPartitionClassDescRegistry.cpp
+++ b/Engine/Source/Runtime/Engine/Private/WorldPartition/WorldPartitionClassDescRegistry.cpp
@@ -42,6 +42,9 @@
  *
  *	[Saving a existing Blueprint]
  *	  - Create or update the class descriptor.
+ *
+ *	[Loading a Blueprint]
+ *	  - Create or update the class descriptor from the assset-registry data.
  */
 
 /*
@@ -91,6 +94,7 @@ void FWorldPartitionClassDescRegistry::Initialize()
 
 	RegisterClasses();
 
+	FCoreUObjectDelegates::OnAssetLoaded.AddRaw(this, &FWorldPartitionClassDescRegistry::OnAssetLoaded);
 	FCoreUObjectDelegates::OnObjectPreSave.AddRaw(this, &FWorldPartitionClassDescRegistry::OnObjectPreSave);
 	FCoreUObjectDelegates::OnObjectPropertyChanged.AddRaw(this, &FWorldPartitionClassDescRegistry::OnObjectPropertyChanged);
 
@@ -113,6 +117,7 @@ void FWorldPartitionClassDescRegistry::Uninitialize()
 	check(IsInitialized());
 	FActorDescList::Empty();
 
+	FCoreUObjectDelegates::OnAssetLoaded.RemoveAll(this);
 	FCoreUObjectDelegates::OnObjectPreSave.RemoveAll(this);
 	FCoreUObjectDelegates::OnObjectPropertyChanged.RemoveAll(this);
 	
@@ -378,6 +383,36 @@ void FWorldPartitionClassDescRegistry::RegisterClassDescriptorFromActorClass(con
 	RegisterClassDescriptor(NewActorDesc);
 }
 
+void FWorldPartitionClassDescRegistry::OnAssetLoaded(UObject* InAssetLoaded)
+{
+	check(InAssetLoaded);
+	check(IsInitialized());
+
+	if (InAssetLoaded->GetPackage()->HasAnyPackageFlags(PKG_PlayInEditor))
+	{
+		return;
+	}
+
+	UBlueprintGeneratedClass* BlueprintGeneratedClass = Cast<UBlueprintGeneratedClass>(InAssetLoaded);
+
+	if (!BlueprintGeneratedClass)
+	{
+		if (UBlueprint* Blueprint = Cast<UBlueprint>(InAssetLoaded))
+		{
+			BlueprintGeneratedClass = Cast<UBlueprintGeneratedClass>(Blueprint->GeneratedClass);
+		}
+	}
+
+	if (!BlueprintGeneratedClass || !BlueprintGeneratedClass->IsChildOf<AActor>())
+	{
+		return;
+	}
+
+	ValidateInternalState();
+	PrefetchClassDescs({ FTopLevelAssetPath(BlueprintGeneratedClass->GetPathName()) });
+	ValidateInternalState();
+}
+
 void FWorldPartitionClassDescRegistry::OnObjectPreSave(UObject* InObject, FObjectPreSaveContext InSaveContext)
 {
 	if (!InSaveContext.IsProceduralSave() && !(InSaveContext.GetSaveFlags() & SAVE_FromAutosave))
@@ -530,52 +565,49 @@ void FWorldPartitionClassDescRegistry::UpdateClassDescriptor(UObject* InObject,
 	const FTopLevelAssetPath ClassPath(ActorCDO->GetClass()->GetPathName());
 	const FTopLevelAssetPath CurrentParentClassPath = FTopLevelAssetPath(ActorCDO->GetClass()->GetSuperClass()->GetPathName());
 
-	ValidateInternalState();
-
-	if (!ClassByPath.Contains(ClassPath))
+	if (TUniquePtr<FWorldPartitionActorDesc>** ExistingClassDesc = ClassByPath.Find(ClassPath))
 	{
-		PrefetchClassDescs({ ClassPath });
-	}
-
-	TUniquePtr<FWorldPartitionActorDesc>* ExistingClassDesc = ClassByPath.FindChecked(ClassPath);
+		const FTopLevelAssetPath PreviousParentClassPath = ParentClassMap.FindChecked(ClassPath);
 
-	ValidateInternalState();
+		ValidateInternalState();
 
-	if (!ExistingClassDesc)
-	{
-		if (!bOnlyIfExists)
+		if (!*ExistingClassDesc)
 		{
-			FWorldPartitionActorDesc* NewActorDesc = ActorCDO->CreateActorDesc().Release();
-			RegisterClassDescriptor(NewActorDesc);
+			if (!bOnlyIfExists)
+			{
+				FWorldPartitionActorDesc* NewActorDesc = ActorCDO->CreateActorDesc().Release();
+				RegisterClassDescriptor(NewActorDesc);
 
-			ParentClassMap.Add(ClassPath, CurrentParentClassPath);
+				ParentClassMap.Add(ClassPath, CurrentParentClassPath);
+			}
 		}
-	}
-	else
-	{
-		FWorldPartitionActorDescUtils::UpdateActorDescriptorFromActor(ActorCDO, *ExistingClassDesc);
-
-		const FTopLevelAssetPath PreviousParentClassPath = ParentClassMap.FindChecked(ClassPath);
-		if (PreviousParentClassPath != CurrentParentClassPath)
+		else
 		{
-			// We are reparenting a blueprint, update our parent map
-			ParentClassMap.Add(ClassPath, CurrentParentClassPath);
+			FWorldPartitionActorDescUtils::UpdateActorDescriptorFromActor(ActorCDO, **ExistingClassDesc);
+
+			if (PreviousParentClassPath != CurrentParentClassPath)
+			{
+				// We are reparenting a blueprint, update our parent map
+				ParentClassMap.Add(ClassPath, CurrentParentClassPath);
+			}
 		}
-	}
 
-	ValidateInternalState();
+		ValidateInternalState();
+	}
 }
 
 void FWorldPartitionClassDescRegistry::ValidateInternalState()
 {
 #if DO_GUARD_SLOW
-	check(ClassByPath.Num() == ParentClassMap.Num());
-	static const FTopLevelAssetPath ActorClassPath(TEXT("/Script/Engine.Actor"));
 	for (auto& [ClassPath, ParentClassPath] : ParentClassMap)
 	{
 		// Validate that the class parents chain is valid
-		check(ClassByPath.Contains(ClassPath));
-		check((ParentClassPath != ActorClassPath) || ParentClassMap.Contains(ParentClassPath));
+		static const FTopLevelAssetPath ActorClassPath(TEXT("/Script/Engine.Actor"));
+		for (FTopLevelAssetPath CurrentClassPath = ClassPath; CurrentClassPath != ActorClassPath; CurrentClassPath = ParentClassMap.FindChecked(CurrentClassPath))
+		{
+			check(CurrentClassPath.IsValid());
+			ClassByPath.FindChecked(CurrentClassPath);
+		}
 	}
 #endif
 }
@@ -595,4 +627,4 @@ FTopLevelAssetPath FWorldPartitionClassDescRegistry::RedirectClassPath(const FTo
 
 	return ResolvedClassPath;
 }
-#endif
+#endif
\ No newline at end of file
diff --git a/Engine/Source/Runtime/Engine/Public/WorldPartition/WorldPartitionClassDescRegistry.h b/Engine/Source/Runtime/Engine/Public/WorldPartition/WorldPartitionClassDescRegistry.h
--- a/Engine/Source/Runtime/Engine/Public/WorldPartition/WorldPartitionClassDescRegistry.h
+++ b/Engine/Source/Runtime/Engine/Public/WorldPartition/WorldPartitionClassDescRegistry.h
@@ -43,6 +43,7 @@ private:
 	const FWorldPartitionActorDesc* GetClassDescDefaultForActor(const FTopLevelAssetPath& InClassPath) const;
 	const FWorldPartitionActorDesc* GetClassDescDefaultForClass(const FTopLevelAssetPath& InClassPath) const;
 
+	void OnAssetLoaded(UObject* InAssetLoaded);
 	void OnObjectPreSave(UObject* InObject, FObjectPreSaveContext InSaveContext);
 	void OnObjectPropertyChanged(UObject* InObject, FPropertyChangedEvent& InPropertyChangedEvent);
 
