Fix ProcessedSkyTexture of SkyLightComponent not being set when bRealTimeCapture == false.
- Context: in some situations, USkyLightComponent::UpdateSkyCaptureContents() is called before entries in SkyCapturesToUpdate are fully initialized causing them to end up in invalid state.
- Only remove entries from SkyCapturesToUpdate when CaptureComponent->GetWorld() == WorldToUpdate (which indicates the component is initialized).

diff --git a/Engine/Source/Runtime/Engine/Private/Components/SkyLightComponent.cpp b/Engine/Source/Runtime/Engine/Private/Components/SkyLightComponent.cpp
--- a/Engine/Source/Runtime/Engine/Private/Components/SkyLightComponent.cpp
+++ b/Engine/Source/Runtime/Engine/Private/Components/SkyLightComponent.cpp
@@ -714,6 +714,11 @@ void USkyLightComponent::UpdateSkyCaptureContentsArray(UWorld* WorldToUpdate, TA
 		USkyLightComponent* CaptureComponent = ComponentArray[CaptureIndex];
 		AActor* Owner = CaptureComponent->GetOwner();
 
+		if (CaptureComponent->GetWorld() != WorldToUpdate)
+		{
+			continue;
+		}
+
 		// Reset the luminance scale in case the texture has been switched
 		CaptureComponent->SpecifiedCubemapColorScale = FLinearColor::White;
 
@@ -737,7 +742,7 @@ void USkyLightComponent::UpdateSkyCaptureContentsArray(UWorld* WorldToUpdate, TA
 		const bool bCaptureSkyLightWaitingForMeshOrTexAssets= CaptureComponent->SourceType == SLS_CapturedScene		&& bSceneIsAsyncCompiling;
 #endif
 
-		if (((!Owner || !Owner->GetLevel() || Owner->GetLevel()->bIsVisible) && CaptureComponent->GetWorld() == WorldToUpdate)
+		if ((!Owner || !Owner->GetLevel() || Owner->GetLevel()->bIsVisible)
 			// Only process sky capture requests once async texture and shader compiling completes, otherwise we will capture the scene with temporary shaders/textures
 			&& (
 #if WITH_EDITOR
@@ -892,7 +897,7 @@ void USkyLightComponent::UpdateSkyCaptureContents(UWorld* WorldToUpdate)
 		{
 			FScopeLock Lock(&SkyCapturesToUpdateLock);
 			// Remove the sky captures if real time capture is enabled. 
-			SkyCapturesToUpdate.RemoveAll([](const USkyLightComponent* CaptureComponent) { return CaptureComponent->IsRealTimeCaptureEnabled(); });
+			SkyCapturesToUpdate.RemoveAll([WorldToUpdate](const USkyLightComponent* CaptureComponent) { return CaptureComponent->GetWorld() == WorldToUpdate && CaptureComponent->IsRealTimeCaptureEnabled(); });
 			UpdateSkyCaptureContentsArray(WorldToUpdate, SkyCapturesToUpdate, true);
 		}
 		
