Blueprint Data Loss with Certain Variable Names
property names with the characters = ( [ . and whitespace break FProperty::ImportSingleProperty. to solve this, this change surrounds any property name containing the illegal characters in quotes and escapes any existing quotes and backslashes when serialized.

diff --git a/Engine/Source/Runtime/CoreUObject/Private/UObject/Property.cpp b/Engine/Source/Runtime/CoreUObject/Private/UObject/Property.cpp
--- a/Engine/Source/Runtime/CoreUObject/Private/UObject/Property.cpp
+++ b/Engine/Source/Runtime/CoreUObject/Private/UObject/Property.cpp
@@ -1652,19 +1652,37 @@ const TCHAR* FProperty::ImportSingleProperty( const TCHAR* Str, void* DestData,
 
 	// strip leading whitespace
 	const TCHAR* Start = FAsciiSet::Skip(Str, Whitespaces);
-	// find first delimiter
-	Str = FAsciiSet::FindFirstOrEnd(Start, Delimiters);
-	// check if delimiter was found...
-	if (*Str)
-	{
-		// strip trailing whitespace
-		int32 Len = UE_PTRDIFF_TO_INT32(Str - Start);
-		while (Len > 0 && Whitespaces.Contains(Start[Len - 1]))
+	FName PropertyName;
+	
+	if (*Start == '"')
+	{
+		int32 OutQuotedLen;
+		FString OutUnquotedString;
+		FParse::QuotedString(Start, OutUnquotedString, &OutQuotedLen);
+		PropertyName = FName(OutUnquotedString);
+		
+		// advance iterator to next delimiter
+		Str = FAsciiSet::FindFirstOrEnd(Start + OutQuotedLen, Delimiters);
+	}
+	else // legacy format requires that we support un-quoted and un-escaped property names
+	{
+		// find first delimiter
+		Str = FAsciiSet::FindFirstOrEnd(Start, Delimiters);
+		// check if delimiter was found...
+		if (*Str)
 		{
-			--Len;
+			// strip trailing whitespace
+			int32 Len = UE_PTRDIFF_TO_INT32(Str - Start);
+			while (Len > 0 && Whitespaces.Contains(Start[Len - 1]))
+			{
+				--Len;
+			}
+			PropertyName = FName(Len, Start);
 		}
-
-		const FName PropertyName(Len, Start);
+	}
+	
+	if (*Str && !PropertyName.IsNone())
+	{
 		FProperty* Property = FindFProperty<FProperty>(ObjectStruct, PropertyName);
 
 		if (Property == nullptr)
diff --git a/Engine/Source/Runtime/Engine/Private/UnrealExporter.cpp b/Engine/Source/Runtime/Engine/Private/UnrealExporter.cpp
--- a/Engine/Source/Runtime/Engine/Private/UnrealExporter.cpp
+++ b/Engine/Source/Runtime/Engine/Private/UnrealExporter.cpp
@@ -16,6 +16,7 @@
 #include "Model.h"
 #include "Misc/FeedbackContext.h"
 #include "AssetExportTask.h"
+#include "Misc/AsciiSet.h"
 #include "UObject/GCObjectScopeGuard.h"
 
 #if WITH_EDITOR
@@ -679,7 +680,6 @@ void ExportProperties
 	UObject*		ExportRootScope
 )
 {
-	FString ThisName = TEXT("(none)");
 	check(ObjectClass != NULL);
 
 	for( FProperty* Property = ObjectClass->PropertyLink; Property; Property = Property->PropertyLinkNext )
@@ -687,7 +687,7 @@ void ExportProperties
 		if (!Property->ShouldPort(PortFlags))
 			continue;
 
-		ThisName = Property->GetName();
+		const FString SanitizedPropertyName = FString::Format(TEXT("\"{0}\""), {Property->GetName().ReplaceCharWithEscapedChar()});
 		FArrayProperty* ArrayProperty = CastField<FArrayProperty>(Property);
 		FObjectPropertyBase* ExportObjectProp = (Property->PropertyFlags & CPF_ExportObject) != 0 ? CastField<FObjectPropertyBase>(Property) : NULL;
 		const uint32 ExportFlags = PortFlags | PPF_Delimited;
@@ -724,7 +724,7 @@ void ExportProperties
 				// If the current size of the array is 0 and the default one is not, add in an empty item so on import it will be empty
 				if( ArrayHelper.Num() == 0 && DiffArrayHelper.Num() != 0 )
 				{
-					Out.Logf(TEXT("%s%s=\r\n"), FCString::Spc(Indent), *Property->GetName());
+					Out.Logf(TEXT("%s%s=\r\n"), FCString::Spc(Indent), *SanitizedPropertyName);
 				}
 				else
 				{
@@ -779,12 +779,12 @@ void ExportProperties
 								}
 							}
 
-							Out.Logf(TEXT("%s%s(%i)=%s\r\n"), FCString::Spc(Indent), *Property->GetName(), DynamicArrayIndex, *Value);
+							Out.Logf(TEXT("%s%s(%i)=%s\r\n"), FCString::Spc(Indent), *SanitizedPropertyName, DynamicArrayIndex, *Value);
 						}
 					}
 					for (int32 DynamicArrayIndex = DiffArrayHelper.Num()-1; DynamicArrayIndex >= ArrayHelper.Num(); --DynamicArrayIndex)
 					{
-						Out.Logf(TEXT("%s%s.RemoveIndex(%d)\r\n"), FCString::Spc(Indent), *Property->GetName(), DynamicArrayIndex);
+						Out.Logf(TEXT("%s%s.RemoveIndex(%d)\r\n"), FCString::Spc(Indent), *SanitizedPropertyName, DynamicArrayIndex);
 					}
 				}
 			}
@@ -837,11 +837,11 @@ void ExportProperties
 
 					if( Property->ArrayDim == 1 )
 					{
-						Out.Logf( TEXT("%s%s=%s\r\n"), FCString::Spc(Indent), *Property->GetName(), *Value );
+						Out.Logf( TEXT("%s%s=%s\r\n"), FCString::Spc(Indent), *SanitizedPropertyName, *Value );
 					}
 					else
 					{
-						Out.Logf( TEXT("%s%s(%i)=%s\r\n"), FCString::Spc(Indent), *Property->GetName(), PropertyArrayIndex, *Value );
+						Out.Logf( TEXT("%s%s(%i)=%s\r\n"), FCString::Spc(Indent), *SanitizedPropertyName, PropertyArrayIndex, *Value );
 					}
 				}
 			}
