Switch Nanite (WIP) scalar displacement material graph work from vertex to pixel inputs, and also hooked up HasDisplacementConnected() to material resource query.

diff --git a/Engine/Shaders/Private/MaterialTemplate.ush b/Engine/Shaders/Private/MaterialTemplate.ush
--- a/Engine/Shaders/Private/MaterialTemplate.ush
+++ b/Engine/Shaders/Private/MaterialTemplate.ush
@@ -2956,10 +2956,10 @@ half GetMaterialMaskedOpacity(FPixelMaterialInputs PixelMaterialInputs)
 }
 #endif
 
-float GetMaterialDisplacement(FMaterialVertexParameters Parameters)
+float GetMaterialDisplacement(FPixelMaterialInputs PixelMaterialInputs)
 {
 #if USES_DISPLACEMENT
-%s;
+	return PixelMaterialInputs.Displacement;
 #else
 	return 0.0f;
 #endif
diff --git a/Engine/Source/Runtime/Engine/Private/Materials/HLSLMaterialTranslator.cpp b/Engine/Source/Runtime/Engine/Private/Materials/HLSLMaterialTranslator.cpp
--- a/Engine/Source/Runtime/Engine/Private/Materials/HLSLMaterialTranslator.cpp
+++ b/Engine/Source/Runtime/Engine/Private/Materials/HLSLMaterialTranslator.cpp
@@ -410,6 +410,7 @@ FHLSLMaterialTranslator::FHLSLMaterialTranslator(FMaterial* InMaterial,
 	SharedPixelProperties[MP_ShadingModel] = true;
 	SharedPixelProperties[MP_SurfaceThickness] = true;
 	SharedPixelProperties[MP_FrontMaterial] = true;
+	SharedPixelProperties[MP_Displacement] = true;
 
 	for (int32 Frequency = 0; Frequency < SF_NumFrequencies; ++Frequency)
 	{
@@ -1016,7 +1017,6 @@ bool FHLSLMaterialTranslator::Translate()
 		Chunk[MP_OpacityMask]					= Material->CompilePropertyAndSetMaterialProperty(MP_OpacityMask			,this);
 		Chunk[MP_Tangent]						= Material->CompilePropertyAndSetMaterialProperty(MP_Tangent				,this);
 		Chunk[MP_WorldPositionOffset]			= Material->CompilePropertyAndSetMaterialProperty(MP_WorldPositionOffset	,this);
-		Chunk[MP_Displacement]					= Material->CompilePropertyAndSetMaterialProperty(MP_Displacement			,this);
 
 		// Make sure to compile this property before using ShadingModelsFromCompilation
 		Chunk[MP_ShadingModel]					= Material->CompilePropertyAndSetMaterialProperty(MP_ShadingModel			,this);
@@ -1088,6 +1088,8 @@ bool FHLSLMaterialTranslator::Translate()
 			}
 
 			Chunk[MP_PixelDepthOffset] = Material->CompilePropertyAndSetMaterialProperty(MP_PixelDepthOffset, this);
+
+			Chunk[MP_Displacement] = Material->CompilePropertyAndSetMaterialProperty(MP_Displacement, this);
 		}
 
 		ResourcesString = TEXT("");
@@ -2644,7 +2646,6 @@ FString FHLSLMaterialTranslator::GetMaterialShaderCode()
 
 	LazyPrintf.PushParam(*FString::Printf(TEXT("return %.5f"), Material->GetOpacityMaskClipValue()));
 
-	LazyPrintf.PushParam(!bEnableExecutionFlow ? *GenerateFunctionCode(MP_Displacement, BaseDerivativeVariation) : TEXT("return Parameters.MaterialAttributes.Displacement"));
 	LazyPrintf.PushParam(!bEnableExecutionFlow ? *GenerateFunctionCode(MP_WorldPositionOffset, BaseDerivativeVariation) : TEXT("return Parameters.MaterialAttributes.WorldPositionOffset"));
 	LazyPrintf.PushParam(!bEnableExecutionFlow ? *GenerateFunctionCode(CompiledMP_PrevWorldPositionOffset, BaseDerivativeVariation) : TEXT("return 0.0f"));
 	LazyPrintf.PushParam(!bEnableExecutionFlow ? *GenerateFunctionCode(MP_CustomData0, BaseDerivativeVariation) : TEXT("return 0.0f"));
diff --git a/Engine/Source/Runtime/Engine/Private/Materials/MaterialAttributeDefinitionMap.cpp b/Engine/Source/Runtime/Engine/Private/Materials/MaterialAttributeDefinitionMap.cpp
--- a/Engine/Source/Runtime/Engine/Private/Materials/MaterialAttributeDefinitionMap.cpp
+++ b/Engine/Source/Runtime/Engine/Private/Materials/MaterialAttributeDefinitionMap.cpp
@@ -261,7 +261,7 @@ void FMaterialAttributeDefinitionMap::InitializeAttributeMap()
 
 	// Advanced attributes
 	Add(FGuid(0xF905F895, 0xD5814314, 0x916D2434, 0x8C40CE9E), TEXT("WorldPositionOffset"),		MP_WorldPositionOffset,		MCT_Float3,	FVector4(0,0,0,0),	SF_Vertex);
-	Add(FGuid(0x199A7166, 0xC67041DC, 0xA68EAD0D, 0x7017D0AD), TEXT("Displacement"),			MP_Displacement,			MCT_Float,	FVector4(0,0,0,0),	SF_Vertex);
+	Add(FGuid(0x199A7166, 0xC67041DC, 0xA68EAD0D, 0x7017D0AD), TEXT("Displacement"),			MP_Displacement,			MCT_Float,	FVector4(0,0,0,0),	SF_Pixel);
 	Add(FGuid(0x5B8FC679, 0x51CE4082, 0x9D777BEE, 0xF4F72C44), TEXT("SubsurfaceColor"),			MP_SubsurfaceColor,			MCT_Float3,	FVector4(1,1,1,0),	SF_Pixel);
 	Add(FGuid(0x9E502E69, 0x3C8F48FA, 0x94645CFD, 0x28E5428D), TEXT("ClearCoat"),				MP_CustomData0,				MCT_Float,	FVector4(1,0,0,0),	SF_Pixel);
 	Add(FGuid(0xBE4F2FFD, 0x12FC4296, 0xB0124EEA, 0x12C28D92), TEXT("ClearCoatRoughness"),		MP_CustomData1,				MCT_Float,	FVector4(.1,0,0,0),	SF_Pixel);
diff --git a/Engine/Source/Runtime/Engine/Private/Materials/MaterialHLSLEmitter.cpp b/Engine/Source/Runtime/Engine/Private/Materials/MaterialHLSLEmitter.cpp
--- a/Engine/Source/Runtime/Engine/Private/Materials/MaterialHLSLEmitter.cpp
+++ b/Engine/Source/Runtime/Engine/Private/Materials/MaterialHLSLEmitter.cpp
@@ -793,6 +793,7 @@ bool MaterialEmitHLSL(const FMaterialCompileTargetParameters& InCompilerTarget,
 	SharedPixelProperties[MP_ShadingModel] = true;
 	SharedPixelProperties[MP_SurfaceThickness] = true;
 	SharedPixelProperties[MP_FrontMaterial] = true;
+	SharedPixelProperties[MP_Displacement] = true;
 
 	bool bUsesWorldPositionOffset = false;
 	bool bUsesPixelDepthOffset = false;
@@ -850,6 +851,11 @@ bool MaterialEmitHLSL(const FMaterialCompileTargetParameters& InCompilerTarget,
 			bUsesPixelDepthOffset = true;
 		}
 
+		if (CachedTree->IsAttributeUsed(EmitContext, *EmitResultScope, PixelResultType0, MP_Displacement))
+		{
+			bUsesDisplacement = true;
+		}
+
 		const bool bUseNormal = EmitMaterialData.IsExternalInputUsed(SF_Pixel, Material::EExternalInput::WorldNormal);
 		const bool bUseReflection = EmitMaterialData.IsExternalInputUsed(SF_Pixel, Material::EExternalInput::WorldReflection);
 		if (!bUseNormal && !bUseReflection)
@@ -931,7 +937,6 @@ bool MaterialEmitHLSL(const FMaterialCompileTargetParameters& InCompilerTarget,
 		EmitMaterialData.EmitInterpolatorStatements(EmitContext, *EmitResultScope);
 
 		bUsesWorldPositionOffset = CachedTree->IsAttributeUsed(EmitContext, *EmitResultScope, VertexResultType, MP_WorldPositionOffset);
-		bUsesDisplacement = CachedTree->IsAttributeUsed(EmitContext, *EmitResultScope, VertexResultType, MP_Displacement);
 
 		CachedTree->GetTree().EmitShader(EmitContext, VertexCode);
 	}
diff --git a/Engine/Source/Runtime/Engine/Private/Materials/MaterialShared.cpp b/Engine/Source/Runtime/Engine/Private/Materials/MaterialShared.cpp
--- a/Engine/Source/Runtime/Engine/Private/Materials/MaterialShared.cpp
+++ b/Engine/Source/Runtime/Engine/Private/Materials/MaterialShared.cpp
@@ -1059,7 +1059,7 @@ bool FMaterial::MaterialMayModifyMeshPosition() const
 {
 	// Conservative estimate when called before material translation has occurred. 
 	// This function is only intended for use in deciding whether or not shader permutations are required.
-	return HasVertexPositionOffsetConnected() || HasPixelDepthOffsetConnected();
+	return HasVertexPositionOffsetConnected() || HasPixelDepthOffsetConnected() || HasDisplacementConnected();
 }
 
 bool FMaterial::MaterialUsesPixelDepthOffset_GameThread() const
@@ -1957,6 +1957,11 @@ bool FMaterialResource::HasAmbientOcclusionConnected() const
 	return Material->HasAmbientOcclusionConnected();
 }
 
+bool FMaterialResource::HasDisplacementConnected() const
+{
+	return Material->HasDisplacementConnected();
+}
+
 bool FMaterialResource::IsStrataMaterial() const
 {
 	// We no longer support both types of material (Strata and non strata) so no need to check if FrontMaterial is plugged in.
@@ -5137,6 +5142,7 @@ FMaterialShaderParameters::FMaterialShaderParameters(const FMaterial* InMaterial
 	}
 	bHasVertexPositionOffsetConnected = InMaterial->HasVertexPositionOffsetConnected();
 	bHasPixelDepthOffsetConnected = InMaterial->HasPixelDepthOffsetConnected();
+	bHasDisplacementConnected = InMaterial->HasDisplacementConnected();
 	bMaterialMayModifyMeshPosition = InMaterial->MaterialMayModifyMeshPosition();
 	bIsUsedWithStaticLighting = InMaterial->IsUsedWithStaticLighting();
 	bIsUsedWithParticleSprites = InMaterial->IsUsedWithParticleSprites();
diff --git a/Engine/Source/Runtime/Engine/Public/MaterialShared.h b/Engine/Source/Runtime/Engine/Public/MaterialShared.h
--- a/Engine/Source/Runtime/Engine/Public/MaterialShared.h
+++ b/Engine/Source/Runtime/Engine/Public/MaterialShared.h
@@ -1944,6 +1944,7 @@ public:
 	virtual bool HasAnisotropyConnected() const { return false; }
 	virtual bool HasAmbientOcclusionConnected() const { return false; }
 	virtual bool HasMaterialPropertyConnected(EMaterialProperty In) const { return false; }
+	virtual bool HasDisplacementConnected() const { return false; }
 	virtual bool IsStrataMaterial() const { return false; }
 	virtual bool RequiresSynchronousCompilation() const { return false; };
 	virtual bool IsDefaultMaterial() const { return false; };
@@ -2629,6 +2630,7 @@ public:
 	ENGINE_API virtual bool HasEmissiveColorConnected() const override;
 	ENGINE_API virtual bool HasAnisotropyConnected() const override;
 	ENGINE_API virtual bool HasAmbientOcclusionConnected() const override;
+	ENGINE_API virtual bool HasDisplacementConnected() const override;
 	ENGINE_API virtual bool IsStrataMaterial() const override;
 	ENGINE_API virtual bool HasMaterialPropertyConnected(EMaterialProperty In) const override;
 	ENGINE_API virtual FMaterialShadingModelField GetShadingModels() const override;
@@ -2989,6 +2991,7 @@ struct FMaterialShaderParameters
 			uint64 bHasEmissiveColorConnected : 1;
 			uint64 bHasAmbientOcclusionConnected : 1;
 			uint64 bHasAnisotropyConnected : 1;
+			uint64 bHasDisplacementConnected : 1;
 			uint64 bHasVertexPositionOffsetConnected : 1;
 			uint64 bHasPixelDepthOffsetConnected : 1;
 			uint64 bMaterialMayModifyMeshPosition : 1;
