Conditionally load PreIntegratedSkinBRDF texture only for platforms that use pre-integrated skin shading. Always loaded in Editor

diff --git a/Engine/Source/Runtime/Engine/Classes/Engine/Engine.h b/Engine/Source/Runtime/Engine/Classes/Engine/Engine.h
--- a/Engine/Source/Runtime/Engine/Classes/Engine/Engine.h
+++ b/Engine/Source/Runtime/Engine/Classes/Engine/Engine.h
@@ -1902,6 +1902,9 @@ public:
 	/** Delay loading this texture until it is needed by the renderer. */
 	void LoadDefaultFilmGrainTexture();
 
+	/** Conditionally load this texture for a platform. Always loaded in Editor */
+	void ConditionallyLoadPreIntegratedSkinBRDFTexture();
+
 private:
 	#if WITH_DYNAMIC_RESOLUTION
 		/** Last dynamic resolution event. */
diff --git a/Engine/Source/Runtime/Engine/Private/UnrealEngine.cpp b/Engine/Source/Runtime/Engine/Private/UnrealEngine.cpp
--- a/Engine/Source/Runtime/Engine/Private/UnrealEngine.cpp
+++ b/Engine/Source/Runtime/Engine/Private/UnrealEngine.cpp
@@ -3112,10 +3112,10 @@ void UEngine::InitializeObjectReferences()
 	LoadEngineTexture(DefaultDiffuseTexture, *DefaultDiffuseTextureName.ToString());
 	LoadEngineTexture(HighFrequencyNoiseTexture, *HighFrequencyNoiseTextureName.ToString());
 	LoadEngineTexture(DefaultBokehTexture, *DefaultBokehTextureName.ToString());
-	LoadEngineTexture(PreIntegratedSkinBRDFTexture, *PreIntegratedSkinBRDFTextureName.ToString());
 	LoadEngineTexture(MiniFontTexture, *MiniFontTextureName.ToString());
 	LoadEngineTexture(WeightMapPlaceholderTexture, *WeightMapPlaceholderTextureName.ToString());
 	LoadEngineTexture(LightMapDensityTexture, *LightMapDensityTextureName.ToString());
+	ConditionallyLoadPreIntegratedSkinBRDFTexture();
 
 #if WITH_EDITOR
 	// Avoid breaking some engine textures that might be cached very early (i.e. BlueNoise)
@@ -3306,6 +3306,19 @@ void UEngine::LoadDefaultFilmGrainTexture()
 	}
 }
 
+void UEngine::ConditionallyLoadPreIntegratedSkinBRDFTexture()
+{
+	if (PreIntegratedSkinBRDFTexture == nullptr)
+	{
+		uint32 ShadingModelsMask = GetPlatformShadingModelsMask(GMaxRHIShaderPlatform);
+		uint32 SkinShadingMask = (1u << (uint32)MSM_PreintegratedSkin);
+		if (GIsEditor || (ShadingModelsMask & SkinShadingMask) != 0)
+		{
+			LoadEngineTexture(PreIntegratedSkinBRDFTexture, *PreIntegratedSkinBRDFTextureName.ToString());
+		}
+	}
+}
+
 void UEngine::InitializePortalServices()
 {
 	IMessagingRpcModule* MessagingRpcModule = nullptr;
diff --git a/Engine/Source/Runtime/Renderer/Private/SceneRendering.cpp b/Engine/Source/Runtime/Renderer/Private/SceneRendering.cpp
--- a/Engine/Source/Runtime/Renderer/Private/SceneRendering.cpp
+++ b/Engine/Source/Runtime/Renderer/Private/SceneRendering.cpp
@@ -1841,6 +1841,7 @@ void FViewInfo::SetupUniformBufferParameters(
 		ViewUniformShaderParameters.FarShadowStaticMeshLODBias = 0;
 	}
 
+	if (GEngine->PreIntegratedSkinBRDFTexture)
 	{
 		const FTextureResource* TextureResource = GEngine->PreIntegratedSkinBRDFTexture->GetResource();
 		if (TextureResource)
