Fix too intense transmission on subsurface profile shading model.
  * Extend ExtinctionScale from 1 to 100. This could reduce the transmission distance when the scale is larger than 1.

How to make transmission look right:
1. Reduce the scattering distribution to minimal 0.01. The current version of anisotropy is not right. We will use isotropic scattering when setting to minimal, 0.01. Path traced SSS uses isotropic SSS by default.
2. Increase the Extinction Scale to reduce the transmission.
3. Turn Transmission Tint Color to max and adjust to match path tracing. Note that Transmission Tint Color is originally added to reduce the intense transmission. This one should only be used later for NPBR.

diff --git a/Engine/Shaders/Private/RayTracing/RayTracingOcclusionRGS.usf b/Engine/Shaders/Private/RayTracing/RayTracingOcclusionRGS.usf
--- a/Engine/Shaders/Private/RayTracing/RayTracingOcclusionRGS.usf
+++ b/Engine/Shaders/Private/RayTracing/RayTracingOcclusionRGS.usf
@@ -575,8 +575,9 @@ FOcclusionResult ComputeOcclusion(
 		const float3 TranslatedWorldViewOrigin = LWCToFloat(LWCAdd(PrimaryView.WorldViewOrigin, PrimaryView.PreViewTranslation));
 		const float3 V = normalize(TranslatedWorldPosition - TranslatedWorldViewOrigin);
 
-		// Apply sqrt to match the behavior of the rasterizer closer visually.
-		const float ExtinctionScale = sqrt(TransmissionProfileParams.ExtinctionScale);
+		// Apply sqrt to match the behavior of the rasterizer closer visually when ExtinctionScale <=1.
+		const float ExtinctionScale = lerp(sqrt(TransmissionProfileParams.ExtinctionScale), 
+											TransmissionProfileParams.ExtinctionScale, TransmissionProfileParams.ExtinctionScale > 1);
 		
 		// Decrease ExtinctionScale should increase the max transmission distance. Keep rasterizer and raytracing shadow behave similar for transmission.
 		const float MaxTransmissionDistance = SSSS_MAX_TRANSMISSION_PROFILE_DISTANCE / ExtinctionScale; 
diff --git a/Engine/Shaders/Private/SubsurfaceProfileCommon.ush b/Engine/Shaders/Private/SubsurfaceProfileCommon.ush
--- a/Engine/Shaders/Private/SubsurfaceProfileCommon.ush
+++ b/Engine/Shaders/Private/SubsurfaceProfileCommon.ush
@@ -53,6 +53,9 @@
 #define ENC_DIFFUSEMEANFREEPATH_IN_MM_TO_UNIT (0.01f*0.2f)
 #define DEC_UNIT_TO_DIFFUSEMEANFREEPATH_IN_MM 1/ENC_DIFFUSEMEANFREEPATH_IN_MM_TO_UNIT
 
+// Make sure UIMax|ClampMax of ExtinctionScale * ENC_EXTINCTIONSCALE_FACTOR <= 1
+#define ENC_EXTINCTIONSCALE_FACTOR	0.01f
+#define DEC_EXTINCTIONSCALE_FACTOR  1/ENC_EXTINCTIONSCALE_FACTOR
 ///////////////////////////////////////////////////////////////////////////////////////////////
 // Accessors
 
@@ -127,6 +130,16 @@ float DecodeScatteringDistribution(float ScatteringDistribution)
 	return ScatteringDistribution * 2.0f - 1.0f;
 }
 
+float EncodeExtinctionScale(float ExtinctionScale)
+{
+	return ExtinctionScale * ENC_EXTINCTIONSCALE_FACTOR;
+}
+
+float DecodeExtinctionScale(float ExtinctionScale)
+{
+	return ExtinctionScale * DEC_EXTINCTIONSCALE_FACTOR;
+}
+
 bool GetSubsurfaceTransmittanceProfileUseBurley(uint SubsurfaceProfileInt)
 {
 	half Type = GetSubsurfaceProfileTexture(SSSS_BOUNDARY_COLOR_BLEED_OFFSET, SubsurfaceProfileInt).a;
diff --git a/Engine/Shaders/Private/TransmissionCommon.ush b/Engine/Shaders/Private/TransmissionCommon.ush
--- a/Engine/Shaders/Private/TransmissionCommon.ush
+++ b/Engine/Shaders/Private/TransmissionCommon.ush
@@ -34,7 +34,7 @@ FTransmissionProfileParams GetTransmissionProfileParams(uint SubsurfaceProfileIn
 	//X:ExtinctionScale, Y:Normal Scale, Z:ScatteringDistribution, W:OneOverIOR
 	float4 Data = View.SSProfilesTexture.Load(int3(SSSS_TRANSMISSION_OFFSET, SubsurfaceProfileInt, 0));
 
-	Result.ExtinctionScale = Data.x;
+	Result.ExtinctionScale = DecodeExtinctionScale(Data.x);
 	Result.NormalScale = Data.y;
 	Result.ScatteringDistribution = DecodeScatteringDistribution(Data.z);
 	Result.OneOverIOR = Data.w;
diff --git a/Engine/Source/Runtime/Engine/Classes/Engine/SubsurfaceProfile.h b/Engine/Source/Runtime/Engine/Classes/Engine/SubsurfaceProfile.h
--- a/Engine/Source/Runtime/Engine/Classes/Engine/SubsurfaceProfile.h
+++ b/Engine/Source/Runtime/Engine/Classes/Engine/SubsurfaceProfile.h
@@ -80,7 +80,7 @@ struct FSubsurfaceProfileStruct
 	UPROPERTY(Category = "Burley Normalized", EditAnywhere, BlueprintReadOnly, meta = (HideAlphaChannel))
 	FLinearColor BoundaryColorBleed;
 
-	UPROPERTY(Category = "Transmission", EditAnywhere, BlueprintReadOnly, meta = (ClampMin = "0.01", UIMax = "1.0", ClampMax = "1.0"))
+	UPROPERTY(Category = "Transmission", EditAnywhere, BlueprintReadOnly, meta = (ClampMin = "0.01", UIMax = "100.0", ClampMax = "100.0", SliderExponent="24"))
 	float ExtinctionScale;
 	
 	UPROPERTY(Category = "Transmission", EditAnywhere, BlueprintReadOnly, meta = (ClampMin = "0.01", UIMax = "0.99", ClampMax = "1.0"))
diff --git a/Engine/Source/Runtime/Engine/Private/Rendering/SubsurfaceProfile.cpp b/Engine/Source/Runtime/Engine/Private/Rendering/SubsurfaceProfile.cpp
--- a/Engine/Source/Runtime/Engine/Private/Rendering/SubsurfaceProfile.cpp
+++ b/Engine/Source/Runtime/Engine/Private/Rendering/SubsurfaceProfile.cpp
@@ -356,6 +356,9 @@ static float GetNextSmallerPositiveFloat(float x)
 // 1 mfp can map to 1.44 dmfp (surface albedo -> 0.0) or 43.50 dmfp (surface albedo -> 1.0).
 #define ENC_DIFFUSEMEANFREEPATH_IN_MM_TO_UNIT (0.01f*0.2f)
 #define DEC_UNIT_TO_DIFFUSEMEANFREEPATH_IN_MM 1/ENC_DIFFUSEMEANFREEPATH_IN_MM_TO_UNIT
+
+#define ENC_EXTINCTIONSCALE_FACTOR	0.01f
+#define DEC_EXTINCTIONSCALE_FACTOR  1/ENC_EXTINCTIONSCALE_FACTOR
 //------------------------------------------------------------------------------------------
 
 //in [0,1]
@@ -391,6 +394,17 @@ float DecodeScatteringDistribution(float ScatteringDistribution)
 	return ScatteringDistribution * 2.0f - 1.0f;
 }
 
+float EncodeExtinctionScale(float ExtinctionScale)
+{
+	return ExtinctionScale * ENC_EXTINCTIONSCALE_FACTOR;
+}
+
+float DecodeExtinctionScale(float ExtinctionScale)
+{
+	return ExtinctionScale * DEC_EXTINCTIONSCALE_FACTOR;
+}
+
+
 void SetupSurfaceAlbedoAndDiffuseMeanFreePath(FLinearColor& SurfaceAlbedo, FLinearColor& Dmfp)
 {
 	int32 SamplingSelectionMethod = FMath::Clamp(CVarSSProfilesSamplingChannelSelection.GetValueOnAnyThread(), 0, 1);
@@ -506,7 +520,7 @@ void FSubsurfaceProfileTexture::CreateTexture(FRHICommandListImmediate& RHICmdLi
 		TextureRow[SSSS_DUAL_SPECULAR_OFFSET].A = FMath::Clamp(MaterialRoughnessToAverage / SSSS_MAX_DUAL_SPECULAR_ROUGHNESS, 0.0f, 1.0f);
 
 		//X:ExtinctionScale, Y:Normal Scale, Z:ScatteringDistribution, W:OneOverIOR
-		TextureRow[SSSS_TRANSMISSION_OFFSET].R = Data.ExtinctionScale;
+		TextureRow[SSSS_TRANSMISSION_OFFSET].R = EncodeExtinctionScale(Data.ExtinctionScale);
 		TextureRow[SSSS_TRANSMISSION_OFFSET].G = Data.NormalScale;
 		TextureRow[SSSS_TRANSMISSION_OFFSET].B = EncodeScatteringDistribution(Data.ScatteringDistribution);
 		TextureRow[SSSS_TRANSMISSION_OFFSET].A = 1.0f / Data.IOR;
