Fixed mismatched copy constructor/assignment operator in FCollisionEnabledMask.

diff --git a/Engine/Source/Runtime/Engine/Classes/Engine/EngineTypes.h b/Engine/Source/Runtime/Engine/Classes/Engine/EngineTypes.h
--- a/Engine/Source/Runtime/Engine/Classes/Engine/EngineTypes.h
+++ b/Engine/Source/Runtime/Engine/Classes/Engine/EngineTypes.h
@@ -1453,7 +1453,6 @@ struct ENGINE_API FCollisionEnabledMask
 {
 	int8 Bits;
 
-	FCollisionEnabledMask(const FCollisionEnabledMask&) = default;
 	FCollisionEnabledMask(int8 InBits = 0);
 	FCollisionEnabledMask(ECollisionEnabled::Type CollisionEnabled);
 
