VT Builder detect uint32 overflow and return false
need to drill out bSucceeded = false

diff --git a/Engine/Source/Runtime/Engine/Private/TextureDerivedDataTask.cpp b/Engine/Source/Runtime/Engine/Private/TextureDerivedDataTask.cpp
--- a/Engine/Source/Runtime/Engine/Private/TextureDerivedDataTask.cpp
+++ b/Engine/Source/Runtime/Engine/Private/TextureDerivedDataTask.cpp
@@ -632,6 +632,8 @@ static void DDC1_BuildTexture(
 	const bool bHasValidMip0 = TextureData.Blocks.Num() && TextureData.Blocks[0].MipsPerLayer.Num() && TextureData.Blocks[0].MipsPerLayer[0].Num();
 	const bool bForVirtualTextureStreamingBuild = EnumHasAnyFlags(CacheFlags, ETextureCacheFlags::ForVirtualTextureStreamingBuild);
 
+	check( bSucceeded == false ); // we set it to true if we succeed
+
 	if (!ensure(Compressor))
 	{
 		UE_LOG(LogTexture, Warning, TEXT("Missing Compressor required to build texture %s"), *TexturePathName);
@@ -656,10 +658,24 @@ static void DDC1_BuildTexture(
 		}
 		
 		FVirtualTextureBuilderDerivedInfo PredictedInfo;
-		PredictedInfo.InitializeFromBuildSettings(TextureData, InBuildSettingsPerLayer.GetData());
+		if ( ! PredictedInfo.InitializeFromBuildSettings(TextureData, InBuildSettingsPerLayer.GetData()) )
+		{
+			UE_LOG(LogTexture, Warning, TEXT("VT InitializeFromBuildSettings failed: %s"), *TexturePathName);
+			delete DerivedData->VTData;
+			DerivedData->VTData = nullptr;
+			bSucceeded = false;
+			return;		
+		}
 
 		FVirtualTextureDataBuilder Builder(*DerivedData->VTData, TexturePathName, Compressor, ImageWrapper);
-		Builder.Build(TextureData, CompositeTextureData, &InBuildSettingsPerLayer[0], true);
+		if ( ! Builder.Build(TextureData, CompositeTextureData, &InBuildSettingsPerLayer[0], true) )
+		{
+			UE_LOG(LogTexture, Warning, TEXT("VT Build failed: %s"), *TexturePathName);
+			delete DerivedData->VTData;
+			DerivedData->VTData = nullptr;
+			bSucceeded = false;
+			return;		
+		}
 
 		// TextureData was freed by Build (FTextureSourceData.ReleaseMemory), don't use it from here down
 
@@ -766,6 +782,7 @@ static void DDC1_BuildTexture(
 			// will log below
 			check( DerivedData->Mips.Num() == 0 );
 			DerivedData->Mips.Empty();
+			bSucceeded = false;
 
 			UE_LOG(LogTexture, Warning, TEXT("BuildTexture failed to build %s derived data for %s"), *InBuildSettingsPerLayer[0].TextureFormatName.GetPlainNameString(), *TexturePathName);
 		}
@@ -1935,9 +1952,6 @@ void FTextureCacheDerivedDataWorker::DoWork()
 					Chunk.bCorruptDataLoadedFromDDC = true;
 				}
 			}
-
-			// this always ignores whether the build failed?
-			bSucceeded = true;
 		}
 		else
 		{
@@ -1945,11 +1959,11 @@ void FTextureCacheDerivedDataWorker::DoWork()
 		}
 	}
 
+	TextureData.ReleaseMemory();
+	CompositeTextureData.ReleaseMemory();
+
 	if (bSucceeded)
 	{
-		TextureData.ReleaseMemory();
-		CompositeTextureData.ReleaseMemory();
-
 		// Populate the VT DDC Cache now if we're asynchronously loading to avoid too many high prio/synchronous request on the render thread
 		if (!IsInGameThread() && DerivedData->VTData && !DerivedData->VTData->Chunks.Last().DerivedDataKey.IsEmpty())
 		{
diff --git a/Engine/Source/Runtime/Engine/Private/VT/VirtualTextureBuiltData.cpp b/Engine/Source/Runtime/Engine/Private/VT/VirtualTextureBuiltData.cpp
--- a/Engine/Source/Runtime/Engine/Private/VT/VirtualTextureBuiltData.cpp
+++ b/Engine/Source/Runtime/Engine/Private/VT/VirtualTextureBuiltData.cpp
@@ -99,7 +99,7 @@ uint32 FVirtualTextureBuiltData::GetMemoryFootprint() const
 
 uint32 FVirtualTextureBuiltData::GetTileMemoryFootprint() const
 {
-	uint32 TotalSize = 0;
+	uint64 TotalSize = 0;
 
 	// Legacy tile offsets are used if tiles are compressed.
 	TotalSize += TileOffsetInChunk.GetAllocatedSize();
@@ -115,7 +115,7 @@ uint32 FVirtualTextureBuiltData::GetTileMemoryFootprint() const
 		TotalSize += Data.Offsets.GetAllocatedSize();
 	}
 
-	return TotalSize;
+	return IntCastChecked<uint32>( TotalSize );
 }
 
 uint32 FVirtualTextureBuiltData::GetNumTileHeaders() const
@@ -210,7 +210,8 @@ uint32 FVirtualTextureBuiltData::GetTileOffset(uint32 vLevel, uint32 vAddress, u
 				const uint32 TileDataSize = TileDataOffsetPerLayer.Last();
 				const uint32 LayerDataOffset = LayerIndex == 0 ? 0 : TileDataOffsetPerLayer[LayerIndex - 1];
 				
-				Offset = BaseOffset + (TileOffset * TileDataSize) + LayerDataOffset;
+				int64 Offset64 = BaseOffset + (int64) TileOffset * TileDataSize + LayerDataOffset;
+				Offset = IntCastChecked<uint32>( Offset64 );
 			}
 		}
 	}
@@ -428,10 +429,10 @@ bool FVirtualTextureBuiltData::ValidateData(FStringView const& InDDCDebugContext
 		FBulkDataLockedScope BulkDataLockedScope;
 
 		const void* ChunkData = nullptr;
-		uint32 ChunkDataSize = 0u;
+		int64 ChunkDataSize = 0;
 		if (Chunk.BulkData.GetBulkDataSize() > 0)
 		{
-			ChunkDataSize = IntCastChecked<uint32>(Chunk.BulkData.GetBulkDataSize());
+			ChunkDataSize = Chunk.BulkData.GetBulkDataSize();
 			ChunkData = Chunk.BulkData.LockReadOnly();
 			BulkDataLockedScope.BulkData = &Chunk.BulkData;
 		}
@@ -455,13 +456,13 @@ bool FVirtualTextureBuiltData::ValidateData(FStringView const& InDDCDebugContext
 			}
 
 			ChunkData = ChunkDataDDC.GetData();
-			ChunkDataSize = IntCastChecked<uint32>(ChunkDataDDC.GetSize());
+			ChunkDataSize = ChunkDataDDC.GetSize();
 		}
 #endif // WITH_EDITORONLY_DATA
 
 		if (!ChunkData || ChunkDataSize < sizeof(FVirtualTextureChunkHeader))
 		{
-			UE_LOG(LogTexture, Error, TEXT("Virtual Texture %s has invalid size %u for chunk %d"), *TextureName, ChunkDataSize, ChunkIndex);
+			UE_LOG(LogTexture, Error, TEXT("Virtual Texture %s has invalid size %lld for chunk %d"), *TextureName, ChunkDataSize, ChunkIndex);
 			bResult = false;
 			break;
 		}
diff --git a/Engine/Source/Runtime/Engine/Private/VT/VirtualTextureDataBuilder.cpp b/Engine/Source/Runtime/Engine/Private/VT/VirtualTextureDataBuilder.cpp
--- a/Engine/Source/Runtime/Engine/Private/VT/VirtualTextureDataBuilder.cpp
+++ b/Engine/Source/Runtime/Engine/Private/VT/VirtualTextureDataBuilder.cpp
@@ -286,8 +286,10 @@ bool FVirtualTextureBuilderDerivedInfo::InitializeFromBuildSettings(const FTextu
 	}
 
 	// We require VT blocks (UDIM pages) to be PoT, but multi block textures may have full logical dimension that's not PoT
-	check(FMath::IsPowerOfTwo(BlockSizeX));
-	check(FMath::IsPowerOfTwo(BlockSizeY));
+	if ( ! FMath::IsPowerOfTwo(BlockSizeX) || ! FMath::IsPowerOfTwo(BlockSizeY) )
+	{
+		return false;
+	}
 
 	// Ensure block size is at least 1 tile, while preserving aspect ratio
 	BlockSizeScale = 1;
@@ -312,10 +314,11 @@ bool FVirtualTextureBuilderDerivedInfo::InitializeFromBuildSettings(const FTextu
 	// the fact that the Min(x, VIRTUALTEXTURE_LOG2_MAX_PAGETABLE_SIZE) is load bearing. In order to get an incorrect
 	// mip count, you need Size to be non pow2 and the result to be larger than VIRTUALTEXTURE_LOG2_MAX_PAGETABLE_SIZE.
 	NumMips = FMath::Min<uint32>(FMath::CeilLogTwo(Size) + 1, VIRTUALTEXTURE_LOG2_MAX_PAGETABLE_SIZE);
+
 	return true;
 }
 
-void FVirtualTextureDataBuilder::Build(FTextureSourceData& InSourceData, FTextureSourceData& InCompositeSourceData, const FTextureBuildSettings* InSettingsPerLayer, bool bAllowAsync)
+bool FVirtualTextureDataBuilder::Build(FTextureSourceData& InSourceData, FTextureSourceData& InCompositeSourceData, const FTextureBuildSettings* InSettingsPerLayer, bool bAllowAsync)
 {
 	const int32 NumLayers = InSourceData.Layers.Num();
 	checkf(NumLayers <= (int32)VIRTUALTEXTURE_DATA_MAXLAYERS, TEXT("The maximum amount of layers is exceeded."));
@@ -363,12 +366,14 @@ void FVirtualTextureDataBuilder::Build(FTextureSourceData& InSourceData, FTextur
 	bAllowAsync = bAllowAsync && CVarVTParallelTileCompression.GetValueOnAnyThread();
 
 	// encode SourceBlocks to VT Tiles and pack into chunks :
-	BuildPagesMacroBlocks(bAllowAsync);
+	bool ok = BuildPagesMacroBlocks(bAllowAsync);
 	// free SourceBlocks :
 	FreeSourcePixels();
+
+	return ok;
 }
 
-void FVirtualTextureDataBuilder::BuildPagesForChunk(const TArray<FVTSourceTileEntry>& ActiveTileList, bool bAllowAsync)
+bool FVirtualTextureDataBuilder::BuildPagesForChunk(const TArray<FVTSourceTileEntry>& ActiveTileList, bool bAllowAsync)
 {
 	TArray<FLayerData> LayerData;
 	LayerData.AddDefaulted(SourceLayers.Num());
@@ -392,21 +397,21 @@ void FVirtualTextureDataBuilder::BuildPagesForChunk(const TArray<FVTSourceTileEn
 		}
 		if (bIsRawGPUData)
 		{
-			uint32 TileDataOffset = 0;
+			int64 TileDataOffset = 0;
 			OutData.TileDataOffsetPerLayer.Reserve(LayerData.Num());
 			for (int32 LayerIndex = 0; LayerIndex < LayerData.Num(); LayerIndex++)
 			{
 				TileDataOffset += LayerData[LayerIndex].TilePayload[0].Num();
-				OutData.TileDataOffsetPerLayer.Add(TileDataOffset);
+				OutData.TileDataOffsetPerLayer.Add( IntCastChecked<uint32>(TileDataOffset) );
 			}
 		}
 	}
 
 	// Write tiles out to chunk.
-	PushDataToChunk(ActiveTileList, LayerData);
+	return PushDataToChunk(ActiveTileList, LayerData);
 }
 
-void FVirtualTextureDataBuilder::BuildPagesMacroBlocks(bool bAllowAsync)
+bool FVirtualTextureDataBuilder::BuildPagesMacroBlocks(bool bAllowAsync)
 {
 	static const uint32 MinSizePerChunk = 1024u; // Each chunk will contain a mip level of at least this size (MinSizePerChunk x MinSizePerChunk)
 	const uint32 NumLayers = SourceLayers.Num();
@@ -496,7 +501,8 @@ void FVirtualTextureDataBuilder::BuildPagesMacroBlocks(bool bAllowAsync)
 			if (!bInFinalChunk && TilesInChunk.Num() >= (int32)MinTilesPerChunk)
 			{
 				OutData.TileIndexPerChunk.Add(TileIndex);
-				BuildPagesForChunk(TilesInChunk, bAllowAsync);
+				if ( ! BuildPagesForChunk(TilesInChunk, bAllowAsync) )
+					return false;
 				TilesInChunk.Reset();
 			}
 			else
@@ -514,7 +520,8 @@ void FVirtualTextureDataBuilder::BuildPagesMacroBlocks(bool bAllowAsync)
 
 		if (TilesInChunk.Num() > 0)
 		{
-			BuildPagesForChunk(TilesInChunk, bAllowAsync);
+			if ( ! BuildPagesForChunk(TilesInChunk, bAllowAsync) )
+				return false;
 		}
 
 		check(OutData.BaseOffsetPerMip.Num() == OutData.NumMips);
@@ -560,6 +567,8 @@ void FVirtualTextureDataBuilder::BuildPagesMacroBlocks(bool bAllowAsync)
 		OutData.TileIndexPerMip.Empty();
 		OutData.TileOffsetInChunk.Empty();
 	}
+
+	return true;
 }
 
 static const TCHAR* GetSafePixelFormatName(EPixelFormat Format)
@@ -730,7 +739,7 @@ void FVirtualTextureDataBuilder::BuildTiles(const TArray<FVTSourceTileEntry>& Ti
 	}
 }
 
-void FVirtualTextureDataBuilder::PushDataToChunk(const TArray<FVTSourceTileEntry>& Tiles, const TArray<FLayerData>& LayerData)
+bool FVirtualTextureDataBuilder::PushDataToChunk(const TArray<FVTSourceTileEntry>& Tiles, const TArray<FLayerData>& LayerData)
 {
 	const int32 NumLayers = SourceLayers.Num();
 
@@ -743,16 +752,20 @@ void FVirtualTextureDataBuilder::PushDataToChunk(const TArray<FVTSourceTileEntry
 			TotalSize += TilePayload.Num();
 		}
 	}
-
+	
 	// Built VT data structures use uint32 :
-	check(TotalSize <= MAX_uint32 );
+	if ( TotalSize >= MAX_uint32 )
+	{
+		UE_LOG(LogVirtualTexturing,Error,TEXT("Cannot build VT; data bigger than 4 GB : %lld"),TotalSize);
+		return false;
+	}
 
 	FVirtualTextureDataChunk& Chunk = OutData.Chunks.AddDefaulted_GetRef();
 	Chunk.SizeInBytes = TotalSize;
 	FByteBulkData& BulkData = Chunk.BulkData;
 	BulkData.Lock(LOCK_READ_WRITE);
 	uint8* NewChunkData = (uint8*)BulkData.Realloc(TotalSize);
-	int64 ChunkOffset = 0u;
+	int64 ChunkOffset = 0;
 
 	// Header for the chunk
 	FVirtualTextureChunkHeader* Header = (FVirtualTextureChunkHeader*)NewChunkData;
@@ -763,7 +776,7 @@ void FVirtualTextureDataBuilder::PushDataToChunk(const TArray<FVTSourceTileEntry
 	// codec payloads
 	for (int32 Layer = 0; Layer < NumLayers; ++Layer)
 	{
-		Chunk.CodecPayloadOffset[Layer] = ChunkOffset;
+		Chunk.CodecPayloadOffset[Layer] = IntCastChecked<uint32>( ChunkOffset );
 		Chunk.CodecType[Layer] = LayerData[Layer].Codec;
 		if (LayerData[Layer].CodecPayload.Num() > 0)
 		{
@@ -780,13 +793,13 @@ void FVirtualTextureDataBuilder::PushDataToChunk(const TArray<FVTSourceTileEntry
 		// Set BaseOffsetPerMip from the first tile we find for the MipIndex.
 		if (OutData.BaseOffsetPerMip[MipIndex] == ~0u)
 		{
-			OutData.BaseOffsetPerMip[MipIndex] = ChunkOffset;
+			OutData.BaseOffsetPerMip[MipIndex] = IntCastChecked<uint32>( ChunkOffset );
 		}
 		int32 TileIndex = Tile.TileIndex;
 		for (int32 Layer = 0; Layer < NumLayers; ++Layer)
 		{
 			check(OutData.TileOffsetInChunk[TileIndex] == ~0u);
-			OutData.TileOffsetInChunk[TileIndex] = ChunkOffset;
+			OutData.TileOffsetInChunk[TileIndex] = IntCastChecked<uint32>( ChunkOffset );
 			++TileIndex;
 
 			const TArray<uint8>& TilePayload = LayerData[Layer].TilePayload[TileIdx];
@@ -804,6 +817,8 @@ void FVirtualTextureDataBuilder::PushDataToChunk(const TArray<FVTSourceTileEntry
 
 	BulkData.Unlock();
 	BulkData.SetBulkDataFlags(BULKDATA_Force_NOT_InlinePayload);
+
+	return true;
 }
 
 int32 FVirtualTextureDataBuilder::FindSourceBlockIndex(int32 MipIndex, int32 BlockX, int32 BlockY) const
diff --git a/Engine/Source/Runtime/Engine/Private/VT/VirtualTextureDataBuilder.h b/Engine/Source/Runtime/Engine/Private/VT/VirtualTextureDataBuilder.h
--- a/Engine/Source/Runtime/Engine/Private/VT/VirtualTextureDataBuilder.h
+++ b/Engine/Source/Runtime/Engine/Private/VT/VirtualTextureDataBuilder.h
@@ -11,6 +11,7 @@ class ITextureCompressorModule;
 class IImageWrapperModule;
 struct FTextureSourceData;
 struct FTextureSourceBlockData;
+struct FTextureBuildSettings;
 
 struct FVTSourceTileEntry
 {
@@ -80,15 +81,15 @@ public:
 	~FVirtualTextureDataBuilder();
 
 	// note: InSourceData is freed by this function
-	void Build(FTextureSourceData& InSourceData, FTextureSourceData& InCompositeSourceData, const FTextureBuildSettings* InSettingsPerLayer, bool bAllowAsync);
+	bool Build(FTextureSourceData& InSourceData, FTextureSourceData& InCompositeSourceData, const FTextureBuildSettings* InSettingsPerLayer, bool bAllowAsync);
 
 private:
 	friend struct FAsyncMacroBlockTask;
 
-	void BuildPagesMacroBlocks(bool bAllowAsync);
-	void BuildPagesForChunk(const TArray<FVTSourceTileEntry>& ActiveTileList, bool bAllowAsync);
+	bool BuildPagesMacroBlocks(bool bAllowAsync);
+	bool BuildPagesForChunk(const TArray<FVTSourceTileEntry>& ActiveTileList, bool bAllowAsync);
 	void BuildTiles(const TArray<FVTSourceTileEntry>& TileList, uint32 layer, FLayerData& GeneratedData, bool bAllowAsync);
-	void PushDataToChunk(const TArray<FVTSourceTileEntry> &Tiles, const TArray<FLayerData>& LayerData);
+	bool PushDataToChunk(const TArray<FVTSourceTileEntry> &Tiles, const TArray<FLayerData>& LayerData);
 
 	int32 FindSourceBlockIndex(int32 MipIndex, int32 BlockX, int32 BlockY) const;
 
