- Add support for TSF_G8 to cut out textures
- Make message more descriptive if the source format is not supported

diff --git a/Engine/Source/Runtime/Engine/Private/Particles/SubUVAnimation.cpp b/Engine/Source/Runtime/Engine/Private/Particles/SubUVAnimation.cpp
--- a/Engine/Source/Runtime/Engine/Private/Particles/SubUVAnimation.cpp
+++ b/Engine/Source/Runtime/Engine/Private/Particles/SubUVAnimation.cpp
@@ -32,12 +32,11 @@ FString FSubUVDerivedData::GetDDCKeyString(const FGuid& StateId, int32 SizeX, in
 	}
 	// adding v2 to the key after fixing color channel offsets
 	// adding v3 to the key after allowing other formats
-	KeyString += TEXT("_V3");
+	// adding v4 to the key after adding G8 support
+	KeyString += TEXT("_V4");
 	return FDerivedDataCacheInterface::BuildCacheKey(TEXT("SUBUV_"), SUBUV_DERIVEDDATA_VER, *KeyString);
 }
 
-
-
 void FSubUVDerivedData::Serialize(FStructuredArchive::FSlot Slot)
 {
 	Slot << BoundingGeometry;
@@ -470,31 +469,31 @@ FIntPoint GNeighbors[] =
 
 bool IsSupportedFormat(ETextureSourceFormat SrcFormat)
 {
-	return SrcFormat == TSF_BGRA8 || SrcFormat == TSF_RGBA16 || SrcFormat == TSF_RGBA16F;
+	return SrcFormat == TSF_G8 || SrcFormat == TSF_BGRA8 || SrcFormat == TSF_RGBA16 || SrcFormat == TSF_RGBA16F;
 }
 
 uint32 GetByteSizePerPixel(ETextureSourceFormat SrcFormat)
 {
-	if (SrcFormat == TSF_BGRA8)
-	{
-		return 4;
-	}
-	else if (SrcFormat == TSF_RGBA16)
+	switch (SrcFormat)
 	{
-		return 8;
+		case TSF_G8:		return 1;
+		case TSF_BGRA8:		return 4;
+		case TSF_RGBA16:	return 8;
+		case TSF_RGBA16F:	return sizeof(FFloat16) * 4;
+		default:			return 0;
 	}
-	else if (SrcFormat == TSF_RGBA16F)
-	{
-		return sizeof(FFloat16) * 4;
-	}
-	return 0;
 }
 
 bool ComputeOpacityValue(const uint8* BGRA, int32 x, int32 y, int32 TextureSizeX,  EOpacitySourceMode OpacitySourceMode, ETextureSourceFormat SrcFormat, float AlphaThresholdF, uint8 AlphaThreshold8, uint16 AlphaThreshold16)
 {
-	int32 Offset = (y * TextureSizeX + x) * GetByteSizePerPixel(SrcFormat);
+	const int32 Offset = (y * TextureSizeX + x) * GetByteSizePerPixel(SrcFormat);
 
-	if (SrcFormat == TSF_BGRA8)
+	if (SrcFormat == TSF_G8)
+	{
+		const uint32 Opacity = BGRA[Offset];
+		return Opacity > AlphaThreshold8;
+	}
+	else if (SrcFormat == TSF_BGRA8)
 	{
 		const uint8* BGRA8 = (const uint8*)(BGRA + Offset);
 		uint32 Opacity = 255;
@@ -601,11 +600,29 @@ void FSubUVDerivedData::GetFeedback(UTexture2D* SubUVTexture, int32 SubImages_Ho
 #if WITH_EDITORONLY_DATA
 	if (SubUVTexture)
 	{
-		ETextureSourceFormat SourceFormat = SubUVTexture->Source.GetFormat();
-
+		const ETextureSourceFormat SourceFormat = SubUVTexture->Source.GetFormat();
 		if (!IsSupportedFormat(SourceFormat))
 		{
-			OutErrors.Add(LOCTEXT("CutoutNotSupported", "Image is not in a supported format for cutouts (BGRA8, RGBA16, RGBA16F). It will be ignored."));
+			FString SupportedFormatsString;
+			UEnum* SourceFormatEnum = StaticEnum<ETextureSourceFormat>();
+			for (int32 Format=0; Format < TSF_MAX; ++Format)
+			{
+				if (IsSupportedFormat(ETextureSourceFormat(Format)))
+				{
+					if (SupportedFormatsString.IsEmpty() == false)
+					{
+						SupportedFormatsString.Append(TEXT(", "));
+					}
+					SupportedFormatsString.Append(SourceFormatEnum->GetNameStringByValue(Format));
+				}
+			}
+
+			OutErrors.Add(
+				FText::Format(LOCTEXT("CutoutNotSupported", "Image source format ({0}) is not in a supported format for cutouts ({1}). It will be ignored."),
+					FText::FromString(SourceFormatEnum->GetNameStringByValue(SourceFormat)),
+					FText::FromString(SupportedFormatsString)
+				)
+			);
 		} 
 		if (SubUVTexture->Source.GetNumMips() == 0)
 		{
