TextureSource : make GetSize return int64 , similar to previous change on FImage

diff --git a/Engine/Source/Runtime/Engine/Classes/Engine/Texture.h b/Engine/Source/Runtime/Engine/Classes/Engine/Texture.h
--- a/Engine/Source/Runtime/Engine/Classes/Engine/Texture.h
+++ b/Engine/Source/Runtime/Engine/Classes/Engine/Texture.h
@@ -336,8 +336,8 @@ struct FTextureSource
 	/** Trivial accessors. These will only give values for Block0 so may not be correct for UDIM/multi-block textures, use GetBlock() for this case. */
 	FGuid GetPersistentId() const { return BulkData.GetIdentifier(); }
 	ENGINE_API FGuid GetId() const;
-	FORCEINLINE int32 GetSizeX() const { return SizeX; }
-	FORCEINLINE int32 GetSizeY() const { return SizeY; }
+	FORCEINLINE int64 GetSizeX() const { return SizeX; }
+	FORCEINLINE int64 GetSizeY() const { return SizeY; }
 	FORCEINLINE int32 GetNumSlices() const { return NumSlices; }
 	FORCEINLINE int32 GetNumMips() const { return NumMips; }
 	FORCEINLINE int32 GetNumLayers() const { return NumLayers; }
