Better error handling for failing to overwrite an existing replay file
Add optional ?flush arg to replay urls

diff --git a/Engine/Source/Runtime/Engine/Private/ReplaySubsystem.cpp b/Engine/Source/Runtime/Engine/Private/ReplaySubsystem.cpp
--- a/Engine/Source/Runtime/Engine/Private/ReplaySubsystem.cpp
+++ b/Engine/Source/Runtime/Engine/Private/ReplaySubsystem.cpp
@@ -140,12 +140,14 @@ void UReplaySubsystem::RecordReplay(const FString& Name, const FString& Friendly
 		DemoURL.AddOption(*Option);
 	}
 
+	const UE::ReplaySubsystem::EStopReplayFlags StopExistingFlags = DemoURL.HasOption(TEXT("flush")) ? UE::ReplaySubsystem::EStopReplayFlags::Flush : UE::ReplaySubsystem::EStopReplayFlags::None;
+
 	UNetDriver* NetDriver = CurrentWorld->GetNetDriver();
 
 	// must be server and using a replication graph to use a replay connection
 	if (NetDriver && NetDriver->IsServer() && NetDriver->GetReplicationDriver() && ReplaySubsystem::CVarUseReplayConnection.GetValueOnAnyThread())
 	{
-		StopExistingReplays(CurrentWorld);
+		StopExistingReplays(CurrentWorld, StopExistingFlags);
 
 		UReplayNetConnection* Connection = NewObject<UReplayNetConnection>();
 
@@ -169,7 +171,7 @@ void UReplaySubsystem::RecordReplay(const FString& Name, const FString& Friendly
 
 	if (!DemoNetDriver || !DemoNetDriver->IsRecordingMapChanges() || !DemoNetDriver->IsRecordingPaused())
 	{
-		StopExistingReplays(CurrentWorld);
+		StopExistingReplays(CurrentWorld, StopExistingFlags);
 
 		bDestroyedDemoNetDriver = true;
 
@@ -231,11 +233,7 @@ bool UReplaySubsystem::PlayReplay(const FString& Name, UWorld* WorldOverride, co
 		return false;
 	}
 
-	StopExistingReplays(CurrentWorld);
-
 	FURL DemoURL;
-	UE_LOG(LogDemo, Log, TEXT("PlayReplay: Attempting to play demo %s"), *Name);
-
 	DemoURL.Map = Name;
 
 	for (const FString& Option : AdditionalOptions)
@@ -243,6 +241,10 @@ bool UReplaySubsystem::PlayReplay(const FString& Name, UWorld* WorldOverride, co
 		DemoURL.AddOption(*Option);
 	}
 
+	StopExistingReplays(CurrentWorld, DemoURL.HasOption(TEXT("flush")) ? UE::ReplaySubsystem::EStopReplayFlags::Flush : UE::ReplaySubsystem::EStopReplayFlags::None);
+
+	UE_LOG(LogDemo, Log, TEXT("PlayReplay: Attempting to play demo %s"), *Name);
+
 	if (!GEngine->CreateNamedNetDriver(CurrentWorld, NAME_DemoNetDriver, NAME_DemoNetDriver))
 	{
 		UE_LOG(LogDemo, Warning, TEXT("PlayReplay: failed to create demo net driver!"));
@@ -295,7 +297,7 @@ void UReplaySubsystem::StopReplay()
 	}
 }
 
-void UReplaySubsystem::StopExistingReplays(UWorld* InWorld)
+void UReplaySubsystem::StopExistingReplays(UWorld* InWorld, UE::ReplaySubsystem::EStopReplayFlags Flags)
 {
 	UWorld* CurrentWorld = InWorld ? InWorld : GetWorld();
 
@@ -309,6 +311,12 @@ void UReplaySubsystem::StopExistingReplays(UWorld* InWorld)
 		Connection->CleanUp();
 		ReplayConnection = nullptr;
 	}
+
+	if (EnumHasAnyFlags(Flags, UE::ReplaySubsystem::EStopReplayFlags::Flush))
+	{
+		//@todo: narrow to specific streamer that was stopped
+		FNetworkReplayStreaming::Get().Flush();
+	}
 }
 
 FString UReplaySubsystem::GetActiveReplayName() const
diff --git a/Engine/Source/Runtime/Engine/Public/ReplaySubsystem.h b/Engine/Source/Runtime/Engine/Public/ReplaySubsystem.h
--- a/Engine/Source/Runtime/Engine/Public/ReplaySubsystem.h
+++ b/Engine/Source/Runtime/Engine/Public/ReplaySubsystem.h
@@ -8,6 +8,17 @@
 
 class UReplayNetConnection;
 
+namespace UE::ReplaySubsystem
+{
+	enum class EStopReplayFlags : uint32
+	{
+		None = 0x0,
+		Flush = 0x1,
+	};
+
+	ENUM_CLASS_FLAGS(EStopReplayFlags);
+};
+
 UCLASS(DisplayName = "Replay Subsystem")
 class ENGINE_API UReplaySubsystem : public UGameInstanceSubsystem
 {
@@ -122,7 +133,7 @@ public:
 	bool bLoadDefaultMapOnStop = true;
 
 private:
-	void StopExistingReplays(UWorld* InWorld);
+	void StopExistingReplays(UWorld* InWorld, UE::ReplaySubsystem::EStopReplayFlags Flags = UE::ReplaySubsystem::EStopReplayFlags::None);
 
 	void OnSeamlessTravelStart(UWorld* CurrentWorld, const FString& LevelName);
 	void OnSeamlessLevelTransition(UWorld* CurrentWorld);
diff --git a/Engine/Source/Runtime/NetworkReplayStreaming/LocalFileNetworkReplayStreaming/Private/LocalFileNetworkReplayStreaming.cpp b/Engine/Source/Runtime/NetworkReplayStreaming/LocalFileNetworkReplayStreaming/Private/LocalFileNetworkReplayStreaming.cpp
--- a/Engine/Source/Runtime/NetworkReplayStreaming/LocalFileNetworkReplayStreaming/Private/LocalFileNetworkReplayStreaming.cpp
+++ b/Engine/Source/Runtime/NetworkReplayStreaming/LocalFileNetworkReplayStreaming/Private/LocalFileNetworkReplayStreaming.cpp
@@ -947,7 +947,15 @@ void FLocalFileNetworkReplayStreamer::StartStreaming(const FStartStreamingParame
 
 				TaskReplayInfo.bIsValid = WriteReplayInfo(CurrentStreamName, TaskReplayInfo);
 
-				RequestData.DelegateResult.Result = EStreamingOperationResult::Success;
+				if (!TaskReplayInfo.bIsValid)
+				{
+					UE_LOG(LogLocalFileReplay, Warning, TEXT("StartStreaming was unable to write to the replay file: %s"), *FullDemoFilename);
+					RequestData.AsyncError = ELocalFileReplayResult::FileWriter;
+				}
+				else
+				{
+					RequestData.DelegateResult.Result = EStreamingOperationResult::Success;
+				}
 			},
 			[this, Delegate](TLocalFileRequestCommonData<FStartStreamingResult>& RequestData)
 			{
@@ -3688,11 +3696,11 @@ void FLocalFileNetworkReplayStreamer::UpdateCurrentReplayInfo(FLocalFileReplayIn
 		
 		if (!EnumHasAnyFlags(UpdateFlags, EUpdateReplayInfoFlags::FullUpdate))
 		{
-		CurrentReplayInfo.LengthInMS = TotalLengthInMS;
-		CurrentReplayInfo.EncryptionKey = MoveTemp(CurrentKey);
+			CurrentReplayInfo.LengthInMS = TotalLengthInMS;
+			CurrentReplayInfo.EncryptionKey = MoveTemp(CurrentKey);
+		}
 	}
 }
-}
 
 int32 FLocalFileNetworkReplayStreamer::GetDecompressedSizeBackCompat(FArchive& InCompressed) const
 {
