Update DDC GUID

diff --git a/Engine/Source/Runtime/Engine/Private/AudioDerivedData.cpp b/Engine/Source/Runtime/Engine/Private/AudioDerivedData.cpp
--- a/Engine/Source/Runtime/Engine/Private/AudioDerivedData.cpp
+++ b/Engine/Source/Runtime/Engine/Private/AudioDerivedData.cpp
@@ -165,7 +165,7 @@ Derived data key generation.
 
 // If you want to bump this version, generate a new guid using
 // VS->Tools->Create GUID and paste it here. https://www.guidgen.com works too.
-#define STREAMEDAUDIO_DERIVEDDATA_VER		TEXT("E9A61FD4C181422A8EAAB060B1690F6A")
+#define STREAMEDAUDIO_DERIVEDDATA_VER		TEXT("acc1e60e03be438ab23c-a3d318425d34")
 
 #ifndef CASE_ENUM_TO_TEXT
 #define CASE_ENUM_TO_TEXT(X) case X: return TEXT(#X);
