[UE] Editor Physics constraint Debug Draw and in viewport editing changes
- Fix instances where a bone index is retrieved from the reference skeleton and used it to index into the editor skeleton component (Editor only)
- Add a new prioritized input chord feature to the EdtiorViewportClient that acts to resolve conflict between overlapping input key combinations. When higher priority key chords evaluate true, the required key presses are marked as consumed and are ignored by lower priority chords.
- Use the new prioritized input chord feature to prevent the keyboard modifier for editing the child frame - Alt + Shift - from clashing with keyboard modifier for the Lock camera to widget feature - Shift - which was causing the translation widget to move erratically when editing the child frame.
- Render the manipulator widget in the child frame when editing the child transform
- Draw the constraint limit cones in the parent frame
- Draw the transform of the child reference frame
- Draw an arrow indicating the constraint�s current twist

diff --git a/Engine/Source/Editor/PhysicsAssetEditor/Private/PhysicsAssetEditorEditMode.cpp b/Engine/Source/Editor/PhysicsAssetEditor/Private/PhysicsAssetEditorEditMode.cpp
--- a/Engine/Source/Editor/PhysicsAssetEditor/Private/PhysicsAssetEditorEditMode.cpp
+++ b/Engine/Source/Editor/PhysicsAssetEditor/Private/PhysicsAssetEditorEditMode.cpp
@@ -27,6 +27,8 @@
 
 FName FPhysicsAssetEditorEditMode::ModeName("PhysicsAssetEditor");
 
+static const FName InputChordName_EditConstraintChildTransform = FName("InputChordName_EditConstraintChildTransform");
+
 FPhysicsAssetEditorEditMode::FPhysicsAssetEditorEditMode()
 	: MinPrimSize(0.5f)
 	, PhysicsAssetEditor_TranslateSpeed(0.25f)
@@ -422,7 +424,7 @@ bool FPhysicsAssetEditorEditMode::InputDelta(FEditorViewportClient* InViewportCl
 			if (SharedData->bManipulating)
 			{
 				float BoneScale = 1.f;
-				SelectedObject.WidgetTM = SharedData->GetConstraintMatrix(SelectedObject.Index, EConstraintFrame::Frame2, 1.f);
+				SelectedObject.WidgetTM = SharedData->GetConstraintMatrix(SelectedObject.Index, GetConstraintFrameForWidget(), 1.f);
 
 				if (InViewportClient->GetWidgetMode() == UE::Widget::WM_Translate)
 				{
@@ -519,11 +521,11 @@ void FPhysicsAssetEditorEditMode::Tick(FEditorViewportClient* ViewportClient, fl
 	{
 		RenderSettings->ConstraintViewportManipulationFlags = EConstraintTransformComponentFlags::All;
 
-		if (ViewportClient->IsShiftPressed() && ViewportClient->IsAltPressed()) // Shft + Alt + Move will rotate or move the child transform only.
+		if (ViewportClient->IsPrioritizedInputChordPressed(InputChordName_EditConstraintChildTransform)) // Rotate or move the child transform only.
 		{
 			EnumRemoveFlags(RenderSettings->ConstraintViewportManipulationFlags, EConstraintTransformComponentFlags::AllParent); // Remove Parent Frame flags.
 		}
-		else if (ViewportClient->IsAltPressed()) // Alt + Move will rotate or move the parent transform only.
+		else if (ViewportClient->IsAltPressed()) // Rotate or move the parent transform only.
 		{
 			EnumRemoveFlags(RenderSettings->ConstraintViewportManipulationFlags, EConstraintTransformComponentFlags::AllChild); // Remove Child Frame flags.
 		}	
@@ -654,6 +656,21 @@ bool FPhysicsAssetEditorEditMode::UsesTransformWidget(UE::Widget::EWidgetMode Ch
 	return ShouldDrawWidget() && (CheckMode == UE::Widget::WM_Scale || CheckMode == UE::Widget::WM_Translate || CheckMode == UE::Widget::WM_Rotate || CheckMode == UE::Widget::WM_None);
 }
 
+EConstraintFrame::Type FPhysicsAssetEditorEditMode::GetConstraintFrameForWidget() const
+{
+	FPhysicsAssetRenderSettings* const RenderSettings = SharedData->GetRenderSettings();
+
+	// Draw widget in the constraint's parent relative frame by default and in the child frame if the user is exclusively editing the child frame.
+	EConstraintFrame::Type ConstraintFrame = EConstraintFrame::Frame2;
+
+	if (RenderSettings && !EnumHasAnyFlags(RenderSettings->ConstraintViewportManipulationFlags, EConstraintTransformComponentFlags::AllParent))
+	{
+		ConstraintFrame = EConstraintFrame::Frame1;
+	}
+
+	return ConstraintFrame;
+}
+
 bool FPhysicsAssetEditorEditMode::GetCustomDrawingCoordinateSystem(FMatrix& InMatrix, void* InData)
 {
 	// Don't draw widget if nothing selected.
@@ -671,7 +688,7 @@ bool FPhysicsAssetEditorEditMode::GetCustomDrawingCoordinateSystem(FMatrix& InMa
 	}
 	else if (SharedData->GetSelectedConstraint())
 	{
-		InMatrix = SharedData->GetConstraintMatrix(SharedData->GetSelectedConstraint()->Index, EConstraintFrame::Frame2, 1.f).ToMatrixNoScale().RemoveTranslation();
+		InMatrix = SharedData->GetConstraintMatrix(SharedData->GetSelectedConstraint()->Index, GetConstraintFrameForWidget(), 1.f).ToMatrixNoScale().RemoveTranslation();
 		return true;
 	}
 
@@ -700,7 +717,7 @@ FVector FPhysicsAssetEditorEditMode::GetWidgetLocation() const
 	}
 	else if (SharedData->GetSelectedConstraint())
 	{
-		return SharedData->GetConstraintMatrix(SharedData->GetSelectedConstraint()->Index, EConstraintFrame::Frame2, 1.f).GetTranslation();
+		return SharedData->GetConstraintMatrix(SharedData->GetSelectedConstraint()->Index, GetConstraintFrameForWidget(), 1.f).GetTranslation();
 	}
 
 	return FVector::ZeroVector;
@@ -793,6 +810,20 @@ bool FPhysicsAssetEditorEditMode::HandleClick(FEditorViewportClient* InViewportC
 	return false;
 }
 
+bool FPhysicsAssetEditorEditMode::ReceivedFocus(FEditorViewportClient* ViewportClient, FViewport* Viewport)
+{
+	ViewportClient->RegisterPrioritizedInputChord(FPrioritizedInputChord(1, InputChordName_EditConstraintChildTransform, EModifierKey::Shift | EModifierKey::Alt));
+
+	return IPersonaEditMode::ReceivedFocus(ViewportClient, Viewport);
+}
+
+bool FPhysicsAssetEditorEditMode::LostFocus(FEditorViewportClient* ViewportClient, FViewport* Viewport)
+{
+	ViewportClient->UnregisterPrioritizedInputChord(InputChordName_EditConstraintChildTransform);
+
+	return IPersonaEditMode::LostFocus(ViewportClient, Viewport);
+}
+
 /** Helper function to open a viewport context menu */
 static void OpenContextMenu(const TSharedRef<FPhysicsAssetEditor>& PhysicsAssetEditor, FEditorViewportClient* InViewportClient, TFunctionRef<void(FMenuBuilder&)> InBuildMenu)
 {
diff --git a/Engine/Source/Editor/PhysicsAssetEditor/Private/PhysicsAssetEditorEditMode.h b/Engine/Source/Editor/PhysicsAssetEditor/Private/PhysicsAssetEditorEditMode.h
--- a/Engine/Source/Editor/PhysicsAssetEditor/Private/PhysicsAssetEditorEditMode.h
+++ b/Engine/Source/Editor/PhysicsAssetEditor/Private/PhysicsAssetEditorEditMode.h
@@ -50,6 +50,8 @@ public:
 	virtual bool GetCustomInputCoordinateSystem(FMatrix& InMatrix, void* InData) override;
 	virtual bool HandleClick(FEditorViewportClient* InViewportClient, HHitProxy *HitProxy, const FViewportClick &Click) override;
 	virtual bool IsCompatibleWith(FEditorModeID OtherModeID) const override { return true; }
+	virtual bool ReceivedFocus(FEditorViewportClient* ViewportClient, FViewport* Viewport);
+	virtual bool LostFocus(FEditorViewportClient* ViewportClient, FViewport* Viewport);
 
 private:
 	/** Simulation mouse forces */
@@ -76,6 +78,9 @@ private:
 
 	void OpenSelectionMenu(FEditorViewportClient* InViewportClient);
 
+	/** Returns the identifier for the constraint frame (child or parent) in which the manipulator widget should be drawn. */
+	EConstraintFrame::Type GetConstraintFrameForWidget() const;
+
 private:
 	/** Shared data */
 	FPhysicsAssetEditorSharedData* SharedData;
diff --git a/Engine/Source/Editor/PhysicsAssetEditor/Private/PhysicsAssetEditorSharedData.cpp b/Engine/Source/Editor/PhysicsAssetEditor/Private/PhysicsAssetEditorSharedData.cpp
--- a/Engine/Source/Editor/PhysicsAssetEditor/Private/PhysicsAssetEditorSharedData.cpp
+++ b/Engine/Source/Editor/PhysicsAssetEditor/Private/PhysicsAssetEditorSharedData.cpp
@@ -2313,10 +2313,10 @@ void FPhysicsAssetEditorSharedData::SetConstraintRelTM(const FPhysicsAssetEditor
 	ConstraintSetup->Modify();
 
 	// Get child bone transform
-	int32 BoneIndex = EditorSkelMesh->GetRefSkeleton().FindBoneIndex(ConstraintSetup->DefaultInstance.ConstraintBone1);
+	const int32 BoneIndex = EditorSkelComp->GetBoneIndex(ConstraintSetup->DefaultInstance.ConstraintBone1);
 	if (BoneIndex != INDEX_NONE)
 	{
-		FTransform BoneTM = EditorSkelComp->GetBoneTransform(BoneIndex);
+		const FTransform BoneTM = EditorSkelComp->GetBoneTransform(BoneIndex);
 		ConstraintSetup->DefaultInstance.SetRefFrame(EConstraintFrame::Frame1, WNewChildFrame.GetRelativeTransform(BoneTM));
 	}
 }
@@ -2638,49 +2638,28 @@ void FPhysicsAssetEditorSharedData::DeleteCurrentPrim()
 
 FTransform FPhysicsAssetEditorSharedData::GetConstraintBodyTM(const UPhysicsConstraintTemplate* ConstraintSetup, EConstraintFrame::Type Frame) const
 {
-	if (ConstraintSetup == NULL)
+	if ((ConstraintSetup != nullptr) && (EditorSkelComp == nullptr))
 	{
-		return FTransform::Identity;
-	}
-
-	USkeletalMesh* EditorSkelMesh = PhysicsAsset->GetPreviewMesh();
-	if(EditorSkelMesh == nullptr)
-	{
-		return FTransform::Identity;
-	}
-
-	int32 BoneIndex;
-	if (Frame == EConstraintFrame::Frame1)
-	{
-		BoneIndex = EditorSkelMesh->GetRefSkeleton().FindBoneIndex(ConstraintSetup->DefaultInstance.ConstraintBone1);
-	}
-	else
-	{
-		BoneIndex = EditorSkelMesh->GetRefSkeleton().FindBoneIndex(ConstraintSetup->DefaultInstance.ConstraintBone2);
-	}
+		const FName BoneName = (Frame == EConstraintFrame::Frame1) ? ConstraintSetup->DefaultInstance.ConstraintBone1 : ConstraintSetup->DefaultInstance.ConstraintBone2;
+		const int32 BoneIndex = EditorSkelComp->GetBoneIndex(BoneName);
 
-	// If we couldn't find the bone - fall back to identity.
-	if (BoneIndex == INDEX_NONE)
-	{
-		return FTransform::Identity;
+		if (BoneIndex != INDEX_NONE)
+		{
+			FTransform BoneTM = EditorSkelComp->GetBoneTransform(BoneIndex);
+			BoneTM.RemoveScaling();
+			return BoneTM;
+		}
 	}
-	else
-	{
-		FTransform BoneTM = EditorSkelComp->GetBoneTransform(BoneIndex);
-		BoneTM.RemoveScaling();
 
-		return BoneTM;
-	}
+	return FTransform::Identity; // If we couldn't find the bone - fall back to identity.
 }
 
 FTransform FPhysicsAssetEditorSharedData::GetConstraintWorldTM(const UPhysicsConstraintTemplate* const ConstraintSetup, const EConstraintFrame::Type Frame, const float Scale) const
 {
-	USkeletalMesh* EditorSkelMesh = PhysicsAsset->GetPreviewMesh();
-
-	if ((ConstraintSetup != nullptr) && (EditorSkelMesh != nullptr))
+	if ((ConstraintSetup != nullptr) && (EditorSkelComp != nullptr))
 	{
 		const FName BoneName = (Frame == EConstraintFrame::Frame1) ? ConstraintSetup->DefaultInstance.ConstraintBone1 : ConstraintSetup->DefaultInstance.ConstraintBone2;
-		const int32 BoneIndex = EditorSkelMesh->GetRefSkeleton().FindBoneIndex(BoneName);
+		const int32 BoneIndex = EditorSkelComp->GetBoneIndex(BoneName);
 
 		if (BoneIndex != INDEX_NONE)
 		{	
diff --git a/Engine/Source/Editor/UnrealEd/Private/EditorViewportClient.cpp b/Engine/Source/Editor/UnrealEd/Private/EditorViewportClient.cpp
--- a/Engine/Source/Editor/UnrealEd/Private/EditorViewportClient.cpp
+++ b/Engine/Source/Editor/UnrealEd/Private/EditorViewportClient.cpp
@@ -228,6 +228,8 @@ FViewportClick::FViewportClick(const FSceneView* View, FEditorViewportClient* Vi
 	AltDown = ViewportClient->IsAltPressed();
 }
 
+static const FName InputChordName_CameraLockedToWidget = FName("InputChordName_CameraLockedToWidget");
+
 FViewportClick::~FViewportClick()
 {
 }
@@ -515,6 +517,8 @@ FEditorViewportClient::FEditorViewportClient(FEditorModeTools* InModeTools, FPre
 	FCoreDelegates::StatDisabled.AddRaw(this, &FEditorViewportClient::HandleViewportStatDisabled);
 	FCoreDelegates::StatDisableAll.AddRaw(this, &FEditorViewportClient::HandleViewportStatDisableAll);
 
+	RegisterPrioritizedInputChord(FPrioritizedInputChord(10, InputChordName_CameraLockedToWidget, EModifierKey::Shift));
+
 	RequestUpdateDPIScale();
 
 	FSlateApplication::Get().OnWindowDPIScaleChanged().AddRaw(this, &FEditorViewportClient::HandleWindowDPIScaleChanged);
@@ -2549,7 +2553,7 @@ void FEditorViewportClient::MarkMouseMovedSinceClick()
 bool FEditorViewportClient::IsUsingAbsoluteTranslation(bool bAlsoCheckAbsoluteRotation) const
 {
 	bool bIsHotKeyAxisLocked = Viewport->KeyState(EKeys::LeftControl) || Viewport->KeyState(EKeys::RightControl);
-	bool bCameraLockedToWidget = !(Widget && Widget->GetCurrentAxis() & EAxisList::Screen) && (Viewport->KeyState(EKeys::LeftShift) || Viewport->KeyState(EKeys::RightShift));
+	bool bCameraLockedToWidget = !(Widget && Widget->GetCurrentAxis() & EAxisList::Screen) && IsPrioritizedInputChordPressed(InputChordName_CameraLockedToWidget);
 	// Screen-space movement must always use absolute translation
 	bool bScreenSpaceTransformation = Widget && (Widget->GetCurrentAxis() == EAxisList::Screen) && GetWidgetMode() != UE::Widget::WM_Rotate;
 	bool bAbsoluteMovementEnabled = GetDefault<ULevelEditorViewportSettings>()->bUseAbsoluteTranslation || bScreenSpaceTransformation;
@@ -3265,6 +3269,51 @@ bool FEditorViewportClient::IsCommandChordPressed(const TSharedPtr<FUICommandInf
 	return bIsChordPressed;
 }
 
+void FEditorViewportClient::RegisterPrioritizedInputChord(const FPrioritizedInputChord& InInputChord)
+{
+	const int32 PreceedingElementIndex = PrioritizedInputChords.FindLastByPredicate([InInputChord](const FPrioritizedInputChord& Element) { return Element.Priority < InInputChord.Priority; });
+	const int32 TargetElementIndex = (PreceedingElementIndex != INDEX_NONE) ? PreceedingElementIndex + 1 : 0;
+	PrioritizedInputChords.Insert(InInputChord, TargetElementIndex);
+}
+
+void FEditorViewportClient::UnregisterPrioritizedInputChord(const FName InInputChordName)
+{
+	PrioritizedInputChords.RemoveAll([InInputChordName](const FPrioritizedInputChord& Element) { return Element.Name == InInputChordName; });
+}
+
+bool FEditorViewportClient::IsPrioritizedInputChordPressed(const FName InInputChordName) const
+{	
+	EModifierKey::Type ConsumedModifiers = EModifierKey::None;
+	TArray<FKey> ConsumedKeys;
+
+	// Iterate over all chords preceding the argument to determine if any key presses should be consumed before the target chord is evaluated.
+	for (const FPrioritizedInputChord& PrioritizedInputChord : PrioritizedInputChords)
+	{
+		const FInputChord& Chord = PrioritizedInputChord.InputChord;
+
+		bool IsPressed = true;
+		IsPressed &= !Chord.Key.IsValid()	|| (!ConsumedKeys.Contains(Chord.Key) && Viewport->KeyState(Chord.Key));
+		IsPressed &= !Chord.NeedsControl()	|| (!(ConsumedModifiers & EModifierKey::Control) && IsCtrlPressed());
+		IsPressed &= !Chord.NeedsAlt()		|| (!(ConsumedModifiers & EModifierKey::Alt) && IsAltPressed());
+		IsPressed &= !Chord.NeedsShift()	|| (!(ConsumedModifiers & EModifierKey::Shift) && IsShiftPressed());
+		IsPressed &= !Chord.NeedsCommand()	|| (!(ConsumedModifiers & EModifierKey::Command) && IsCmdPressed());
+
+		if (IsPressed)
+		{
+			// Mark all of the keys required by this chord as used so that they cannot be considered by any other, lower priority chords.
+			if (Chord.Key.IsValid()) { ConsumedKeys.Add(Chord.Key); }
+			ConsumedModifiers |= EModifierKey::FromBools(Chord.NeedsControl(), Chord.NeedsAlt(), Chord.NeedsShift(), Chord.NeedsCommand());
+		}
+
+		if (PrioritizedInputChord.Name == InInputChordName)
+		{
+			return IsPressed;
+		}
+	}	
+
+	return false;
+}
+
 void FEditorViewportClient::ProcessDoubleClickInViewport( const struct FInputEventState& InputState, FSceneView& View )
 {
 	// Stop current tracking
diff --git a/Engine/Source/Editor/UnrealEd/Public/EditorViewportClient.h b/Engine/Source/Editor/UnrealEd/Public/EditorViewportClient.h
--- a/Engine/Source/Editor/UnrealEd/Public/EditorViewportClient.h
+++ b/Engine/Source/Editor/UnrealEd/Public/EditorViewportClient.h
@@ -79,6 +79,26 @@ public:
 	TSharedPtr< FUICommandInfo > FovZoomOut;
 };
 
+// FPrioritizedInputChord
+//
+// Provides a way to resolve conflict between overlapping keyboard commands. Instances of this class 
+// are registered with an FEditorViewportClient instance and held in a prioritized list. When input 
+// matches one of them the keys referenced by that one are ignored when evaluating all lower priority 
+// instances.
+//
+struct FPrioritizedInputChord
+{	
+	FPrioritizedInputChord(const int32 InPriority, const FName InName, const EModifierKey::Type InModifierKeyFlags, const FKey InKey = EKeys::Invalid)
+	: InputChord(InModifierKeyFlags, InKey)
+	, Name(InName)
+	, Priority(InPriority)
+	{}
+
+	FInputChord InputChord;
+	FName Name;
+	int32 Priority;
+};
+
 struct FInputEventState
 {
 public:
@@ -1054,6 +1074,14 @@ public:
 	bool IsShiftPressed() const;
 	bool IsCmdPressed() const;
 	
+	// Functions for registering and evaluating prioritized input chords. 
+	// When a prioritized chord evaluates true its required keys are 
+	// ignored by all lower priority chords, providing a way to resolve 
+	// conflict between input key combinations.
+	void RegisterPrioritizedInputChord(const FPrioritizedInputChord& InInputCord);
+	void UnregisterPrioritizedInputChord(const FName InInputCordName);
+	bool IsPrioritizedInputChordPressed(const FName InInputCordName) const;
+
 	/**
 	 * Utility function to return whether the command accepts the key states
 	 * @param InCommand The command being checked
@@ -2010,6 +2038,9 @@ private:
 
 	TArray<FIntPoint> CapturedMouseMoves;
 
+	/** A list of input key combinations, sorted by priority. See FPrioritizedInputChord. */
+	TArray<FPrioritizedInputChord> PrioritizedInputChords;
+
 	/** True if this is the current viewport client in editing */
 	bool bIsCurrentLevelEditingFocus = false;
 };
diff --git a/Engine/Source/Runtime/Engine/Classes/PhysicsEngine/ConstraintInstance.h b/Engine/Source/Runtime/Engine/Classes/PhysicsEngine/ConstraintInstance.h
--- a/Engine/Source/Runtime/Engine/Classes/PhysicsEngine/ConstraintInstance.h
+++ b/Engine/Source/Runtime/Engine/Classes/PhysicsEngine/ConstraintInstance.h
@@ -1025,7 +1025,10 @@ private:
 
 		FPrimitiveDrawInterface* GetPDI() const;
 
-		void DrawCylinder(const FVector& Start, const FVector& End, float Thickness, FMaterialRenderProxy* MaterialProxy, ESceneDepthPriorityGroup DepthPriority) const;
+		void DrawCylinder(const FVector& Start, const FVector& End, const float Thickness, const FMaterialRenderProxy* const MaterialProxy, const ESceneDepthPriorityGroup DepthPriority) const;
+		void DrawCone(const FMatrix& ConeTransform, const float AngleWidth, const float AngleHeight, const uint32 NumSides, const FColor& PDIColor, const FMaterialRenderProxy* const MaterialRenderProxy, const ESceneDepthPriorityGroup DepthPriority) const;
+		void DrawArrow(const FMatrix& ArrowTransform, const float Length, const float Thickness, const uint32 NumSides, const FColor& PDIColor, const FMaterialRenderProxy* const MaterialRenderProxy, const ESceneDepthPriorityGroup DepthPriority) const;
+
 	};
 
 	void DrawConstraintImp(const FPDIOrCollector& PDIOrCollector,
diff --git a/Engine/Source/Runtime/Engine/Private/PhysicsEngine/PhysDrawing.cpp b/Engine/Source/Runtime/Engine/Private/PhysicsEngine/PhysDrawing.cpp
--- a/Engine/Source/Runtime/Engine/Private/PhysicsEngine/PhysDrawing.cpp
+++ b/Engine/Source/Runtime/Engine/Private/PhysicsEngine/PhysDrawing.cpp
@@ -1358,24 +1358,6 @@ static void DrawLinearLimit(FPrimitiveDrawInterface* PDI, const FVector& Origin,
 	}
 }
 
-//creates fan shape along visualized axis for rotation axis of length Length
-FMatrix HelpBuildFan(const FTransform& Con1Frame, const FTransform& Con2Frame, EAxis::Type DrawOnAxis, EAxis::Type RotationAxis, float Length)
-{
-	FVector Con1DrawOnAxis = Con1Frame.GetScaledAxis(DrawOnAxis);
-	FVector Con2DrawOnAxis = Con2Frame.GetScaledAxis(DrawOnAxis);
-
-	FVector Con1RotationAxis = Con1Frame.GetScaledAxis(RotationAxis);
-	FVector Con2RotationAxis = Con2Frame.GetScaledAxis(RotationAxis);
-
-	// Rotate parent twist ref axis
-	FQuat Con2ToCon1Rot = FQuat::FindBetween(Con2RotationAxis, Con1RotationAxis);
-	FVector Con2InCon1DrawOnAxis = Con2ToCon1Rot.RotateVector(Con2DrawOnAxis);
-
-	FTransform ConeLimitTM(Con2InCon1DrawOnAxis, Con1RotationAxis ^ Con2InCon1DrawOnAxis, Con1RotationAxis, Con1Frame.GetTranslation());
-	FMatrix ConeToWorld = FScaleMatrix(FVector(Length * 0.9f)) * ConeLimitTM.ToMatrixWithScale();
-	return ConeToWorld;
-}
-
 //builds radians for limit based on limit type
 float HelpBuildAngle(float LimitAngle, EAngularConstraintMotion LimitType)
 {
@@ -1387,13 +1369,12 @@ float HelpBuildAngle(float LimitAngle, EAngularConstraintMotion LimitType)
 	}
 }
 
-
 FPrimitiveDrawInterface* FConstraintInstance::FPDIOrCollector::GetPDI() const
 {
 	return PDI ? PDI : Collector->GetPDI(ViewIndex);
 }
 
-void FConstraintInstance::FPDIOrCollector::DrawCylinder(const FVector& Start, const FVector& End, float Thickness, FMaterialRenderProxy* MaterialProxy, ESceneDepthPriorityGroup DepthPriority) const
+void FConstraintInstance::FPDIOrCollector::DrawCylinder(const FVector& Start, const FVector& End, const float Thickness, const FMaterialRenderProxy* const MaterialProxy, const ESceneDepthPriorityGroup DepthPriority) const
 {
 	if (HasCollector())
 	{
@@ -1405,6 +1386,35 @@ void FConstraintInstance::FPDIOrCollector::DrawCylinder(const FVector& Start, co
 	}
 }
 
+void FConstraintInstance::FPDIOrCollector::DrawCone(const FMatrix& ConeTransform, const float AngleWidth, const float AngleHeight, const uint32 NumSides, const FColor& PDIColor, const FMaterialRenderProxy* MaterialRenderProxy, const ESceneDepthPriorityGroup DepthPriority) const
+{
+	if (HasCollector())
+	{
+		GetConeMesh(ConeTransform, FMath::RadiansToDegrees(AngleWidth), FMath::RadiansToDegrees(AngleHeight), NumSides, MaterialRenderProxy, DepthPriority, ViewIndex, *Collector);
+	}
+	else
+	{
+		::DrawCone(GetPDI(), ConeTransform, AngleWidth, AngleHeight, NumSides, false, PDIColor, MaterialRenderProxy, DepthPriority);
+	}
+}
+
+void FConstraintInstance::FPDIOrCollector::DrawArrow(const FMatrix& ArrowTransform, const float Length, const float Thickness, const uint32 NumSides, const FColor& PDIColor, const FMaterialRenderProxy* MaterialRenderProxy, const ESceneDepthPriorityGroup DepthPriority) const
+{
+	const FVector ArrowDirection = ArrowTransform.GetScaledAxis(EAxis::X);
+	const FVector ArrowPosition = ArrowTransform.GetOrigin();
+	
+	// Draw the arrow shaft as a cylinder.	
+	DrawCylinder(ArrowPosition, ArrowPosition + Length * ArrowDirection, Thickness, MaterialRenderProxy, DepthPriority);
+
+	// Draw the arrow head as a cone on the end of the shaft.
+	FMatrix ArrowHeadTM = ArrowTransform;
+	ArrowHeadTM.SetOrigin(ArrowPosition + Length * 1.05 * ArrowDirection);
+	const float ConeAngle = UE_PI / 4;
+	const FMatrix ConeToWorld = FScaleMatrix(FVector(Length * -0.1f)) * ArrowHeadTM;
+
+	DrawCone(ConeToWorld, ConeAngle, ConeAngle, NumSides, PDIColor, MaterialRenderProxy, DepthPriority);
+}
+
 void FConstraintInstance::GetUsedMaterials(TArray<UMaterialInterface*>& Materials)
 {
 	Materials.AddUnique(GEngine->ConstraintLimitMaterialX);
@@ -1434,8 +1444,8 @@ void FConstraintInstance::DrawConstraintImp(const FPDIOrCollector& PDIOrCollecto
 	static UMaterialInterface * LimitMaterialZ = GEngine->ConstraintLimitMaterialZ;
 	static UMaterialInterface * LimitMaterialZAxis = GEngine->ConstraintLimitMaterialZAxis;
 	
-	FVector Con1Pos = Con1Frame.GetTranslation();
-	FVector Con2Pos = Con2Frame.GetTranslation();
+	const FVector Con1Pos = Con1Frame.GetTranslation();
+	const FVector Con2Pos = Con2Frame.GetTranslation();
 
 	float Length = (bDrawSelected ? SelectedJointRenderSize : UnselectedJointRenderSize) * Scale;
 	float Thickness = JointRenderThickness;
@@ -1461,7 +1471,7 @@ void FConstraintInstance::DrawConstraintImp(const FPDIOrCollector& PDIOrCollecto
 		const bool bLockSwing2 = GetAngularSwing2Motion() == ACM_Locked;
 		const bool bLockAllSwing = bLockSwing1 && bLockSwing2;
 
-		// If swing is limited (but not locked) - draw the limit cone.
+		// If swing is limited (but not locked) - draw the swing limit cone.
 		if (!bLockAllSwing)
 		{
 			if (ProfileInstance.ConeLimit.Swing1Motion == ACM_Free && ProfileInstance.ConeLimit.Swing2Motion == ACM_Free)
@@ -1477,71 +1487,52 @@ void FConstraintInstance::DrawConstraintImp(const FPDIOrCollector& PDIOrCollecto
 			}
 			else
 			{
-				FTransform ConeLimitTM = Con2Frame;
-				ConeLimitTM.SetTranslation(Con1Frame.GetTranslation());
-
+				const FTransform ConeLimitTM = Con2Frame;
 				const float Swing1Ang = HelpBuildAngle(GetAngularSwing1Limit(), GetAngularSwing1Motion());
 				const float Swing2Ang = HelpBuildAngle(GetAngularSwing2Limit(), GetAngularSwing2Motion());
-				FMatrix ConeToWorld = FScaleMatrix(FVector(Length * 0.9f)) * ConeLimitTM.ToMatrixWithScale();
-
-				if (PDIOrCollector.HasCollector())
-				{
-					GetConeMesh(ConeToWorld, FMath::RadiansToDegrees(Swing1Ang), FMath::RadiansToDegrees(Swing2Ang), DrawConeLimitSides, LimitMaterialX->GetRenderProxy(), Layer, PDIOrCollector.ViewIndex, *PDIOrCollector.Collector);
-				}
-				else
-				{
-					DrawCone(PDI, ConeToWorld, Swing1Ang, Swing2Ang, DrawConeLimitSides, false, JointLimitColor, LimitMaterialX->GetRenderProxy(), Layer);
-				}
+				const FMatrix ConeToWorld = FScaleMatrix(FVector(Length * 0.9f)) * ConeLimitTM.ToMatrixWithScale();
+				PDIOrCollector.DrawCone(ConeToWorld, Swing1Ang, Swing2Ang, DrawConeLimitSides, JointLimitColor, LimitMaterialX->GetRenderProxy(), Layer);
 			}
+	
+			// Draw the swing Dial indicator - shows the current orientation of the child frame relative to the parent frame on the swing axis.	
+			FTransform ArrowTM = Con1Frame;
+			ArrowTM.SetTranslation(Con2Pos);
+			PDIOrCollector.DrawArrow(ArrowTM.ToMatrixWithScale(), Length, Thickness, DrawConeLimitSides, JointLimitColor, LimitMaterialX->GetRenderProxy(), Layer);
 		}
 
-		//twist
+		// Draw the twist limit - A green arc that shows the allowed range of rotation about the parent frame's x axis.
 		if (GetAngularTwistMotion() != ACM_Locked)
 		{
-			FMatrix ConeToWorld = HelpBuildFan(Con1Frame, Con2Frame, EAxis::Y, EAxis::X, Length);
-			float Limit = HelpBuildAngle(GetAngularTwistLimit(), GetAngularTwistMotion());
-			if (PDIOrCollector.HasCollector())
-			{
-				GetConeMesh(ConeToWorld, FMath::RadiansToDegrees(Limit), 0, DrawConeLimitSides, LimitMaterialY->GetRenderProxy(), Layer, PDIOrCollector.ViewIndex, *PDIOrCollector.Collector);
-			}
-			else
-			{
-				DrawCone(PDI, ConeToWorld, Limit, 0, DrawConeLimitSides, false, JointLimitColor, LimitMaterialY->GetRenderProxy(), Layer);
-			}
+			const FTransform ConeLimitTM(Con2Frame.GetScaledAxis(EAxis::Y), Con2Frame.GetScaledAxis(EAxis::Z), Con2Frame.GetScaledAxis(EAxis::X), Con2Frame.GetTranslation()); // Draw the fan in the parent frame (Con2Frame).
+			const FMatrix ConeToWorld = FScaleMatrix(FVector(Length * 0.9f)) * ConeLimitTM.ToMatrixWithScale();
+			const float LimitAngle = HelpBuildAngle(GetAngularTwistLimit(), GetAngularTwistMotion());
+			PDIOrCollector.DrawCone(ConeToWorld, LimitAngle, 0, DrawConeLimitSides, JointLimitColor, LimitMaterialY->GetRenderProxy(), Layer);
+
+			// Draw the twist Dial indicator - shows the current orientation of the child frame relative to the parent frame on the twist axis.
+			const FVector TwistIndicator = FVector::PointPlaneProject(Con1Frame.GetScaledAxis(EAxis::Y), FVector::ZeroVector, Con2Frame.GetUnitAxis(EAxis::X)).GetSafeNormal(); // project the y axis of the child frame into the parent frame's yz plane
+			const FTransform ArrowTM(TwistIndicator, Con2Frame.GetScaledAxis(EAxis::X), TwistIndicator ^ Con2Frame.GetScaledAxis(EAxis::X), Con2Frame.GetTranslation());
+			PDIOrCollector.DrawArrow(ArrowTM.ToMatrixWithScale(), Length, Thickness, DrawConeLimitSides, JointLimitColor, LimitMaterialYAxis->GetRenderProxy(), Layer);
 		}
 	}
 
-
 	//////////////////////////////////////////////////////////////////////////
 	// COORDINATE AXES
-	FVector Position = Con1Frame.GetTranslation();
-
-	PDIOrCollector.DrawCylinder(Position, Position + Length * Con1Frame.GetScaledAxis(EAxis::X), Thickness, LimitMaterialXAxis->GetRenderProxy(), Layer);
-	PDIOrCollector.DrawCylinder(Position, Position + Length * Con1Frame.GetScaledAxis(EAxis::Y), Thickness, LimitMaterialYAxis->GetRenderProxy(), Layer);
-	PDIOrCollector.DrawCylinder(Position, Position + Length * Con1Frame.GetScaledAxis(EAxis::Z), Thickness, LimitMaterialZAxis->GetRenderProxy(), Layer);
-
-	PDIOrCollector.DrawCylinder(Position, Position + Length * Con2Frame.GetScaledAxis(EAxis::X), Thickness, LimitMaterialXAxis->GetRenderProxy(), Layer);
-	PDIOrCollector.DrawCylinder(Position, Position + Length * Con2Frame.GetScaledAxis(EAxis::Y), Thickness, LimitMaterialYAxis->GetRenderProxy(), Layer);
-	PDIOrCollector.DrawCylinder(Position, Position + Length * Con2Frame.GetScaledAxis(EAxis::Z), Thickness, LimitMaterialZAxis->GetRenderProxy(), Layer);
 
+	const float FrameTransformRenderSize = Length * 0.1f;
 
-	//Draw arrow on twist axist
+	// Child Transform
 	{
-		FTransform ConeLimitTM = Con2Frame;
-		ConeLimitTM.SetTranslation(Con1Frame.GetTranslation() + Length*1.05*Con2Frame.GetScaledAxis(EAxis::X));
-
-		const float Swing1Ang = UE_PI / 4;
-		const float Swing2Ang = UE_PI / 4;
-		FMatrix ConeToWorld = FScaleMatrix(FVector(Length * -0.1f)) * ConeLimitTM.ToMatrixWithScale();
+		PDIOrCollector.DrawCylinder(Con1Pos, Con1Pos + FrameTransformRenderSize * Con1Frame.GetScaledAxis(EAxis::X), Thickness, LimitMaterialXAxis->GetRenderProxy(), Layer);
+		PDIOrCollector.DrawCylinder(Con1Pos, Con1Pos + FrameTransformRenderSize * Con1Frame.GetScaledAxis(EAxis::Y), Thickness, LimitMaterialYAxis->GetRenderProxy(), Layer);
+		PDIOrCollector.DrawCylinder(Con1Pos, Con1Pos + FrameTransformRenderSize * Con1Frame.GetScaledAxis(EAxis::Z), Thickness, LimitMaterialZAxis->GetRenderProxy(), Layer);
+	}
 
-		if (PDIOrCollector.HasCollector())
-		{
-			GetConeMesh(ConeToWorld, FMath::RadiansToDegrees(Swing1Ang), FMath::RadiansToDegrees(Swing2Ang), DrawConeLimitSides, LimitMaterialXAxis->GetRenderProxy(), Layer, PDIOrCollector.ViewIndex, *PDIOrCollector.Collector);
-		}
-		else
-		{
-			DrawCone(PDI, ConeToWorld, Swing1Ang, Swing2Ang, DrawConeLimitSides, false, JointLimitColor, LimitMaterialXAxis->GetRenderProxy(), Layer);
-		}
+	// Parent Transform
+	{
+		const FVector Position = Con2Frame.GetTranslation();
+		PDIOrCollector.DrawCylinder(Con2Pos, Con2Pos + FrameTransformRenderSize * Con2Frame.GetScaledAxis(EAxis::X), Thickness, LimitMaterialXAxis->GetRenderProxy(), Layer);
+		PDIOrCollector.DrawCylinder(Con2Pos, Con2Pos + FrameTransformRenderSize * Con2Frame.GetScaledAxis(EAxis::Y), Thickness, LimitMaterialYAxis->GetRenderProxy(), Layer);
+		PDIOrCollector.DrawCylinder(Con2Pos, Con2Pos + FrameTransformRenderSize * Con2Frame.GetScaledAxis(EAxis::Z), Thickness, LimitMaterialZAxis->GetRenderProxy(), Layer);
 	}
 
 	//////////////////////////////////////////////////////////////////////////
