Fix for missing replaced SubObjects reference.

diff --git a/Engine/Source/Runtime/Engine/Private/UnrealEngine.cpp b/Engine/Source/Runtime/Engine/Private/UnrealEngine.cpp
--- a/Engine/Source/Runtime/Engine/Private/UnrealEngine.cpp
+++ b/Engine/Source/Runtime/Engine/Private/UnrealEngine.cpp
@@ -17053,6 +17053,10 @@ void UEngine::CopyPropertiesForUnrelatedObjects(UObject* OldObject, UObject* New
 			{
 				FInstancedObjectRecord& Record = SavedInstances[*pOldInstanceIndex];
 				ReferenceReplacementMap.Add(Record.OldInstance, NewInstance);
+				if(Params.OptionalReplacementMappings)
+				{
+					Params.OptionalReplacementMappings->Add(Record.OldInstance, NewInstance);
+				}
 				if (bDoAggressiveReplacement || bAuditAggressiveReplacement)
 				{
 					UClass* Class = OldObject->GetClass()->GetSuperClass();
