Prevent all RVT from being invalidated when entering landscape mode just because we build a thumbnail material

diff --git a/Engine/Source/Editor/MaterialEditor/Private/MaterialEditor.cpp b/Engine/Source/Editor/MaterialEditor/Private/MaterialEditor.cpp
--- a/Engine/Source/Editor/MaterialEditor/Private/MaterialEditor.cpp
+++ b/Engine/Source/Editor/MaterialEditor/Private/MaterialEditor.cpp
@@ -5048,7 +5048,7 @@ void FMaterialEditor::OnNumericParameterDefaultChanged(class UMaterialExpression
 void FMaterialEditor::OnParameterDefaultChanged()
 {
 	// Brute force all flush virtual textures if this material writes to any runtime virtual texture.
-	if (Material->GetCachedExpressionData().bHasRuntimeVirtualTextureOutput)
+	if (Material->WritesToRuntimeVirtualTexture())
 	{
 		ENQUEUE_RENDER_COMMAND(FlushVTCommand)([](FRHICommandListImmediate& RHICmdList)
 		{
diff --git a/Engine/Source/Runtime/Engine/Classes/Materials/Material.h b/Engine/Source/Runtime/Engine/Classes/Materials/Material.h
--- a/Engine/Source/Runtime/Engine/Classes/Materials/Material.h
+++ b/Engine/Source/Runtime/Engine/Classes/Materials/Material.h
@@ -776,7 +776,7 @@ public:
 	UPROPERTY(EditAnywhere, BlueprintReadOnly, Category = Usage)
 	uint32 bUsedWithNanite : 1;
 
-	/**
+	/** 
 	 * Indicates that the material and its instances with volumetric cloud. Without that flag, it can only be used on volumetric fog.
 	 * This will result in the shaders required to support Volumetric Cloud rendering being compiled which will increase shader compile time and memory usage.
 	 */
@@ -1004,7 +1004,7 @@ public:
 	UPROPERTY(EditAnywhere, Category=Refraction)
 	float RefractionDepthBias;
 
-	/**
+	/** 
 	 * Specifies the max world position offset of the material. Use this value to resolve issues with culling and self-occlusion caused by
 	 * World Position Offset, and/or to restrict how much offset is permitted (i.e. values are clamped on each axis).
 	 * NOTE: A value of 0.0 effectively means "no maximum", and will not clamp the offsets, however it will also not extend culling bounds.
@@ -1124,6 +1124,7 @@ public:
 	ENGINE_API virtual USubsurfaceProfile* GetSubsurfaceProfile_Internal() const override;
 	ENGINE_API virtual bool CastsRayTracedShadows() const override;
 	ENGINE_API virtual float GetMaxWorldPositionOffsetDisplacement() const override;
+	ENGINE_API virtual bool WritesToRuntimeVirtualTexture() const override;
 
 	ENGINE_API virtual FGraphEventArray PrecachePSOs(const FPSOPrecacheVertexFactoryDataList& VertexFactoryDataList, const FPSOPrecacheParams& PreCacheParams, EPSOPrecachePriority Priority, TArray<FMaterialPSOPrecacheRequestID>& OutMaterialPSORequestIDs) override;
 
diff --git a/Engine/Source/Runtime/Engine/Classes/Materials/MaterialInstance.h b/Engine/Source/Runtime/Engine/Classes/Materials/MaterialInstance.h
--- a/Engine/Source/Runtime/Engine/Classes/Materials/MaterialInstance.h
+++ b/Engine/Source/Runtime/Engine/Classes/Materials/MaterialInstance.h
@@ -799,8 +799,8 @@ public:
 	ENGINE_API virtual bool IsTranslucencyWritingVelocity() const override;
 	ENGINE_API virtual bool IsDitheredLODTransition() const override;
 	ENGINE_API virtual bool IsMasked() const override;
+	ENGINE_API virtual bool WritesToRuntimeVirtualTexture() const override;
 	ENGINE_API virtual float GetMaxWorldPositionOffsetDisplacement() const override;
-	
 	ENGINE_API virtual USubsurfaceProfile* GetSubsurfaceProfile_Internal() const override;
 	ENGINE_API virtual bool CastsRayTracedShadows() const override;
 
diff --git a/Engine/Source/Runtime/Engine/Classes/Materials/MaterialInterface.h b/Engine/Source/Runtime/Engine/Classes/Materials/MaterialInterface.h
--- a/Engine/Source/Runtime/Engine/Classes/Materials/MaterialInterface.h
+++ b/Engine/Source/Runtime/Engine/Classes/Materials/MaterialInterface.h
@@ -906,8 +906,8 @@ public:
 	ENGINE_API virtual bool IsTranslucencyWritingFrontLayerTransparency() const;
 	ENGINE_API virtual bool IsMasked() const;
 	ENGINE_API virtual bool IsDeferredDecal() const;
+	ENGINE_API virtual bool WritesToRuntimeVirtualTexture() const;
 	ENGINE_API virtual float GetMaxWorldPositionOffsetDisplacement() const;
-
 	ENGINE_API virtual USubsurfaceProfile* GetSubsurfaceProfile_Internal() const;
 	ENGINE_API virtual bool CastsRayTracedShadows() const;
 
diff --git a/Engine/Source/Runtime/Engine/Private/Materials/Material.cpp b/Engine/Source/Runtime/Engine/Private/Materials/Material.cpp
--- a/Engine/Source/Runtime/Engine/Private/Materials/Material.cpp
+++ b/Engine/Source/Runtime/Engine/Private/Materials/Material.cpp
@@ -6566,6 +6566,11 @@ bool UMaterial::IsPostProcessMaterial() const
 	return MaterialDomain == MD_PostProcess;
 }
 
+bool UMaterial::WritesToRuntimeVirtualTexture() const
+{
+	return GetCachedExpressionData().bHasRuntimeVirtualTextureOutput;
+}
+
 USubsurfaceProfile* UMaterial::GetSubsurfaceProfile_Internal() const
 {
 	checkSlow(IsInGameThread());
diff --git a/Engine/Source/Runtime/Engine/Private/Materials/MaterialInstance.cpp b/Engine/Source/Runtime/Engine/Private/Materials/MaterialInstance.cpp
--- a/Engine/Source/Runtime/Engine/Private/Materials/MaterialInstance.cpp
+++ b/Engine/Source/Runtime/Engine/Private/Materials/MaterialInstance.cpp
@@ -1768,6 +1768,11 @@ FMaterialResource* UMaterialInstance::GetMaterialResource(ERHIFeatureLevel::Type
 	return Parent ? Parent->GetMaterialResource(InFeatureLevel, QualityLevel) : nullptr;
 }
 
+bool UMaterialInstance::WritesToRuntimeVirtualTexture() const
+{
+	return Parent ? Parent->WritesToRuntimeVirtualTexture() : false;
+}
+
 const FMaterialResource* UMaterialInstance::GetMaterialResource(ERHIFeatureLevel::Type InFeatureLevel, EMaterialQualityLevel::Type QualityLevel) const
 {
 	return const_cast<UMaterialInstance*>(this)->GetMaterialResource(InFeatureLevel, QualityLevel);
@@ -4013,8 +4018,7 @@ void UMaterialInstance::PostEditChangeProperty(FPropertyChangedEvent& PropertyCh
 	if (GIsEditor)
 	{
 		// Brute force all flush virtual textures if this material writes to any runtime virtual texture.
-		const UMaterial* BaseMaterial = GetMaterial();
-		if (BaseMaterial != nullptr && BaseMaterial->GetCachedExpressionData().bHasRuntimeVirtualTextureOutput)
+		if (WritesToRuntimeVirtualTexture())
 		{
 			ENQUEUE_RENDER_COMMAND(FlushVTCommand)([ResourcePtr = Resource](FRHICommandListImmediate& RHICmdList) 
 			{
diff --git a/Engine/Source/Runtime/Engine/Private/Materials/MaterialInterface.cpp b/Engine/Source/Runtime/Engine/Private/Materials/MaterialInterface.cpp
--- a/Engine/Source/Runtime/Engine/Private/Materials/MaterialInterface.cpp
+++ b/Engine/Source/Runtime/Engine/Private/Materials/MaterialInterface.cpp
@@ -1286,11 +1286,17 @@ bool UMaterialInterface::IsDeferredDecal() const
 {
 	return false;
 }
+
 bool UMaterialInterface::GetCastDynamicShadowAsMasked() const
 {
 	return false;
 }
 
+bool UMaterialInterface::WritesToRuntimeVirtualTexture() const
+{
+	return false;
+}
+
 FMaterialShadingModelField UMaterialInterface::GetShadingModels() const
 {
 	return MSM_DefaultLit;
diff --git a/Engine/Source/Runtime/Landscape/Classes/LandscapeMaterialInstanceConstant.h b/Engine/Source/Runtime/Landscape/Classes/LandscapeMaterialInstanceConstant.h
--- a/Engine/Source/Runtime/Landscape/Classes/LandscapeMaterialInstanceConstant.h
+++ b/Engine/Source/Runtime/Landscape/Classes/LandscapeMaterialInstanceConstant.h
@@ -43,6 +43,8 @@ class ULandscapeMaterialInstanceConstant : public UMaterialInstanceConstant
 	UPROPERTY()
 	uint32 bEditorToolUsage:1;
 
+	virtual bool WritesToRuntimeVirtualTexture() const override;
+
 	virtual void PostLoad() override;
 #if WITH_EDITOR
 	virtual void PostEditChangeProperty(FPropertyChangedEvent& PropertyChangedEvent) override;
diff --git a/Engine/Source/Runtime/Landscape/Private/LandscapeEditLayers.cpp b/Engine/Source/Runtime/Landscape/Private/LandscapeEditLayers.cpp
--- a/Engine/Source/Runtime/Landscape/Private/LandscapeEditLayers.cpp
+++ b/Engine/Source/Runtime/Landscape/Private/LandscapeEditLayers.cpp
@@ -3586,7 +3586,7 @@ void ALandscape::InvalidateRVTForTextures(const TSet<UTexture2D*>& InTextures)
 				//  gets fully streamed in so that we're not left with low-res mips being rendered in the RVT tiles : 
 				for (UMaterialInterface* MaterialInterface : ObjectCacheScope.GetContext().GetMaterialsAffectedByTexture(Texture))
 				{
-					if (MaterialInterface->GetCachedExpressionData().bHasRuntimeVirtualTextureOutput)
+					if (MaterialInterface->WritesToRuntimeVirtualTexture())
 					{
 						for (UPrimitiveComponent* PrimitiveComponent : ObjectCacheScope.GetContext().GetPrimitivesAffectedByMaterial(MaterialInterface))
 						{
diff --git a/Engine/Source/Runtime/Landscape/Private/LandscapeRender.cpp b/Engine/Source/Runtime/Landscape/Private/LandscapeRender.cpp
--- a/Engine/Source/Runtime/Landscape/Private/LandscapeRender.cpp
+++ b/Engine/Source/Runtime/Landscape/Private/LandscapeRender.cpp
@@ -3163,6 +3163,12 @@ float ULandscapeMaterialInstanceConstant::GetLandscapeTexelFactor(const FName& T
 	return 1.0f;
 }
 
+bool ULandscapeMaterialInstanceConstant::WritesToRuntimeVirtualTexture() const
+{
+	// Don't invalidate the RVTs for the thumbnail material updates, that would be a waste : 
+	return !bIsLayerThumbnail && Super::WritesToRuntimeVirtualTexture();
+}
+
 #if WITH_EDITOR
 
 void ULandscapeMaterialInstanceConstant::PostEditChangeProperty(FPropertyChangedEvent& PropertyChangedEvent)
