Fixed Inactive world types interacting with RVT that caused changelist validation (which loads actors in an Inactive world) to create a new VT producer.

diff --git a/Engine/Source/Runtime/Engine/Classes/Components/RuntimeVirtualTextureComponent.h b/Engine/Source/Runtime/Engine/Classes/Components/RuntimeVirtualTextureComponent.h
--- a/Engine/Source/Runtime/Engine/Classes/Components/RuntimeVirtualTextureComponent.h
+++ b/Engine/Source/Runtime/Engine/Classes/Components/RuntimeVirtualTextureComponent.h
@@ -184,6 +184,9 @@ protected:
 	/** Calculate a hash used to determine if the StreamingTexture contents are valid for use. The hash doesn't include whether the contents are up to date. */
 	uint64 CalculateStreamingTextureSettingsHash() const;
 
+	/** @return true if the owning World is one where URuntimeVirtualTextureComponent should actually do anything (avoids updating RVT for non-game/PIE/editor world types) */
+	bool IsActiveInWorld() const;
+
 public:
 	/** Scene proxy object. Managed by the scene but stored here. */
 	class FRuntimeVirtualTextureSceneProxy* SceneProxy;
diff --git a/Engine/Source/Runtime/Engine/Private/Components/RuntimeVirtualTextureComponent.cpp b/Engine/Source/Runtime/Engine/Private/Components/RuntimeVirtualTextureComponent.cpp
--- a/Engine/Source/Runtime/Engine/Private/Components/RuntimeVirtualTextureComponent.cpp
+++ b/Engine/Source/Runtime/Engine/Private/Components/RuntimeVirtualTextureComponent.cpp
@@ -44,11 +44,21 @@ bool URuntimeVirtualTextureComponent::IsReadyForFinishDestroy()
 	return bResult;
 }
 
+bool URuntimeVirtualTextureComponent::IsActiveInWorld() const
+{
+	UWorld* World = GetWorld();
+	return ((World != nullptr) 
+		&& ((World->WorldType == EWorldType::Game) 
+			|| (World->WorldType == EWorldType::Editor) 
+			|| (World->WorldType == EWorldType::PIE)));
+}
+
 #if WITH_EDITOR
 
 void URuntimeVirtualTextureComponent::OnRegister()
 {
 	Super::OnRegister();
+
 	// PIE duplicate will take ownership of the URuntimeVirtualTexture, so we add a delegate to be called when PIE finishes allowing us to retake ownership.
 	PieEndDelegateHandle = FGameDelegates::Get().GetEndPlayMapDelegate().AddUObject(this, &URuntimeVirtualTextureComponent::MarkRenderStateDirty);
 }
@@ -79,7 +89,8 @@ void URuntimeVirtualTextureComponent::GetHidePrimitiveSettings(bool& OutHidePrim
 
 bool URuntimeVirtualTextureComponent::IsVisible() const
 {
-	return Super::IsVisible() && UseVirtualTexturing(GetScene()->GetFeatureLevel());
+	// Make sure to have the component do nothing if VT is disabled or if the world is not compatible with RVT
+	return Super::IsVisible() && IsActiveInWorld() && UseVirtualTexturing(GetScene()->GetFeatureLevel());
 }
 
 void URuntimeVirtualTextureComponent::ApplyWorldOffset(const FVector& InOffset, bool bWorldShift)
@@ -91,7 +102,8 @@ void URuntimeVirtualTextureComponent::ApplyWorldOffset(const FVector& InOffset,
 
 void URuntimeVirtualTextureComponent::CreateRenderState_Concurrent(FRegisterComponentContext* Context)
 {
-	if (ShouldRender() && VirtualTexture != nullptr)
+	// Make sure not to create a render state if the world is not compatible with RVT :
+	if (IsActiveInWorld() && ShouldRender() && VirtualTexture != nullptr)
 	{
 		// This will modify the URuntimeVirtualTexture and allocate its VT
 		GetScene()->AddRuntimeVirtualTexture(this);
@@ -102,8 +114,11 @@ void URuntimeVirtualTextureComponent::CreateRenderState_Concurrent(FRegisterComp
 
 void URuntimeVirtualTextureComponent::SendRenderTransform_Concurrent()
 {
-	if (ShouldRender() && VirtualTexture != nullptr)
+	// We don't have a render state if the world is not compatible with RVT :
+	if (IsActiveInWorld() && ShouldRender() && VirtualTexture != nullptr)
 	{
+		checkf(IsActiveInWorld(), TEXT("ShouldRender should never return true for a world where we're inactive"));
+
 		// We do a full recreate of the URuntimeVirtualTexture here which can cause a visual glitch.
 		// We do this because, for an arbitrary transform, there is no way to only modify the transform and maintain the VT contents.
 		// Possibly, with some work, the contents could be maintained for any transform change that is an exact multiple of the page size in world space.
@@ -115,15 +130,18 @@ void URuntimeVirtualTextureComponent::SendRenderTransform_Concurrent()
 
 void URuntimeVirtualTextureComponent::DestroyRenderState_Concurrent()
 {
-	// This will modify the URuntimeVirtualTexture and free its VT
-	GetScene()->RemoveRuntimeVirtualTexture(this);
+	if (IsActiveInWorld())
+	{
+		// This will modify the URuntimeVirtualTexture and free its VT
+		GetScene()->RemoveRuntimeVirtualTexture(this);
+	}
 
 	Super::DestroyRenderState_Concurrent();
 }
 
 void URuntimeVirtualTextureComponent::Invalidate(FBoxSphereBounds const& InWorldBounds)
 {
-	if (GetScene() != nullptr)
+	if (IsActiveInWorld() && (GetScene() != nullptr))
 	{
 		GetScene()->InvalidateRuntimeVirtualTexture(this, InWorldBounds);
 	}
@@ -194,6 +212,8 @@ uint64 URuntimeVirtualTextureComponent::CalculateStreamingTextureSettingsHash()
 
 bool URuntimeVirtualTextureComponent::IsStreamingLowMips(EShadingPath ShadingPath) const
 {
+	checkf(IsActiveInWorld(), TEXT("This function should never be called for a world where we're inactive"));
+	
 #if WITH_EDITOR
 	if (!bUseStreamingLowMipsInEditor && GIsEditor)
 	{
@@ -205,6 +225,8 @@ bool URuntimeVirtualTextureComponent::IsStreamingLowMips(EShadingPath ShadingPat
 
 bool URuntimeVirtualTextureComponent::IsStreamingTextureInvalid(EShadingPath ShadingPath) const
 {
+	checkf(IsActiveInWorld(), TEXT("This function should never be called for a world where we're inactive"));
+
 	return 
 		VirtualTexture != nullptr && 
 		StreamingTexture != nullptr && 
@@ -216,6 +238,8 @@ bool URuntimeVirtualTextureComponent::IsStreamingTextureInvalid(EShadingPath Sha
 
 bool URuntimeVirtualTextureComponent::IsStreamingTextureInvalid() const
 {
+	checkf(IsActiveInWorld(), TEXT("This function should never be called for a world where we're inactive"));
+
 	return 
 		VirtualTexture != nullptr && 
 		StreamingTexture != nullptr && 
@@ -270,7 +294,7 @@ static void GetLayerFormatSettings(FTextureFormatSettings& OutFormatSettings, EP
 void URuntimeVirtualTextureComponent::InitializeStreamingTexture(EShadingPath ShadingPath, uint32 InSizeX, uint32 InSizeY, uint8* InData)
 {
 	// We need an existing StreamingTexture object to update.
-	if (VirtualTexture != nullptr && StreamingTexture != nullptr)
+	if (IsActiveInWorld() && VirtualTexture != nullptr && StreamingTexture != nullptr)
 	{
 		FScopedRuntimeVirtualTextureRecreate ProducerRecreate(StreamingTexture);
 
@@ -326,7 +350,7 @@ void URuntimeVirtualTextureComponent::CheckForErrors()
 	Super::CheckForErrors();
 
 	// Check if streaming texture has been built with the latest settings. If not then it won't be used which would cause a performance regression.
-	if (IsStreamingTextureInvalid())
+	if (IsActiveInWorld() && IsStreamingTextureInvalid())
 	{
 		FMessageLog("MapCheck").PerformanceWarning()
 			->AddToken(FUObjectToken::Create(this))
