*Don't try to find a manager for FActorInstanceHandles constructed with an actor.
*Fix FindLightWeightInstanceManager being called with AActor::StaticClass instead of the subclass of the actors, resulting in never finding a manager.

diff --git a/Engine/Source/Runtime/Engine/Private/Engine/ActorInstanceHandle.cpp b/Engine/Source/Runtime/Engine/Private/Engine/ActorInstanceHandle.cpp
--- a/Engine/Source/Runtime/Engine/Private/Engine/ActorInstanceHandle.cpp
+++ b/Engine/Source/Runtime/Engine/Private/Engine/ActorInstanceHandle.cpp
@@ -19,17 +19,6 @@ FActorInstanceHandle::FActorInstanceHandle(AActor* InActor)
 	, InstanceIndex(INDEX_NONE)
 	, InstanceUID(0)
 {
-	if (InActor)
-	{
-#if WITH_EDITOR
-		// use the first data layer the actor is in if it's in multiple layers
-		TArray<const UDataLayerInstance*> DataLayerInstances = InActor->GetDataLayerInstances();
-		const UDataLayerInstance* DataLayerInstance = DataLayerInstances.Num() > 0 ? DataLayerInstances[0] : nullptr;
-#else
-		const UDataLayerInstance* DataLayerInstance = nullptr;
-#endif // WITH_EDITOR
-		Manager = FLightWeightInstanceSubsystem::Get().FindLightWeightInstanceManager(*InActor->StaticClass(), *InActor->GetWorld(), InActor->GetActorLocation(), DataLayerInstance);
-	}
 }
 
 FActorInstanceHandle::FActorInstanceHandle(ALightWeightInstanceManager* InManager, int32 InInstanceIndex)
@@ -288,7 +277,7 @@ bool FActorInstanceHandle::operator==(const AActor* OtherActor) const
 	const UDataLayerInstance* DataLayerInstance = nullptr;
 #endif // WITH_EDITOR
 
-	if (ALightWeightInstanceManager* LWIManager = FLightWeightInstanceSubsystem::Get().FindLightWeightInstanceManager(*OtherActor->StaticClass(), *OtherActor->GetWorld(), OtherActor->GetActorLocation(), DataLayerInstance))
+	if (ALightWeightInstanceManager* LWIManager = FLightWeightInstanceSubsystem::Get().FindLightWeightInstanceManager(*OtherActor->GetClass(), *OtherActor->GetWorld(), OtherActor->GetActorLocation(), DataLayerInstance))
 	{
 		if (Manager.Get() != LWIManager)
 		{
