Fixed unity build

#rnx

diff --git a/Engine/Source/Runtime/Engine/Private/Components/RuntimeVirtualTextureComponent.cpp b/Engine/Source/Runtime/Engine/Private/Components/RuntimeVirtualTextureComponent.cpp
--- a/Engine/Source/Runtime/Engine/Private/Components/RuntimeVirtualTextureComponent.cpp
+++ b/Engine/Source/Runtime/Engine/Private/Components/RuntimeVirtualTextureComponent.cpp
@@ -5,6 +5,7 @@
 #include "Async/TaskGraphInterfaces.h"
 #include "GameDelegates.h"
 #include "Engine/Texture.h"
+#include "Engine/World.h"
 #include "Logging/MessageLog.h"
 #include "GameFramework/Actor.h"
 #include "Misc/UObjectToken.h"
