Material parameter collection instance optimization
- Saves 2MB after land from skydive. (obj list class=materialparametercollectioninstance)
- Removes 0.15ms (0.8ms on ps4) material parameter collection instance setup per wp cell when streaming in.
- Reduces material parameter collection postload cost from 1ms (5ms on ps4) to 0.07ms.

Perf data from a threadripper pc.

diff --git a/Engine/Source/Runtime/Engine/Private/Materials/ParameterCollection.cpp b/Engine/Source/Runtime/Engine/Private/Materials/ParameterCollection.cpp
--- a/Engine/Source/Runtime/Engine/Private/Materials/ParameterCollection.cpp
+++ b/Engine/Source/Runtime/Engine/Private/Materials/ParameterCollection.cpp
@@ -56,7 +56,12 @@ void UMaterialParameterCollection::SetupWorldParameterCollectionInstances()
 	for (TObjectIterator<UWorld> It; It; ++It)
 	{
 		UWorld* CurrentWorld = *It;
-		CurrentWorld->AddParameterCollectionInstance(this, true);
+		ULevel* Level = CurrentWorld->PersistentLevel;
+		const bool bIsWorldPartitionRuntimeCell = Level && Level->IsWorldPartitionRuntimeCell();
+		if (!bIsWorldPartitionRuntimeCell)
+		{
+			CurrentWorld->AddParameterCollectionInstance(this, true);
+		}
 	}
 }
 
diff --git a/Engine/Source/Runtime/Engine/Private/World.cpp b/Engine/Source/Runtime/Engine/Private/World.cpp
--- a/Engine/Source/Runtime/Engine/Private/World.cpp
+++ b/Engine/Source/Runtime/Engine/Private/World.cpp
@@ -1522,15 +1522,20 @@ UWorld* UWorld::GetWorld() const
 
 void UWorld::SetupParameterCollectionInstances()
 {
-	QUICK_SCOPE_CYCLE_COUNTER(STAT_World_SetupParameterCollectionInstances);
-
-	// Create an instance for each parameter collection in memory
-	for (UMaterialParameterCollection* CurrentCollection : TObjectRange<UMaterialParameterCollection>())
+	ULevel* Level = PersistentLevel;
+	const bool bIsWorldPartitionRuntimeCell = Level && Level->IsWorldPartitionRuntimeCell();
+	if (!bIsWorldPartitionRuntimeCell)
 	{
-		AddParameterCollectionInstance(CurrentCollection, false);
-	}
+		QUICK_SCOPE_CYCLE_COUNTER(STAT_World_SetupParameterCollectionInstances);
 
-	UpdateParameterCollectionInstances(false, false);
+		// Create an instance for each parameter collection in memory
+		for (UMaterialParameterCollection* CurrentCollection : TObjectRange<UMaterialParameterCollection>())
+		{
+			AddParameterCollectionInstance(CurrentCollection, false);
+		}
+
+		UpdateParameterCollectionInstances(false, false);
+	}
 }
 
 void UWorld::AddParameterCollectionInstance(UMaterialParameterCollection* Collection, bool bUpdateScene)
