Fixed missing include.

diff --git a/Engine/Source/Runtime/Engine/Private/Materials/ParameterCollection.cpp b/Engine/Source/Runtime/Engine/Private/Materials/ParameterCollection.cpp
--- a/Engine/Source/Runtime/Engine/Private/Materials/ParameterCollection.cpp
+++ b/Engine/Source/Runtime/Engine/Private/Materials/ParameterCollection.cpp
@@ -3,6 +3,7 @@
 #include "ParameterCollection.h"
 #include "UObject/UObjectIterator.h"
 #include "RenderingThread.h"
+#include "Engine/Level.h"
 #include "Engine/World.h"
 #include "MaterialShared.h"
 #include "MaterialCachedData.h"
