Chaos: Remove allocation of multifind array in GetPendingCollisionForContactPair

diff --git a/Engine/Source/Runtime/Engine/Private/PhysicsEngine/Experimental/PhysScene_Chaos.cpp b/Engine/Source/Runtime/Engine/Private/PhysicsEngine/Experimental/PhysScene_Chaos.cpp
--- a/Engine/Source/Runtime/Engine/Private/PhysicsEngine/Experimental/PhysScene_Chaos.cpp
+++ b/Engine/Source/Runtime/Engine/Private/PhysicsEngine/Experimental/PhysScene_Chaos.cpp
@@ -790,15 +790,14 @@ const FBodyInstance* FPhysScene_Chaos::GetBodyInstanceFromProxyAndShape(IPhysics
 FCollisionNotifyInfo& FPhysScene_Chaos::GetPendingCollisionForContactPair(const void* P0, const void* P1, Chaos::FReal SolverTime, bool& bNewEntry)
 {
 	const FUniqueContactPairKey Key = { P0, P1 };
-	TArray<int32> ExistingPairs;
-	ContactPairToPendingNotifyMap.MultiFind(Key, ExistingPairs);
-	for(int32 ExistingIdx : ExistingPairs)
+	for(TMultiMap<FUniqueContactPairKey, int32>::TConstKeyIterator It = ContactPairToPendingNotifyMap.CreateConstKeyIterator(Key); It; ++It)
 	{
+		int32 ExistingIdx = It.Value();
 		if(FMath::IsNearlyEqual(PendingCollisionNotifies[ExistingIdx].SolverTime, SolverTime))
 		{
 			// we already have one for this pair
 			bNewEntry = false;
-			return PendingCollisionNotifies[ExistingIdx];	
+			return PendingCollisionNotifies[ExistingIdx];
 		}
 	}
 
