Adding in a CSV stat to count the number of streaming levels that are pending purge cleanup.

diff --git a/Engine/Source/Runtime/Engine/Private/World.cpp b/Engine/Source/Runtime/Engine/Private/World.cpp
--- a/Engine/Source/Runtime/Engine/Private/World.cpp
+++ b/Engine/Source/Runtime/Engine/Private/World.cpp
@@ -136,6 +136,7 @@ CSV_DECLARE_CATEGORY_MODULE_EXTERN(CORE_API, Basic);
 CSV_DEFINE_CATEGORY(LevelStreamingAdaptive, true);
 CSV_DEFINE_CATEGORY(LevelStreamingAdaptiveDetail, false);
 CSV_DEFINE_CATEGORY(LevelStreamingDetail, false);
+CSV_DEFINE_CATEGORY(LevelStreamingPendingPurge, (!UE_BUILD_SHIPPING));
 
 #define LOCTEXT_NAMESPACE "World"
 
@@ -4188,10 +4189,14 @@ void UWorld::UpdateLevelStreaming()
 
 	StreamingLevelsToConsider.EndConsideration();
 
+	const int32 CurrentNumLevelsPendingPurge = FLevelStreamingGCHelper::GetNumLevelsPendingPurge();
+
+	CSV_CUSTOM_STAT(LevelStreamingPendingPurge, NumlevelsPendingPurge, CurrentNumLevelsPendingPurge, ECsvCustomStatOp::Set);
+
 	if (GLevelStreamingContinuouslyIncrementalGCWhileLevelsPendingPurge)
 	{
 		// Figure out whether there are any levels we haven't collected garbage yet.
-		const bool bAreLevelsPendingPurge = FLevelStreamingGCHelper::GetNumLevelsPendingPurge() >= GLevelStreamingContinuouslyIncrementalGCWhileLevelsPendingPurge;
+		const bool bAreLevelsPendingPurge = CurrentNumLevelsPendingPurge >= GLevelStreamingContinuouslyIncrementalGCWhileLevelsPendingPurge;
 
 		// Request a 'soft' GC if there are levels pending purge and there are levels to be loaded. In the case of a blocking
 		// load this is going to guarantee GC firing first thing afterwards and otherwise it is going to sneak in right before
