"Obj list" report sorting function not working well.

diff --git a/Engine/Source/Runtime/Engine/Private/UnrealEngine.cpp b/Engine/Source/Runtime/Engine/Private/UnrealEngine.cpp
--- a/Engine/Source/Runtime/Engine/Private/UnrealEngine.cpp
+++ b/Engine/Source/Runtime/Engine/Private/UnrealEngine.cpp
@@ -9048,6 +9048,7 @@ bool UEngine::HandleObjCommand( const TCHAR* Cmd, FOutputDevice& Ar )
 				FCompareFItem( bool InAlphaSort, bool InCountSort, bool InResourceSizeSort )
 					: bAlphaSort( InAlphaSort )
 					, bCountSort( InCountSort )
+					, bResourceSizeSort( InResourceSizeSort )
 				{}
 				FORCEINLINE bool operator()( const FItem& A, const FItem& B ) const
 				{
