Fixed regression in FArchiveGatherExternalActorRefs where weak object pointers on garbage objects were treated as valid references.
One side effect of this was that creating a Level instance with any existing actor would generate validation errors on invalid references.

diff --git a/Engine/Source/Runtime/Engine/Private/ActorReferencesUtils.cpp b/Engine/Source/Runtime/Engine/Private/ActorReferencesUtils.cpp
--- a/Engine/Source/Runtime/Engine/Private/ActorReferencesUtils.cpp
+++ b/Engine/Source/Runtime/Engine/Private/ActorReferencesUtils.cpp
@@ -49,7 +49,7 @@ public:
 
 	virtual FArchive& operator<<(FWeakObjectPtr& Value) override
 	{
-		if (UObject* Object = Value.Get(true))
+		if (UObject* Object = Value.Get())
 		{
 			return *this << Object;
 		}
