Increase material connected property bitmask from 32b to 64b, and fixed up invalid static_asserts that were supposed to detect bitmask overflow (but were not).

diff --git a/Engine/Source/Runtime/Core/Public/UObject/UE5MainStreamObjectVersion.h b/Engine/Source/Runtime/Core/Public/UObject/UE5MainStreamObjectVersion.h
--- a/Engine/Source/Runtime/Core/Public/UObject/UE5MainStreamObjectVersion.h
+++ b/Engine/Source/Runtime/Core/Public/UObject/UE5MainStreamObjectVersion.h
@@ -346,6 +346,9 @@ struct CORE_API FUE5MainStreamObjectVersion
 		// Change the default for facing & alignment to be automatic
 		GroomAssetRemoveInAssetSerialization,
 
+		// Changed the material property connected bitmasks from 32bit to 64bit
+		IncreaseMaterialAttributesInputMask,
+
 		// -----<new versions can be added above this line>-------------------------------------------------
 		VersionPlusOne,
 		LatestVersion = VersionPlusOne - 1
diff --git a/Engine/Source/Runtime/Engine/Classes/Engine/EngineTypes.h b/Engine/Source/Runtime/Engine/Classes/Engine/EngineTypes.h
--- a/Engine/Source/Runtime/Engine/Classes/Engine/EngineTypes.h
+++ b/Engine/Source/Runtime/Engine/Classes/Engine/EngineTypes.h
@@ -723,6 +723,8 @@ public:
 	FStrataMaterialInfo() {}
 	FStrataMaterialInfo(EStrataShadingModel InShadingModel) { AddShadingModel(InShadingModel); }
 
+	bool Serialize(FArchive& Ar);
+
 	// Shading model
 	void AddShadingModel(EStrataShadingModel InShadingModel) { check(InShadingModel < SSM_NUM); ShadingModelField |= (uint16)(1 << (uint16)InShadingModel); }
 	void SetSingleShadingModel(EStrataShadingModel InShadingModel) { check(InShadingModel < SSM_NUM); ShadingModelField = (uint16)(1 << (uint16)InShadingModel); }
@@ -740,10 +742,10 @@ public:
 	void SetShadingModelFromExpression(bool bIn) { bHasShadingModelFromExpression = bIn ? 1u : 0u; }
 	bool HasShadingModelFromExpression() const { return bHasShadingModelFromExpression > 0u; }
 
-	uint32 GetPropertyConnected() const { return ConnectedProperties; }
-	void AddPropertyConnected(uint32 In) { ConnectedProperties |= (1 << In); }
-	bool HasPropertyConnected(uint32 In) const { return !!(ConnectedProperties & (1 << In)); }
-	static bool HasPropertyConnected(uint32 InConnectedProperties, uint32 In) { return !!(InConnectedProperties & (1 << In)); }
+	uint64 GetPropertyConnected() const { return ConnectedPropertyMask; }
+	void AddPropertyConnected(uint64 In) { ConnectedPropertyMask |= (1ull << In); }
+	bool HasPropertyConnected(uint64 In) const { return !!(ConnectedPropertyMask & (1ull << In)); }
+	static bool HasPropertyConnected(uint64 InConnectedPropertyMask, uint64 In) { return !!(InConnectedPropertyMask & (1ull << In)); }
 
 	bool IsValid() const { return (ShadingModelField > 0) && (ShadingModelField < (1 << SSM_NUM)); }
 
@@ -776,10 +778,13 @@ private:
 	UPROPERTY()
 	uint8 bHasShadingModelFromExpression = 0;
 
+	UPROPERTY()
+	uint32 ConnectedProperties_DEPRECATED = 0;
+
 	/* Indicates which (legacy) inputs are connected */
 	UPROPERTY()
-	uint32 ConnectedProperties = 0;
-	
+	uint64 ConnectedPropertyMask = 0;
+
 	UPROPERTY()
 	TArray<TObjectPtr<USubsurfaceProfile>> SubsurfaceProfiles;
 
@@ -790,6 +795,15 @@ private:
 #endif
 };
 
+template<>
+struct TStructOpsTypeTraits<FStrataMaterialInfo> : public TStructOpsTypeTraitsBase2<FStrataMaterialInfo>
+{
+	enum
+	{
+		WithSerializer = true,
+	};
+};
+
 /** Describes how textures are sampled for materials */
 UENUM(BlueprintType)
 enum EMaterialSamplerType : int
diff --git a/Engine/Source/Runtime/Engine/Classes/Materials/MaterialExpression.h b/Engine/Source/Runtime/Engine/Classes/Materials/MaterialExpression.h
--- a/Engine/Source/Runtime/Engine/Classes/Materials/MaterialExpression.h
+++ b/Engine/Source/Runtime/Engine/Classes/Materials/MaterialExpression.h
@@ -83,7 +83,7 @@ USTRUCT(noexport)
 struct FMaterialAttributesInput : public FExpressionInput
 {
 	UPROPERTY(transient)
-	int32 PropertyConnectedBitmask;
+	int64 PropertyConnectedMask;
 };
 
 #endif
diff --git a/Engine/Source/Runtime/Engine/Private/Materials/Material.cpp b/Engine/Source/Runtime/Engine/Private/Materials/Material.cpp
--- a/Engine/Source/Runtime/Engine/Private/Materials/Material.cpp
+++ b/Engine/Source/Runtime/Engine/Private/Materials/Material.cpp
@@ -4818,29 +4818,29 @@ void UMaterial::RebuildShadingModelField()
 		if (EditorOnly->FrontMaterial.Expression->IsResultStrataMaterial(EditorOnly->FrontMaterial.OutputIndex))
 		{
 			// Mask of all input collected by Strata BSDF nodes
-			static const uint32 ConnectionMask = 
-				  (1u << MP_BaseColor)
-				| (1u << MP_Metallic)
-				| (1u << MP_Specular)
-				| (1u << MP_Roughness)
-				| (1u << MP_Anisotropy)
-				| (1u << MP_EmissiveColor)
-				| (1u << MP_Normal)
-				| (1u << MP_Tangent)
-				| (1u << MP_SubsurfaceColor)
-				| (1u << MP_CustomData0)
-				| (1u << MP_CustomData1)
-				| (1u << MP_Opacity)
-				| (1u << MP_ShadingModel)
-				| (1u << MP_DiffuseColor)
-				| (1u << MP_SpecularColor);
+			static const uint64 ConnectionMask = 
+				  (1ull << MP_BaseColor)
+				| (1ull << MP_Metallic)
+				| (1ull << MP_Specular)
+				| (1ull << MP_Roughness)
+				| (1ull << MP_Anisotropy)
+				| (1ull << MP_EmissiveColor)
+				| (1ull << MP_Normal)
+				| (1ull << MP_Tangent)
+				| (1ull << MP_SubsurfaceColor)
+				| (1ull << MP_CustomData0)
+				| (1ull << MP_CustomData1)
+				| (1ull << MP_Opacity)
+				| (1ull << MP_ShadingModel)
+				| (1ull << MP_DiffuseColor)
+				| (1ull << MP_SpecularColor);
 
 			EditorOnly->FrontMaterial.Expression->GatherStrataMaterialInfo(StrataMaterialInfo, EditorOnly->FrontMaterial.OutputIndex);
 
 			// Override the cached expression data with collected connection from StrataMaterialInfo, but preserve all other input (e.g., refraction)
 			check(this->CachedExpressionData);
-			this->CachedExpressionData->PropertyConnectedBitmask &= ~ConnectionMask;
-			this->CachedExpressionData->PropertyConnectedBitmask |= StrataMaterialInfo.GetPropertyConnected();
+			this->CachedExpressionData->PropertyConnectedMask &= ~ConnectionMask;
+			this->CachedExpressionData->PropertyConnectedMask |= StrataMaterialInfo.GetPropertyConnected();
 
 			if (StrataMaterialInfo.GetStrataTreeOutOfStackDepthOccurred())
 			{
diff --git a/Engine/Source/Runtime/Engine/Private/Materials/MaterialCachedData.cpp b/Engine/Source/Runtime/Engine/Private/Materials/MaterialCachedData.cpp
--- a/Engine/Source/Runtime/Engine/Private/Materials/MaterialCachedData.cpp
+++ b/Engine/Source/Runtime/Engine/Private/Materials/MaterialCachedData.cpp
@@ -36,6 +36,7 @@
 #include "LandscapeGrassType.h"
 #include "Curves/CurveLinearColor.h"
 #include "Curves/CurveLinearColorAtlas.h"
+#include "UObject/UE5MainStreamObjectVersion.h"
 
 #include UE_INLINE_GENERATED_CPP_BY_NAME(MaterialCachedData)
 
@@ -43,7 +44,7 @@ const FMaterialCachedParameterEntry FMaterialCachedParameterEntry::EmptyData{};
 const FMaterialCachedExpressionData FMaterialCachedExpressionData::EmptyData{};
 const FMaterialCachedExpressionEditorOnlyData FMaterialCachedExpressionEditorOnlyData::EmptyData{};
 
-static_assert((uint32)(EMaterialProperty::MP_MaterialAttributes)-1 <= (8 * sizeof(FMaterialCachedExpressionData::PropertyConnectedBitmask)), "PropertyConnectedBitmask cannot contain entire EMaterialProperty enumeration.");
+static_assert((uint64)(EMaterialProperty::MP_MaterialAttributes)-1 < (8 * sizeof(FMaterialCachedExpressionData::PropertyConnectedMask)), "PropertyConnectedMask cannot contain entire EMaterialProperty enumeration.");
 
 static bool GExperimentalMaterialCachedDataAnalysisEnabled = false;
 static FAutoConsoleVariableRef CVarExperimentalMaterialCachedDataAnalysisEnabled(TEXT("r.Material.ExperimentalMaterialCachedDataAnalysisEnabled"), GExperimentalMaterialCachedDataAnalysisEnabled, TEXT("Enables material cached data experimental graph based analysis"));
@@ -1099,6 +1100,36 @@ void FMaterialCachedExpressionData::GetParameterValueByIndex(EMaterialParameterT
 	}
 }
 
+bool FStrataMaterialInfo::Serialize(FArchive& Ar)
+{
+	Ar.UsingCustomVersion(FUE5MainStreamObjectVersion::GUID);
+
+	if (Ar.IsLoading())
+	{
+		if (Ar.CustomVer(FUE5MainStreamObjectVersion::GUID) < FUE5MainStreamObjectVersion::IncreaseMaterialAttributesInputMask)
+		{
+			ConnectedPropertyMask = uint64(ConnectedProperties_DEPRECATED);
+		}
+	}
+
+	return true;
+}
+
+bool FMaterialCachedExpressionData::Serialize(FArchive& Ar)
+{
+	Ar.UsingCustomVersion(FUE5MainStreamObjectVersion::GUID);
+
+	if (Ar.IsLoading())
+	{
+		if (Ar.CustomVer(FUE5MainStreamObjectVersion::GUID) < FUE5MainStreamObjectVersion::IncreaseMaterialAttributesInputMask)
+		{
+			PropertyConnectedMask = uint64(PropertyConnectedBitmask_DEPRECATED);
+		}
+	}
+
+	return true;
+}
+
 void FMaterialCachedExpressionData::PostSerialize(const FArchive& Ar)
 {
 #if WITH_EDITORONLY_DATA
diff --git a/Engine/Source/Runtime/Engine/Private/Materials/MaterialShared.cpp b/Engine/Source/Runtime/Engine/Private/Materials/MaterialShared.cpp
--- a/Engine/Source/Runtime/Engine/Private/Materials/MaterialShared.cpp
+++ b/Engine/Source/Runtime/Engine/Private/Materials/MaterialShared.cpp
@@ -1988,7 +1988,7 @@ bool FMaterialResource::HasMaterialPropertyConnected(EMaterialProperty In) const
 			return Material->HasAmbientOcclusionConnected();
 		}
 		// Strata material traversal is cached as this is an expensive operation
-		return FStrataMaterialInfo::HasPropertyConnected(Material->GetCachedExpressionData().PropertyConnectedBitmask, In);
+		return FStrataMaterialInfo::HasPropertyConnected(Material->GetCachedExpressionData().PropertyConnectedMask, In);
 	}
 	else
 	{
diff --git a/Engine/Source/Runtime/Engine/Public/MaterialCachedData.h b/Engine/Source/Runtime/Engine/Public/MaterialCachedData.h
--- a/Engine/Source/Runtime/Engine/Public/MaterialCachedData.h
+++ b/Engine/Source/Runtime/Engine/Public/MaterialCachedData.h
@@ -248,14 +248,15 @@ struct FMaterialCachedExpressionData
 
 	bool IsPropertyConnected(EMaterialProperty Property) const
 	{
-		return ((PropertyConnectedBitmask >> (uint32)Property) & 0x1) != 0;
+		return ((PropertyConnectedMask >> (uint64)Property) & 0x1) != 0;
 	}
 
 	void SetPropertyConnected(EMaterialProperty Property)
 	{
-		PropertyConnectedBitmask |= (1 << (uint32)Property);
+		PropertyConnectedMask |= (1ull << (uint64)Property);
 	}
 
+	bool Serialize(FArchive& Ar);
 	void PostSerialize(const FArchive& Ar);
 
 #if WITH_EDITORONLY_DATA
@@ -347,9 +348,12 @@ struct FMaterialCachedExpressionData
 	UPROPERTY()
 	uint32 bHasVertexInterpolator : 1;
 
+	UPROPERTY()
+	uint32 PropertyConnectedBitmask_DEPRECATED = 0;
+
 	/** Each bit corresponds to EMaterialProperty connection status. */
 	UPROPERTY()
-	uint32 PropertyConnectedBitmask = 0;
+	uint64 PropertyConnectedMask = 0;
 
 #if WITH_EDITOR
 	/** Array of errors reporting a parameter being set multiple times to distinct values. */
@@ -357,3 +361,11 @@ struct FMaterialCachedExpressionData
 #endif
 };
 
+template<>
+struct TStructOpsTypeTraits<FMaterialCachedExpressionData> : public TStructOpsTypeTraitsBase2<FMaterialCachedExpressionData>
+{
+	enum
+	{
+		WithSerializer = true,
+	};
+};
diff --git a/Engine/Source/Runtime/Engine/Public/MaterialExpressionIO.h b/Engine/Source/Runtime/Engine/Public/MaterialExpressionIO.h
--- a/Engine/Source/Runtime/Engine/Public/MaterialExpressionIO.h
+++ b/Engine/Source/Runtime/Engine/Public/MaterialExpressionIO.h
@@ -299,27 +299,27 @@ struct TStructOpsTypeTraits<FVector2MaterialInput>
 struct FMaterialAttributesInput : FExpressionInput
 {
 	FMaterialAttributesInput() 
-	: PropertyConnectedBitmask(0)
+	: PropertyConnectedMask(0)
 	{ 
-		// ensure PropertyConnectedBitmask can contain all properties.
-		static_assert((uint32)(MP_MaterialAttributes)-1 <= (8 * sizeof(PropertyConnectedBitmask)), "PropertyConnectedBitmask cannot contain entire EMaterialProperty enumeration.");
+		// Ensure PropertyConnectedMask can contain all properties.
+		static_assert((uint64)(MP_MaterialAttributes)-1 < (8 * sizeof(PropertyConnectedMask)), "PropertyConnectedMask cannot contain entire EMaterialProperty enumeration.");
 	}
 
 #if WITH_EDITOR
 	ENGINE_API int32 CompileWithDefault(class FMaterialCompiler* Compiler, const FGuid& AttributeID);
-	inline bool IsConnected(EMaterialProperty Property) { return ((PropertyConnectedBitmask >> (uint32)Property) & 0x1) != 0; }
+	inline bool IsConnected(EMaterialProperty Property) { return ((PropertyConnectedMask >> (uint64)Property) & 0x1) != 0; }
 	inline bool IsConnected() const { return FExpressionInput::IsConnected(); }
 	inline void SetConnectedProperty(EMaterialProperty Property, bool bIsConnected)
 	{
-		PropertyConnectedBitmask = bIsConnected ? PropertyConnectedBitmask | (1 << (uint32)Property) : PropertyConnectedBitmask & ~(1 << (uint32)Property);
+		PropertyConnectedMask = bIsConnected ? PropertyConnectedMask | (1ull << (uint64)Property) : PropertyConnectedMask & ~(1ull << (uint64)Property);
 	}
 #endif  // WITH_EDITOR
 
 	/** ICPPStructOps interface */
 	ENGINE_API bool Serialize(FArchive& Ar);
 
-	// each bit corresponds to EMaterialProperty connection status.
-	uint32 PropertyConnectedBitmask;
+	// Each bit corresponds to EMaterialProperty connection status.
+	uint64 PropertyConnectedMask;
 };
 
 template<>
