Remove a stale field from material resources.

diff --git a/Engine/Source/Runtime/Engine/Public/MaterialShared.h b/Engine/Source/Runtime/Engine/Public/MaterialShared.h
--- a/Engine/Source/Runtime/Engine/Public/MaterialShared.h
+++ b/Engine/Source/Runtime/Engine/Public/MaterialShared.h
@@ -2713,7 +2713,6 @@ public:
 protected:
 	UMaterial* Material;
 	UMaterialInstance* MaterialInstance;
-	mutable FStrataMaterialInfo CachedStrataMaterialInfo;
 
 	/** Entry point for compiling a specific material property.  This must call SetMaterialProperty. */
 	ENGINE_API virtual int32 CompilePropertyAndSetMaterialProperty(EMaterialProperty Property, class FMaterialCompiler* Compiler, EShaderFrequency OverrideShaderFrequency, bool bUsePreviousFrameTime) const override;
