#UE Silence spammy message

diff --git a/Engine/Source/Runtime/Engine/Private/Animation/AnimationCompressionDerivedData.cpp b/Engine/Source/Runtime/Engine/Private/Animation/AnimationCompressionDerivedData.cpp
--- a/Engine/Source/Runtime/Engine/Private/Animation/AnimationCompressionDerivedData.cpp
+++ b/Engine/Source/Runtime/Engine/Private/Animation/AnimationCompressionDerivedData.cpp
@@ -138,7 +138,7 @@ void FAnimationSequenceAsyncCacheTask::EndCache(UE::DerivedData::FCacheGetValueR
 				FMemoryReaderView Ar(RecordData, /*bIsPersistent*/ true);
 				CompressedData->SerializeCompressedData(Ar, true, AnimSequence, AnimSequence->GetSkeleton(), CompressibleAnimPtr->BoneCompressionSettings, CompressibleAnimPtr->CurveCompressionSettings);
 
-				UE_LOG(LogAnimationCompression, Display, TEXT("Fetched compressed animation data for %s"), *CompressibleAnimPtr->FullName);
+				UE_LOG(LogAnimationCompression, Verbose, TEXT("Fetched compressed animation data for %s"), *CompressibleAnimPtr->FullName);
 				if (Compression::FAnimationCompressionMemorySummaryScope::ShouldStoreCompressionResults())
 				{
 					const double CompressionEndTime = FPlatformTime::Seconds();
