ODSC + Staged builds

The ability to launch a staged build and only use the COTF server to process `recompileshaders` commands.

- Adds a cmd line option, `-odschost=127.0.0.1` which allows the game to connect directly to the COTF server, but only to send ODSC recompileshaders commands.
- If we don't have a default COTF connection (for cooking assets), we create our own connection to the COTF server.
- Refactor code to send a command to the COTF server into a function.  That way we if we don't have a "special" ODSC only connection we can the command that way.
- Fixed bug with the `recompileshaders material` command.  It would only search and find a _single_ material and recompile it.  This ignored any material instances.  Now the code takes into account finding all dependent instances and issuing a recompile command for them.

diff --git a/Engine/Source/Runtime/Engine/Private/ODSC/ODSCManager.cpp b/Engine/Source/Runtime/Engine/Private/ODSC/ODSCManager.cpp
--- a/Engine/Source/Runtime/Engine/Private/ODSC/ODSCManager.cpp
+++ b/Engine/Source/Runtime/Engine/Private/ODSC/ODSCManager.cpp
@@ -15,10 +15,13 @@ FODSCManager* GODSCManager = nullptr;
 FODSCManager::FODSCManager()
 	: FTSTickerObjectBase(0.0f, FTSBackgroundableTicker::GetCoreTicker())
 {
-	if (IsRunningCookOnTheFly())
+	FString Host;
+	const bool bODSCEnabled = FParse::Value(FCommandLine::Get(), TEXT("-odschost="), Host);
+
+	if (IsRunningCookOnTheFly() || bODSCEnabled)
 	{
 		FCoreDelegates::OnEnginePreExit.AddRaw(this, &FODSCManager::OnEnginePreExit);
-		Thread = new FODSCThread();
+		Thread = new FODSCThread(Host);
 		Thread->StartThread();
 	}
 }
diff --git a/Engine/Source/Runtime/Engine/Private/ODSC/ODSCThread.cpp b/Engine/Source/Runtime/Engine/Private/ODSC/ODSCThread.cpp
--- a/Engine/Source/Runtime/Engine/Private/ODSC/ODSCThread.cpp
+++ b/Engine/Source/Runtime/Engine/Private/ODSC/ODSCThread.cpp
@@ -1,8 +1,10 @@
 // Copyright Epic Games, Inc. All Rights Reserved.
 
 #include "ODSCThread.h"
+#include "CookOnTheFly.h"
 #include "ODSCLog.h"
 #include "HAL/FileManager.h"
+#include "Modules/ModuleManager.h"
 
 FODSCRequestPayload::FODSCRequestPayload(
 	EShaderPlatform InShaderPlatform,
@@ -98,11 +100,25 @@ bool FODSCMessageHandler::ReloadGlobalShaders() const
 	return RecompileCommandType == ODSCRecompileCommand::Global;
 }
 
-FODSCThread::FODSCThread()
+FODSCThread::FODSCThread(const FString& HostIP)
 	: Thread(nullptr),
 	  WakeupEvent(FPlatformProcess::GetSynchEventFromPool(true))
 {
 	UE_LOG(LogODSC, Log, TEXT("ODSC Thread active."));
+
+	// Attempt to get a default connection to the COTF server (which cooks assets).
+	UE::Cook::ICookOnTheFlyModule& CookOnTheFlyModule = FModuleManager::LoadModuleChecked<UE::Cook::ICookOnTheFlyModule>(TEXT("CookOnTheFly"));
+	if (!CookOnTheFlyModule.GetDefaultServerConnection())
+	{
+		// If we don't have a default connection make a specific connection to the HostIP provided.
+		UE::Cook::FCookOnTheFlyHostOptions CookOnTheFlyHostOptions;
+		CookOnTheFlyHostOptions.Hosts.Add(HostIP);
+		CookOnTheFlyServerConnection = CookOnTheFlyModule.ConnectToServer(CookOnTheFlyHostOptions);
+		if (!CookOnTheFlyServerConnection)
+		{
+			UE_LOG(LogODSC, Warning, TEXT("Failed to connect to cook on the fly server."));
+		}
+	}
 }
 
 FODSCThread::~FODSCThread()
@@ -241,7 +257,7 @@ void FODSCThread::Process()
 	for (FODSCMessageHandler* NextRequest : RequestsToStart)
 	{
 		// send the info, the handler will process the response (and update shaders, etc)
-		IFileManager::Get().SendMessageToServer(TEXT("RecompileShaders"), NextRequest);
+		SendMessageToServer(NextRequest);
 
 		CompletedThreadedRequests.Enqueue(NextRequest);
 	}
@@ -256,10 +272,36 @@ void FODSCThread::Process()
 		}
 
 		// send the info, the handler will process the response (and update shaders, etc)
-		IFileManager::Get().SendMessageToServer(TEXT("RecompileShaders"), RequestHandler);
+		SendMessageToServer(RequestHandler);
 
 		CompletedThreadedRequests.Enqueue(RequestHandler);
 	}
 
 	WakeupEvent->Reset();
 }
+
+void FODSCThread::SendMessageToServer(IPlatformFile::IFileServerMessageHandler* Handler)
+{
+	// If we have a default connection that already exists, send directly to that.
+	if ((CookOnTheFlyServerConnection == nullptr) || (!CookOnTheFlyServerConnection->IsConnected()))
+	{
+		IFileManager::Get().SendMessageToServer(TEXT("RecompileShaders"), Handler);
+		return;
+	}
+
+	// We don't have a default COTF connection so use our specific connection to send our command.
+	UE::Cook::FCookOnTheFlyRequest Request(UE::Cook::ECookOnTheFlyMessage::RecompileShaders);
+	{
+		TUniquePtr<FArchive> Ar = Request.WriteBody();
+		Handler->FillPayload(*Ar);
+	}
+
+	UE::Cook::FCookOnTheFlyResponse Response = CookOnTheFlyServerConnection->SendRequest(Request).Get();
+	if (Response.IsOk())
+	{
+		TUniquePtr<FArchive> Ar = Response.ReadBody();
+		Handler->ProcessResponse(*Ar);
+	}
+
+	check(Response.IsOk());
+}
diff --git a/Engine/Source/Runtime/Engine/Private/ODSC/ODSCThread.h b/Engine/Source/Runtime/Engine/Private/ODSC/ODSCThread.h
--- a/Engine/Source/Runtime/Engine/Private/ODSC/ODSCThread.h
+++ b/Engine/Source/Runtime/Engine/Private/ODSC/ODSCThread.h
@@ -13,6 +13,15 @@
 class FEvent;
 class FRunnableThread;
 
+namespace UE
+{
+	namespace Cook
+	{
+		class ICookOnTheFlyServerConnection;
+		class FCookOnTheFlyMessage;
+	}
+}
+
 class FODSCMessageHandler : public IPlatformFile::IFileServerMessageHandler
 {
 public:
@@ -72,7 +81,7 @@ class FODSCThread
 {
 public:
 
-	FODSCThread();
+	FODSCThread(const FString& HostIP);
 	virtual ~FODSCThread();
 
 	/**
@@ -188,4 +197,9 @@ private:
 
 	/** Holds an event signaling the thread to wake up. */
 	FEvent* WakeupEvent;
+
+	void SendMessageToServer(IPlatformFile::IFileServerMessageHandler* Handler);
+
+	/** Special connection to the cooking server.  This is only used to send recompileshaders commands on. */
+	TUniquePtr<UE::Cook::ICookOnTheFlyServerConnection> CookOnTheFlyServerConnection;
 };
diff --git a/Engine/Source/Runtime/Engine/Private/ShaderCompiler/ShaderCompiler.cpp b/Engine/Source/Runtime/Engine/Private/ShaderCompiler/ShaderCompiler.cpp
--- a/Engine/Source/Runtime/Engine/Private/ShaderCompiler/ShaderCompiler.cpp
+++ b/Engine/Source/Runtime/Engine/Private/ShaderCompiler/ShaderCompiler.cpp
@@ -24,6 +24,7 @@
 #include "Internationalization/LocKeyFuncs.h"
 #include "MaterialShared.h"
 #include "Materials/Material.h"
+#include "Materials/MaterialInstance.h"
 #include "Misc/Compression.h"
 #include "Misc/ConfigCacheIni.h"
 #include "Misc/CoreDelegates.h"
@@ -6763,17 +6764,31 @@ namespace
 
 			// tell other side the material to load, by pathname
 			FString RequestedMaterialName( FParse::Token( CmdString, 0 ) );
-
-			for( TObjectIterator<UMaterialInterface> It; It; ++It )
+			UMaterialInterface* MatchingMaterial = nullptr;
+			for (TObjectIterator<UMaterialInterface> It; It; ++It)
 			{
 				UMaterial* Material = It->GetMaterial();
 
-				if( Material && Material->GetName() == RequestedMaterialName)
+				if (Material && Material->GetName() == RequestedMaterialName)
 				{
 					OutMaterialsToLoad.Add(It->GetPathName());
+					MatchingMaterial = Material;
 					break;
 				}
 			}
+
+			// Find all material instances from the requested material and 
+			// request a compile for them.
+			if (MatchingMaterial)
+			{
+				for (TObjectIterator<UMaterialInstance> It; It; ++It)
+				{
+					if (It && It->IsDependent(MatchingMaterial))
+					{
+						OutMaterialsToLoad.Add(It->GetPathName());
+					}
+				}
+			}
 		}
 		else if (!CmdName.IsEmpty() && FCString::Stricmp(*CmdName, TEXT("Global")) == 0)
 		{
