InitWorld: When calling CleanupWorld and InitWorld during a map build, we need to also refresh streaming levels, or they will be removed from world and will fail to render and give errors in the lighting build.
#rnx

diff --git a/Engine/Source/Editor/UnrealEd/Private/UnrealEdMisc.cpp b/Engine/Source/Editor/UnrealEd/Private/UnrealEdMisc.cpp
--- a/Engine/Source/Editor/UnrealEd/Private/UnrealEdMisc.cpp
+++ b/Engine/Source/Editor/UnrealEd/Private/UnrealEdMisc.cpp
@@ -1214,8 +1214,7 @@ void FUnrealEdMisc::CB_MapChange( uint32 InFlags )
 			if (World->IsInitialized())
 			{
 				// Call CleanupWorld/InitWorld to detach components and some other stuff
-				World->CleanupWorld();
-				World->InitWorld();
+				World->ReInitWorld();
 			}
 		}
 
diff --git a/Engine/Source/Runtime/Engine/Classes/Engine/World.h b/Engine/Source/Runtime/Engine/Classes/Engine/World.h
--- a/Engine/Source/Runtime/Engine/Classes/Engine/World.h
+++ b/Engine/Source/Runtime/Engine/Classes/Engine/World.h
@@ -3079,6 +3079,12 @@ public:
 	bool IsInitializedAndNeedsCleanup() const { return bIsWorldInitialized; }
 	/** Returns whether InitWorld has ever been called since this World was created.  */
 	bool HasEverBeenInitialized() const { return bHasEverBeenInitialized; }
+
+	/**
+	 * Calls CleanupWorld and InitWorld, while handling preservation of StreamingLevels and WorldInitialization values
+	 * Gives a warning and does nothing if IsInitialized is not currently true.
+	 */
+	void ReInitWorld();
 #endif
 
 	/** Returns whether InitWorld has been called without yet calling CleanupWorld.  */
diff --git a/Engine/Source/Runtime/Engine/Private/World.cpp b/Engine/Source/Runtime/Engine/Private/World.cpp
--- a/Engine/Source/Runtime/Engine/Private/World.cpp
+++ b/Engine/Source/Runtime/Engine/Private/World.cpp
@@ -2177,6 +2177,35 @@ void UWorld::InitWorld(const InitializationValues IVS)
 }
 
 #if WITH_EDITOR
+void UWorld::ReInitWorld()
+{
+	if (!IsInitialized())
+	{
+		UE_LOG(LogWorld, Warning, TEXT("ReInitWorld called on World %s, but it is not yet initialized. ReInitWorld will be ignored."), *GetPathName());
+		return;
+	}
+
+	InitializationValues IVS;
+	IVS.bInitializeScenes = Scene != nullptr;
+	IVS.bAllowAudioPlayback = bAllowAudioPlayback;
+	IVS.bRequiresHitProxies = bRequiresHitProxies;
+	IVS.bCreatePhysicsScene = PhysicsScene != nullptr;
+	IVS.bCreateNavigation = NavigationSystem != nullptr;
+	IVS.bCreateAISystem = AISystem != nullptr;
+	IVS.bShouldSimulatePhysics = bShouldSimulatePhysics;
+	IVS.bEnableTraceCollision = bEnableTraceCollision;
+	IVS.bForceUseMovementComponentInNonGameWorld = bForceUseMovementComponentInNonGameWorld;
+	IVS.bTransactional = HasAnyFlags(RF_Transactional);
+	IVS.bCreateFXSystem = FXSystem != nullptr;
+	// IVS.bCreateWorldPartition; // Only used in InitializeNewWorld, not needed by InitWorld
+	// IVS.DefaultGameMode; // Only used inInitializeNewWorld, not needed by InitWorld
+
+	CleanupWorld();
+	InitWorld(IVS);
+	RefreshStreamingLevels();
+	UpdateWorldComponents(false /* bRerunConstructionScripts */, false /* bCurrentLevelOnly */);
+}
+
 const FName UWorld::KeepInitializedDuringLoadTag(TEXT("KeepInitializedDuringLoadTag"));
 #endif
 
