Cleaned up some functions in UAnimatedSparseVolumeTextureController, made some private members uproperties (which should hopefully fix a crash caused by a uobject ptr pointing to invalid data) and added GetFrameByIndex(int32). Added a guard against empty Frames array in UStreamableSparseVolumeTexture::GetStreamedFrameProxyOrFallback(). Silenced log message about the efficiency of the tile deduplication during SVT cooking.

#rnx

diff --git a/Engine/Source/Runtime/Engine/Classes/SparseVolumeTexture/SparseVolumeTexture.h b/Engine/Source/Runtime/Engine/Classes/SparseVolumeTexture/SparseVolumeTexture.h
--- a/Engine/Source/Runtime/Engine/Classes/SparseVolumeTexture/SparseVolumeTexture.h
+++ b/Engine/Source/Runtime/Engine/Classes/SparseVolumeTexture/SparseVolumeTexture.h
@@ -272,6 +272,15 @@ class ENGINE_API UAnimatedSparseVolumeTextureController : public UObject
 
 public:
 
+	UPROPERTY(BlueprintReadWrite, Category = "Animation")
+	TObjectPtr<USparseVolumeTexture> SparseVolumeTexture;
+	
+	UPROPERTY(BlueprintReadWrite, Category = "Animation")
+	float Time;
+
+	UPROPERTY(BlueprintReadWrite, Category = "Animation")
+	bool bIsPlaying;
+
 	UPROPERTY(BlueprintReadWrite, Category = "Animation")
 	float FrameRate = 24.0f;
 
@@ -290,41 +299,21 @@ public:
 	UFUNCTION(BlueprintCallable, Category = "Animation")
 	void Stop();
 
-	UFUNCTION(BlueprintCallable, Category = "Animation")
-	bool IsPlaying();
-
 	UFUNCTION(BlueprintCallable, Category = "Animation")
 	void Update(float DeltaTime);
 
 	UFUNCTION(BlueprintCallable, Category = "Animation")
-	void SetSparseVolumeTexture(USparseVolumeTexture* Texture);
-
-	UFUNCTION(BlueprintCallable, Category = "Animation")
-	void SetTime(float Time);
-
-	UFUNCTION(BlueprintCallable, Category = "Animation")
-	void SetFractionalFrameIndex(float Frame);
-
-	UFUNCTION(BlueprintCallable, Category = "Animation")
-	USparseVolumeTexture* GetSparseVolumeTexture();
-
-	UFUNCTION(BlueprintCallable, Category = "Animation")
-	float GetTime();
+	float GetFractionalFrameIndex();
 
 	UFUNCTION(BlueprintCallable, Category = "Animation")
-	float GetFractionalFrameIndex();
+	USparseVolumeTextureFrame* GetFrameByIndex(int32 FrameIndex);
 
 	UFUNCTION(BlueprintCallable, Category = "Animation")
 	USparseVolumeTextureFrame* GetCurrentFrame();
 
 	UFUNCTION(BlueprintCallable, Category = "Animation")
-	void GetLerpFrames(USparseVolumeTextureFrame*& Frame0, USparseVolumeTextureFrame*& Frame1, float& LerpAlpha);
+	void GetCurrentFramesForInterpolation(USparseVolumeTextureFrame*& Frame0, USparseVolumeTextureFrame*& Frame1, float& LerpAlpha);
 
 	UFUNCTION(BlueprintCallable, Category = "Animation")
 	float GetDuration();
-
-private:
-	TObjectPtr<USparseVolumeTexture> SparseVolumeTexture;
-	float Time;
-	bool bIsPlaying;
 };
diff --git a/Engine/Source/Runtime/Engine/Private/SparseVolumeTexture/SparseVolumeTexture.cpp b/Engine/Source/Runtime/Engine/Private/SparseVolumeTexture/SparseVolumeTexture.cpp
--- a/Engine/Source/Runtime/Engine/Private/SparseVolumeTexture/SparseVolumeTexture.cpp
+++ b/Engine/Source/Runtime/Engine/Private/SparseVolumeTexture/SparseVolumeTexture.cpp
@@ -467,6 +467,12 @@ void UStreamableSparseVolumeTexture::GetResourceSizeEx(FResourceSizeEx& Cumulati
 
 const FSparseVolumeTextureSceneProxy* UStreamableSparseVolumeTexture::GetStreamedFrameProxyOrFallback(int32 FrameIndex, int32 MipLevel) const
 {
+	if (Frames.IsEmpty())
+	{
+		UE_LOG(LogSparseVolumeTexture, Warning, TEXT("SVT is empty and has no frames at all! SVT: %s"), *GetName());
+		return nullptr;
+	}
+	FrameIndex = FMath::Clamp(FrameIndex, 0, Frames.Num() - 1);
 	ISparseVolumeTextureStreamingManager& StreamingManager = IStreamingManager::Get().GetSparseVolumeTextureStreamingManager();
 	const FSparseVolumeTextureSceneProxy* Proxy = StreamingManager.GetSparseVolumeTextureSceneProxy(this, FrameIndex, MipLevel, true);
 
@@ -754,14 +760,9 @@ void UAnimatedSparseVolumeTextureController::Stop()
 	}
 }
 
-bool UAnimatedSparseVolumeTextureController::IsPlaying()
-{
-	return bIsPlaying;
-}
-
 void UAnimatedSparseVolumeTextureController::Update(float DeltaTime)
 {
-	if (!SparseVolumeTexture || !bIsPlaying)
+	if (!IsValid(SparseVolumeTexture) || !bIsPlaying)
 	{
 		return;
 	}
@@ -771,61 +772,34 @@ void UAnimatedSparseVolumeTextureController::Update(float DeltaTime)
 	Time = FMath::Fmod(Time + DeltaTime, AnimationDuration + UE_SMALL_NUMBER);
 }
 
-void UAnimatedSparseVolumeTextureController::SetSparseVolumeTexture(USparseVolumeTexture* Texture)
-{
-	if (Texture == SparseVolumeTexture)
-	{
-		return;
-	}
-
-	SparseVolumeTexture = Texture;
-	bIsPlaying = bIsPlaying && (SparseVolumeTexture != nullptr);
-	Time = 0.0f;
-}
-
-void UAnimatedSparseVolumeTextureController::SetTime(float InTime)
-{
-	const float AnimationDuration = GetDuration();
-	Time = FMath::Fmod(InTime, AnimationDuration + UE_SMALL_NUMBER);
-}
-
-void UAnimatedSparseVolumeTextureController::SetFractionalFrameIndex(float Frame)
+float UAnimatedSparseVolumeTextureController::GetFractionalFrameIndex()
 {
-	if (!SparseVolumeTexture)
+	if (!IsValid(SparseVolumeTexture))
 	{
-		return;
+		return 0.0f;
 	}
 
 	const int32 FrameCount = SparseVolumeTexture->GetNumFrames();
-	Frame = FMath::Fmod(Frame, (float)FrameCount);
-	Time = Frame / (FrameRate + UE_SMALL_NUMBER);
-}
-
-USparseVolumeTexture* UAnimatedSparseVolumeTextureController::GetSparseVolumeTexture()
-{
-	return SparseVolumeTexture;
-}
-
-float UAnimatedSparseVolumeTextureController::GetTime()
-{
-	return Time;
+	const float FrameIndexF = FMath::Fmod(Time * FrameRate, (float)FrameCount);
+	return FrameIndexF;
 }
 
-float UAnimatedSparseVolumeTextureController::GetFractionalFrameIndex()
+USparseVolumeTextureFrame* UAnimatedSparseVolumeTextureController::GetFrameByIndex(int32 FrameIndex)
 {
-	if (!SparseVolumeTexture)
+	if (!IsValid(SparseVolumeTexture))
 	{
-		return 0.0f;
+		return nullptr;
 	}
 
-	const int32 FrameCount = SparseVolumeTexture->GetNumFrames();
-	const float FrameIndexF = FMath::Fmod(Time * FrameRate, (float)FrameCount);
-	return FrameIndexF;
+	// Create and initialize a USparseVolumeTextureFrame which holds the frame to sample and can be bound to shaders
+	USparseVolumeTextureFrame* Frame = USparseVolumeTextureFrame::CreateFrame(SparseVolumeTexture, FrameIndex, MipLevel);
+
+	return Frame;
 }
 
 USparseVolumeTextureFrame* UAnimatedSparseVolumeTextureController::GetCurrentFrame()
 {
-	if (!SparseVolumeTexture)
+	if (!IsValid(SparseVolumeTexture))
 	{
 		return nullptr;
 	}
@@ -840,9 +814,9 @@ USparseVolumeTextureFrame* UAnimatedSparseVolumeTextureController::GetCurrentFra
 	return Frame;
 }
 
-void UAnimatedSparseVolumeTextureController::GetLerpFrames(USparseVolumeTextureFrame*& Frame0, USparseVolumeTextureFrame*& Frame1, float& LerpAlpha)
+void UAnimatedSparseVolumeTextureController::GetCurrentFramesForInterpolation(USparseVolumeTextureFrame*& Frame0, USparseVolumeTextureFrame*& Frame1, float& LerpAlpha)
 {
-	if (!SparseVolumeTexture)
+	if (!IsValid(SparseVolumeTexture))
 	{
 		return;
 	}
@@ -859,7 +833,7 @@ void UAnimatedSparseVolumeTextureController::GetLerpFrames(USparseVolumeTextureF
 
 float UAnimatedSparseVolumeTextureController::GetDuration()
 {
-	if (!SparseVolumeTexture)
+	if (!IsValid(SparseVolumeTexture))
 	{
 		return 0.0f;
 	}
diff --git a/Engine/Source/Runtime/Engine/Private/SparseVolumeTexture/SparseVolumeTextureData.cpp b/Engine/Source/Runtime/Engine/Private/SparseVolumeTexture/SparseVolumeTextureData.cpp
--- a/Engine/Source/Runtime/Engine/Private/SparseVolumeTexture/SparseVolumeTextureData.cpp
+++ b/Engine/Source/Runtime/Engine/Private/SparseVolumeTexture/SparseVolumeTextureData.cpp
@@ -626,9 +626,11 @@ bool FSparseVolumeTextureData::DeduplicateTiles()
 		}
 	}
 
+#if 0 // Disabled in order to not spam the log when importing long sequences
 	const int32 NumRemovedTiles = NumProcessedTiles - NumSurvivingTiles;
 	const float RemovedTilesPercentage = ((float)NumRemovedTiles / (float)NumProcessedTiles) * 100.0f;
 	UE_LOG(LogSparseVolumeTextureData, Display, TEXT("SparseVolumeTexture tile deduplication removed %i tiles out of %i (%f%%)"), NumRemovedTiles, NumProcessedTiles, RemovedTilesPercentage);
+#endif
 
 	return true;
 }
