Divide by zero build fix (Shipping)

#rnx

diff --git a/Engine/Source/Runtime/Engine/Private/Collision/SceneQueryLowLevel.cpp b/Engine/Source/Runtime/Engine/Private/Collision/SceneQueryLowLevel.cpp
--- a/Engine/Source/Runtime/Engine/Private/Collision/SceneQueryLowLevel.cpp
+++ b/Engine/Source/Runtime/Engine/Private/Collision/SceneQueryLowLevel.cpp
@@ -31,28 +31,6 @@ FAutoConsoleVariableRef CVarSerializeBadSweeps(TEXT("p.SerializeBadSQs"), Serial
 FAutoConsoleVariableRef CVarSerializeSQsRaycastEnabled(TEXT("p.SerializeSQsRaycastEnabled"), EnableRaycastSQCapture, TEXT("If disabled, p.SerializeSQs will not consider raycasts"));
 FAutoConsoleVariableRef CVarSerializeSQsOverlapEnabled(TEXT("p.SerializeSQsOverlapEnabled"), EnableOverlapSQCapture, TEXT("If disabled, p.SerializeSQs will not consider overlaps"));
 FAutoConsoleVariableRef CVarSerializeSQsSweepEnabled(TEXT("p.SerializeSQsSweepEnabled"), EnableSweepSQCapture, TEXT("If disabled, p.SerializeSQs will not consider sweeps"));
-
-void FinalizeCapture(FPhysTestSerializer& Serializer)
-{
-	if (SerializeSQs)
-	{
-		Serializer.Serialize(TEXT("SQCapture"));
-	}
-#if 0
-	if (ReplaySQs)
-	{
-		const bool bReplaySuccess = SQComparisonHelper(Serializer);
-		if (!bReplaySuccess)
-		{
-			UE_LOG(LogPhysicsCore, Warning, TEXT("Chaos SQ does not match physx"));
-			if (SerializeBadSQs && !SerializeSQs)
-			{
-				Serializer.Serialize(TEXT("BadSQCapture"));
-			}
-		}
-	}
-#endif
-}
 #else
 constexpr int32 SerializeSQs = 0;
 constexpr int32 ReplaySQs = 0;
@@ -60,15 +38,37 @@ constexpr int32 SerializeSQSamples = 0;
 constexpr int32 EnableRaycastSQCapture = 0;
 constexpr int32 EnableOverlapSQCapture = 0;
 constexpr int32 EnableSweepSQCapture = 0;
-
-// No-op in shipping
-void FinalizeCapture(FPhysTestSerializer& Serializer) {}
 #endif
 
 namespace
 {
+	void FinalizeCapture(FPhysTestSerializer& Serializer)
+	{
+#if !UE_BUILD_SHIPPING
+		if (SerializeSQs)
+		{
+			Serializer.Serialize(TEXT("SQCapture"));
+		}
+#if 0
+		if (ReplaySQs)
+		{
+			const bool bReplaySuccess = SQComparisonHelper(Serializer);
+			if (!bReplaySuccess)
+			{
+				UE_LOG(LogPhysicsCore, Warning, TEXT("Chaos SQ does not match physx"));
+				if (SerializeBadSQs && !SerializeSQs)
+				{
+					Serializer.Serialize(TEXT("BadSQCapture"));
+				}
+			}
+		}
+#endif
+#endif
+	}
+
 	void SweepSQCaptureHelper(float QueryDurationSeconds, const FChaosSQAccelerator& SQAccelerator, FPhysScene& Scene, const FPhysicsGeometry& QueryGeom, const FTransform& StartTM, const FVector& Dir, float DeltaMag, const FPhysicsHitCallback<FHitSweep>& HitBuffer, EHitFlags OutputFlags, FQueryFlags QueryFlags, const FCollisionFilterData& Filter, const ChaosInterface::FQueryFilterData& QueryFilterData, ICollisionQueryFilterCallbackBase* QueryCallback, const ChaosInterface::FQueryDebugParams& DebugParams)
 	{
+#if !UE_BUILD_SHIPPING
 		float QueryDurationMicro = QueryDurationSeconds * 1000.0 * 1000.0;
 		if (((SerializeSQs && QueryDurationMicro > SerializeSQs)) && IsInGameThread())
 		{
@@ -99,10 +99,12 @@ namespace
 				FinalizeCapture(Serializer);
 			}
 		}
+#endif
 	}
 
 	void RaycastSQCaptureHelper(float QueryDurationSeconds, const FChaosSQAccelerator& SQAccelerator, FPhysScene& Scene, const FVector& Start, const FVector& Dir, float DeltaMag, const FPhysicsHitCallback<FHitRaycast>& HitBuffer, EHitFlags OutputFlags, FQueryFlags QueryFlags, const FCollisionFilterData& Filter, const ChaosInterface::FQueryFilterData& QueryFilterData, ICollisionQueryFilterCallbackBase* QueryCallback, const ChaosInterface::FQueryDebugParams& DebugParams)
 	{
+#if !UE_BUILD_SHIPPING
 		float QueryDurationMicro = QueryDurationSeconds * 1000.0 * 1000.0;
 		if (((!!SerializeSQs && QueryDurationMicro > SerializeSQs)) && IsInGameThread())
 		{
@@ -133,10 +135,12 @@ namespace
 				FinalizeCapture(Serializer);
 			}
 		}
+#endif
 	}
 
 	void OverlapSQCaptureHelper(float QueryDurationSeconds, const FChaosSQAccelerator& SQAccelerator, FPhysScene& Scene, const FPhysicsGeometry& QueryGeom, const FTransform& GeomPose, const FPhysicsHitCallback<FHitOverlap>& HitBuffer, FQueryFlags QueryFlags, const FCollisionFilterData& Filter, const ChaosInterface::FQueryFilterData& QueryFilterData, ICollisionQueryFilterCallbackBase* QueryCallback, const ChaosInterface::FQueryDebugParams& DebugParams)
 	{
+#if !UE_BUILD_SHIPPING
 		float QueryDurationMicro = QueryDurationSeconds * 1000.0 * 1000.0;
 		if (((!!SerializeSQs && QueryDurationMicro > SerializeSQs)) && IsInGameThread())
 		{
@@ -167,6 +171,7 @@ namespace
 				FinalizeCapture(Serializer);
 			}
 		}
+#endif
 	}
 }
 
