[shaders]
- add job cache validation functionality, enabled via cvar (will always submit all jobs for compilation even if they would be cache hits, and compare the results to that of the cache to ensure they are binary identical)
- fix poor cache hit rate due to debug group name being included in input hash (this seems to have always been the case but the hit rate was worsened by the various fixes for debug info collisions)

diff --git a/Engine/Source/Runtime/Engine/Private/ShaderCompiler/ShaderCompiler.cpp b/Engine/Source/Runtime/Engine/Private/ShaderCompiler/ShaderCompiler.cpp
--- a/Engine/Source/Runtime/Engine/Private/ShaderCompiler/ShaderCompiler.cpp
+++ b/Engine/Source/Runtime/Engine/Private/ShaderCompiler/ShaderCompiler.cpp
@@ -87,6 +87,13 @@ static FAutoConsoleVariableRef CVarShaderCompilerJobCache(
 	ECVF_Default
 );
 
+static TAutoConsoleVariable<bool> CVarShaderCompilerDebugValidateJobCache(
+	TEXT("r.ShaderCompiler.DebugValidateJobCache"),
+	false,
+	TEXT("Enables debug mode for job cache which will fully execute all jobs and validate that job outputs with matching input hashes match."),
+	ECVF_Default
+);
+
 int32 GShaderCompilerMaxJobCacheMemoryMB = 16LL * 1024LL;
 static FAutoConsoleVariableRef CVarShaderCompilerMaxJobCacheMemoryMB(
 	TEXT("r.ShaderCompiler.MaxJobCacheMemoryMB"),
@@ -309,6 +316,11 @@ namespace ShaderCompiler
 		return GShaderCompilerJobCache != 0;
 	}
 
+	bool IsJobCacheDebugValidateEnabled()
+	{
+		return IsJobCacheEnabled() && CVarShaderCompilerDebugValidateJobCache.GetValueOnAnyThread();
+	}
+
 	bool IsRemoteCompilingAllowed()
 	{
 		// commandline switches override the CVars
@@ -583,7 +595,8 @@ void FShaderCompileJobCollection::SubmitJobs(const TArray<FShaderCommonCompileJo
 
 				const int32 PriorityIndex = (int32)Job->Priority;
 				bool bNewJob = true;
-				if (ShaderCompiler::IsJobCacheEnabled())
+				// check caches unless we're running in validation mode (which runs _all_ jobs and compares hashes of outputs)
+				if (ShaderCompiler::IsJobCacheEnabled() && !ShaderCompiler::IsJobCacheDebugValidateEnabled())
 				{
 					const FShaderCommonCompileJob::FInputHash& InputHash = Job->GetInputHash();
 
@@ -8809,13 +8822,46 @@ void FShaderJobCache::AddJobOutput(const FShaderCommonCompileJob* FinishedJob, c
 	}
 
 	FJobOutputHash* ExistingOutputHash = InputHashToOutput.Find(Hash);
-	if (ExistingOutputHash)
+	if (ExistingOutputHash && !ShaderCompiler::IsJobCacheDebugValidateEnabled())
 	{
 		return;
 	}
 
 	FJobOutputHash OutputHash = FBlake3::HashBuffer(Contents.GetData(), Contents.GetSize());
 
+	if (ExistingOutputHash && ShaderCompiler::IsJobCacheDebugValidateEnabled())
+	{
+		if (OutputHash != *ExistingOutputHash)
+		{
+			FString OriginalFilename = 
+				GShaderCompilingManager->GetAbsoluteShaderDebugInfoDirectory() / 
+				TEXT("CacheMismatches") / 
+				FString::Printf(TEXT("shaderoutput-%s.orig"), *LexToString(Hash));
+			FString NewFilename = 
+				GShaderCompilingManager->GetAbsoluteShaderDebugInfoDirectory() / 
+				TEXT("CacheMismatches") / 
+				FString::Printf(TEXT("shaderoutput-%s.mismatch"), *LexToString(Hash));
+
+			UE_LOG(
+				LogShaderCompilers,
+				Warning,
+				TEXT("Job cache validation found output mismatch. Mismatching outputs dumped to %s and %s"),
+				*OriginalFilename, *NewFilename);
+
+			FFileHelper::SaveArrayToFile(
+				TArrayView<const uint8>((const uint8*)Contents.GetData(), Contents.GetSize()), 
+				*NewFilename);
+
+			FStoredOutput** StoredOutput = Outputs.Find(*ExistingOutputHash);
+			check(StoredOutput);
+
+			FFileHelper::SaveArrayToFile(
+				TArrayView<const uint8>((const uint8*)(*StoredOutput)->JobOutput.GetData(), (*StoredOutput)->JobOutput.GetSize()),
+				*OriginalFilename);
+		}
+		return;
+	}
+
 	const bool bDumpCachedDebugInfo = CVarDumpShaderOutputCacheHits.GetValueOnAnyThread();
 
 	// Get dump shader debug output path
diff --git a/Engine/Source/Runtime/RenderCore/Private/ShaderCore.cpp b/Engine/Source/Runtime/RenderCore/Private/ShaderCore.cpp
--- a/Engine/Source/Runtime/RenderCore/Private/ShaderCore.cpp
+++ b/Engine/Source/Runtime/RenderCore/Private/ShaderCore.cpp
@@ -2657,7 +2657,10 @@ FShaderCommonCompileJob::FInputHash FShaderCompileJob::GetInputHash()
 	{
 		checkf(Archive.IsSaving() && !Archive.IsLoading(), TEXT("A loading archive is passed to FShaderCompileJob::GetInputHash(), this is not supported as it may corrupt its data"));
 
+		// Don't include debug group name in the hashing; this drastically worsens our cache hit rate
+		FString DebugGroupNameTmp(MoveTemp(Input.DebugGroupName));
 		Archive << Input;
+		Input.DebugGroupName = MoveTemp(DebugGroupNameTmp);
 		Input.Environment.SerializeEverythingButFiles(Archive);
 
 		// hash the source file so changes to files during the development are picked up
