- Add velocity based motion vector calculation to Cascade sprite / gpu sprite / mesh renderers
- Enabled by default, can be disabled using 'fx.Cascade.UseVelocityForMotionBlur' or by using the control on the required module for each emitter
- Mesh renderers will prefer accurate calculations if that is enabled before failing back to this method

diff --git a/Engine/Shaders/Private/MeshParticleVertexFactory.ush b/Engine/Shaders/Private/MeshParticleVertexFactory.ush
--- a/Engine/Shaders/Private/MeshParticleVertexFactory.ush
+++ b/Engine/Shaders/Private/MeshParticleVertexFactory.ush
@@ -430,7 +430,7 @@ float4 CalcWorldPosition(FVertexFactoryInput Input)
 	return float4(WorldPosition, Input.Position.w);
 }
 
-float4 CalcPrevWorldPosition(FVertexFactoryInput Input)
+float4 CalcPrevWorldPosition(FVertexFactoryInput Input, FVertexFactoryIntermediates Intermediates)
 {
 #if USE_PREVIOUS_TRANSFORM
 	if (MeshParticleVF.PrevTransformAvailable > 0)
@@ -442,7 +442,9 @@ float4 CalcPrevWorldPosition(FVertexFactoryInput Input)
 	else
 #endif
 	{
-		return CalcWorldPosition(Input);
+		float4 Position = CalcWorldPosition(Input);
+		Position.xyz -= Intermediates.ParticleVelocity.xyz * Intermediates.ParticleVelocity.w * MeshParticleVF.UseVelocityForMotionBlur;
+		return Position;
 	}
 }
 
@@ -666,7 +668,7 @@ FVertexFactoryInterpolantsVSToPS VertexFactoryGetInterpolantsVSToPS(FVertexFacto
 // @return previous translated world position
 float4 VertexFactoryGetPreviousWorldPosition(FVertexFactoryInput Input, FVertexFactoryIntermediates Intermediates)
 {	
-	return CalcPrevWorldPosition(Input);
+	return CalcPrevWorldPosition(Input, Intermediates);
 }
 
 #if NEEDS_VERTEX_FACTORY_INTERPOLATION
diff --git a/Engine/Shaders/Private/ParticleGPUSpriteVertexFactory.ush b/Engine/Shaders/Private/ParticleGPUSpriteVertexFactory.ush
--- a/Engine/Shaders/Private/ParticleGPUSpriteVertexFactory.ush
+++ b/Engine/Shaders/Private/ParticleGPUSpriteVertexFactory.ush
@@ -690,7 +690,9 @@ FVertexFactoryInterpolantsVSToPS VertexFactoryGetInterpolantsVSToPS( FVertexFact
  */
 float4 VertexFactoryGetPreviousWorldPosition( FVertexFactoryInput Input, FVertexFactoryIntermediates Intermediates )
 {
-	return float4(Intermediates.VertexWorldPosition,1);
+	float3 Position = Intermediates.VertexWorldPosition;
+	Position.xyz -= Intermediates.ParticleVelocity * EmitterUniforms.UseVelocityForMotionBlur;
+	return float4(Position, 1);
 }
 
 float4 VertexFactoryGetTranslatedPrimitiveVolumeBounds(FVertexFactoryInterpolantsVSToPS Interpolants)
diff --git a/Engine/Shaders/Private/ParticleSpriteVertexFactory.ush b/Engine/Shaders/Private/ParticleSpriteVertexFactory.ush
--- a/Engine/Shaders/Private/ParticleSpriteVertexFactory.ush
+++ b/Engine/Shaders/Private/ParticleSpriteVertexFactory.ush
@@ -115,6 +115,8 @@ struct FVertexFactoryIntermediates
 	float3 TangentRight;
 	/** The color of the sprite. */
 	float4 Color;
+	/** The delta between previous and current particle locations. */
+	float3 ParticlePrevToCurrDelta;
 	/** The velocity of the particle, XYZ: direction, W: speed. */
 	float4 ParticleVelocity;
 	/** Dynamic parameter. */
@@ -533,11 +535,17 @@ FVertexFactoryIntermediates GetVertexFactoryIntermediates(FVertexFactoryInput In
 #endif
 #endif
 
-#if USE_PARTICLE_VELOCITY
-	float3 ParticleVelocity = LWCToFloat(LWCSubtract(ParticleWorldPosition, ParticleOldWorldPosition)) * SpriteVF.InvDeltaSeconds;
-	Intermediates.ParticleVelocity.xyz = normalize(ParticleVelocity);
-	Intermediates.ParticleVelocity.w = length(ParticleVelocity);
+#if !USE_PARTICLE_VELOCITY
+	if (SpriteVF.UseVelocityForMotionBlur > 0.0f)
 #endif
+	{
+		Intermediates.ParticlePrevToCurrDelta = LWCToFloat(LWCSubtract(ParticleWorldPosition, ParticleOldWorldPosition));
+	#if USE_PARTICLE_VELOCITY
+		float3 ParticleVelocity = Intermediates.ParticlePrevToCurrDelta * SpriteVF.InvDeltaSeconds;
+		Intermediates.ParticleVelocity.xyz = normalize(ParticleVelocity);
+		Intermediates.ParticleVelocity.w = length(ParticleVelocity);
+	#endif
+	}
 
 #if USE_PARTICLE_TIME
 	Intermediates.RelativeTime = Input.Position.w;
@@ -669,7 +677,9 @@ FVertexFactoryInterpolantsVSToPS VertexFactoryGetInterpolantsVSToPS(FVertexFacto
 
 float4 VertexFactoryGetPreviousWorldPosition(FVertexFactoryInput Input, FVertexFactoryIntermediates Intermediates)
 {
-	return VertexFactoryGetWorldPosition(Input, Intermediates);
+	float4 Position = VertexFactoryGetWorldPosition(Input, Intermediates);
+	Position.xyz -= Intermediates.ParticlePrevToCurrDelta * SpriteVF.UseVelocityForMotionBlur;
+	return Position;
 }
 
 /**
diff --git a/Engine/Source/Runtime/Engine/Classes/Particles/ParticleModuleRequired.h b/Engine/Source/Runtime/Engine/Classes/Particles/ParticleModuleRequired.h
--- a/Engine/Source/Runtime/Engine/Classes/Particles/ParticleModuleRequired.h
+++ b/Engine/Source/Runtime/Engine/Classes/Particles/ParticleModuleRequired.h
@@ -113,6 +113,7 @@ struct FParticleRequiredModule
 	TArray<FVector2f> FrameData;
 	FRHIShaderResourceView* BoundingGeometryBufferSRV;
 	uint8 bCutoutTexureIsValid : 1;
+	uint8 bUseVelocityForMotionBlur : 1;
 };
 
 
@@ -211,6 +212,17 @@ class UParticleModuleRequired : public UParticleModule
 	UPROPERTY(EditAnywhere, Category = Duration)
 	uint8 bEmitterDurationUseRange : 1;
 
+
+	UPROPERTY(EditAnywhere, Category = Rendering)
+	uint8 bOverrideUseVelocityForMotionBlur : 1;
+
+	/**
+	 * When supported by the vertex factory will use particle velocity for motion blur approximation.
+	 * This will be inaccurate in some cases, i.e. sprite rotation, but may provide a reasonable result vs having this disabled
+	 */
+	UPROPERTY(EditAnywhere, Category = Rendering, meta=(EditCondition = "bOverrideUseVelocityForMotionBlur"))
+	uint8 bUseVelocityForMotionBlur : 1;
+
 	/** 
 	 *	How long, in seconds, the emitter will run before looping.
 	 */
@@ -440,6 +452,8 @@ class UParticleModuleRequired : public UParticleModule
 	virtual bool GenerateLODModuleValues(UParticleModule* SourceModule, float Percentage, UParticleLODLevel* LODLevel) override;
 	//~ End UParticleModule Interface
 
+	bool ShouldUseVelocityForMotionBlur() const;
+
 	inline int32 GetNumFrames() const
 	{
 		return SubImages_Vertical * SubImages_Horizontal;
@@ -474,6 +488,7 @@ class UParticleModuleRequired : public UParticleModule
 	{
 		FParticleRequiredModule *FReqMod = new FParticleRequiredModule();
 		FReqMod->bCutoutTexureIsValid = IsBoundingGeometryValid();
+		FReqMod->bUseVelocityForMotionBlur = ShouldUseVelocityForMotionBlur();
 		FReqMod->NumFrames = GetNumFrames();
 		FReqMod->FrameData = DerivedData.BoundingGeometry;
 		FReqMod->NumBoundingVertices = GetNumBoundingVertices();
diff --git a/Engine/Source/Runtime/Engine/Classes/Particles/TypeData/ParticleModuleTypeDataGpu.h b/Engine/Source/Runtime/Engine/Classes/Particles/TypeData/ParticleModuleTypeDataGpu.h
--- a/Engine/Source/Runtime/Engine/Classes/Particles/TypeData/ParticleModuleTypeDataGpu.h
+++ b/Engine/Source/Runtime/Engine/Classes/Particles/TypeData/ParticleModuleTypeDataGpu.h
@@ -181,6 +181,10 @@ struct FGPUSpriteEmitterInfo
 	UPROPERTY()
 	TEnumAsByte<EParticleCollisionMode::Type> CollisionMode;
 	
+	/** If true, use velocity to approximate motion blur */
+	UPROPERTY()
+	uint32 bUseVelocityForMotionBlur :1;
+
 	/** If true, removes the HMD view roll (e.g. in VR) */
 	UPROPERTY()
 	uint32 bRemoveHMDRoll : 1;
@@ -226,6 +230,7 @@ struct FGPUSpriteEmitterInfo
 		, LockAxisFlag(0)
 		, bEnableCollision(false)
 		, CollisionMode(EParticleCollisionMode::SceneDepth)
+		, bUseVelocityForMotionBlur(0)
 		, bRemoveHMDRoll(0)
 		, MinFacingCameraBlendDistance(0.f)
 		, MaxFacingCameraBlendDistance(0.f)
@@ -388,6 +393,10 @@ struct FGPUSpriteResourceData
 	UPROPERTY()
 	FVector2D PivotOffset;
 	
+	/** If true, use velocity to approximate motion blur */
+	UPROPERTY()
+	uint32 bUseVelocityForMotionBlur :1;
+
 	/** If true, removes the HMD view roll (e.g. in VR) */
 	UPROPERTY()
 	uint32 bRemoveHMDRoll:1;
@@ -435,6 +444,7 @@ struct FGPUSpriteResourceData
 		, ScreenAlignment(0)
 		, LockAxisFlag(0)
 		, PivotOffset(-0.5f,-0.5f)
+		, bUseVelocityForMotionBlur(0)
 		, bRemoveHMDRoll(0)
 		, MinFacingCameraBlendDistance(0.f)
 		, MaxFacingCameraBlendDistance(0.f)		
diff --git a/Engine/Source/Runtime/Engine/Private/Particles/ParticleComponents.cpp b/Engine/Source/Runtime/Engine/Private/Particles/ParticleComponents.cpp
--- a/Engine/Source/Runtime/Engine/Private/Particles/ParticleComponents.cpp
+++ b/Engine/Source/Runtime/Engine/Private/Particles/ParticleComponents.cpp
@@ -2251,6 +2251,11 @@ bool UParticleSystem::DoesAnyEmitterHaveMotionBlur(int32 LODLevelIndex) const
 			{
 				return true;
 			}
+
+			if (EmitterLOD->RequiredModule && EmitterLOD->RequiredModule->ShouldUseVelocityForMotionBlur())
+			{
+				return true;
+			}
 		}
 	}
 
diff --git a/Engine/Source/Runtime/Engine/Private/Particles/ParticleEmitterInstances.cpp b/Engine/Source/Runtime/Engine/Private/Particles/ParticleEmitterInstances.cpp
--- a/Engine/Source/Runtime/Engine/Private/Particles/ParticleEmitterInstances.cpp
+++ b/Engine/Source/Runtime/Engine/Private/Particles/ParticleEmitterInstances.cpp
@@ -370,6 +370,7 @@ FParticleEmitterBuildInfo::FParticleEmitterBuildInfo()
 	, bLocalVectorFieldTileY(false)
 	, bLocalVectorFieldTileZ(false)
 	, bLocalVectorFieldUseFixDT(false)
+	, bUseVelocityForMotionBlur(0)
 	, bRemoveHMDRoll(0)
 	, MinFacingCameraBlendDistance(0.0f)
 	, MaxFacingCameraBlendDistance(0.0f)
@@ -2899,6 +2900,7 @@ bool FParticleEmitterInstance::FillReplayData( FDynamicEmitterReplayDataBase& Ou
 
 		NewReplayData->PivotOffset = FVector2f(PivotOffset);
 
+		NewReplayData->bUseVelocityForMotionBlur = LODLevel->RequiredModule->ShouldUseVelocityForMotionBlur();
 		NewReplayData->bRemoveHMDRoll = LODLevel->RequiredModule->bRemoveHMDRoll;
 		NewReplayData->MinFacingCameraBlendDistance = LODLevel->RequiredModule->MinFacingCameraBlendDistance;
 		NewReplayData->MaxFacingCameraBlendDistance = LODLevel->RequiredModule->MaxFacingCameraBlendDistance;
@@ -4027,6 +4029,7 @@ FDynamicSpriteEmitterReplayDataBase::FDynamicSpriteEmitterReplayDataBase()
 	, EmitterRenderMode(0)
 	, EmitterNormalsMode(0)
 	, PivotOffset(-0.5f, -0.5f)
+	, bUseVelocityForMotionBlur(false)
 	, bRemoveHMDRoll(false)
 	, MinFacingCameraBlendDistance(0.f)
 	, MaxFacingCameraBlendDistance(0.f)
@@ -4069,6 +4072,7 @@ void FDynamicSpriteEmitterReplayDataBase::Serialize( FArchive& Ar )
 
 	Ar << PivotOffset;
 
+	Ar << bUseVelocityForMotionBlur;
 	Ar << bRemoveHMDRoll;
 	Ar << MinFacingCameraBlendDistance;
 	Ar << MaxFacingCameraBlendDistance;
diff --git a/Engine/Source/Runtime/Engine/Private/Particles/ParticleGpuSimulation.cpp b/Engine/Source/Runtime/Engine/Private/Particles/ParticleGpuSimulation.cpp
--- a/Engine/Source/Runtime/Engine/Private/Particles/ParticleGpuSimulation.cpp
+++ b/Engine/Source/Runtime/Engine/Private/Particles/ParticleGpuSimulation.cpp
@@ -489,6 +489,7 @@ BEGIN_GLOBAL_SHADER_PARAMETER_STRUCT(FGPUSpriteEmitterUniformParameters,)
 	SHADER_PARAMETER(float, RotationBias)
 	SHADER_PARAMETER(float, CameraMotionBlurAmount)
 	SHADER_PARAMETER(FVector2f, PivotOffset)
+	SHADER_PARAMETER(float, UseVelocityForMotionBlur)
 END_GLOBAL_SHADER_PARAMETER_STRUCT()
 
 IMPLEMENT_GLOBAL_SHADER_PARAMETER_STRUCT(FGPUSpriteEmitterUniformParameters, "EmitterUniforms");
@@ -4960,6 +4961,8 @@ static void SetGPUSpriteResourceData( FGPUSpriteResources* Resources, const FGPU
 		Resources->UniformParameters.RotationBias = (LockAxisFlag == EPAL_ROTATE_Z) ? (0.5f * UE_PI) : 0.0f;
 	}
 
+	Resources->UniformParameters.UseVelocityForMotionBlur = InResourceData.bUseVelocityForMotionBlur ? 1.0f : 0.0f;
+
 	// Alignment overrides
 	Resources->UniformParameters.RemoveHMDRoll = InResourceData.bRemoveHMDRoll ? 1.f : 0.f;
 
diff --git a/Engine/Source/Runtime/Engine/Private/Particles/ParticleModules.cpp b/Engine/Source/Runtime/Engine/Private/Particles/ParticleModules.cpp
--- a/Engine/Source/Runtime/Engine/Private/Particles/ParticleModules.cpp
+++ b/Engine/Source/Runtime/Engine/Private/Particles/ParticleModules.cpp
@@ -87,12 +87,16 @@
 #include "Particles/Collision/ParticleModuleCollisionGPU.h"
 #include "DerivedDataCacheInterface.h"
 
+static TAutoConsoleVariable<bool> CVarFxCascadeUseVelocityForMotionBlur(
+	TEXT("fx.Cascade.UseVelocityForMotionBlur"),
+	true,
+	TEXT("When enabled velocity will be used to approximate velocity for vertex factories that support this.")
+);
+
 /*-----------------------------------------------------------------------------
 	Abstract base modules used for categorization.
 -----------------------------------------------------------------------------*/
 
-
-
 /*-----------------------------------------------------------------------------
 	Helper functions.
 -----------------------------------------------------------------------------*/
@@ -1310,6 +1314,15 @@ void UParticleModuleRequired::CacheDerivedData()
 #endif
 }
 
+bool UParticleModuleRequired::ShouldUseVelocityForMotionBlur() const
+{
+	if ( bOverrideUseVelocityForMotionBlur )
+	{
+		return bUseVelocityForMotionBlur;
+	}
+	return CVarFxCascadeUseVelocityForMotionBlur.GetValueOnAnyThread();
+}
+
 void UParticleModuleRequired::InitBoundingGeometryBuffer()
 {
 	// The SRV is only needed for platforms that can render particles with instancing
@@ -4847,6 +4860,9 @@ void UParticleModuleTypeDataGpu::Build( FParticleEmitterBuildInfo& EmitterBuildI
 		// Compute the maximum number of particles allowed for this emitter.
 		EmitterInfo.MaxParticleCount = FMath::Max<int32>( 1, EmitterBuildInfo.EstimatedMaxActiveParticleCount );
 
+		EmitterInfo.bUseVelocityForMotionBlur = EmitterBuildInfo.RequiredModule->ShouldUseVelocityForMotionBlur();
+		ResourceData.bUseVelocityForMotionBlur = EmitterBuildInfo.RequiredModule->ShouldUseVelocityForMotionBlur();
+
 		// Store screen alignment for particles.
 		EmitterInfo.ScreenAlignment = EmitterBuildInfo.RequiredModule->ScreenAlignment;
 		ResourceData.ScreenAlignment = EmitterBuildInfo.RequiredModule->ScreenAlignment;
diff --git a/Engine/Source/Runtime/Engine/Private/Particles/ParticleSystemRender.cpp b/Engine/Source/Runtime/Engine/Private/Particles/ParticleSystemRender.cpp
--- a/Engine/Source/Runtime/Engine/Private/Particles/ParticleSystemRender.cpp
+++ b/Engine/Source/Runtime/Engine/Private/Particles/ParticleSystemRender.cpp
@@ -1035,6 +1035,7 @@ void FDynamicSpriteEmitterData::UpdateRenderThreadResourcesEmitter(const FPartic
 		UniformParameters.TangentSelector = FVector4f(0.0f, 0.0f, 0.0f, 0.0f);
 		UniformParameters.InvDeltaSeconds = SourceData->InvDeltaSeconds;
 		UniformParameters.LWCTile = SourceData->LWCTile;
+		UniformParameters.UseVelocityForMotionBlur = Source.RequiredModule->bUseVelocityForMotionBlur ? 1.0f : 0.0f;
 
 		// Parameters for computing sprite tangents.
 		const FMatrix& LocalToWorld = InOwnerProxy->GetLocalToWorld();
@@ -1375,6 +1376,7 @@ void FDynamicMeshEmitterData::GetDynamicMeshElementsEmitter(const FParticleSyste
 			UniformParameters.PrevTransformAvailable = Source.MeshMotionBlurOffset ? 1 : 0;
 			UniformParameters.bUseLocalSpace = Source.bUseLocalSpace;
 			UniformParameters.LWCTile = Source.LWCTile;
+			UniformParameters.UseVelocityForMotionBlur = Source.RequiredModule->bUseVelocityForMotionBlur ? 1.0f : 0.0f;
 
 			CollectorResources.UniformBuffer = FMeshParticleUniformBufferRef::CreateUniformBufferImmediate(UniformParameters, UniformBuffer_MultiFrame);
 			MeshVertexFactory->SetUniformBuffer(CollectorResources.UniformBuffer);
diff --git a/Engine/Source/Runtime/Engine/Public/MeshParticleVertexFactory.h b/Engine/Source/Runtime/Engine/Public/MeshParticleVertexFactory.h
--- a/Engine/Source/Runtime/Engine/Public/MeshParticleVertexFactory.h
+++ b/Engine/Source/Runtime/Engine/Public/MeshParticleVertexFactory.h
@@ -32,6 +32,7 @@ BEGIN_GLOBAL_SHADER_PARAMETER_STRUCT( FMeshParticleUniformParameters, ENGINE_API
 	SHADER_PARAMETER( float, DeltaSeconds )
 	SHADER_PARAMETER( uint32, bUseLocalSpace )
 	SHADER_PARAMETER( FVector3f, LWCTile )
+	SHADER_PARAMETER( float, UseVelocityForMotionBlur )
 END_GLOBAL_SHADER_PARAMETER_STRUCT()
 typedef TUniformBufferRef<FMeshParticleUniformParameters> FMeshParticleUniformBufferRef;
 
diff --git a/Engine/Source/Runtime/Engine/Public/ParticleEmitterInstances.h b/Engine/Source/Runtime/Engine/Public/ParticleEmitterInstances.h
--- a/Engine/Source/Runtime/Engine/Public/ParticleEmitterInstances.h
+++ b/Engine/Source/Runtime/Engine/Public/ParticleEmitterInstances.h
@@ -196,7 +196,9 @@ struct FParticleEmitterBuildInfo
 	uint32 bLocalVectorFieldTileZ : 1;
 	/** Use fix delta time in the simulation? */
 	uint32 bLocalVectorFieldUseFixDT : 1;
-	
+
+	uint32 bUseVelocityForMotionBlur : 1;
+
 	/** Particle alignment overrides */
 	uint32 bRemoveHMDRoll : 1;
 	float MinFacingCameraBlendDistance;
diff --git a/Engine/Source/Runtime/Engine/Public/ParticleHelper.h b/Engine/Source/Runtime/Engine/Public/ParticleHelper.h
--- a/Engine/Source/Runtime/Engine/Public/ParticleHelper.h
+++ b/Engine/Source/Runtime/Engine/Public/ParticleHelper.h
@@ -1416,6 +1416,7 @@ struct FDynamicSpriteEmitterReplayDataBase
 	uint8						EmitterRenderMode;
 	uint8						EmitterNormalsMode;
 	FVector2f					PivotOffset;
+	bool						bUseVelocityForMotionBlur;
 	bool						bRemoveHMDRoll;
 	float						MinFacingCameraBlendDistance;
 	float						MaxFacingCameraBlendDistance;
diff --git a/Engine/Source/Runtime/Engine/Public/ParticleVertexFactory.h b/Engine/Source/Runtime/Engine/Public/ParticleVertexFactory.h
--- a/Engine/Source/Runtime/Engine/Public/ParticleVertexFactory.h
+++ b/Engine/Source/Runtime/Engine/Public/ParticleVertexFactory.h
@@ -85,6 +85,7 @@ BEGIN_GLOBAL_SHADER_PARAMETER_STRUCT( FParticleSpriteUniformParameters, ENGINE_A
 	SHADER_PARAMETER_EX( float, InvDeltaSeconds, EShaderPrecisionModifier::Half )
 	SHADER_PARAMETER( FVector3f, LWCTile )
 	SHADER_PARAMETER_EX( FVector2f, PivotOffset, EShaderPrecisionModifier::Half )
+	SHADER_PARAMETER(float, UseVelocityForMotionBlur)
 END_GLOBAL_SHADER_PARAMETER_STRUCT()
 typedef TUniformBufferRef<FParticleSpriteUniformParameters> FParticleSpriteUniformBufferRef;
 
