Adding shader compilation stats to cook analytics

- Adds a GatherAnalytics to the various shader compilation classes which allows each of them to pass their own stats back up to the cooker.
- Stats have the base name "Shaders_"
- The goal is to capture high level stats which are good indicators of performance and are useful for catching regression (job cache hits, unique output percentage, and mem usage)

Shader Job Cache Stats:
- Queries
- Hits
- NumInputs
- NumOutputs
- MemUsed
- MemBudget

Shader Stats:
- ShadersCompiled
- TotalThreadTimeForAllShaders
- TotalThreadPreprocessTimeForAllShaders

diff --git a/Engine/Source/Editor/UnrealEd/Private/Cooker/CookProfiling.cpp b/Engine/Source/Editor/UnrealEd/Private/Cooker/CookProfiling.cpp
--- a/Engine/Source/Editor/UnrealEd/Private/Cooker/CookProfiling.cpp
+++ b/Engine/Source/Editor/UnrealEd/Private/Cooker/CookProfiling.cpp
@@ -948,6 +948,8 @@ void LogCookStats(ECookMode::Type CookMode)
 		}
 #endif
 
+		GShaderCompilingManager->GatherAnalytics(Attributes);
+
 		// Record them all under cooking event
 		FStudioAnalytics::GetProvider().RecordEvent(TEXT("Core.Cooking"), Attributes);
 
diff --git a/Engine/Source/Runtime/Engine/Private/ShaderCompiler/ShaderCompiler.cpp b/Engine/Source/Runtime/Engine/Private/ShaderCompiler/ShaderCompiler.cpp
--- a/Engine/Source/Runtime/Engine/Private/ShaderCompiler/ShaderCompiler.cpp
+++ b/Engine/Source/Runtime/Engine/Private/ShaderCompiler/ShaderCompiler.cpp
@@ -5,6 +5,7 @@
 =============================================================================*/
 
 #include "ShaderCompiler.h"
+#include "AnalyticsEventAttribute.h"
 #include "Async/ParallelFor.h"
 #include "ClearReplacementShaders.h"
 #include "ComponentRecreateRenderStateContext.h"
@@ -789,6 +790,12 @@ void FShaderCompileJobCollection::LogCachingStats()
 	CompletedJobsCache.LogStats();
 }
 
+void FShaderCompileJobCollection::GatherAnalytics(const FString& BaseName, TArray<FAnalyticsEventAttribute>& Attributes) const
+{
+	FWriteScopeLock Locker(Lock);
+	CompletedJobsCache.GatherAnalytics(BaseName, Attributes);
+}
+
 int32 FShaderCompileJobCollection::GetNumPendingJobs() const
 {
 	FReadScopeLock Locker(Lock);
@@ -3523,6 +3530,35 @@ void FShaderCompilerStats::WriteStatSummary()
 	}
 }
 
+void FShaderCompilerStats::GatherAnalytics(const FString& BaseName, TArray<FAnalyticsEventAttribute>& Attributes) const
+{
+	{
+		FString AttrName = BaseName + TEXT("ShadersCompiled");
+		Attributes.Emplace(MoveTemp(AttrName), JobsCompleted);
+	}
+
+	if (ShaderTimings.Num())
+	{
+		double TotalThreadTimeForAllShaders = 0.0;
+		double TotalThreadPreprocessTimeForAllShaders = 0.0;
+		for (TMap<FString, FShaderTimings>::TConstIterator Iter(ShaderTimings); Iter; ++Iter)
+		{
+			TotalThreadTimeForAllShaders += Iter.Value().TotalCompileTime;
+			TotalThreadPreprocessTimeForAllShaders += Iter.Value().TotalPreprocessTime;
+		}
+
+		{
+			FString AttrName = BaseName + TEXT("TotalThreadTime");
+			Attributes.Emplace(MoveTemp(AttrName), TotalThreadTimeForAllShaders);
+		}
+
+		{
+			FString AttrName = BaseName + TEXT("TotalThreadPreprocessTime");
+			Attributes.Emplace(MoveTemp(AttrName), TotalThreadPreprocessTimeForAllShaders);
+		}
+	}
+}
+
 uint32 FShaderCompilerStats::GetTotalShadersCompiled()
 {
 	FScopeLock Lock(&CompileStatsLock);
@@ -4260,6 +4296,18 @@ int32 FShaderCompilingManager::GetNumOutstandingJobs() const
 	return AllJobs.GetNumOutstandingJobs();
 }
 
+void FShaderCompilingManager::GatherAnalytics(TArray<FAnalyticsEventAttribute>& Attributes) const
+{
+	const FString BaseName = TEXT("Shaders_");
+
+	if (ShaderCompiler::IsJobCacheEnabled())
+	{
+		AllJobs.GatherAnalytics(BaseName, Attributes);
+	}
+
+	GShaderCompilerStats->GatherAnalytics(BaseName, Attributes);
+}
+
 FShaderCompilingManager::EDumpShaderDebugInfo FShaderCompilingManager::GetDumpShaderDebugInfo() const
 {
 	if (GDumpShaderDebugInfo < EDumpShaderDebugInfo::Never || GDumpShaderDebugInfo > EDumpShaderDebugInfo::OnErrorOrWarning)
@@ -9004,7 +9052,7 @@ void FShaderJobCache::RemoveByInputHash(const FJobInputHash& InputHash)
 }
 
 /** Calculates memory used by the cache*/
-uint64 FShaderJobCache::GetAllocatedMemory()
+uint64 FShaderJobCache::GetAllocatedMemory() const
 {
 	return CurrentlyAllocatedMemory;
 }
@@ -9042,4 +9090,39 @@ void FShaderJobCache::LogStats()
 	}
 }
 
+void FShaderJobCache::GatherAnalytics(const FString& BaseName, TArray<FAnalyticsEventAttribute>& Attributes) const
+{
+	const FString ChildName = TEXT("JobCache_");
+
+	{
+		FString AttrName = BaseName + ChildName + TEXT("Queries");
+		Attributes.Emplace(MoveTemp(AttrName), TotalSearchAttempts);
+	}
+
+	{
+		FString AttrName = BaseName + ChildName + TEXT("Hits");
+		Attributes.Emplace(MoveTemp(AttrName), TotalCacheHits);
+	}
+
+	{
+		FString AttrName = BaseName + ChildName + TEXT("NumInputs");
+		Attributes.Emplace(MoveTemp(AttrName), Outputs.Num());
+	}
+
+	{
+		FString AttrName = BaseName + ChildName + TEXT("NumOutputs");
+		Attributes.Emplace(MoveTemp(AttrName), Outputs.Num());
+	}
+
+	{
+		FString AttrName = BaseName + ChildName + TEXT("MemUsed");
+		Attributes.Emplace(MoveTemp(AttrName), GetAllocatedMemory());
+	}
+
+	{
+		FString AttrName = BaseName + ChildName + TEXT("MemBudget");
+		Attributes.Emplace(MoveTemp(AttrName), GetCurrentMemoryBudget());
+	}
+}
+
 #undef LOCTEXT_NAMESPACE
diff --git a/Engine/Source/Runtime/Engine/Public/ShaderCompiler.h b/Engine/Source/Runtime/Engine/Public/ShaderCompiler.h
--- a/Engine/Source/Runtime/Engine/Public/ShaderCompiler.h
+++ b/Engine/Source/Runtime/Engine/Public/ShaderCompiler.h
@@ -38,6 +38,7 @@ DECLARE_LOG_CATEGORY_EXTERN(LogShaderCompilers, Log, All);
 
 class FShaderCompileJob;
 class FShaderPipelineCompileJob;
+struct FAnalyticsEventAttribute;
 
 #define DEBUG_INFINITESHADERCOMPILE 0
 
@@ -62,11 +63,14 @@ public:
 	void AddJobOutput(const FShaderCommonCompileJob* FinishedJob, const FJobInputHash& Hash, const FJobCachedOutput& Contents, int InitialHitCount, const bool bAddToDDC);
 
 	/** Calculates memory used by the cache*/
-	uint64 GetAllocatedMemory();
+	uint64 GetAllocatedMemory() const;
 
 	/** Logs out the statistics */
 	void LogStats();
 
+	/** Gather statistics to send to analytics */
+	void GatherAnalytics(const FString& BaseName, TArray<FAnalyticsEventAttribute>& Attributes) const;
+
 	/** Calculates current memory budget, in bytes */
 	uint64 GetCurrentMemoryBudget() const;
 
@@ -116,11 +120,11 @@ public:
 	/** Adds the job to cache. */
 	void AddToCacheAndProcessPending(FShaderCommonCompileJob* FinishedJob);
 
-	/** Log caching statistics.
-	 *
-	 */
+	/** Log caching statistics.*/
 	void LogCachingStats();
 
+	void GatherAnalytics(const FString& BaseName, TArray<FAnalyticsEventAttribute>& Attributes) const;
+
 	inline int32 GetNumPendingJobs(EShaderCompileJobPriority InPriority) const
 	{
 		return NumPendingJobs[(int32)InPriority];
@@ -510,6 +514,8 @@ public:
 	/** Informs statistics about a new job batch, so we can tally up batches. */
 	void RegisterJobBatch(int32 NumJobs, EExecutionType ExecType);
 
+	void GatherAnalytics(const FString& BaseName, TArray<FAnalyticsEventAttribute>& Attributes) const;
+
 private:
 	FCriticalSection CompileStatsLock;
 	TSparseArray<ShaderCompilerStats> CompileStats;
@@ -752,6 +758,8 @@ public:
 	ENGINE_API int32 GetNumPendingJobs() const;
 	ENGINE_API int32 GetNumOutstandingJobs() const;
 
+	ENGINE_API void GatherAnalytics(TArray<FAnalyticsEventAttribute>& Attributes) const;
+
 	/** 
 	 * Returns whether to display a notification that shader compiling is happening in the background. 
 	 * Note: This is dependent on NumOutstandingJobs which is updated from another thread, so the results are non-deterministic.
