When running with either `-noshaderddc` or `-nomaterialshaderddc` we should still make the DDC request instead of skipping it.

- Those cmd line options are used to simulate a cold DDC, but unfortunately they skipped issuing the DDC request.
- This didn't allow you to do performance tests compared to an actual cold DDC cook.
- Instead we should still make the DDC requets, but just have it fail (so we still get the round trip time performance of the request).

diff --git a/Engine/Source/Runtime/Engine/Private/Materials/MaterialShader.cpp b/Engine/Source/Runtime/Engine/Private/Materials/MaterialShader.cpp
--- a/Engine/Source/Runtime/Engine/Private/Materials/MaterialShader.cpp
+++ b/Engine/Source/Runtime/Engine/Private/Materials/MaterialShader.cpp
@@ -1517,7 +1517,7 @@ TSharedRef<FMaterialShaderMap::FAsyncLoadContext> FMaterialShaderMap::BeginLoadF
 			// Do not check the DDC if the material isn't persistent, because
 			//   - this results in a lot of DDC requests when editing materials which are almost always going to return nothing.
 			//   - since the get call is synchronous, this can cause a hitch if there's network latency
-			if (bCheckCache && Material->IsPersistent())
+			if (Material->IsPersistent())
 			{
 				FCacheGetValueRequest Request;
 				Request.Name = GetMaterialShaderMapName(Material->GetFullPath(), ShaderMapId, InPlatform);
@@ -1525,12 +1525,15 @@ TSharedRef<FMaterialShaderMap::FAsyncLoadContext> FMaterialShaderMap::BeginLoadF
 				Result->AssetName = Material->GetAssetName();
 				Result->Platform = InPlatform;
 
-				GetCache().GetValue({Request}, Result->RequestOwner, [Result](FCacheGetValueResponse&& Response)
+				GetCache().GetValue({Request}, Result->RequestOwner, [Result, bCheckCache](FCacheGetValueResponse&& Response)
 				{
-					Result->CachedData = Response.Value.GetData().Decompress();
-					// This callback might hold the last reference to Result, which owns RequestOwner, so
-					// we must not cancel in the Owner's destructor; Canceling in a callback will deadlock.
-					Result->RequestOwner.KeepAlive();
+					if (bCheckCache)
+					{
+						Result->CachedData = Response.Value.GetData().Decompress();
+						// This callback might hold the last reference to Result, which owns RequestOwner, so
+						// we must not cancel in the Owner's destructor; Canceling in a callback will deadlock.
+						Result->RequestOwner.KeepAlive();
+					}
 				});
 			}
 		}
diff --git a/Engine/Source/Runtime/Engine/Private/ShaderCompiler/ShaderCompiler.cpp b/Engine/Source/Runtime/Engine/Private/ShaderCompiler/ShaderCompiler.cpp
--- a/Engine/Source/Runtime/Engine/Private/ShaderCompiler/ShaderCompiler.cpp
+++ b/Engine/Source/Runtime/Engine/Private/ShaderCompiler/ShaderCompiler.cpp
@@ -7911,11 +7911,9 @@ void CompileGlobalShaderMap(EShaderPlatform Platform, const ITargetPlatform* Tar
 			static bool bNoShaderDDC =
 				FParse::Param(FCommandLine::Get(), TEXT("noshaderddc")) || 
 				FParse::Param(FCommandLine::Get(), TEXT("noglobalshaderddc"));
-			if (UNLIKELY(bNoShaderDDC))
-			{
-				bShaderMapIsBeingCompiled = true;
-			}
-			else
+
+			const bool bTempNoShaderDDC = bNoShaderDDC;
+
 			{
 				using namespace UE::DerivedData;
 
@@ -7952,8 +7950,13 @@ void CompileGlobalShaderMap(EShaderPlatform Platform, const ITargetPlatform* Tar
 					COOK_STAT(auto Timer = GlobalShaderCookStats::UsageStats.TimeSyncWork());
 					COOK_STAT(Timer.TrackCyclesOnly());
 					FRequestOwner BlockingOwner(EPriority::Blocking);
-					GetCache().GetValue(Requests, BlockingOwner, [&GlobalShaderMapBuffers](FCacheGetValueResponse&& Response)
+					GetCache().GetValue(Requests, BlockingOwner, [&GlobalShaderMapBuffers, &bTempNoShaderDDC](FCacheGetValueResponse&& Response)
 					{
+						if (bTempNoShaderDDC)
+						{
+							return;
+						}
+
 						GlobalShaderMapBuffers[int32(Response.UserData)] = MoveTemp(Response.Value);
 					});
 					BlockingOwner.Wait();
