WorldPartition: PIE Cell name now use the same suffix as the main world
- GetCellNameString uses source world path to avoid including LevelInstance suffix in Hash
- LevelInstance suffix is appended to Cell hashed name

#rnx

diff --git a/Engine/Source/Runtime/Engine/Private/WorldPartition/RuntimeSpatialHash/RuntimeSpatialHashHLOD.cpp b/Engine/Source/Runtime/Engine/Private/WorldPartition/RuntimeSpatialHash/RuntimeSpatialHashHLOD.cpp
--- a/Engine/Source/Runtime/Engine/Private/WorldPartition/RuntimeSpatialHash/RuntimeSpatialHashHLOD.cpp
+++ b/Engine/Source/Runtime/Engine/Private/WorldPartition/RuntimeSpatialHash/RuntimeSpatialHashHLOD.cpp
@@ -178,7 +178,7 @@ static TArray<FGuid> GenerateHLODsForGrid(UWorldPartition* WorldPartition, const
 				SlowTask.EnterProgressFrame(1);
 
 				const FGuid CellGuid = UWorldPartitionRuntimeSpatialHash::GetCellGuid(RuntimeGrid.GridName, CellGlobalCoord, GridCellDataChunk.GetDataLayersID(), GridCellDataChunk.GetContentBundleID());
-				const FString CellName = UWorldPartitionRuntimeSpatialHash::GetCellNameString(OuterWorld->GetPackage()->GetFName(), RuntimeGrid.GridName, CellGlobalCoord, GridCellDataChunk.GetDataLayersID(), GridCellDataChunk.GetContentBundleID());
+				const FString CellName = UWorldPartitionRuntimeSpatialHash::GetCellNameString(OuterWorld, RuntimeGrid.GridName, CellGlobalCoord, GridCellDataChunk.GetDataLayersID(), GridCellDataChunk.GetContentBundleID());
 
 				UE_LOG(LogWorldPartitionRuntimeSpatialHashHLOD, Verbose, TEXT("Creating HLOD for cell %s at %s..."), *CellName, *CellCoord.ToString());
 
diff --git a/Engine/Source/Runtime/Engine/Private/WorldPartition/WorldPartitionRuntimeHash.cpp b/Engine/Source/Runtime/Engine/Private/WorldPartition/WorldPartitionRuntimeHash.cpp
--- a/Engine/Source/Runtime/Engine/Private/WorldPartition/WorldPartitionRuntimeHash.cpp
+++ b/Engine/Source/Runtime/Engine/Private/WorldPartition/WorldPartitionRuntimeHash.cpp
@@ -96,7 +96,7 @@ void UWorldPartitionRuntimeHash::OnEndPlay()
 	ModifiedActorDescListForPIE.Empty();
 }
 
-UWorldPartitionRuntimeCell* UWorldPartitionRuntimeHash::CreateRuntimeCell(UClass* CellClass, UClass* CellDataClass, const FString& CellName)
+UWorldPartitionRuntimeCell* UWorldPartitionRuntimeHash::CreateRuntimeCell(UClass* CellClass, UClass* CellDataClass, const FString& CellName, const FString& CellInstanceSuffix)
 {
 	//@todo_ow: to reduce file paths on windows, we compute a MD5 hash from the unique cell name and use that hash as the cell filename. We
 	//			should create the runtime cell using the name but make sure the package that gets associated with it when cooking gets the
@@ -115,7 +115,9 @@ UWorldPartitionRuntimeCell* UWorldPartitionRuntimeHash::CreateRuntimeCell(UClass
 		return CellNameGuid.ToString(EGuidFormats::Base36Encoded);
 	};
 
-	const FString CellObjectName = GetCellObjectName(CellName);
+	// Cooking should have an empty CellInstanceSuffix
+	check(!IsRunningCookCommandlet() || CellInstanceSuffix.IsEmpty());
+	const FString CellObjectName = GetCellObjectName(CellName) + CellInstanceSuffix;
 	UWorldPartitionRuntimeCell* RuntimeCell = NewObject<UWorldPartitionRuntimeCell>(this, CellClass, *CellObjectName);
 	RuntimeCell->RuntimeCellData = NewObject<UWorldPartitionRuntimeCellData>(RuntimeCell, CellDataClass);
 	return RuntimeCell;
diff --git a/Engine/Source/Runtime/Engine/Private/WorldPartition/WorldPartitionRuntimeSpatialHash.cpp b/Engine/Source/Runtime/Engine/Private/WorldPartition/WorldPartitionRuntimeSpatialHash.cpp
--- a/Engine/Source/Runtime/Engine/Private/WorldPartition/WorldPartitionRuntimeSpatialHash.cpp
+++ b/Engine/Source/Runtime/Engine/Private/WorldPartition/WorldPartitionRuntimeSpatialHash.cpp
@@ -1259,10 +1259,32 @@ void UWorldPartitionRuntimeSpatialHash::DumpStateLog(FHierarchicalLogArchive& Ar
 	});
 }
 
-FString UWorldPartitionRuntimeSpatialHash::GetCellNameString(FName InWorldPackageName, FName InGridName, const FGridCellCoord& InCellGlobalCoord, const FDataLayersID& InDataLayerID, const FGuid& InContentBundleID)
+FString UWorldPartitionRuntimeSpatialHash::GetCellNameString(UWorld* InOuterWorld, FName InGridName, const FGridCellCoord& InCellGlobalCoord, const FDataLayersID& InDataLayerID, const FGuid& InContentBundleID, FString* OutInstanceSuffix)
 {
-	const FString ShortPackageName = UWorld::RemovePIEPrefix(FPackageName::GetShortName(InWorldPackageName));
-	FString CellName = FString::Printf(TEXT("%s_%s_%s"), *ShortPackageName, *InGridName.ToString(), *GetCellCoordString(InCellGlobalCoord));
+	FString WorldName = FPackageName::GetShortName(InOuterWorld->GetPackage());
+	
+	if (!IsRunningCookCommandlet() && InOuterWorld->IsGameWorld())
+	{
+		FString SourceWorldPath, InstancedWorldPath;
+		if (InOuterWorld->GetSoftObjectPathMapping(SourceWorldPath, InstancedWorldPath))
+		{
+			const FTopLevelAssetPath SourceAssetPath(SourceWorldPath);
+			WorldName = FPackageName::GetShortName(SourceAssetPath.GetPackageName());
+						
+			InstancedWorldPath = UWorld::RemovePIEPrefix(InstancedWorldPath);
+
+			const FString SourcePackageName = SourceAssetPath.GetPackageName().ToString();
+			const FTopLevelAssetPath InstanceAssetPath(InstancedWorldPath);
+			const FString InstancePackageName = InstanceAssetPath.GetPackageName().ToString();
+
+			if (InstancePackageName.StartsWith(SourcePackageName))
+			{
+				*OutInstanceSuffix = InstancePackageName.Mid(SourcePackageName.Len());
+			}
+		}
+	}
+			
+	FString CellName = FString::Printf(TEXT("%s_%s_%s"), *WorldName, *InGridName.ToString(), *GetCellCoordString(InCellGlobalCoord));
 
 	if (InDataLayerID.GetHash())
 	{
@@ -1372,11 +1394,12 @@ bool UWorldPartitionRuntimeSpatialHash::CreateStreamingGrid(const FSpatialHashRu
 				if (PopulateCellActorInstances(GridCellDataChunk.GetActorSetInstances(), bIsMainWorldPartition, bIsCellAlwaysLoaded, FilteredActors))
 				{
 					FGridCellCoord CellGlobalCoords;
+					FString WorldInstanceSuffix;
 					verify(PartionedActors.GetCellGlobalCoords(FGridCellCoord(CellCoordX, CellCoordY, Level), CellGlobalCoords));
-					const FString CellName = GetCellNameString(OuterWorld->GetPackage()->GetFName(), CurrentStreamingGrid.GridName, CellGlobalCoords, GridCellDataChunk.GetDataLayersID(), GridCellDataChunk.GetContentBundleID());
+					const FString CellName = GetCellNameString(OuterWorld, CurrentStreamingGrid.GridName, CellGlobalCoords, GridCellDataChunk.GetDataLayersID(), GridCellDataChunk.GetContentBundleID(), &WorldInstanceSuffix);
 					const FGuid CellGuid = GetCellGuid(CurrentStreamingGrid.GridName, CellGlobalCoords, GridCellDataChunk.GetDataLayersID(), GridCellDataChunk.GetContentBundleID());
 
-					UWorldPartitionRuntimeCell* StreamingCell = CreateRuntimeCell(StreamingPolicy->GetRuntimeCellClass(), UWorldPartitionRuntimeCellDataSpatialHash::StaticClass(), CellName);
+					UWorldPartitionRuntimeCell* StreamingCell = CreateRuntimeCell(StreamingPolicy->GetRuntimeCellClass(), UWorldPartitionRuntimeCellDataSpatialHash::StaticClass(), CellName, WorldInstanceSuffix);
 					UWorldPartitionRuntimeCellDataSpatialHash* CellDataSpatialHash = CastChecked<UWorldPartitionRuntimeCellDataSpatialHash>(StreamingCell->RuntimeCellData);
 
 					StreamingCell->SetIsAlwaysLoaded(bIsCellAlwaysLoaded);
@@ -1397,7 +1420,7 @@ bool UWorldPartitionRuntimeSpatialHash::CreateStreamingGrid(const FSpatialHashRu
 					CellDataSpatialHash->Position = FVector(Bounds.GetCenter(), 0.f);
 					CellDataSpatialHash->Extent = (float)CellExtent;
 					CellDataSpatialHash->GridName = RuntimeGrid.GridName;			
-					CellDataSpatialHash->DebugName = CellName;
+					CellDataSpatialHash->DebugName = StreamingCell->GetName();
 
 					PopulateRuntimeCell(StreamingCell, FilteredActors, OutPackagesToGenerate);
 
diff --git a/Engine/Source/Runtime/Engine/Public/WorldPartition/WorldPartitionRuntimeHash.h b/Engine/Source/Runtime/Engine/Public/WorldPartition/WorldPartitionRuntimeHash.h
--- a/Engine/Source/Runtime/Engine/Public/WorldPartition/WorldPartitionRuntimeHash.h
+++ b/Engine/Source/Runtime/Engine/Public/WorldPartition/WorldPartitionRuntimeHash.h
@@ -99,7 +99,7 @@ class ENGINE_API UWorldPartitionRuntimeHash : public UObject
 	void OnBeginPlay();
 	void OnEndPlay();
 
-	UWorldPartitionRuntimeCell* CreateRuntimeCell(UClass* CellClass, UClass* CellDataClass, const FString& CellName);
+	UWorldPartitionRuntimeCell* CreateRuntimeCell(UClass* CellClass, UClass* CellDataClass, const FString& CellName, const FString& CellInstanceSuffix);
 
 protected:
 	bool ConditionalRegisterAlwaysLoadedActorsForPIE(const IStreamingGenerationContext::FActorSetInstance* ActorSetInstance, bool bIsMainWorldPartition, bool bIsMainContainer, bool bIsCellAlwaysLoaded);
diff --git a/Engine/Source/Runtime/Engine/Public/WorldPartition/WorldPartitionRuntimeSpatialHash.h b/Engine/Source/Runtime/Engine/Public/WorldPartition/WorldPartitionRuntimeSpatialHash.h
--- a/Engine/Source/Runtime/Engine/Public/WorldPartition/WorldPartitionRuntimeSpatialHash.h
+++ b/Engine/Source/Runtime/Engine/Public/WorldPartition/WorldPartitionRuntimeSpatialHash.h
@@ -266,7 +266,7 @@ public:
 
 	virtual URuntimeHashExternalStreamingObjectBase* StoreToExternalStreamingObject(UObject* StreamingObjectOuter, FName StreamingObjectName) override;
 
-	static FString GetCellNameString(FName InWorldPackageName, FName InGridName, const FGridCellCoord& InCellGlobalCoord, const FDataLayersID& InDataLayerID, const FGuid& InContentBundleID);
+	static FString GetCellNameString(UWorld* InOuterWorld, FName InGridName, const FGridCellCoord& InCellGlobalCoord, const FDataLayersID& InDataLayerID, const FGuid& InContentBundleID, FString* OutInstanceSuffix = nullptr);
 	static FGuid GetCellGuid(FName InGridName, const FGridCellCoord& InCellGlobalCoord, const FDataLayersID& InDataLayerID, const FGuid& InContentBundleID);
 
 	bool GetPreviewGrids() const;
