Implemented multi-phase initialization for SparseVolumeTextures.

#rnx

diff --git a/Engine/Source/Editor/SparseVolumeTexture/Private/SparseVolumeTextureFactory.cpp b/Engine/Source/Editor/SparseVolumeTexture/Private/SparseVolumeTextureFactory.cpp
--- a/Engine/Source/Editor/SparseVolumeTexture/Private/SparseVolumeTextureFactory.cpp
+++ b/Engine/Source/Editor/SparseVolumeTexture/Private/SparseVolumeTextureFactory.cpp
@@ -509,7 +509,7 @@ UObject* USparseVolumeTextureFactory::FactoryCreateFile(UClass* InClass, UObject
 		}
 
 		UStaticSparseVolumeTexture* StaticSVTexture = NewObject<UStaticSparseVolumeTexture>(InParent, UStaticSparseVolumeTexture::StaticClass(), InName, Flags);
-		const bool bInitSuccess = StaticSVTexture->InitializeFromUncooked(TextureData);
+		const bool bInitSuccess = StaticSVTexture->Initialize(MakeArrayView<FSparseVolumeTextureData>(&TextureData, 1));
 		if (!bInitSuccess)
 		{
 			UE_LOG(LogSparseVolumeTextureFactory, Error, TEXT("Failed to initialize SparseVolumeTexture: %s"), *Filename);
@@ -721,7 +721,7 @@ UObject* USparseVolumeTextureFactory::FactoryCreateFile(UClass* InClass, UObject
 		}
 
 		UAnimatedSparseVolumeTexture* AnimatedSVTexture = NewObject<UAnimatedSparseVolumeTexture>(InParent, UAnimatedSparseVolumeTexture::StaticClass(), InName, Flags);
-		const bool bInitSuccess = AnimatedSVTexture->InitializeFromUncooked(UncookedFramesData);
+		const bool bInitSuccess = AnimatedSVTexture->Initialize(UncookedFramesData);
 		if (!bInitSuccess)
 		{
 			UE_LOG(LogSparseVolumeTextureFactory, Error, TEXT("Failed to initialize SparseVolumeTexture: %s"), *Filename);
diff --git a/Engine/Source/Runtime/Engine/Classes/SparseVolumeTexture/SparseVolumeTexture.h b/Engine/Source/Runtime/Engine/Classes/SparseVolumeTexture/SparseVolumeTexture.h
--- a/Engine/Source/Runtime/Engine/Classes/SparseVolumeTexture/SparseVolumeTexture.h
+++ b/Engine/Source/Runtime/Engine/Classes/SparseVolumeTexture/SparseVolumeTexture.h
@@ -117,7 +117,17 @@ public:
 	/** Getter for the shader uniform parameter type with index as ESparseVolumeTextureShaderUniform. */
 	static UE::Shader::EValueType GetUniformParameterType(int32 Index);
 
-private:
+protected:
+#if WITH_EDITOR
+	enum class ENotifyMaterialsEffectOnShaders
+	{
+		Default,
+		DoesNotInvalidate
+	};
+
+	/** Notify any loaded material instances that the texture has changed. */
+	void NotifyMaterials(const ENotifyMaterialsEffectOnShaders EffectOnShaders = ENotifyMaterialsEffectOnShaders::Default);
+#endif // WITH_EDITOR
 };
 
 UCLASS(ClassGroup = Rendering, BlueprintType)//, hidecategories = (Object))
@@ -127,6 +137,14 @@ class UStreamableSparseVolumeTexture : public USparseVolumeTexture
 
 public:
 
+	enum class EInitState : uint8
+	{
+		Uninitialized,
+		Pending,
+		Done,
+		Failed,
+	};
+
 	UPROPERTY(VisibleAnywhere, Category = "Texture")
 	FIntVector VolumeResolution;
 
@@ -148,7 +166,18 @@ public:
 	UStreamableSparseVolumeTexture();
 	virtual ~UStreamableSparseVolumeTexture() = default;
 
-	virtual bool InitializeFromUncooked(const TArrayView<FSparseVolumeTextureData>& UncookedData, int32 NumMipLevels = INDEX_NONE /*Create entire mip chain by default*/);
+	// Multi-phase initialization: Call BeginInitialize(), then call AppendFrame() for each frame to add and then finish initialization with a call to EndInitialize().
+	// The NumExpectedFrames parameter on BeginInitialize() just serves as a potential optimization to reserve memory for the frames to be appended
+	// and doesn't need to match the exact number if it is not known at the time.
+	virtual bool BeginInitialize(int32 NumExpectedFrames);
+	virtual bool AppendFrame(FSparseVolumeTextureData& UncookedFrame);
+	virtual bool EndInitialize(int32 NumMipLevels = INDEX_NONE /*Create entire mip chain by default*/);
+
+	// Convenience function wrapping the multi-phase initialization functions above
+	virtual bool Initialize(const TArrayView<FSparseVolumeTextureData>& UncookedData, int32 NumMipLevels = INDEX_NONE /*Create entire mip chain by default*/);
+
+	const FSparseVolumeTextureSceneProxy* GetStreamedFrameProxyOrFallback(int32 FrameIndex, int32 MipLevel) const;
+	TArrayView<const FSparseVolumeTextureFrame> GetFrames() const { return Frames; };
 
 	//~ Begin UObject Interface.
 	virtual void PostLoad() override;
@@ -171,23 +200,14 @@ public:
 	virtual const FSparseVolumeTextureSceneProxy* GetSparseVolumeTextureSceneProxy() const override { return GetStreamedFrameProxyOrFallback(0 /*FrameIndex*/, 0 /*MipLevel*/); };
 	//~ End USparseVolumeTexture Interface.
 
-	const FSparseVolumeTextureSceneProxy* GetStreamedFrameProxyOrFallback(int32 FrameIndex, int32 MipLevel) const;
-	TArrayView<const FSparseVolumeTextureFrame> GetFrames() const;
-
 protected:
 
 	TArray<FSparseVolumeTextureFrame> Frames;
-
-#if WITH_EDITOR
-	enum class ENotifyMaterialsEffectOnShaders
-	{
-		Default,
-		DoesNotInvalidate
-	};
-
-	/** Notify any loaded material instances that the texture has changed. */
-	void NotifyMaterials(const ENotifyMaterialsEffectOnShaders EffectOnShaders = ENotifyMaterialsEffectOnShaders::Default);
-#endif // WITH_EDITOR
+#if WITH_EDITORONLY_DATA
+	FIntVector VolumeBoundsMin;
+	FIntVector VolumeBoundsMax;
+	EInitState InitState = EInitState::Uninitialized;
+#endif // WITH_EDITORONLY_DATA
 
 	void GenerateOrLoadDDCRuntimeDataAndCreateSceneProxy();
 	void GenerateOrLoadDDCRuntimeDataForFrame(FSparseVolumeTextureFrame& Frame, UE::DerivedData::FRequestOwner& DDCRequestOwner);
@@ -203,8 +223,8 @@ public:
 	UStaticSparseVolumeTexture();
 	virtual ~UStaticSparseVolumeTexture() = default;
 
-	virtual bool InitializeFromUncooked(const TArrayView<FSparseVolumeTextureData>& UncookedData, int32 NumMipLevels = INDEX_NONE /*Create entire mip chain by default*/) override;
-	bool InitializeFromUncooked(FSparseVolumeTextureData& UncookedData, int32 NumMipLevels = INDEX_NONE /*Create entire mip chain by default*/);
+	// Override AppendFrame() to ensure that there is never more than a single frame in a static SVT
+	virtual bool AppendFrame(FSparseVolumeTextureData& UncookedFrame) override;
 
 	//~ Begin USparseVolumeTexture Interface.
 	int32 GetNumFrames() const override { return 1; }
diff --git a/Engine/Source/Runtime/Engine/Private/SparseVolumeTexture/SparseVolumeTexture.cpp b/Engine/Source/Runtime/Engine/Private/SparseVolumeTexture/SparseVolumeTexture.cpp
--- a/Engine/Source/Runtime/Engine/Private/SparseVolumeTexture/SparseVolumeTexture.cpp
+++ b/Engine/Source/Runtime/Engine/Private/SparseVolumeTexture/SparseVolumeTexture.cpp
@@ -308,6 +308,50 @@
 	return UE::Shader::EValueType::Float4;
 }
 
+#if WITH_EDITOR
+void USparseVolumeTexture::NotifyMaterials(const ENotifyMaterialsEffectOnShaders EffectOnShaders)
+{
+	// Create a material update context to safely update materials.
+	{
+		FMaterialUpdateContext UpdateContext;
+
+		// Notify any material that uses this texture
+		TSet<UMaterial*> BaseMaterialsThatUseThisTexture;
+		for (TObjectIterator<UMaterialInterface> It; It; ++It)
+		{
+			UMaterialInterface* MaterialInterface = *It;
+			if (!FPlatformProperties::IsServerOnly() && MaterialInterface->GetReferencedTextures().Contains(this))
+			{
+				UpdateContext.AddMaterialInterface(MaterialInterface);
+				// This is a bit tricky. We want to make sure all materials using this texture are
+				// updated. Materials are always updated. Material instances may also have to be
+				// updated and if they have static permutations their children must be updated
+				// whether they use the texture or not! The safe thing to do is to add the instance's
+				// base material to the update context causing all materials in the tree to update.
+				BaseMaterialsThatUseThisTexture.Add(MaterialInterface->GetMaterial());
+			}
+		}
+
+		// Go ahead and update any base materials that need to be.
+		if (EffectOnShaders == ENotifyMaterialsEffectOnShaders::Default)
+		{
+			for (TSet<UMaterial*>::TConstIterator It(BaseMaterialsThatUseThisTexture); It; ++It)
+			{
+				(*It)->PostEditChange();
+			}
+		}
+		else
+		{
+			FPropertyChangedEvent EmptyPropertyUpdateStruct(nullptr);
+			for (TSet<UMaterial*>::TConstIterator It(BaseMaterialsThatUseThisTexture); It; ++It)
+			{
+				(*It)->PostEditChangePropertyInternal(EmptyPropertyUpdateStruct, UMaterial::EPostEditChangeEffectOnShaders::DoesNotInvalidate);
+			}
+		}
+	}
+}
+#endif
+
 ////////////////////////////////////////////////////////////////////////////////////////////////
 
 UStreamableSparseVolumeTexture::UStreamableSparseVolumeTexture(const FObjectInitializer& ObjectInitializer)
@@ -315,64 +359,165 @@
 {
 }
 
-bool UStreamableSparseVolumeTexture::InitializeFromUncooked(const TArrayView<FSparseVolumeTextureData>& InUncookedData, int32 InNumMipLevels)
+bool UStreamableSparseVolumeTexture::BeginInitialize(int32 NumExpectedFrames)
 {
-	if (InUncookedData.IsEmpty())
+#if WITH_EDITORONLY_DATA
+	if (InitState != EInitState::Uninitialized)
 	{
-		UE_LOG(LogSparseVolumeTexture, Error, TEXT("Tried to initialize a SparseVolumeTexture with no frames"));
+		UE_LOG(LogSparseVolumeTexture, Error, TEXT("Tried to call UStreamableSparseVolumeTexture::BeginInitialize() while not in the Uninitialized init state."));
 		return false;
 	}
 
-	auto ExpandVolumeBounds = [](const FSparseVolumeTextureHeader& Header, FIntVector3& VolumeBoundsMin, FIntVector3& VolumeBoundsMax)
+	check(Frames.IsEmpty());
+	Frames.Empty(NumExpectedFrames);
+	VolumeBoundsMin = FIntVector(INT32_MAX, INT32_MAX, INT32_MAX);
+	VolumeBoundsMax = FIntVector(INT32_MIN, INT32_MIN, INT32_MIN);
+
+	InitState = EInitState::Pending;
+
+	return true;
+#else
+	return false;
+#endif
+}
+
+bool UStreamableSparseVolumeTexture::AppendFrame(FSparseVolumeTextureData& UncookedFrame)
+{
+#if WITH_EDITORONLY_DATA
+	if (InitState != EInitState::Pending)
 	{
-		VolumeBoundsMin.X = FMath::Min(VolumeBoundsMin.X, Header.VirtualVolumeAABBMin.X);
-		VolumeBoundsMin.Y = FMath::Min(VolumeBoundsMin.Y, Header.VirtualVolumeAABBMin.Y);
-		VolumeBoundsMin.Z = FMath::Min(VolumeBoundsMin.Z, Header.VirtualVolumeAABBMin.Z);
-
-		VolumeBoundsMax.X = FMath::Max(VolumeBoundsMax.X, Header.VirtualVolumeAABBMax.X);
-		VolumeBoundsMax.Y = FMath::Max(VolumeBoundsMax.Y, Header.VirtualVolumeAABBMax.Y);
-		VolumeBoundsMax.Z = FMath::Max(VolumeBoundsMax.Z, Header.VirtualVolumeAABBMax.Z);
-	};
-
-	// Compute union of all frame AABBs
-	FIntVector3 VolumeBoundsMin = FIntVector3(INT32_MAX, INT32_MAX, INT32_MAX);
-	FIntVector3 VolumeBoundsMax = FIntVector3(INT32_MIN, INT32_MIN, INT32_MIN);
-	for (const FSparseVolumeTextureData& Data : InUncookedData)
-	{
-		ExpandVolumeBounds(Data.Header, VolumeBoundsMin, VolumeBoundsMax);
+		UE_LOG(LogSparseVolumeTexture, Error, TEXT("Tried to call UStreamableSparseVolumeTexture::AppendFrame() while not in the Pending init state."));
+		return false;
 	}
 
 	// The the minimum of the union of all frame AABBs should ideally be at (0, 0, 0), but it should also be fine if it is greater than that.
 	// A mimimum of less than (0, 0, 0) is not permitted.
-	if (VolumeBoundsMin.X < 0 || VolumeBoundsMin.Y < 0 || VolumeBoundsMin.Z < 0)
+	if (UncookedFrame.Header.VirtualVolumeAABBMin.X < 0 || UncookedFrame.Header.VirtualVolumeAABBMin.Y < 0 || UncookedFrame.Header.VirtualVolumeAABBMin.Z < 0)
 	{
-		UE_LOG(LogSparseVolumeTexture, Error, TEXT("Tried to initialize a SparseVolumeTexture with a VirtualVolumeAABBMin < 0 (%i, %i, %i)"),
+		UE_LOG(LogSparseVolumeTexture, Error, TEXT("Tried to add a frame to a SparseVolumeTexture with a VirtualVolumeAABBMin < 0 (%i, %i, %i)"),
 			VolumeBoundsMin.X, VolumeBoundsMin.Y, VolumeBoundsMin.Z);
 		return false;
 	}
+
+	// Compute union of all frame AABBs
+	VolumeBoundsMin.X = FMath::Min(VolumeBoundsMin.X, UncookedFrame.Header.VirtualVolumeAABBMin.X);
+	VolumeBoundsMin.Y = FMath::Min(VolumeBoundsMin.Y, UncookedFrame.Header.VirtualVolumeAABBMin.Y);
+	VolumeBoundsMin.Z = FMath::Min(VolumeBoundsMin.Z, UncookedFrame.Header.VirtualVolumeAABBMin.Z);
+	VolumeBoundsMax.X = FMath::Max(VolumeBoundsMax.X, UncookedFrame.Header.VirtualVolumeAABBMax.X);
+	VolumeBoundsMax.Y = FMath::Max(VolumeBoundsMax.Y, UncookedFrame.Header.VirtualVolumeAABBMax.Y);
+	VolumeBoundsMax.Z = FMath::Max(VolumeBoundsMax.Z, UncookedFrame.Header.VirtualVolumeAABBMax.Z);
+
+	VolumeResolution = VolumeBoundsMax;
+	Frames.AddDefaulted();
+	{
+		UE::Serialization::FEditorBulkDataWriter RawDataArchiveWriter(Frames.Last().RawData);
+		UncookedFrame.Serialize(RawDataArchiveWriter);
+	}
+
+	return true;
+#else
+	return false;
+#endif
+}
+
+bool UStreamableSparseVolumeTexture::EndInitialize(int32 InNumMipLevels)
+{
+#if WITH_EDITORONLY_DATA
+	if (InitState != EInitState::Pending)
+	{
+		UE_LOG(LogSparseVolumeTexture, Error, TEXT("Tried to call UStreamableSparseVolumeTexture::EndInitialize() while not in the Pending init state."));
+		return false;
+	}
+
+	// Ensure that at least one frame of data exists
+	if (Frames.IsEmpty())
+	{
+		UE_LOG(LogSparseVolumeTexture, Warning, TEXT("SVT has zero frames! Adding a dummy frame. SVT: %s"), *GetName());
+		FSparseVolumeTextureData DummyFrame;
+		DummyFrame.ConstructDefault();
+		AppendFrame(DummyFrame);
+	}
+
+	check(VolumeResolution.X > 0 && VolumeResolution.Y > 0 && VolumeResolution.Z > 0);
+	check(VolumeBoundsMin.X >= 0 && VolumeBoundsMin.Y >= 0 && VolumeBoundsMin.Z >= 0);
+
 	if (VolumeBoundsMin.X > 0 || VolumeBoundsMin.Y > 0 || VolumeBoundsMin.Z > 0)
 	{
 		UE_LOG(LogSparseVolumeTexture, Warning, TEXT("Initialized a SparseVolumeTexture with a VirtualVolumeAABBMin > 0 (%i, %i, %i). This wastes memory"),
 			VolumeBoundsMin.X, VolumeBoundsMin.Y, VolumeBoundsMin.Z);
 	}
 
-	VolumeResolution = VolumeBoundsMax;
-	Frames.SetNum(InUncookedData.Num());
-
 	const int32 NumMipLevelsFullMipChain = SVTComputeNumMipLevels(VolumeResolution);
 	check(NumMipLevelsFullMipChain > 0);
 
 	NumMipLevels = (InNumMipLevels <= INDEX_NONE) ? NumMipLevelsFullMipChain : FMath::Clamp(InNumMipLevels, 1, NumMipLevelsFullMipChain);
 
-#if WITH_EDITORONLY_DATA
-	for (int32 FrameIndex = 0; FrameIndex < Frames.Num(); ++FrameIndex)
+	InitState = EInitState::Done;
+
+	return true;
+#else
+	return false;
+#endif
+}
+
+bool UStreamableSparseVolumeTexture::Initialize(const TArrayView<FSparseVolumeTextureData>& InUncookedData, int32 InNumMipLevels)
+{
+	if (InUncookedData.IsEmpty())
 	{
-		UE::Serialization::FEditorBulkDataWriter RawDataArchiveWriter(Frames[FrameIndex].RawData);
-		InUncookedData[FrameIndex].Serialize(RawDataArchiveWriter);
+		UE_LOG(LogSparseVolumeTexture, Error, TEXT("Tried to initialize a SparseVolumeTexture with no frames"));
+		return false;
 	}
-#endif
+
+	if (!BeginInitialize(InUncookedData.Num()))
+	{
+		return false;
+	}
+	for (FSparseVolumeTextureData& UncookedFrame : InUncookedData)
+	{
+		if (!AppendFrame(UncookedFrame))
+		{
+			return false;
+		}
+	}
+	if (!EndInitialize(InNumMipLevels))
+	{
+		return false;
+	}
 
 	return true;
 }
 
+const FSparseVolumeTextureSceneProxy* UStreamableSparseVolumeTexture::GetStreamedFrameProxyOrFallback(int32 FrameIndex, int32 MipLevel) const
+{
+#if WITH_EDITORONLY_DATA
+	if (InitState != EInitState::Done)
+	{
+		UE_LOG(LogSparseVolumeTexture, Warning, TEXT("Tried to call GetStreamedFrameProxyOrFallback() on uninitialized SVT: %s"), *GetName());
+		return nullptr;
+	}
+#endif
+	if (Frames.IsEmpty())
+	{
+		UE_LOG(LogSparseVolumeTexture, Warning, TEXT("SVT is empty and has no frames at all! SVT: %s"), *GetName());
+		return nullptr;
+	}
+	FrameIndex = FMath::Clamp(FrameIndex, 0, Frames.Num() - 1);
+	ISparseVolumeTextureStreamingManager& StreamingManager = IStreamingManager::Get().GetSparseVolumeTextureStreamingManager();
+	const FSparseVolumeTextureSceneProxy* Proxy = StreamingManager.GetSparseVolumeTextureSceneProxy(this, FrameIndex, MipLevel, true);
+
+	int32 FallbackFrameIndex = FrameIndex;
+	while (!Proxy)
+	{
+		FallbackFrameIndex = FallbackFrameIndex > 0 ? (FallbackFrameIndex - 1) : (Frames.Num() - 1);
+		if (FallbackFrameIndex == FrameIndex)
+		{
+			UE_LOG(LogSparseVolumeTexture, Warning, TEXT("Failed to get ANY streamed SparseVolumeTexture frame  SVT: %s, FrameIndex: %i"), *GetName(), FrameIndex);
+			return nullptr;
+		}
+		Proxy = StreamingManager.GetSparseVolumeTextureSceneProxy(this, FallbackFrameIndex, MipLevel, false);
+	}
+
+	return Proxy;
+}
+
 void UStreamableSparseVolumeTexture::PostLoad()
@@ -415,6 +560,32 @@
 {
 	Super::Serialize(Ar);
 
+	// To ensure that we keep the same binary format between builds with and without editor-only data, we serialize dummy data if the editor-only data is stripped.
+	// SVT_TODO: Is there a cleaner way of doing this?
+	FIntVector* VolumeBoundsMinPtr = nullptr;
+	FIntVector* VolumeBoundsMaxPtr = nullptr;
+	EInitState* InitStatePtr = nullptr;
+
+#if WITH_EDITORONLY_DATA
+	// Check that we are not trying to cook unitialized data!
+	check(!Ar.IsCooking() || InitState == EInitState::Done);
+
+	VolumeBoundsMinPtr = &VolumeBoundsMin;
+	VolumeBoundsMaxPtr = &VolumeBoundsMax;
+	InitStatePtr = &InitState;
+#else
+	FIntVector VolumeBoundsMinDummy{};
+	FIntVector VolumeBoundsMaxDummy{};
+	EInitState InitStateDummy{};
+	VolumeBoundsMinPtr = &VolumeBoundsMinDummy;
+	VolumeBoundsMaxPtr = &VolumeBoundsMaxDummy;
+	InitStatePtr = &InitStateDummy;
+#endif
+
+	UE::SVT::Private::SerializeEnumAs<uint8>(Ar, *InitStatePtr);
+	Ar << *VolumeBoundsMinPtr;
+	Ar << *VolumeBoundsMaxPtr;
+
 	int32 NumFrames = Frames.Num();
 	Ar << NumFrames;
 
@@ -449,7 +620,7 @@
 void UStreamableSparseVolumeTexture::GetResourceSizeEx(FResourceSizeEx& CumulativeResourceSize)
 {
 	Super::GetResourceSizeEx(CumulativeResourceSize);
-	SIZE_T SizeCPU = sizeof(Frames);
+	SIZE_T SizeCPU = sizeof(*this);
 	SIZE_T SizeGPU = 0;
 	SizeCPU += Frames.GetAllocatedSize();
 	for (const FSparseVolumeTextureFrame& Frame : Frames)
@@ -465,84 +636,16 @@
 	CumulativeResourceSize.AddDedicatedVideoMemoryBytes(SizeGPU);
 }
 
-const FSparseVolumeTextureSceneProxy* UStreamableSparseVolumeTexture::GetStreamedFrameProxyOrFallback(int32 FrameIndex, int32 MipLevel) const
-{
-	if (Frames.IsEmpty())
-	{
-		UE_LOG(LogSparseVolumeTexture, Warning, TEXT("SVT is empty and has no frames at all! SVT: %s"), *GetName());
-		return nullptr;
-	}
-	FrameIndex = FMath::Clamp(FrameIndex, 0, Frames.Num() - 1);
-	ISparseVolumeTextureStreamingManager& StreamingManager = IStreamingManager::Get().GetSparseVolumeTextureStreamingManager();
-	const FSparseVolumeTextureSceneProxy* Proxy = StreamingManager.GetSparseVolumeTextureSceneProxy(this, FrameIndex, MipLevel, true);
-
-	int32 FallbackFrameIndex = FrameIndex;
-	while (!Proxy)
-	{
-		FallbackFrameIndex = FallbackFrameIndex > 0 ? (FallbackFrameIndex - 1) : (Frames.Num() - 1);
-		if (FallbackFrameIndex == FrameIndex)
-		{
-			UE_LOG(LogSparseVolumeTexture, Warning, TEXT("Failed to get ANY streamed SparseVolumeTexture frame  SVT: %s, FrameIndex: %i"), *GetName(), FrameIndex);
-			return nullptr;
-		}
-		Proxy = StreamingManager.GetSparseVolumeTextureSceneProxy(this, FallbackFrameIndex, MipLevel, false);
-	}
-
-	return Proxy;
-}
-
-TArrayView<const FSparseVolumeTextureFrame> UStreamableSparseVolumeTexture::GetFrames() const
-{
-	return Frames;
-}
-
-#if WITH_EDITOR
-void UStreamableSparseVolumeTexture::NotifyMaterials(const ENotifyMaterialsEffectOnShaders EffectOnShaders)
-{
-	// Create a material update context to safely update materials.
-	{
-		FMaterialUpdateContext UpdateContext;
-
-		// Notify any material that uses this texture
-		TSet<UMaterial*> BaseMaterialsThatUseThisTexture;
-		for (TObjectIterator<UMaterialInterface> It; It; ++It)
-		{
-			UMaterialInterface* MaterialInterface = *It;
-			if (!FPlatformProperties::IsServerOnly() && MaterialInterface->GetReferencedTextures().Contains(this))
-			{
-				UpdateContext.AddMaterialInterface(MaterialInterface);
-				// This is a bit tricky. We want to make sure all materials using this texture are
-				// updated. Materials are always updated. Material instances may also have to be
-				// updated and if they have static permutations their children must be updated
-				// whether they use the texture or not! The safe thing to do is to add the instance's
-				// base material to the update context causing all materials in the tree to update.
-				BaseMaterialsThatUseThisTexture.Add(MaterialInterface->GetMaterial());
-			}
-		}
-
-		// Go ahead and update any base materials that need to be.
-		if (EffectOnShaders == ENotifyMaterialsEffectOnShaders::Default)
-		{
-			for (TSet<UMaterial*>::TConstIterator It(BaseMaterialsThatUseThisTexture); It; ++It)
-			{
-				(*It)->PostEditChange();
-			}
-		}
-		else
-		{
-			FPropertyChangedEvent EmptyPropertyUpdateStruct(nullptr);
-			for (TSet<UMaterial*>::TConstIterator It(BaseMaterialsThatUseThisTexture); It; ++It)
-			{
-				(*It)->PostEditChangePropertyInternal(EmptyPropertyUpdateStruct, UMaterial::EPostEditChangeEffectOnShaders::DoesNotInvalidate);
-			}
-		}
-	}
-}
-#endif //WITH_EDITOR
-
 void UStreamableSparseVolumeTexture::GenerateOrLoadDDCRuntimeDataAndCreateSceneProxy()
 {
 #if WITH_EDITORONLY_DATA
+
+	if (InitState != EInitState::Done)
+	{
+		UE_LOG(LogSparseVolumeTexture, Warning, TEXT("Tried to cache derived data of an uninitialized SVT: %s"), *GetName());
+		return;
+	}
+
 	UE::DerivedData::FRequestOwner DDCRequestOwner(UE::DerivedData::EPriority::Normal);
 	{
 		UE::DerivedData::FRequestBarrier DDCRequestBarrier(DDCRequestOwner);
@@ -580,6 +683,12 @@
 #if WITH_EDITORONLY_DATA
 	using namespace UE::DerivedData;
 
+	if (InitState != EInitState::Done)
+	{
+		UE_LOG(LogSparseVolumeTexture, Warning, TEXT("Tried to cache derived data of an uninitialized SVT: %s"), *GetName());
+		return;
+	}
+
 	static const FString SparseVolumeTextureDDCVersion = TEXT("381AE2A9-A903-4C8F-8486-891E24D6EC70");	// Bump this if you want to ignore all cached data so far.
 	const FString DerivedDataKey = Frame.RawData.GetIdentifier().ToString() 
 		+ FString::Format(TEXT("{0},{1},{2}"), {AddressX.GetIntValue(), AddressY.GetIntValue(), AddressZ.GetIntValue()}) 
@@ -658,19 +767,14 @@
 {
 }
 
-bool UStaticSparseVolumeTexture::InitializeFromUncooked(const TArrayView<FSparseVolumeTextureData>& InUncookedData, int32 InNumMipLevels)
+bool UStaticSparseVolumeTexture::AppendFrame(FSparseVolumeTextureData& UncookedFrame)
 {
-	if (InUncookedData.Num() > 1)
+	if (!Frames.IsEmpty())
 	{
-		UE_LOG(LogSparseVolumeTexture, Error, TEXT("Tried to initialize a UStaticSparseVolumeTexture with more than 1 frame: %i"), InUncookedData.Num());
+		UE_LOG(LogSparseVolumeTexture, Error, TEXT("Tried to initialize a UStaticSparseVolumeTexture with more than 1 frame"));
 		return false;
 	}
-	return Super::InitializeFromUncooked(InUncookedData, InNumMipLevels);
+	return Super::AppendFrame(UncookedFrame);
 }
 
-bool UStaticSparseVolumeTexture::InitializeFromUncooked(FSparseVolumeTextureData& InUncookedData, int32 InNumMipLevels)
-{
-	return Super::InitializeFromUncooked(MakeArrayView<FSparseVolumeTextureData>(&InUncookedData, 1), InNumMipLevels);
-}
-
 ////////////////////////////////////////////////////////////////////////////////////////////////
diff --git a/Engine/Source/Runtime/Engine/Private/SparseVolumeTexture/SparseVolumeTextureData.cpp b/Engine/Source/Runtime/Engine/Private/SparseVolumeTexture/SparseVolumeTextureData.cpp
--- a/Engine/Source/Runtime/Engine/Private/SparseVolumeTexture/SparseVolumeTextureData.cpp
+++ b/Engine/Source/Runtime/Engine/Private/SparseVolumeTexture/SparseVolumeTextureData.cpp
@@ -237,6 +237,25 @@ bool FSparseVolumeTextureData::Construct(const ISparseVolumeTextureDataConstruct
 	return true;
 }
 
+void FSparseVolumeTextureData::ConstructDefault()
+{
+	// Store a single 1x1x1 mip level with the page table pointing to a (zero-valued) null tile.
+
+	Header.VirtualVolumeResolution = FIntVector3(1, 1, 1);
+	Header.VirtualVolumeAABBMin = FIntVector3::ZeroValue;
+	Header.VirtualVolumeAABBMax = Header.VirtualVolumeResolution;
+	Header.PageTableVolumeResolution = FIntVector3(1, 1, 1);
+	Header.PageTableVolumeAABBMin = FIntVector3::ZeroValue;
+	Header.PageTableVolumeAABBMax = Header.PageTableVolumeResolution;
+	Header.AttributesFormats[0] = PF_R8;
+
+	MipMaps.SetNum(1);
+	MipMaps[0].PageTable.SetNumZeroed(1);
+	MipMaps[0].PhysicalTileDataA.Empty();
+	MipMaps[0].PhysicalTileDataB.Empty();
+	MipMaps[0].NumPhysicalTiles = 0;
+}
+
 uint32 FSparseVolumeTextureData::ReadPageTable(const FIntVector3& PageTableCoord, int32 MipLevel) const
 {
 	if (!MipMaps.IsValidIndex(MipLevel))
diff --git a/Engine/Source/Runtime/Engine/Public/SparseVolumeTexture/SparseVolumeTextureData.h b/Engine/Source/Runtime/Engine/Public/SparseVolumeTexture/SparseVolumeTextureData.h
--- a/Engine/Source/Runtime/Engine/Public/SparseVolumeTexture/SparseVolumeTextureData.h
+++ b/Engine/Source/Runtime/Engine/Public/SparseVolumeTexture/SparseVolumeTextureData.h
@@ -55,6 +55,7 @@ struct ENGINE_API FSparseVolumeTextureData
 	void Serialize(FArchive& Ar);
 	
 	bool Construct(const ISparseVolumeTextureDataConstructionAdapter& Adapter);
+	void ConstructDefault();
 	FVector4f Load(const FIntVector3& VolumeCoord, int32 MipLevel, int32 AttributesIdx, const FSparseVolumeTextureDataAddressingInfo& AddressingInfo) const;
 	bool BuildDerivedData(const FSparseVolumeTextureDataAddressingInfo& AddressingInfo, int32 NumMipLevels, bool bMoveMip0FromThis, FSparseVolumeTextureData& OutDerivedData);
 
