Add CVar to always update GuidReferenceMap entries for NetSerialize structs with object properties, in order to prevent issues with old property values being applied when the object reference becomes mapped.

diff --git a/Engine/Source/Runtime/Engine/Private/RepLayout.cpp b/Engine/Source/Runtime/Engine/Private/RepLayout.cpp
--- a/Engine/Source/Runtime/Engine/Private/RepLayout.cpp
+++ b/Engine/Source/Runtime/Engine/Private/RepLayout.cpp
@@ -106,6 +106,10 @@ namespace UE::Net::Private
 	 */
 	static bool bReplicateCustomDeltaPropertiesInRepIndexOrder = false;
 	static FAutoConsoleVariableRef CVarReplicateCustomDeltaPropertiesInRepIndexOrder(TEXT("net.ReplicateCustomDeltaPropertiesInRepIndexOrder"), bReplicateCustomDeltaPropertiesInRepIndexOrder, TEXT("If false (default) custom delta properties will replicate in the same order as they're added to the lifetime property array during the call to GetLifetimeReplicatedProps. If true custom delta properties will be replicated in the property RepIndex order, which is typically in increasing property offset order. Note that custom delta properties are always serialized after regular properties."));
+
+	static bool bAlwaysUpdateGuidReferenceMapForNetSerializeObjectStruct = false;
+	static FAutoConsoleVariableRef CVarAlwaysUpdateGuidReferenceMapForNetSerializeStruct(TEXT("net.AlwaysUpdateGuidReferenceMapForNetSerializeObjectStruct"), bAlwaysUpdateGuidReferenceMapForNetSerializeObjectStruct,
+		TEXT("Requires net.TrackNetSerializeObjectReferences. If true, entries in the GuidReferenceMap for NetSerialize struct properties with object properties will always be updated, not just when the Guid changes or goes NULL. This should prevent issues with old property data being applied when an unmapped actor ref in the struct is mapped."));
 }
 
 extern int32 GNumSharedSerializationHit;
@@ -3397,6 +3401,13 @@ static bool ReceivePropertyHelper(
 				GuidReferencesMap->Add(AbsOffset, FGuidReferences(Bunch, Mark, TrackedUnmappedGuids, TrackedDynamicMappedGuids, Cmd.ParentIndex, CmdIndex));
 				bOutGuidsChanged = true;
 			}
+			else if (UE::Net::Private::bAlwaysUpdateGuidReferenceMapForNetSerializeObjectStruct && Cmd.Type == ERepLayoutCmdType::NetSerializeStructWithObjectReferences)
+			{
+				// If this is a NetSerialize struct with object references, there may be other properties "wrapped up" with this GUID reference.
+				// In this case, the entry in the map should be always be updated, so there isn't outdated data in the entry that also gets
+				// applied when the Guid possibly goes unmapped and then mapped later.
+				GuidReferencesMap->Add(AbsOffset, FGuidReferences(Bunch, Mark, TrackedUnmappedGuids, TrackedDynamicMappedGuids, Cmd.ParentIndex, CmdIndex));
+			}
 		}
 		else
 		{
