Hierarchical ISM - Consider MaxWorldPositionOffsetDistance of materials when performing CPU culling of the HISM clusters

diff --git a/Engine/Source/Runtime/Engine/Private/HierarchicalInstancedStaticMesh.cpp b/Engine/Source/Runtime/Engine/Private/HierarchicalInstancedStaticMesh.cpp
--- a/Engine/Source/Runtime/Engine/Private/HierarchicalInstancedStaticMesh.cpp
+++ b/Engine/Source/Runtime/Engine/Private/HierarchicalInstancedStaticMesh.cpp
@@ -910,9 +910,12 @@ void FHierarchicalStaticMeshSceneProxy::SetupOcclusion(UHierarchicalInstancedSta
 		//UE_LOG(LogTemp, Display, TEXT("Occlusion level %d   %d inst / node"), NumNodes, NumPerNode);
 		OcclusionBounds.Reserve(NumNodes);
 		FMatrix XForm = InComponent->GetRenderMatrix();
+
+		const float MaxWorldPositionOffset = GetMaxWorldPositionOffsetDisplacement();
+
 		for (int32 Index = FirstOcclusionNode; Index <= LastOcclusionNode; Index++)
 		{
-			OcclusionBounds.Add(FBoxSphereBounds(FBox(ClusterTree[Index].BoundMin, ClusterTree[Index].BoundMax).TransformBy(XForm)));
+			OcclusionBounds.Add(FBoxSphereBounds(FBox(ClusterTree[Index].BoundMin, ClusterTree[Index].BoundMax).ExpandBy(MaxWorldPositionOffset).TransformBy(XForm)));
 		}
 	}
 }
@@ -1055,7 +1058,7 @@ struct FFoliageCullInstanceParams : public FFoliageRenderInstanceParams
 	int32 LastOcclusionNode;
 	const TArray<bool>* OcclusionResults;
 	int32 OcclusionResultsStart;
-
+	float MaxWPODisplacement;
 
 
 	FFoliageCullInstanceParams(bool InbNeedsSingleLODRuns, bool InbNeedsMultipleLODRuns, bool InbOverestimate, const TArray<FClusterNode>& InTree)
@@ -1065,6 +1068,7 @@ struct FFoliageCullInstanceParams : public FFoliageRenderInstanceParams
 	,	LastOcclusionNode(-1)
 	,	OcclusionResults(nullptr)
 	,	OcclusionResultsStart(0)
+	,	MaxWPODisplacement(0)
 	{
 	}
 };
@@ -1098,15 +1102,15 @@ static FAutoConsoleCommand LogFoliageFrameCmd(
 const VectorRegister		VECTOR_HALF_HALF_HALF_ZERO				= DECLARE_VECTOR_REGISTER(0.5f, 0.5f, 0.5f, 0.0f);
 
 template<bool TUseVector>
-static FORCEINLINE_DEBUGGABLE bool CullNode(const FFoliageCullInstanceParams& Params, const FClusterNode& Node, bool& bOutFullyContained)
+static FORCEINLINE_DEBUGGABLE bool CullNode(const FFoliageCullInstanceParams& Params, const FVector& BoundMin, const FVector& BoundMax, bool& bOutFullyContained)
 {
 	if (TUseVector)
 	{
 		checkSlow(Params.ViewFrustumLocal.PermutedPlanes.Num() == 4);
 
 		//@todo, once we have more than one mesh per tree, these should be aligned
-		VectorRegister BoxMin = VectorLoadFloat3(&Node.BoundMin);
-		VectorRegister BoxMax = VectorLoadFloat3(&Node.BoundMax);
+		VectorRegister BoxMin = VectorLoadFloat3(&BoundMin);
+		VectorRegister BoxMax = VectorLoadFloat3(&BoundMax);
 
 		VectorRegister BoxDiff = VectorSubtract(BoxMax,BoxMin);
 		VectorRegister BoxSum = VectorAdd(BoxMax,BoxMin);
@@ -1145,8 +1149,8 @@ static FORCEINLINE_DEBUGGABLE bool CullNode(const FFoliageCullInstanceParams& Pa
 		// Check for completely outside
 		return !!VectorAnyGreaterThan(Distance,PushOut);
 	}
-	FVector Center = (FVector)(Node.BoundMin + Node.BoundMax) * 0.5f;
-	FVector Extent = (FVector)(Node.BoundMax - Node.BoundMin) * 0.5f;
+	FVector Center = (BoundMin + BoundMax) * 0.5f;
+	FVector Extent = (BoundMax - BoundMin) * 0.5f;
 	if (!Params.ViewFrustumLocal.IntersectBox(Center, Extent, bOutFullyContained)) 
 	{
 		return true;
@@ -1190,13 +1194,23 @@ inline bool CanGroup(const FVector& BoundMin, const FVector& BoundMax, const FVe
 
 
 
-template<bool TUseVector>
+template<bool TUseVector, bool THasWPODisplacement>
 void FHierarchicalStaticMeshSceneProxy::Traverse(const FFoliageCullInstanceParams& Params, int32 Index, int32 MinLOD, int32 MaxLOD, bool bFullyContained) const
 {
 	const FClusterNode& Node = Params.Tree[Index];
+
+	FVector BoundMin = (FVector)Node.BoundMin;
+	FVector BoundMax = (FVector)Node.BoundMax;
+
+	if (THasWPODisplacement)
+	{
+		BoundMin -= FVector(Params.MaxWPODisplacement);
+		BoundMax += FVector(Params.MaxWPODisplacement);
+	}
+
 	if (!bFullyContained)
 	{
-		if (CullNode<TUseVector>(Params, Node, bFullyContained))
+		if (CullNode<TUseVector>(Params, BoundMin, BoundMax, bFullyContained))
 		{
 			return;
 		}
@@ -1204,7 +1218,7 @@ void FHierarchicalStaticMeshSceneProxy::Traverse(const FFoliageCullInstanceParam
 
 	if (MinLOD != MaxLOD)
 	{
-		CalcLOD(MinLOD, MaxLOD, (FVector)Node.BoundMin, (FVector)Node.BoundMax, Params.ViewOriginInLocalZero, Params.ViewOriginInLocalOne, Params.LODPlanesMin, Params.LODPlanesMax);
+		CalcLOD(MinLOD, MaxLOD, BoundMin, BoundMax, Params.ViewOriginInLocalZero, Params.ViewOriginInLocalOne, Params.LODPlanesMin, Params.LODPlanesMax);
 
 		if (MinLOD >= Params.LODs)
 		{
@@ -1225,7 +1239,7 @@ void FHierarchicalStaticMeshSceneProxy::Traverse(const FFoliageCullInstanceParam
 
 	bool bShouldGroup = Node.FirstChild < 0
 		|| ((Node.LastInstance - Node.FirstInstance + 1) < Params.MinInstancesToSplit[MinLOD]
-			&& CanGroup((FVector)Node.BoundMin, (FVector)Node.BoundMax, Params.ViewOriginInLocalZero, Params.ViewOriginInLocalOne, Params.LODPlanesMax[Params.LODs - 1]));
+			&& CanGroup(BoundMin, BoundMax, Params.ViewOriginInLocalZero, Params.ViewOriginInLocalOne, Params.LODPlanesMax[Params.LODs - 1]));
 	bool bSplit = (!bFullyContained || MinLOD < MaxLOD || Index < Params.FirstOcclusionNode)
 		&& !bShouldGroup;
 
@@ -1237,7 +1251,7 @@ void FHierarchicalStaticMeshSceneProxy::Traverse(const FFoliageCullInstanceParam
 	}
 	for (int32 ChildIndex = Node.FirstChild; ChildIndex <= Node.LastChild; ChildIndex++)
 	{
-		Traverse<TUseVector>(Params, ChildIndex, MinLOD, MaxLOD, bFullyContained);
+		Traverse<TUseVector, THasWPODisplacement>(Params, ChildIndex, MinLOD, MaxLOD, bFullyContained);
 	}
 }
 
@@ -1657,6 +1671,7 @@ void FHierarchicalStaticMeshSceneProxy::GetDynamicMeshElements(const TArray<cons
 
 				InstanceParams.ViewOriginInLocalZero = WorldToLocal.TransformPosition(View->GetTemporalLODOrigin(0, bMultipleSections));
 				InstanceParams.ViewOriginInLocalOne = WorldToLocal.TransformPosition(View->GetTemporalLODOrigin(1, bMultipleSections));
+				InstanceParams.MaxWPODisplacement = GetMaxWorldPositionOffsetDisplacement();
 
 				float MinSize = bIsOrtho ? 0.0f : CVarFoliageMinimumScreenSize.GetValueOnRenderThread();
 				float LODScale = UserData_AllInstances.LODDistanceScale;
@@ -1666,7 +1681,7 @@ void FHierarchicalStaticMeshSceneProxy::GetDynamicMeshElements(const TArray<cons
 								
 				FVector AverageScale(InstanceParams.Tree[0].MinInstanceScale + (InstanceParams.Tree[0].MaxInstanceScale - InstanceParams.Tree[0].MinInstanceScale) / 2.0f);
 				FBoxSphereBounds ScaledBounds = RenderData->Bounds.TransformBy(FTransform(FRotator::ZeroRotator, FVector::ZeroVector, AverageScale));
-				float SphereRadius = ScaledBounds.SphereRadius;
+				float SphereRadius = ScaledBounds.SphereRadius + InstanceParams.MaxWPODisplacement;
 
 				float FinalCull = MAX_flt;
 				if (MinSize > 0.0)
@@ -1792,13 +1807,29 @@ void FHierarchicalStaticMeshSceneProxy::GetDynamicMeshElements(const TArray<cons
 
 					if (CVarCullAll.GetValueOnRenderThread() < 1)
 					{
+						const bool bHasWPODisplacement = InstanceParams.MaxWPODisplacement != 0.0f;
+												
 						if (bUseVectorCull)
 						{
-							Traverse<true>(InstanceParams, 0, UseMinLOD, UseMaxLOD, bDisableCull);
+							if (bHasWPODisplacement)
+							{
+								Traverse<true, true>(InstanceParams, 0, UseMinLOD, UseMaxLOD, bDisableCull);
+							}
+							else
+							{
+								Traverse<true, false>(InstanceParams, 0, UseMinLOD, UseMaxLOD, bDisableCull);
+							}
 						}
 						else
 						{
-							Traverse<false>(InstanceParams, 0, UseMinLOD, UseMaxLOD, bDisableCull);
+							if (bHasWPODisplacement)
+							{
+								Traverse<false, true>(InstanceParams, 0, UseMinLOD, UseMaxLOD, bDisableCull);
+							}
+							else
+							{
+								Traverse<false, false>(InstanceParams, 0, UseMinLOD, UseMaxLOD, bDisableCull);
+							}
 						}
 					}
 #if !(UE_BUILD_SHIPPING || UE_BUILD_TEST)
@@ -1861,7 +1892,7 @@ void FHierarchicalStaticMeshSceneProxy::GetDynamicMeshElements(const TArray<cons
 						int MaxEndCullDistance = CVarFoliageMaxEndCullDistance.GetValueOnRenderThread();
 						const float LODRandom = CVarRandomLODRange.GetValueOnRenderThread();
 						const float MaxDrawDistanceScale = GetCachedScalabilityCVars().ViewDistanceScale;
-						const float SphereRadius = RenderData->Bounds.SphereRadius;
+						const float SphereRadius = RenderData->Bounds.SphereRadius + GetMaxWorldPositionOffsetDisplacement();
 
 						checkSlow(NumLODs > 0);
 
@@ -1951,10 +1982,11 @@ void FHierarchicalStaticMeshSceneProxy::GetDynamicMeshElements(const TArray<cons
 			if (View->Family->EngineShowFlags.HISMCClusterTree)
 			{
 				FColor StartingColor(100, 0, 0);
+				const float MaxWorldPositionOffset = GetMaxWorldPositionOffsetDisplacement();
 
 				for (const FClusterNode& CulsterNode : ClusterTree)
 				{
-					DrawWireBox(Collector.GetPDI(ViewIndex), GetLocalToWorld(), FBox(CulsterNode.BoundMin, CulsterNode.BoundMax), StartingColor, View->Family->EngineShowFlags.Game ? SDPG_World : SDPG_Foreground);
+					DrawWireBox(Collector.GetPDI(ViewIndex), GetLocalToWorld(), FBox(CulsterNode.BoundMin, CulsterNode.BoundMax).ExpandBy(MaxWPODisplacement), StartingColor, View->Family->EngineShowFlags.Game ? SDPG_World : SDPG_Foreground);
 					StartingColor.R += 5;
 					StartingColor.G += 5;
 					StartingColor.B += 5;
diff --git a/Engine/Source/Runtime/Engine/Public/HierarchicalStaticMeshSceneProxy.h b/Engine/Source/Runtime/Engine/Public/HierarchicalStaticMeshSceneProxy.h
--- a/Engine/Source/Runtime/Engine/Public/HierarchicalStaticMeshSceneProxy.h
+++ b/Engine/Source/Runtime/Engine/Public/HierarchicalStaticMeshSceneProxy.h
@@ -93,6 +93,6 @@ public:
 
 	void FillDynamicMeshElements(FMeshElementCollector& Collector, const FFoliageElementParams& ElementParams, const FFoliageRenderInstanceParams& Instances) const;
 
-	template<bool TUseVector>
+	template<bool TUseVector, bool THasWPODisplacement>
 	void Traverse(const FFoliageCullInstanceParams& Params, int32 Index, int32 MinLOD, int32 MaxLOD, bool bFullyContained = false) const;
 };
