Don't use sRGB for debug canvas stereo layer

On some platforms, we use compute to clear, which usually needs special plumbing if writing to an sRGB format texture. In our case, we don't actually need sRGB for the texture.

Also fixes an underlying platform RHI disabled check, which improves debugging experience

diff --git a/Engine/Source/Runtime/Engine/Private/Slate/DebugCanvas.cpp b/Engine/Source/Runtime/Engine/Private/Slate/DebugCanvas.cpp
--- a/Engine/Source/Runtime/Engine/Private/Slate/DebugCanvas.cpp
+++ b/Engine/Source/Runtime/Engine/Private/Slate/DebugCanvas.cpp
@@ -256,7 +256,7 @@ void FDebugCanvasDrawer::DrawRenderThread(FRHICommandListImmediate& RHICmdList,
 			if (!LayerTexture)
 			{
 				// Set TexCreate_NoFastClear because the fast CMASK clear was not working on ps4.
-				FPooledRenderTargetDesc Desc(FPooledRenderTargetDesc::Create2DDesc(RenderThreadCanvas->GetParentCanvasSize(), PF_B8G8R8A8, FClearValueBinding(), TexCreate_SRGB, TexCreate_RenderTargetable | TexCreate_NoFastClear, false));
+				FPooledRenderTargetDesc Desc(FPooledRenderTargetDesc::Create2DDesc(RenderThreadCanvas->GetParentCanvasSize(), PF_B8G8R8A8, FClearValueBinding(), ETextureCreateFlags::None, ETextureCreateFlags::RenderTargetable | ETextureCreateFlags::NoFastClear, false));
 				Desc.DebugName = TEXT("DebugCanvasLayerTexture");
 				GRenderTargetPool.FindFreeElement(RHICmdList, Desc, LayerTexture, TEXT("DebugCanvasLayerTexture"));
 				UE_LOG(LogProfilingDebugging, Log, TEXT("Allocated a %d x %d texture for HMD canvas layer"), RenderThreadCanvas->GetParentCanvasSize().X, RenderThreadCanvas->GetParentCanvasSize().Y);
