Change displacement scaling center default from 0.0 to 0.5

diff --git a/Engine/Source/Runtime/Engine/Classes/Engine/EngineTypes.h b/Engine/Source/Runtime/Engine/Classes/Engine/EngineTypes.h
--- a/Engine/Source/Runtime/Engine/Classes/Engine/EngineTypes.h
+++ b/Engine/Source/Runtime/Engine/Classes/Engine/EngineTypes.h
@@ -2900,7 +2900,7 @@ struct ENGINE_API FDisplacementScaling
 public:
 	FDisplacementScaling()
 	: Magnitude(4.0f)
-	, Center(0.0f)
+	, Center(0.5f)
 	{
 	}
 
