Fix regression: Non-deterministic behavior with changing WorldAsset/CookedWorldAsset in PreSave (UDN Report)

diff --git a/Engine/Source/Runtime/Engine/Private/LevelInstance/LevelInstanceActor.cpp b/Engine/Source/Runtime/Engine/Private/LevelInstance/LevelInstanceActor.cpp
--- a/Engine/Source/Runtime/Engine/Private/LevelInstance/LevelInstanceActor.cpp
+++ b/Engine/Source/Runtime/Engine/Private/LevelInstance/LevelInstanceActor.cpp
@@ -152,6 +152,8 @@ void ALevelInstance::PostLoad()
 
 	LevelInstanceComponent = Cast<ULevelInstanceComponent>(RootComponent);
 	OnLevelInstanceActorPostLoad.Broadcast(this);
+
+	UpdateCookedAsset();
 }
 
 void ALevelInstance::PreEditUndo()
@@ -307,30 +309,24 @@ void ALevelInstance::PushLevelInstanceEditingStateToProxies(bool bInEditingState
 	LevelInstanceActorImpl.PushLevelInstanceEditingStateToProxies(bInEditingState);
 }
 
-void ALevelInstance::PreSave(FObjectPreSaveContext ObjectSaveContext)
+void ALevelInstance::ResetUnsupportedWorldAsset()
 {
-	Super::PreSave(ObjectSaveContext);
-
-	if (IsRunningCookCommandlet())
+	if (!ULevelInstanceSubsystem::CanUsePackage(*WorldAsset.GetLongPackageName()))
 	{
-		ResetUnsupportedWorldAsset();
-
-#if WITH_EDITORONLY_DATA
-		if (IsLoadingEnabled())
-		{
-			CookedWorldAsset = WorldAsset;
-		}
-#endif
+		UE_LOG(LogLevelInstance, Warning, TEXT("LevelInstance doesn't support partitioned world %s, make sure to flag world partition's 'Can be Used by Level Instance'."), *WorldAsset.GetLongPackageName());
+		WorldAsset.Reset();
+		UpdateCookedAsset();
 	}
 }
 
-void ALevelInstance::ResetUnsupportedWorldAsset()
+void ALevelInstance::UpdateCookedAsset()
 {
-	if (!ULevelInstanceSubsystem::CanUsePackage(*WorldAsset.GetLongPackageName()))
+#if WITH_EDITORONLY_DATA
+	if (IsRunningCookCommandlet() && IsLoadingEnabled())
 	{
-		UE_LOG(LogLevelInstance, Warning, TEXT("LevelInstance doesn't support partitioned world %s, make sure to flag world partition's 'Can be Used by Level Instance'."), *WorldAsset.GetLongPackageName());
-		WorldAsset.Reset();
+		CookedWorldAsset = WorldAsset;
 	}
+#endif
 }
 
 #endif
diff --git a/Engine/Source/Runtime/Engine/Public/LevelInstance/LevelInstanceActor.h b/Engine/Source/Runtime/Engine/Public/LevelInstance/LevelInstanceActor.h
--- a/Engine/Source/Runtime/Engine/Public/LevelInstance/LevelInstanceActor.h
+++ b/Engine/Source/Runtime/Engine/Public/LevelInstance/LevelInstanceActor.h
@@ -81,7 +81,6 @@ public:
 	virtual void PreEditChange(FProperty* PropertyThatWillChange) override;
 	virtual void PostEditChangeProperty(FPropertyChangedEvent& PropertyChangedEvent) override;
 	virtual void PostEditImport() override;
-	virtual void PreSave(FObjectPreSaveContext ObjectSaveContext) override;
 
 	// AActor overrides
 	virtual void CheckForErrors() override;
@@ -105,6 +104,7 @@ public:
 	static FOnLevelInstanceActorPostLoad OnLevelInstanceActorPostLoad;
 
 private:
+	void UpdateCookedAsset();
 	void ResetUnsupportedWorldAsset();
 #endif
 };
