LevelInstance: Check LevelInstanceID before accessing Subsystem (reported live coding issue in UDN)

diff --git a/Engine/Source/Runtime/Engine/Private/LevelInstance/LevelInstanceActorImpl.cpp b/Engine/Source/Runtime/Engine/Private/LevelInstance/LevelInstanceActorImpl.cpp
--- a/Engine/Source/Runtime/Engine/Private/LevelInstance/LevelInstanceActorImpl.cpp
+++ b/Engine/Source/Runtime/Engine/Private/LevelInstance/LevelInstanceActorImpl.cpp
@@ -38,14 +38,14 @@ void FLevelInstanceActorImpl::RegisterLevelInstance()
 
 void FLevelInstanceActorImpl::UnregisterLevelInstance()
 {
-	if (ULevelInstanceSubsystem* LevelInstanceSubsystem = LevelInstance->GetLevelInstanceSubsystem())
+	// If LevelInstance has already been unregistered it will have an Invalid LevelInstanceID. Avoid processing it.
+	if (!LevelInstanceID.IsValid())
 	{
-		// If LevelInstance has already been unregistered it will have an Invalid LevelInstanceID. Avoid processing it.
-		if (!LevelInstanceID.IsValid())
-		{
-			return;
-		}
+		return;
+	}
 
+	if (ULevelInstanceSubsystem* LevelInstanceSubsystem = LevelInstance->GetLevelInstanceSubsystem())
+	{
 		LevelInstanceSubsystem->UnregisterLevelInstance(LevelInstance);
 
 		LevelInstance->UnloadLevelInstance();
