Add optional use of RemoveAtSwap when removing instances from InstancedStaticMeshComponent.
Also added CVar to force the behavior: r.InstancedStaticMeshes.ForceRemoveAtSwap

diff --git a/Engine/Source/Runtime/Engine/Classes/Components/InstancedStaticMeshComponent.h b/Engine/Source/Runtime/Engine/Classes/Components/InstancedStaticMeshComponent.h
--- a/Engine/Source/Runtime/Engine/Classes/Components/InstancedStaticMeshComponent.h
+++ b/Engine/Source/Runtime/Engine/Classes/Components/InstancedStaticMeshComponent.h
@@ -292,7 +292,11 @@ class ENGINE_API UInstancedStaticMeshComponent : public UStaticMeshComponent, pu
 
 	virtual void PostLoad() override;
 	virtual void OnRegister() override;
-	virtual bool SupportsRemoveSwap() const { return false; }
+	
+	/** Sets to use RemoveAtSwap on instance removal. This is an optimization, but will change the resultant instance reordering. */
+	void SetRemoveSwap() { bSupportRemoveAtSwap = true; }
+	/** Returns true if RemoveAtSwap is enabled. Derived classes may always return true regardless of whether SetRemoveSwapEnabled() was called. */
+	virtual bool SupportsRemoveSwap() const;
 
 #if WITH_EDITOR
 	virtual bool ComponentIsTouchingSelectionBox(const FBox& InSelBBox, const bool bConsiderOnlyBSP, const bool bMustEncompassEntireComponent) const override;
@@ -439,6 +443,12 @@ public:
 	/** Recorded modifications to per-instance data */
 	FInstanceUpdateCmdBuffer InstanceUpdateCmdBuffer;
 
+	/** 
+	 * Flag for using RemoveAtSwap on instance removal. 
+	 * The implementation is free to ignore this flag, but should honor whatever behavior is being returned by SupportsRemoveSwap().
+	 */
+	bool bSupportRemoveAtSwap = false;
+
 private:
 
 	/** Sets up new instance data to sensible defaults, creates physics counterparts if possible. */
diff --git a/Engine/Source/Runtime/Engine/Private/InstancedStaticMesh.cpp b/Engine/Source/Runtime/Engine/Private/InstancedStaticMesh.cpp
--- a/Engine/Source/Runtime/Engine/Private/InstancedStaticMesh.cpp
+++ b/Engine/Source/Runtime/Engine/Private/InstancedStaticMesh.cpp
@@ -141,6 +141,11 @@ static TAutoConsoleVariable<int32> CVarEnableViewportSMInstanceSelection(
 );
 #endif
 
+static TAutoConsoleVariable<int32> CVarISMForceRemoveAtSwap(
+	TEXT("r.InstancedStaticMeshes.ForceRemoveAtSwap"),
+	0,
+	TEXT("Force the RemoveAtSwap optimization when removing instances from an ISM."));
+
 class FISMExecHelper : public FSelfRegisteringExec
 {
 	virtual bool Exec_Runtime(class UWorld* InWorld, const TCHAR* Cmd, FOutputDevice& Ar)
@@ -3413,12 +3418,21 @@ bool UInstancedStaticMeshComponent::SetCustomData(int32 InstanceIndex, TArrayVie
 	return true;
 }
 
+bool UInstancedStaticMeshComponent::SupportsRemoveSwap() const
+{
+	return bSupportRemoveAtSwap || CVarISMForceRemoveAtSwap.GetValueOnGameThread() != 0;
+}
+
 bool UInstancedStaticMeshComponent::RemoveInstanceInternal(int32 InstanceIndex, bool InstanceAlreadyRemoved)
 {
 #if WITH_EDITOR
 	DeletionState = InstanceAlreadyRemoved ? EInstanceDeletionReason::EntryAlreadyRemoved : EInstanceDeletionReason::EntryRemoval;
 #endif
 
+	// For performance we would prefer to use RemoveAtSwap() but some old code may be relying on the old
+	// RemoveAt() behavior, since there was no explicit contract about how instance indices can move around.
+	const bool bUseRemoveAtSwap = bSupportRemoveAtSwap || CVarISMForceRemoveAtSwap.GetValueOnGameThread() != 0;
+
 	// remove instance
 	if (!InstanceAlreadyRemoved && PerInstanceSMData.IsValidIndex(InstanceIndex))
 	{
@@ -3427,8 +3441,16 @@ bool UInstancedStaticMeshComponent::RemoveInstanceInternal(int32 InstanceIndex,
 		// Request navigation update
 		PartialNavigationUpdate(InstanceIndex);
 
-		PerInstanceSMData.RemoveAt(InstanceIndex);
-		PerInstanceSMCustomData.RemoveAt(InstanceIndex * NumCustomDataFloats, NumCustomDataFloats);
+		if (bUseRemoveAtSwap)
+		{
+			PerInstanceSMData.RemoveAtSwap(InstanceIndex, 1, false);
+			PerInstanceSMCustomData.RemoveAt(InstanceIndex * NumCustomDataFloats, NumCustomDataFloats, false);
+		}
+		else
+		{
+			PerInstanceSMData.RemoveAt(InstanceIndex);
+			PerInstanceSMCustomData.RemoveAt(InstanceIndex * NumCustomDataFloats, NumCustomDataFloats);
+		}
 
 		// If it's the last instance, unregister the component since component with no instances are not registered. 
 		// (because of GetInstanceCount() > 0 in UInstancedStaticMeshComponent::IsNavigationRelevant())
@@ -3443,10 +3465,19 @@ bool UInstancedStaticMeshComponent::RemoveInstanceInternal(int32 InstanceIndex,
 	// remove selection flag if array is filled in
 	if (SelectedInstances.IsValidIndex(InstanceIndex))
 	{
-		SelectedInstances.RemoveAt(InstanceIndex);
+		if (bUseRemoveAtSwap)
+		{
+			SelectedInstances.RemoveAtSwap(InstanceIndex);
+		}
+		else
+		{
+			SelectedInstances.RemoveAt(InstanceIndex);
+		}
 	}
 #endif
 
+	const int32 LastInstanceIndex = PerInstanceSMData.Num();
+
 	// update the physics state
 	if (bPhysicsStateCreated && InstanceBodies.IsValidIndex(InstanceIndex))
 	{
@@ -3473,17 +3504,32 @@ bool UInstancedStaticMeshComponent::RemoveInstanceInternal(int32 InstanceIndex,
 	// Notify that these instances have been removed/relocated
 	if (FInstancedStaticMeshDelegates::OnInstanceIndexUpdated.IsBound())
 	{
-		TArray<FInstancedStaticMeshDelegates::FInstanceIndexUpdateData, TInlineAllocator<128>> IndexUpdates;
-		IndexUpdates.Reserve(1 + (PerInstanceSMData.Num() - InstanceIndex));
-
-		IndexUpdates.Add(FInstancedStaticMeshDelegates::FInstanceIndexUpdateData{ FInstancedStaticMeshDelegates::EInstanceIndexUpdateType::Removed, InstanceIndex });
-		for (int32 MovedInstanceIndex = InstanceIndex; MovedInstanceIndex < PerInstanceSMData.Num(); ++MovedInstanceIndex)
+		if (bUseRemoveAtSwap)
 		{
-			// ISMs use standard remove, so each instance above our removal point is shuffled down by 1
-			IndexUpdates.Add(FInstancedStaticMeshDelegates::FInstanceIndexUpdateData{ FInstancedStaticMeshDelegates::EInstanceIndexUpdateType::Relocated, MovedInstanceIndex, MovedInstanceIndex + 1 });
+			TArray<FInstancedStaticMeshDelegates::FInstanceIndexUpdateData, TInlineAllocator<2>> IndexUpdates;
+			IndexUpdates.Add(FInstancedStaticMeshDelegates::FInstanceIndexUpdateData{ FInstancedStaticMeshDelegates::EInstanceIndexUpdateType::Removed, InstanceIndex });
+			if (InstanceIndex != LastInstanceIndex)
+			{
+				// used swap remove, so the last index has been moved to the spot we removed from
+				IndexUpdates.Add(FInstancedStaticMeshDelegates::FInstanceIndexUpdateData{ FInstancedStaticMeshDelegates::EInstanceIndexUpdateType::Relocated, InstanceIndex, LastInstanceIndex });
+			}
+
+			FInstancedStaticMeshDelegates::OnInstanceIndexUpdated.Broadcast(this, IndexUpdates);
 		}
+		else
+		{
+			TArray<FInstancedStaticMeshDelegates::FInstanceIndexUpdateData> IndexUpdates;
+			IndexUpdates.Reserve(IndexUpdates.Num() + 1 + (PerInstanceSMData.Num() - InstanceIndex));
 
-		FInstancedStaticMeshDelegates::OnInstanceIndexUpdated.Broadcast(this, IndexUpdates);
+			IndexUpdates.Add(FInstancedStaticMeshDelegates::FInstanceIndexUpdateData{ FInstancedStaticMeshDelegates::EInstanceIndexUpdateType::Removed, InstanceIndex });
+			for (int32 MovedInstanceIndex = InstanceIndex; MovedInstanceIndex < PerInstanceSMData.Num(); ++MovedInstanceIndex)
+			{
+				// ISMs use standard remove, so each instance above our removal point is shuffled down by 1
+				IndexUpdates.Add(FInstancedStaticMeshDelegates::FInstanceIndexUpdateData{ FInstancedStaticMeshDelegates::EInstanceIndexUpdateType::Relocated, MovedInstanceIndex, MovedInstanceIndex + 1 });
+			}
+
+			FInstancedStaticMeshDelegates::OnInstanceIndexUpdated.Broadcast(this, IndexUpdates);
+		}
 	}
 
 	// Force recreation of the render data
