Implement deep copy for sets in CopySinglePropertyRecursive and remove now unused bNeedsStringCopy.

#tests Used a tset of tobjectptrs locally that were not serializing properly.

diff --git a/Engine/Source/Runtime/Engine/Private/ObjectEditorUtils.cpp b/Engine/Source/Runtime/Engine/Private/ObjectEditorUtils.cpp
--- a/Engine/Source/Runtime/Engine/Private/ObjectEditorUtils.cpp
+++ b/Engine/Source/Runtime/Engine/Private/ObjectEditorUtils.cpp
@@ -150,7 +150,6 @@ namespace FObjectEditorUtils
 	static void CopySinglePropertyRecursive(UObject* SourceObject, const void* const InSourcePtr, FProperty* InSourceProperty, void* const InTargetPtr, UObject* InDestinationObject, FProperty* InDestinationProperty)
 	{
 		bool bNeedsShallowCopy = true;
-		bool bNeedsStringCopy = false;
 
 		if (FStructProperty* const DestStructProperty = CastField<FStructProperty>(InDestinationProperty))
 		{
@@ -225,20 +224,37 @@ namespace FObjectEditorUtils
 			}
 
 			bNeedsShallowCopy = false;
-			bNeedsStringCopy = false;
 		}
 		else if ( FSetProperty* const DestSetProperty = CastField<FSetProperty>(InDestinationProperty) )
 		{
-			//FSetProperty* const SrcSetProperty = CastField<FSetProperty>(InSourceProperty);
+			FSetProperty* const SrcSetProperty = CastField<FSetProperty>(InSourceProperty);
 
-			//check(InDestinationProperty->ArrayDim == 1);
-			//FScriptSetHelper SourceSetHelper(SrcSetProperty, SrcSetProperty->ContainerPtrToValuePtr<void>(InSourcePtr));
-			//FScriptSetHelper TargetSetHelper(DestSetProperty, DestSetProperty->ContainerPtrToValuePtr<void>(InTargetPtr));
+			check(InDestinationProperty->ArrayDim == 1);
+			FScriptSetHelper SourceSetHelper(SrcSetProperty, SrcSetProperty->ContainerPtrToValuePtr<void>(InSourcePtr));
+			FScriptSetHelper TargetSetHelper(DestSetProperty, DestSetProperty->ContainerPtrToValuePtr<void>(InTargetPtr));
+
+			TargetSetHelper.EmptyElements();
 
-			//TargetSetHelper.EmptyElements();
+			int32 Num = SourceSetHelper.Num();
+			for ( int32 Index = 0; Num; Index++ )
+			{
+				if ( SourceSetHelper.IsValidIndex(Index) )
+				{
+					uint8* SrcPtr = SourceSetHelper.GetElementPtr(Index);
+
+					int32 NewIndex = TargetSetHelper.AddDefaultValue_Invalid_NeedsRehash();
+					TargetSetHelper.Rehash();
+
+					uint8* TargetPtr = TargetSetHelper.GetElementPtr(NewIndex);
+					CopySinglePropertyRecursive(SourceObject, SrcPtr, SrcSetProperty->ElementProp, TargetPtr, InDestinationObject, DestSetProperty->ElementProp);
+
+					TargetSetHelper.Rehash();
+
+					--Num;
+				}
+			}
 
 			bNeedsShallowCopy = false;
-			bNeedsStringCopy = true;
 		}
 		else if ( FObjectPropertyBase* SourceObjectProperty = CastField<FObjectPropertyBase>(InSourceProperty) )
 		{
@@ -290,8 +306,6 @@ namespace FObjectEditorUtils
 			}
 		}
 
-		check(!( bNeedsShallowCopy && bNeedsStringCopy ));
-
 		if ( bNeedsShallowCopy )
 		{
 			const uint8* SourceAddr = InSourceProperty->ContainerPtrToValuePtr<uint8>(InSourcePtr);
@@ -299,14 +313,6 @@ namespace FObjectEditorUtils
 
 			InSourceProperty->CopyCompleteValue(DestinationAddr, SourceAddr);
 		}
-		else if ( bNeedsStringCopy )
-		{
-			FString ExportedTextString;
-			if ( InSourceProperty->ExportText_InContainer(0, ExportedTextString, InSourcePtr, InSourcePtr, SourceObject, PPF_Copy, SourceObject) )
-			{
-				InDestinationProperty->ImportText_InContainer(*ExportedTextString, InTargetPtr, InDestinationObject, 0);
-			}
-		}
 	}
 
 	bool MigratePropertyValue(UObject* SourceObject, FProperty* SourceProperty, UObject* DestinationObject, FProperty* DestinationProperty)
