Create the UnsavedActorMonitor only if the content bundle was successfly injected.
FContentBundle::DoRemove is not called when the bundle failed to inject

diff --git a/Engine/Source/Runtime/Engine/Private/WorldPartition/ContentBundle/ContentBundleEditor.cpp b/Engine/Source/Runtime/Engine/Private/WorldPartition/ContentBundle/ContentBundleEditor.cpp
--- a/Engine/Source/Runtime/Engine/Private/WorldPartition/ContentBundle/ContentBundleEditor.cpp
+++ b/Engine/Source/Runtime/Engine/Private/WorldPartition/ContentBundle/ContentBundleEditor.cpp
@@ -54,9 +54,6 @@ void FContentBundleEditor::DoInjectContent()
 	bool bCreatedContainerPath = ContentBundlePaths::BuildActorDescContainerPackagePath(GetDescriptor()->GetPackageRoot(), GetDescriptor()->GetGuid(), GetInjectedWorld()->GetPackage()->GetName(), ActorDescContainerPackage);
 	if (bCreatedContainerPath)
 	{
-		UnsavedActorMonitor = NewObject<UContentBundleUnsavedActorMonitor>(GetTransientPackage(), NAME_None, RF_Transactional);
-		UnsavedActorMonitor->Initialize(*this);
-
 		UWorldPartition* WorldPartition = GetInjectedWorld()->GetWorldPartition();
 		UWorldPartition::FContainerRegistrationParams RegistrationsParams(*ActorDescContainerPackage);
 		ActorDescContainer = WorldPartition->RegisterActorDescContainer(RegistrationsParams);
@@ -77,11 +74,14 @@ void FContentBundleEditor::DoInjectContent()
 				SetStatus(EContentBundleStatus::ReadyToInject);
 			}
 
+			UnsavedActorMonitor = NewObject<UContentBundleUnsavedActorMonitor>(GetTransientPackage(), NAME_None, RF_Transactional);
+			UnsavedActorMonitor->Initialize(*this);
+
 			RegisterDelegates();
 		}
 		else
 		{
-			UE_LOG(LogContentBundle, Log, TEXT("%s Failed to register actor desc container with %s"), *ContentBundle::Log::MakeDebugInfoString(*this), *ActorDescContainerPackage);
+			UE_LOG(LogContentBundle, Error, TEXT("%s Failed to register actor desc container with %s"), *ContentBundle::Log::MakeDebugInfoString(*this), *ActorDescContainerPackage);
 			SetStatus(EContentBundleStatus::FailedToInject);
 		}
 	}
