Add `CachedMaterialResourcesForCooking` to referenced objects so GC knows about them.

diff --git a/Engine/Source/Runtime/Engine/Private/Materials/Material.cpp b/Engine/Source/Runtime/Engine/Private/Materials/Material.cpp
--- a/Engine/Source/Runtime/Engine/Private/Materials/Material.cpp
+++ b/Engine/Source/Runtime/Engine/Private/Materials/Material.cpp
@@ -5192,6 +5192,21 @@ void UMaterial::AddReferencedObjects(UObject* InThis, FReferenceCollector& Colle
 	{
 		CurrentResource->AddReferencedObjects(Collector);
 	}
+
+#if WITH_EDITOR
+	for (auto& It : This->CachedMaterialResourcesForCooking)
+	{
+		TArray<FMaterialResource*>& CachedMaterialResourcesForPlatform = It.Value;
+		for (FMaterialResource* CurrentResource : CachedMaterialResourcesForPlatform)
+		{
+			if (CurrentResource)
+			{
+				CurrentResource->AddReferencedObjects(Collector);
+			}
+		}
+	}
+#endif
+
 #if WITH_EDITORONLY_DATA
 	Collector.AddReferencedObject(This->MaterialGraph, This);
 #endif
