Added advanced options to control per-partitioned world ServerStreamingMode and ServerStreamingOutMode.
Default values are set to Project default which rely on existing wp.Runtime.EnableServerStreaming and  wp.Runtime.EnableServerStreamingOut.

diff --git a/Engine/Source/Runtime/Engine/Private/WorldPartition/WorldPartition.cpp b/Engine/Source/Runtime/Engine/Private/WorldPartition/WorldPartition.cpp
--- a/Engine/Source/Runtime/Engine/Private/WorldPartition/WorldPartition.cpp
+++ b/Engine/Source/Runtime/Engine/Private/WorldPartition/WorldPartition.cpp
@@ -88,18 +88,18 @@ FAutoConsoleVariableRef UWorldPartition::CVarEnableSimulationStreamingSource(
 	ECVF_Default);
 #endif
 
-int32 UWorldPartition::EnableServerStreaming = 0;
+int32 UWorldPartition::GlobalEnableServerStreaming = 0;
 FAutoConsoleVariableRef UWorldPartition::CVarEnableServerStreaming(
 	TEXT("wp.Runtime.EnableServerStreaming"),
-	UWorldPartition::EnableServerStreaming,
+	UWorldPartition::GlobalEnableServerStreaming,
 	TEXT("Set to 1 to enable server streaming, set to 2 to only enable it in PIE.\n")
 	TEXT("Changing the value while the game is running won't be considered."),
 	WorldPartition::ECVF_Runtime_ReadOnly);
 
-bool UWorldPartition::bEnableServerStreamingOut = false;
+bool UWorldPartition::bGlobalEnableServerStreamingOut = false;
 FAutoConsoleVariableRef UWorldPartition::CVarEnableServerStreamingOut(
 	TEXT("wp.Runtime.EnableServerStreamingOut"),
-	UWorldPartition::bEnableServerStreamingOut,
+	UWorldPartition::bGlobalEnableServerStreamingOut,
 	TEXT("Turn on/off to allow or not the server to stream out levels (only relevant when server streaming is enabled)\n")
 	TEXT("Changing the value while the game is running won't be considered."),
 	WorldPartition::ECVF_Runtime_ReadOnly);
@@ -283,6 +283,8 @@ UWorldPartition::UWorldPartition(const FObjectInitializer& ObjectInitializer)
 	, Replay(nullptr)
 {
 	bEnableStreaming = true;
+	ServerStreamingMode = EWorldPartitionServerStreamingMode::ProjectDefault;
+	ServerStreamingOutMode = EWorldPartitionServerStreamingOutMode::ProjectDefault;
 
 #if WITH_EDITOR
 	WorldPartitionStreamingPolicyClass = UWorldPartitionLevelStreamingPolicy::StaticClass();
@@ -408,6 +410,25 @@ void UWorldPartition::OnEndPlay()
 	RuntimeHash->OnEndPlay();
 }
 
+bool UWorldPartition::CanEditChange(const FProperty* InProperty) const
+{
+	if (!Super::CanEditChange(InProperty))
+	{
+		return false;
+	}
+
+	if (InProperty->GetFName() == GET_MEMBER_NAME_CHECKED(UWorldPartition, ServerStreamingOutMode))
+	{
+		return bEnableStreaming && (ServerStreamingMode != EWorldPartitionServerStreamingMode::Disabled);
+	}
+	else if (InProperty->GetFName() == GET_MEMBER_NAME_CHECKED(UWorldPartition, ServerStreamingMode))
+	{
+		return bEnableStreaming;
+	}
+
+	return true;
+}
+
 FName UWorldPartition::GetWorldPartitionEditorName() const
 {
 	if (SupportsStreaming())
@@ -1043,13 +1064,28 @@ bool UWorldPartition::IsServerStreamingEnabled() const
 	if (!bCachedIsServerStreamingEnabled.IsSet())
 	{
 		bool bIsEnabled = false;
-		switch (UWorldPartition::EnableServerStreaming)
+		if (ServerStreamingMode == EWorldPartitionServerStreamingMode::ProjectDefault)
 		{
-		case 1:	bIsEnabled = true;
+			switch (UWorldPartition::GlobalEnableServerStreaming)
+			{
+			case 1:	bIsEnabled = true;
 #if WITH_EDITOR
-		case 2: bIsEnabled = bIsPIE;
+			case 2: bIsEnabled = bIsPIE;
 #endif
+			}
 		}
+		else
+		{
+			if ((ServerStreamingMode == EWorldPartitionServerStreamingMode::Enabled)
+#if WITH_EDITOR
+				|| (bIsPIE && (ServerStreamingMode == EWorldPartitionServerStreamingMode::EnabledInPIE))
+#endif
+				)
+			{
+				bIsEnabled = true;
+			}
+		}
+
 		UWorld* OwningWorld = GetWorld();
 		bCachedIsServerStreamingEnabled = (OwningWorld && OwningWorld->IsGameWorld() && bIsEnabled);
 	}
@@ -1063,7 +1099,8 @@ bool UWorldPartition::IsServerStreamingOutEnabled() const
 	if (!bCachedIsServerStreamingOutEnabled.IsSet())
 	{
 		UWorld* OwningWorld = GetWorld();
-		bCachedIsServerStreamingOutEnabled = OwningWorld && OwningWorld->IsGameWorld() && IsServerStreamingEnabled() && UWorldPartition::bEnableServerStreamingOut;
+		const bool bEnableServerStreamingOut = (ServerStreamingOutMode == EWorldPartitionServerStreamingOutMode::ProjectDefault) ? UWorldPartition::bGlobalEnableServerStreamingOut : (ServerStreamingOutMode == EWorldPartitionServerStreamingOutMode::Enabled);
+		bCachedIsServerStreamingOutEnabled = OwningWorld && OwningWorld->IsGameWorld() && IsServerStreamingEnabled() && bEnableServerStreamingOut;
 	}
 
 	return bCachedIsServerStreamingOutEnabled.Get(false);
diff --git a/Engine/Source/Runtime/Engine/Public/WorldPartition/WorldPartition.h b/Engine/Source/Runtime/Engine/Public/WorldPartition/WorldPartition.h
--- a/Engine/Source/Runtime/Engine/Public/WorldPartition/WorldPartition.h
+++ b/Engine/Source/Runtime/Engine/Public/WorldPartition/WorldPartition.h
@@ -56,6 +56,23 @@ enum class EWorldPartitionInitState
 	Uninitializing
 };
 
+UENUM()
+enum class EWorldPartitionServerStreamingMode : uint8
+{
+	ProjectDefault = 0 UMETA(ToolTip = "Use project default (wp.Runtime.EnableServerStreaming)"),
+	Disabled = 1 UMETA(ToolTip = "Server streaming is disabled"),
+	Enabled = 2 UMETA(ToolTip = "Server streaming is enabled"),
+	EnabledInPIE = 3 UMETA(ToolTip = "Server streaming is only enabled in PIE"),
+};
+
+UENUM()
+enum class EWorldPartitionServerStreamingOutMode : uint8
+{
+	ProjectDefault = 0 UMETA(ToolTip = "Use project default (wp.Runtime.EnableServerStreamingOut)"),
+	Disabled = 1 UMETA(ToolTip = "Server streaming out is disabled"),
+	Enabled = 2 UMETA(ToolTip = "Server streaming out is enabled"),
+};
+
 #if WITH_EDITOR
 /**
  * Interface for the world partition editor
@@ -138,6 +155,9 @@ public:
 	//~ End UActorDescContainer Interface
 
 	//~ Begin UObject Interface
+#if WITH_EDITOR
+	virtual bool CanEditChange(const FProperty* InProperty) const override;
+#endif //WITH_EDITOR
 	virtual void Serialize(FArchive& Ar) override;
 	virtual UWorld* GetWorld() const override;
 	virtual bool ResolveSubobject(const TCHAR* SubObjectPath, UObject*& OutObject, bool bLoadIfExists) override;
@@ -375,6 +395,12 @@ public:
 	UPROPERTY()
 	bool bEnableStreaming;
 
+	UPROPERTY(EditAnywhere, Category = WorldPartitionSetup, AdvancedDisplay, meta = (EditConditionHides, EditCondition = "bEnableStreaming", HideEditConditionToggle))
+	EWorldPartitionServerStreamingMode ServerStreamingMode;
+
+	UPROPERTY(EditAnywhere, Category = WorldPartitionSetup, AdvancedDisplay, meta = (EditConditionHides, EditCondition = "bEnableStreaming", HideEditConditionToggle))
+	EWorldPartitionServerStreamingOutMode ServerStreamingOutMode;
+
 private:
 	TObjectPtr<UWorld> World;
 
@@ -435,8 +461,8 @@ private:
 	static FAutoConsoleVariableRef CVarDebugDedicatedServerStreaming;
 #endif
 
-	static int32 EnableServerStreaming;
-	static bool bEnableServerStreamingOut;
+	static int32 GlobalEnableServerStreaming;
+	static bool bGlobalEnableServerStreamingOut;
 	static bool bUseMakingVisibleTransactionRequests;
 	static bool bUseMakingInvisibleTransactionRequests;
 	static FAutoConsoleVariableRef CVarEnableServerStreaming;
