Temporarily put an error as a log while investigation is ongoing.
The error can occur at cook time

diff --git a/Engine/Source/Runtime/Engine/Private/WorldPartition/ContentBundle/ContentBundleEditor.cpp b/Engine/Source/Runtime/Engine/Private/WorldPartition/ContentBundle/ContentBundleEditor.cpp
--- a/Engine/Source/Runtime/Engine/Private/WorldPartition/ContentBundle/ContentBundleEditor.cpp
+++ b/Engine/Source/Runtime/Engine/Private/WorldPartition/ContentBundle/ContentBundleEditor.cpp
@@ -81,7 +81,8 @@ void FContentBundleEditor::DoInjectContent()
 		}
 		else
 		{
-			UE_LOG(LogContentBundle, Error, TEXT("%s Failed to register actor desc container with %s"), *ContentBundle::Log::MakeDebugInfoString(*this), *ActorDescContainerPackage);
+			// @todo_ow : Temporarily put this error as a log. It can fire during cook.
+			UE_LOG(LogContentBundle, Log, TEXT("%s Failed to register actor desc container with %s"), *ContentBundle::Log::MakeDebugInfoString(*this), *ActorDescContainerPackage);
 			SetStatus(EContentBundleStatus::FailedToInject);
 		}
 	}
