Adds a Low Memory state variable. Set when platform is in low memory situation. When in low memory, checks if a soft GC should be kicked off based on number of streaming levels pending purge.
Removed the 15s increased GC timer as it would never fire due to the code path execution.
#test Tested with cooked builds and replays

diff --git a/Engine/Source/Runtime/Engine/Classes/Engine/CoreSettings.h b/Engine/Source/Runtime/Engine/Classes/Engine/CoreSettings.h
--- a/Engine/Source/Runtime/Engine/Classes/Engine/CoreSettings.h
+++ b/Engine/Source/Runtime/Engine/Classes/Engine/CoreSettings.h
@@ -155,6 +155,8 @@ extern ENGINE_API int32 GLevelStreamingMaxLevelRequestsAtOnceWhileInMatch;
 extern ENGINE_API int32 GLevelStreamingForceVerifyLevelsGotRemovedByGC;
 /** Whether to force routing actor initialize phase in its own frame. */
 extern ENGINE_API int32 GLevelStreamingForceRouteActorInitializeNextFrame;
+/** In a low memory situation, will kick off a 'soft' GC if the number of streaming levels pending purge meets or exceeds this count. */
+extern ENGINE_API int32 GLevelStreamingLowMemoryPendingPurgeCount;
 
 /**
 * Implements the settings for garbage collection.
diff --git a/Engine/Source/Runtime/Engine/Private/CoreSettings.cpp b/Engine/Source/Runtime/Engine/Private/CoreSettings.cpp
--- a/Engine/Source/Runtime/Engine/Private/CoreSettings.cpp
+++ b/Engine/Source/Runtime/Engine/Private/CoreSettings.cpp
@@ -25,6 +25,7 @@ int32 GLevelStreamingAllowLevelRequestsWhileAsyncLoadingInMatch = 1;
 int32 GLevelStreamingMaxLevelRequestsAtOnceWhileInMatch = 0;
 int32 GLevelStreamingForceVerifyLevelsGotRemovedByGC = 0;
 int32 GLevelStreamingForceRouteActorInitializeNextFrame = 0;
+int32 GLevelStreamingLowMemoryPendingPurgeCount = MAX_int32;
 
 static FAutoConsoleVariableRef CVarUseBackgroundLevelStreaming(
 	TEXT("s.UseBackgroundLevelStreaming"),
@@ -145,6 +146,14 @@ static FAutoConsoleVariableRef CVarMaxLevelRequestsAtOnceWhileInMatch(
 	ECVF_Default
 );
 
+static FAutoConsoleVariableRef CVarLevelStreamingLowMemoryPendingPurgeCount(
+	TEXT("s.LevelStreamingLowMemoryPendingPurgeCount"),
+	GLevelStreamingLowMemoryPendingPurgeCount,
+	TEXT("When system is in low memory state, if the number of streaming levels to purge meets or exceeds this value, perform a 'soft' GC."),
+	ECVF_Default
+);
+
+
 UStreamingSettings::UStreamingSettings()
 	: Super()
 {
diff --git a/Engine/Source/Runtime/Engine/Private/UnrealEngine.cpp b/Engine/Source/Runtime/Engine/Private/UnrealEngine.cpp
--- a/Engine/Source/Runtime/Engine/Private/UnrealEngine.cpp
+++ b/Engine/Source/Runtime/Engine/Private/UnrealEngine.cpp
@@ -629,6 +629,7 @@ void FTemporaryPlayInEditorIDOverride::SetID(int32 NewID)
 ENGINE_API float GAverageFPS = 0.0f;
 ENGINE_API float GAverageMS = 0.0f;
 ENGINE_API double GLastMemoryWarningTime = 0.f;
+ENGINE_API bool GIsLowMemory = false;
 
 
 void CalculateFPSTimings()
@@ -1644,18 +1645,14 @@ float GetLowMemoryGCTimer(const float DefaultTimeBetweenGC)
 #if !UE_BUILD_SHIPPING
 	MBFree -= float(FPlatformMemory::GetExtraDevelopmentMemorySize() / 1024 / 1024);
 #endif
-	if (MBFree <= GLowMemoryMemoryThresholdMB)
+
+	GIsLowMemory = (MBFree <= GLowMemoryMemoryThresholdMB);
+	if (GIsLowMemory)
 	{
 		if (GLowMemoryTimeBetweenPurgingPendingKillObjects < DefaultTimeBetweenGC)
 		{
 			return GLowMemoryTimeBetweenPurgingPendingKillObjects;
 		}
-
-		// Go even faster if there are levels that need to be unloaded
-		if (GLowMemoryTimeBetweenPurgingPendingLevels < DefaultTimeBetweenGC && FLevelStreamingGCHelper::GetNumLevelsPendingPurge() > 0)
-		{
-			return GLowMemoryTimeBetweenPurgingPendingLevels;
-		}
 	}
 
 	return DefaultTimeBetweenGC;
diff --git a/Engine/Source/Runtime/Engine/Private/World.cpp b/Engine/Source/Runtime/Engine/Private/World.cpp
--- a/Engine/Source/Runtime/Engine/Private/World.cpp
+++ b/Engine/Source/Runtime/Engine/Private/World.cpp
@@ -4153,6 +4153,8 @@ bool UWorld::CanAddLoadedLevelToWorld(ULevel* Level) const
 	return true;
 }
 
+extern ENGINE_API bool GIsLowMemory;
+
 void UWorld::UpdateLevelStreaming()
 {
 	TRACE_CPUPROFILER_EVENT_SCOPE(UWorld::UpdateLevelStreaming);
@@ -4222,18 +4224,20 @@ void UWorld::UpdateLevelStreaming()
 
 	CSV_CUSTOM_STAT(LevelStreamingPendingPurge, NumlevelsPendingPurge, CurrentNumLevelsPendingPurge, ECsvCustomStatOp::Set);
 
-	if (GLevelStreamingContinuouslyIncrementalGCWhileLevelsPendingPurge)
-	{
-		// Figure out whether there are any levels we haven't collected garbage yet.
-		const bool bAreLevelsPendingPurge = CurrentNumLevelsPendingPurge >= GLevelStreamingContinuouslyIncrementalGCWhileLevelsPendingPurge;
+	// Are we currently in a low memory situation and number of pending levels to purge meets or exceeds our threshold?
+	const bool bShouldDoLowMemoryGC = GIsLowMemory && CurrentNumLevelsPendingPurge >= GLevelStreamingLowMemoryPendingPurgeCount;
 
+	// Low memory GC takes precedence over Continuous GC condition
+	if (bShouldDoLowMemoryGC)
+	{
+		GEngine->ForceGarbageCollection(false);
+	}
+	else if (GLevelStreamingContinuouslyIncrementalGCWhileLevelsPendingPurge && CurrentNumLevelsPendingPurge >= GLevelStreamingContinuouslyIncrementalGCWhileLevelsPendingPurge)
+	{
 		// Request a 'soft' GC if there are levels pending purge and there are levels to be loaded. In the case of a blocking
 		// load this is going to guarantee GC firing first thing afterwards and otherwise it is going to sneak in right before
 		// kicking off the async load.
-		if (bAreLevelsPendingPurge)
-		{
-			GEngine->ForceGarbageCollection(false);
-		}
+		GEngine->ForceGarbageCollection(false);
 	}
 
 	// In case more levels has been requested to unload, force GC on next tick 
