on the client - defer adding into the cluster union until the cluster union gets its position replicated for the first time

diff --git a/Engine/Source/Runtime/Engine/Classes/PhysicsEngine/ClusterUnionComponent.h b/Engine/Source/Runtime/Engine/Classes/PhysicsEngine/ClusterUnionComponent.h
--- a/Engine/Source/Runtime/Engine/Classes/PhysicsEngine/ClusterUnionComponent.h
+++ b/Engine/Source/Runtime/Engine/Classes/PhysicsEngine/ClusterUnionComponent.h
@@ -151,6 +151,8 @@ public:
 	UFUNCTION()
 	bool IsComponentAdded(UPrimitiveComponent* Component) { return ComponentToPhysicsObjects.Contains(Component); }
 
+	bool HasReceivedTransform() const { return bHasReceivedTransform; }
+
 	friend class UClusterUnionReplicatedProxyComponent;
 protected:
 
@@ -197,6 +199,7 @@ private:
 	FPhysScene_Chaos* GetChaosScene() const;
 
 	Chaos::FClusterUnionPhysicsProxy* PhysicsProxy;
+	bool bHasReceivedTransform;
 
 	// User data to be able to tie the cluster particle back to this component.
 	FChaosUserData PhysicsUserData;
diff --git a/Engine/Source/Runtime/Engine/Classes/PhysicsEngine/ClusterUnionReplicatedProxyComponent.h b/Engine/Source/Runtime/Engine/Classes/PhysicsEngine/ClusterUnionReplicatedProxyComponent.h
--- a/Engine/Source/Runtime/Engine/Classes/PhysicsEngine/ClusterUnionReplicatedProxyComponent.h
+++ b/Engine/Source/Runtime/Engine/Classes/PhysicsEngine/ClusterUnionReplicatedProxyComponent.h
@@ -94,6 +94,9 @@ private:
 	FTimerHandle DeferSetChildToParentHandle;
 	void DeferSetChildToParentChildUntilClusteredComponentInParentUnion();
 
+	FTimerHandle DeferAddComponentToClusterHandle;
+	void DeferAddComponentToClusterHandleUntilInitialTransformUpdate();
+
 	//~ Begin UActorComponent Interface
 public:
 	virtual void EndPlay(const EEndPlayReason::Type EndPlayReason) override;
diff --git a/Engine/Source/Runtime/Engine/Private/PhysicsEngine/ClusterUnionComponent.cpp b/Engine/Source/Runtime/Engine/Private/PhysicsEngine/ClusterUnionComponent.cpp
--- a/Engine/Source/Runtime/Engine/Private/PhysicsEngine/ClusterUnionComponent.cpp
+++ b/Engine/Source/Runtime/Engine/Private/PhysicsEngine/ClusterUnionComponent.cpp
@@ -19,7 +19,7 @@ UClusterUnionComponent::UClusterUnionComponent(const FObjectInitializer& ObjectI
 {
 	PhysicsProxy = nullptr;
 	SetIsReplicatedByDefault(true);
-
+	bHasReceivedTransform = false;
 #if WITH_EDITORONLY_DATA
 	bVisualizeComponent = true;
 #endif
@@ -208,6 +208,7 @@ void UClusterUnionComponent::OnCreatePhysicsState()
 	InitData.ActorId = GetOwner()->GetUniqueID();
 	InitData.ComponentId = GetUniqueID();
 	InitData.bNeedsClusterXRInitialization = GetOwner()->HasAuthority();
+	bHasReceivedTransform = false;
 	PhysicsProxy = new Chaos::FClusterUnionPhysicsProxy{ this, Parameters, InitData };
 	PhysicsProxy->Initialize_External();
 	if (FPhysScene_Chaos* Scene = GetChaosScene())
@@ -266,6 +267,7 @@ void UClusterUnionComponent::OnDestroyPhysicsState()
 void UClusterUnionComponent::OnUpdateTransform(EUpdateTransformFlags UpdateTransformFlags, ETeleportType Teleport)
 {
 	USceneComponent::OnUpdateTransform(UpdateTransformFlags, Teleport);
+	bHasReceivedTransform = true;
 
 	if (PhysicsProxy && !(UpdateTransformFlags & EUpdateTransformFlags::SkipPhysicsUpdate))
 	{
diff --git a/Engine/Source/Runtime/Engine/Private/PhysicsEngine/ClusterUnionReplicatedProxyComponent.cpp b/Engine/Source/Runtime/Engine/Private/PhysicsEngine/ClusterUnionReplicatedProxyComponent.cpp
--- a/Engine/Source/Runtime/Engine/Private/PhysicsEngine/ClusterUnionReplicatedProxyComponent.cpp
+++ b/Engine/Source/Runtime/Engine/Private/PhysicsEngine/ClusterUnionReplicatedProxyComponent.cpp
@@ -123,7 +123,10 @@ void UClusterUnionReplicatedProxyComponent::PostRepNotifies()
 	const bool bIsInitialReplication = bNetUpdateParentClusterUnion || bNetUpdateChildClusteredComponent || bNetUpdateParticleBoneIds;
 	if (bIsInitialReplication)
 	{
-		ParentClusterUnion->AddComponentToCluster(ChildClusteredComponent, ParticleBoneIds);
+		if (!DeferAddComponentToClusterHandle.IsValid())
+		{
+			DeferAddComponentToClusterHandleUntilInitialTransformUpdate();
+		}
 		bNetUpdateParentClusterUnion = false;
 		bNetUpdateChildClusteredComponent = false;
 		bNetUpdateParticleBoneIds = false;
@@ -141,6 +144,25 @@ void UClusterUnionReplicatedProxyComponent::PostRepNotifies()
 	}
 }
 
+void UClusterUnionReplicatedProxyComponent::DeferAddComponentToClusterHandleUntilInitialTransformUpdate()
+{
+	if (!ParentClusterUnion || !ChildClusteredComponent || IsPendingDeletion())
+	{
+		return;
+	}
+
+	DeferAddComponentToClusterHandle.Invalidate();
+
+	if (ParentClusterUnion->HasReceivedTransform())
+	{
+		ParentClusterUnion->AddComponentToCluster(ChildClusteredComponent, ParticleBoneIds);
+	}
+	else if (AActor* Owner = GetOwner())
+	{
+		DeferAddComponentToClusterHandle = Owner->GetWorldTimerManager().SetTimerForNextTick(this, &UClusterUnionReplicatedProxyComponent::DeferAddComponentToClusterHandleUntilInitialTransformUpdate);
+	}
+}
+
 void UClusterUnionReplicatedProxyComponent::DeferSetChildToParentChildUntilClusteredComponentInParentUnion()
 {
 	if (!ParentClusterUnion || !ChildClusteredComponent || IsPendingDeletion())
