Exposing GetSlotMontageLocalWeight() to Blueprints

diff --git a/Engine/Source/Runtime/Engine/Classes/Animation/AnimInstance.h b/Engine/Source/Runtime/Engine/Classes/Animation/AnimInstance.h
--- a/Engine/Source/Runtime/Engine/Classes/Animation/AnimInstance.h
+++ b/Engine/Source/Runtime/Engine/Classes/Animation/AnimInstance.h
@@ -553,6 +553,10 @@ public:
 	UFUNCTION(BlueprintPure, Category="Animation|Montage", meta=(BlueprintThreadSafe))
 	bool IsSlotActive(FName SlotNodeName) const;
 	
+	/** Get local weight of any montages this slot node is playing. If this slot is not currently playing a montage, it will return 0. */
+	UFUNCTION(BlueprintPure, Category="Animation|Montage", meta=(BlueprintThreadSafe, DisplayName="GetSlotLocalWeight"))
+	float Blueprint_GetSlotMontageLocalWeight(FName SlotNodeName) const;
+	
 	/*********************************************************************************************
 	 * AnimMontage
 	 ********************************************************************************************* */
diff --git a/Engine/Source/Runtime/Engine/Private/Animation/AnimInstance.cpp b/Engine/Source/Runtime/Engine/Private/Animation/AnimInstance.cpp
--- a/Engine/Source/Runtime/Engine/Private/Animation/AnimInstance.cpp
+++ b/Engine/Source/Runtime/Engine/Private/Animation/AnimInstance.cpp
@@ -2106,6 +2106,11 @@ bool UAnimInstance::IsSlotActive(FName SlotNodeName) const
    return false;
 }
 
+float UAnimInstance::Blueprint_GetSlotMontageLocalWeight(FName SlotNodeName) const
+{
+	return GetProxyOnAnyThread<FAnimInstanceProxy>().GetSlotMontageLocalWeight(SlotNodeName);
+}
+
 float UAnimInstance::Montage_PlayInternal(UAnimMontage* MontageToPlay, const FMontageBlendSettings& BlendInSettings, float InPlayRate /*= 1.f*/, EMontagePlayReturnType ReturnValueType /*= EMontagePlayReturnType::MontageLength*/, float InTimeToStartMontageAt /*= 0.f*/, bool bStopAllMontages /*= true*/)
 {
 	LLM_SCOPE(ELLMTag::Animation);
