Fix for issue where the counter to generate unique ids on FSimpleParticleUniqueIndices was not initialized

diff --git a/Engine/Source/Runtime/Engine/Private/PhysicsEngine/ImmediatePhysics/ImmediatePhysicsChaos/ImmediatePhysicsSimulation_Chaos.cpp b/Engine/Source/Runtime/Engine/Private/PhysicsEngine/ImmediatePhysics/ImmediatePhysicsChaos/ImmediatePhysicsSimulation_Chaos.cpp
--- a/Engine/Source/Runtime/Engine/Private/PhysicsEngine/ImmediatePhysics/ImmediatePhysicsChaos/ImmediatePhysicsSimulation_Chaos.cpp
+++ b/Engine/Source/Runtime/Engine/Private/PhysicsEngine/ImmediatePhysics/ImmediatePhysicsChaos/ImmediatePhysicsSimulation_Chaos.cpp
@@ -271,7 +271,7 @@ namespace ImmediatePhysics_Chaos
 		~FSimpleParticleUniqueIndices() = default;
 
 	private:
-		int32 NextUniqueIndex; // this includes all valid and freed indices
+		int32 NextUniqueIndex = 0; // this includes all valid and freed indices
 		TArray<int32> FreeIndices;
 	};
 
