UStreamableSparseVolumeTexture::GetStreamedFrameProxyOrFallback() now directly returns the loaded scene proxy of the requested frame instead of going through the SVT streaming manager logic. This is more robust and fixes certain issues in editor builds.

#rnx

diff --git a/Engine/Source/Runtime/Engine/Private/SparseVolumeTexture/SparseVolumeTexture.cpp b/Engine/Source/Runtime/Engine/Private/SparseVolumeTexture/SparseVolumeTexture.cpp
--- a/Engine/Source/Runtime/Engine/Private/SparseVolumeTexture/SparseVolumeTexture.cpp
+++ b/Engine/Source/Runtime/Engine/Private/SparseVolumeTexture/SparseVolumeTexture.cpp
@@ -502,6 +502,10 @@ const FSparseVolumeTextureSceneProxy* UStreamableSparseVolumeTexture::GetStreame
 		return nullptr;
 	}
 	FrameIndex = FMath::Clamp(FrameIndex, 0, Frames.Num() - 1);
+
+#if WITH_EDITORONLY_DATA
+	return Frames[FrameIndex].SparseVolumeTextureSceneProxy;
+#else
 	ISparseVolumeTextureStreamingManager& StreamingManager = IStreamingManager::Get().GetSparseVolumeTextureStreamingManager();
 	const FSparseVolumeTextureSceneProxy* Proxy = StreamingManager.GetSparseVolumeTextureSceneProxy(this, FrameIndex, MipLevel, true);
 
@@ -518,6 +522,7 @@ const FSparseVolumeTextureSceneProxy* UStreamableSparseVolumeTexture::GetStreame
 	}
 
 	return Proxy;
+#endif
 }
 
 void UStreamableSparseVolumeTexture::PostLoad()
