World Partition Class Descriptors: don't use delta serialization in the (improbable) eventuality of a missing class descriptor.

diff --git a/Engine/Source/Runtime/Engine/Private/WorldPartition/WorldPartitionActorDescArchive.cpp b/Engine/Source/Runtime/Engine/Private/WorldPartition/WorldPartitionActorDescArchive.cpp
--- a/Engine/Source/Runtime/Engine/Private/WorldPartition/WorldPartitionActorDescArchive.cpp
+++ b/Engine/Source/Runtime/Engine/Private/WorldPartition/WorldPartitionActorDescArchive.cpp
@@ -63,8 +63,15 @@ PRAGMA_ENABLE_DEPRECATION_WARNINGS
 	ClassDesc = InActorDesc->bIsDefaultActorDesc ? ClassDescRegistry.GetClassDescDefaultForClass(ClassPath) : ClassDescRegistry.GetClassDescDefaultForActor(ClassPath);
 	if (!ClassDesc)
 	{
-		ClassDesc = ClassDescRegistry.GetClassDescDefault(FTopLevelAssetPath(TEXT("/Script/Engine.Actor")));
-		UE_LOG(LogWorldPartition, Log, TEXT("Can't find class descriptor '%s' for '%s', using '%s'"), *ClassPath.ToString(), *InActorDesc->GetActorSoftPath().ToString(), *ClassDesc->GetActorSoftPath().ToString());
+		if (IsLoading())
+		{
+			ClassDesc = ClassDescRegistry.GetClassDescDefault(FTopLevelAssetPath(TEXT("/Script/Engine.Actor")));
+			UE_LOG(LogWorldPartition, Log, TEXT("Can't find class descriptor '%s' for loading '%s', using '%s'"), *ClassPath.ToString(), *InActorDesc->GetActorSoftPath().ToString(), *ClassDesc->GetActorSoftPath().ToString());
+		}
+		else
+		{
+			UE_LOG(LogWorldPartition, Warning, TEXT("Can't find class descriptor '%s' for saving '%s'"), *ClassPath.ToString(), *InActorDesc->GetActorSoftPath().ToString());
+		}
 	}
 	check(ClassDesc);
 }
diff --git a/Engine/Source/Runtime/Engine/Public/WorldPartition/WorldPartitionActorDescArchive.h b/Engine/Source/Runtime/Engine/Public/WorldPartition/WorldPartitionActorDescArchive.h
--- a/Engine/Source/Runtime/Engine/Public/WorldPartition/WorldPartitionActorDescArchive.h
+++ b/Engine/Source/Runtime/Engine/Public/WorldPartition/WorldPartitionActorDescArchive.h
@@ -34,7 +34,7 @@ public:
 		{
 			FActorDescArchive& ActorDescAr = (FActorDescArchive&)Ar;
 			
-			check(ActorDescAr.ClassDesc);
+			check(ActorDescAr.ClassDesc || Ar.IsSaving());
 			check(ActorDescAr.ActorDesc);
 
 			Ar.UsingCustomVersion(FFortniteMainBranchObjectVersion::GUID);
@@ -50,7 +50,7 @@ public:
 
 			if (Ar.CustomVer(FFortniteMainBranchObjectVersion::GUID) >= FFortniteMainBranchObjectVersion::WorldPartitionActorClassDescSerialize)
 			{
-				if (Ar.IsSaving())
+				if (Ar.IsSaving() && ActorDescAr.ClassDesc)
 				{
 					bSerialize = (V.Value != GetClassDefaultValue()) ? 1 : 0;
 				}
