Changed a warning to a log while investigating cooking issue.

diff --git a/Engine/Source/Runtime/Engine/Private/WorldPartition/WorldPartitionActorDescArchive.cpp b/Engine/Source/Runtime/Engine/Private/WorldPartition/WorldPartitionActorDescArchive.cpp
--- a/Engine/Source/Runtime/Engine/Private/WorldPartition/WorldPartitionActorDescArchive.cpp
+++ b/Engine/Source/Runtime/Engine/Private/WorldPartition/WorldPartitionActorDescArchive.cpp
@@ -70,7 +70,7 @@ PRAGMA_ENABLE_DEPRECATION_WARNINGS
 		}
 		else
 		{
-			UE_LOG(LogWorldPartition, Warning, TEXT("Can't find class descriptor '%s' for saving '%s'"), *ClassPath.ToString(), *InActorDesc->GetActorSoftPath().ToString());
+			UE_LOG(LogWorldPartition, Log, TEXT("Can't find class descriptor '%s' for saving '%s'"), *ClassPath.ToString(), *InActorDesc->GetActorSoftPath().ToString());
 		}
 	}
 	check(ClassDesc);
