Re-enable error message.
Data was cleaned by 25051966 and latest cooks do not repor that error anymore.

diff --git a/Engine/Source/Runtime/Engine/Private/WorldPartition/ContentBundle/ContentBundleEditor.cpp b/Engine/Source/Runtime/Engine/Private/WorldPartition/ContentBundle/ContentBundleEditor.cpp
--- a/Engine/Source/Runtime/Engine/Private/WorldPartition/ContentBundle/ContentBundleEditor.cpp
+++ b/Engine/Source/Runtime/Engine/Private/WorldPartition/ContentBundle/ContentBundleEditor.cpp
@@ -81,8 +81,7 @@ void FContentBundleEditor::DoInjectContent()
 		}
 		else
 		{
-			// @todo_ow : Temporarily put this error as a log. It can fire during cook.
-			UE_LOG(LogContentBundle, Log, TEXT("%s Failed to register actor desc container with %s"), *ContentBundle::Log::MakeDebugInfoString(*this), *ActorDescContainerPackage);
+			UE_LOG(LogContentBundle, Error, TEXT("%s Failed to register actor desc container with %s"), *ContentBundle::Log::MakeDebugInfoString(*this), *ActorDescContainerPackage);
 			SetStatus(EContentBundleStatus::FailedToInject);
 		}
 	}
