Cameras: check for null world pointer before starting/stopping shakes

diff --git a/Engine/Source/Runtime/Engine/Private/Camera/CameraShakeSourceComponent.cpp b/Engine/Source/Runtime/Engine/Private/Camera/CameraShakeSourceComponent.cpp
--- a/Engine/Source/Runtime/Engine/Private/Camera/CameraShakeSourceComponent.cpp
+++ b/Engine/Source/Runtime/Engine/Private/Camera/CameraShakeSourceComponent.cpp
@@ -81,36 +81,45 @@ void UCameraShakeSourceComponent::Start()
 
 void UCameraShakeSourceComponent::StartCameraShake(TSubclassOf<UCameraShakeBase> InCameraShake, float Scale, ECameraShakePlaySpace PlaySpace, FRotator UserPlaySpaceRot)
 {
-	for (FConstPlayerControllerIterator Iterator = GetWorld()->GetPlayerControllerIterator(); Iterator; ++Iterator)
+	if (UWorld* World = GetWorld())
 	{
-		APlayerController* PlayerController = Iterator->Get();
-		if (PlayerController != nullptr && PlayerController->PlayerCameraManager != nullptr)
+		for (FConstPlayerControllerIterator Iterator = World->GetPlayerControllerIterator(); Iterator; ++Iterator)
 		{
-			PlayerController->PlayerCameraManager->StartCameraShakeFromSource(InCameraShake, this, Scale, PlaySpace, UserPlaySpaceRot);
+			APlayerController* PlayerController = Iterator->Get();
+			if (PlayerController != nullptr && PlayerController->PlayerCameraManager != nullptr)
+			{
+				PlayerController->PlayerCameraManager->StartCameraShakeFromSource(InCameraShake, this, Scale, PlaySpace, UserPlaySpaceRot);
+			}
 		}
 	}
 }
 
 void UCameraShakeSourceComponent::StopAllCameraShakesOfType(TSubclassOf<UCameraShakeBase> InCameraShake, bool bImmediately)
 {
-	for (FConstPlayerControllerIterator Iterator = GetWorld()->GetPlayerControllerIterator(); Iterator; ++Iterator)
+	if (UWorld* World = GetWorld())
 	{
-		APlayerController* PlayerController = Iterator->Get();
-		if (PlayerController != nullptr && PlayerController->PlayerCameraManager != nullptr)
+		for (FConstPlayerControllerIterator Iterator = World->GetPlayerControllerIterator(); Iterator; ++Iterator)
 		{
-			PlayerController->PlayerCameraManager->StopAllInstancesOfCameraShakeFromSource(InCameraShake, this, bImmediately);
+			APlayerController* PlayerController = Iterator->Get();
+			if (PlayerController != nullptr && PlayerController->PlayerCameraManager != nullptr)
+			{
+				PlayerController->PlayerCameraManager->StopAllInstancesOfCameraShakeFromSource(InCameraShake, this, bImmediately);
+			}
 		}
 	}
 }
 
 void UCameraShakeSourceComponent::StopAllCameraShakes(bool bImmediately)
 {
-	for (FConstPlayerControllerIterator Iterator = GetWorld()->GetPlayerControllerIterator(); Iterator; ++Iterator)
+	if (UWorld* World = GetWorld())
 	{
-		APlayerController* PlayerController = Iterator->Get();
-		if (PlayerController != nullptr && PlayerController->PlayerCameraManager != nullptr)
+		for (FConstPlayerControllerIterator Iterator = World->GetPlayerControllerIterator(); Iterator; ++Iterator)
 		{
-			PlayerController->PlayerCameraManager->StopAllCameraShakesFromSource(this, bImmediately);
+			APlayerController* PlayerController = Iterator->Get();
+			if (PlayerController != nullptr && PlayerController->PlayerCameraManager != nullptr)
+			{
+				PlayerController->PlayerCameraManager->StopAllCameraShakesFromSource(this, bImmediately);
+			}
 		}
 	}
 }
