Crash while connecting to the match

diff --git a/Engine/Source/Runtime/Engine/Private/PhysicsEngine/ClusterUnionComponent.cpp b/Engine/Source/Runtime/Engine/Private/PhysicsEngine/ClusterUnionComponent.cpp
--- a/Engine/Source/Runtime/Engine/Private/PhysicsEngine/ClusterUnionComponent.cpp
+++ b/Engine/Source/Runtime/Engine/Private/PhysicsEngine/ClusterUnionComponent.cpp
@@ -110,7 +110,7 @@ void UClusterUnionComponent::AddComponentToCluster(UPrimitiveComponent* InCompon
 
 void UClusterUnionComponent::RemoveComponentFromCluster(UPrimitiveComponent* InComponent)
 {
-	if (!InComponent || !PhysicsProxy || !InComponent->HasValidPhysicsState())
+	if (!InComponent || !PhysicsProxy)
 	{
 		return;
 	}
@@ -137,7 +137,16 @@ void UClusterUnionComponent::RemoveComponentFromCluster(UPrimitiveComponent* InC
 			}
 		}
 
-		PhysicsObjectsToRemove = ComponentData->PhysicsObjects;
+		if (!InComponent->HasValidPhysicsState())
+		{
+			ComponentData->PhysicsObjects.Reset();
+			ComponentData->AllPhysicsObjects.Reset();
+			return;
+		}
+		else
+		{
+			PhysicsObjectsToRemove = ComponentData->PhysicsObjects;
+		}
 	}
 
 	PhysicsProxy->RemovePhysicsObjects_External(PhysicsObjectsToRemove);
