- Fix an issue with constant interpolation key evaluation where the value from Key1 would be used rather than Key2
- One repro case was with a 2 Key curve both constant interpolation, the second value would never be used leading to incorrect results

#rnx

diff --git a/Engine/Source/Runtime/Engine/Private/Curves/CurveEvaluation.cpp b/Engine/Source/Runtime/Engine/Private/Curves/CurveEvaluation.cpp
--- a/Engine/Source/Runtime/Engine/Private/Curves/CurveEvaluation.cpp
+++ b/Engine/Source/Runtime/Engine/Private/Curves/CurveEvaluation.cpp
@@ -206,34 +206,41 @@ namespace UE
 		{
 			const float Diff = Key2.Time - Key1.Time;
 
-			if (Diff > 0.f && Key1.InterpMode != RCIM_Constant)
+			if (Diff > 0.f)
 			{
-				const float Alpha = (InTime - Key1.Time) / Diff;
-				const float P0 = Key1.Value;
-				const float P3 = Key2.Value;
-
-				if (Key1.InterpMode == RCIM_Linear)
+				if (Key1.InterpMode != RCIM_Constant)
 				{
-					return FMath::Lerp(P0, P3, Alpha);
-				}
-				else 
-				{					
-					if (IsItNotWeighted(Key1,Key2))
-					{
-					    const float OneThird = 1.0f / 3.0f;
-					    const float P1 = P0 + (Key1.LeaveTangent * Diff * OneThird);
-					    const float P2 = P3 - (Key2.ArriveTangent * Diff * OneThird);
+					const float Alpha = (InTime - Key1.Time) / Diff;
+					const float P0 = Key1.Value;
+					const float P3 = Key2.Value;
 
-					    return BezierInterp(P0, P1, P2, P3, Alpha);
-					}
-					//if (IsWeighted(Key1, Key2))//it's weighted
-					else	
+					if (Key1.InterpMode == RCIM_Linear)
 					{
-					return WeightedEvalForTwoKeys(
-						Key1.Value, Key1.Time, Key1.LeaveTangent, Key1.LeaveTangentWeight, Key1.TangentWeightMode,
-						Key2.Value, Key2.Time, Key2.ArriveTangent, Key2.ArriveTangentWeight, Key2.TangentWeightMode,
-						InTime);
+						return FMath::Lerp(P0, P3, Alpha);
 					}
+					else 
+					{					
+						if (IsItNotWeighted(Key1,Key2))
+						{
+							const float OneThird = 1.0f / 3.0f;
+							const float P1 = P0 + (Key1.LeaveTangent * Diff * OneThird);
+							const float P2 = P3 - (Key2.ArriveTangent * Diff * OneThird);
+
+							return BezierInterp(P0, P1, P2, P3, Alpha);
+						}
+						//if (IsWeighted(Key1, Key2))//it's weighted
+						else	
+						{
+						return WeightedEvalForTwoKeys(
+							Key1.Value, Key1.Time, Key1.LeaveTangent, Key1.LeaveTangentWeight, Key1.TangentWeightMode,
+							Key2.Value, Key2.Time, Key2.ArriveTangent, Key2.ArriveTangentWeight, Key2.TangentWeightMode,
+							InTime);
+						}
+					}
+				}
+				else
+				{
+					return InTime < Key2.Time ? Key1.Value : Key2.Value;
 				}
 			}
 			else
diff --git a/Engine/Source/Runtime/MovieScene/Private/Channels/MovieSceneCurveChannelImpl.cpp b/Engine/Source/Runtime/MovieScene/Private/Channels/MovieSceneCurveChannelImpl.cpp
--- a/Engine/Source/Runtime/MovieScene/Private/Channels/MovieSceneCurveChannelImpl.cpp
+++ b/Engine/Source/Runtime/MovieScene/Private/Channels/MovieSceneCurveChannelImpl.cpp
@@ -59,26 +59,33 @@ EvalForTwoKeys(
 	Diff /= DecimalRate;
 	const int CheckBothLinear = GSequencerLinearCubicInterpolation;
 
-	if (Diff > 0 && Key1.InterpMode != RCIM_Constant)
+	if (Diff > 0)
 	{
-		const double Alpha = ((double)(InTime - Key1Time).Value / DecimalRate) / Diff;
-		const CurveValueType P0 = Key1.Value;
-		const CurveValueType P3 = Key2.Value;
-
-		if (Key1.InterpMode == RCIM_Linear && (!CheckBothLinear || Key2.InterpMode != RCIM_Cubic))
+		if (Key1.InterpMode != RCIM_Constant)
 		{
-			return FMath::Lerp(P0, P3, Alpha);
+			const double Alpha = ((double)(InTime - Key1Time).Value / DecimalRate) / Diff;
+			const CurveValueType P0 = Key1.Value;
+			const CurveValueType P3 = Key2.Value;
+
+			if (Key1.InterpMode == RCIM_Linear && (!CheckBothLinear || Key2.InterpMode != RCIM_Cubic))
+			{
+				return FMath::Lerp(P0, P3, Alpha);
+			}
+			else
+			{
+				double LeaveTangent = Key1.Tangent.LeaveTangent * DecimalRate;
+				double ArriveTangent = Key2.Tangent.ArriveTangent * DecimalRate;
+
+				const double OneThird = 1.0 / 3.0;
+				const CurveValueType P1 = P0 + (LeaveTangent * Diff * OneThird);
+				const CurveValueType P2 = P3 - (ArriveTangent * Diff * OneThird);
+
+				return UE::Curves::BezierInterp(P0, P1, P2, P3, Alpha);
+			}
 		}
 		else
 		{
-			double LeaveTangent = Key1.Tangent.LeaveTangent * DecimalRate;
-			double ArriveTangent = Key2.Tangent.ArriveTangent * DecimalRate;
-
-			const double OneThird = 1.0 / 3.0;
-			const CurveValueType P1 = P0 + (LeaveTangent * Diff*OneThird);
-			const CurveValueType P2 = P3 - (ArriveTangent * Diff*OneThird);
-
-			return UE::Curves::BezierInterp(P0, P1, P2, P3, Alpha);
+			return InTime < Key2Time ? Key1.Value : Key2.Value;
 		}
 	}
 	else
