Disable NoMips and LeaveExistingMips for VTs, as they require mips.

diff --git a/Engine/Source/Runtime/Engine/Private/TextureDerivedData.cpp b/Engine/Source/Runtime/Engine/Private/TextureDerivedData.cpp
--- a/Engine/Source/Runtime/Engine/Private/TextureDerivedData.cpp
+++ b/Engine/Source/Runtime/Engine/Private/TextureDerivedData.cpp
@@ -896,6 +896,17 @@ static void GetTextureBuildSettings(
 	const bool bForVirtualTextureStreamingBuild = ULightMapVirtualTexture2D::StaticClass() == Texture.GetClass();
 	const bool bVirtualTextureStreaming = bForVirtualTextureStreamingBuild || (CVarVirtualTexturesEnabled->GetValueOnAnyThread() && bPlatformSupportsVirtualTextureStreaming && Texture.VirtualTextureStreaming);
 
+
+	// Virtual textures must have mips as VT memory management relies on a 1:1 texel/pixel mapping, which in turn
+	// requires that we be able to swap in lower mips when that density gets too high for a given texture.
+	if (bVirtualTextureStreaming &&
+		(MipGenSettings == TMGS_NoMipmaps || MipGenSettings == TMGS_LeaveExistingMips)
+		)
+	{
+		MipGenSettings = TMGS_SimpleAverage;
+		UE_LOG(LogTexture, Warning, TEXT("Texture %s is virtual and has NoMips / LeaveExistingMips - forcing to SimpleAverage."), *Texture.GetPathName());
+	}
+
 	const FIntPoint SourceSize = Texture.Source.GetLogicalSize();
 
 	OutBuildSettings.MipGenSettings = MipGenSettings;
