Sanitize some virtual texture settings that are invalid.

diff --git a/Engine/Source/Runtime/Engine/Private/Texture.cpp b/Engine/Source/Runtime/Engine/Private/Texture.cpp
--- a/Engine/Source/Runtime/Engine/Private/Texture.cpp
+++ b/Engine/Source/Runtime/Engine/Private/Texture.cpp
@@ -586,10 +586,23 @@ void UTexture::ValidateSettingsAfterImportOrEdit(bool * pRequiresNotifyMaterials
 					bRequiresNotifyMaterials = true;
 				}
 			}
+
+			// VTs require mips as VT memory management assumes 1:1 texel/pixel mapping, which requires mips to enforce.
+			if (LODGroup == TEXTUREGROUP_ColorLookupTable)
+			{
+				UE_LOG(LogTexture, Warning, TEXT("VirtualTextureStreaming is not compatible with ColorLookupTable LODGroup as virtual textures require mips (%s)"), *GetName());
+				VirtualTextureStreaming = false;
+				bRequiresNotifyMaterials = true;
+			}
+			if (MipGenSettings == TMGS_LeaveExistingMips || MipGenSettings == TMGS_NoMipmaps)
+			{
+				UE_LOG(LogTexture, Warning, TEXT("Virtual textures require mips and MipGenSettings is NoMipmaps or LeaveExistingMips: Forcing to SimpleAverage (%s)"), *GetName());
+				MipGenSettings = TMGS_SimpleAverage;
+			}
 		}
 
 		// Make sure settings are correct for LUT textures.
-		if(LODGroup == TEXTUREGROUP_ColorLookupTable)
+		if (LODGroup == TEXTUREGROUP_ColorLookupTable)
 		{
 			if ( MipGenSettings != TMGS_NoMipmaps || SRGB != false )
 			{
