Cast to bool node is missing

Fixed a bug from CL 23313465. Bool->Float conversion is perfectly valid to use in Blueprints. It should not be marked for internal use only.

#rnx

diff --git a/Engine/Source/Runtime/Engine/Classes/Kismet/KismetMathLibrary.h b/Engine/Source/Runtime/Engine/Classes/Kismet/KismetMathLibrary.h
--- a/Engine/Source/Runtime/Engine/Classes/Kismet/KismetMathLibrary.h
+++ b/Engine/Source/Runtime/Engine/Classes/Kismet/KismetMathLibrary.h
@@ -3645,7 +3645,7 @@ class ENGINE_API UKismetMathLibrary : public UBlueprintFunctionLibrary
 	static float Conv_BoolToFloat(bool InBool);
 
 	/** Converts a bool to a float (0.0 or 1.0) */
-	UFUNCTION(BlueprintPure, meta = (BlueprintInternalUseOnly = "true"))
+	UFUNCTION(BlueprintPure, meta = (DisplayName = "To Float (Boolean)", CompactNodeTitle = "->", Keywords = "cast convert", BlueprintAutocast), Category = "Math|Conversions")
 	static double Conv_BoolToDouble(bool InBool);
 
 	/** Converts a bool to a byte */
