UClusterUnionComponent::HandleRemovedClusteredComponent should not try to operate on invalid physics object handles

diff --git a/Engine/Source/Runtime/Engine/Private/PhysicsEngine/ClusterUnionComponent.cpp b/Engine/Source/Runtime/Engine/Private/PhysicsEngine/ClusterUnionComponent.cpp
--- a/Engine/Source/Runtime/Engine/Private/PhysicsEngine/ClusterUnionComponent.cpp
+++ b/Engine/Source/Runtime/Engine/Private/PhysicsEngine/ClusterUnionComponent.cpp
@@ -507,15 +507,18 @@ void UClusterUnionComponent::HandleRemovedClusteredComponent(UPrimitiveComponent
 
 	if (FClusteredComponentData* ComponentData = ComponentToPhysicsObjects.Find(ChangedComponent))
 	{
-		FLockedReadPhysicsObjectExternalInterface Interface = FPhysicsObjectExternalInterface::LockRead(ComponentData->AllPhysicsObjects);
-		for (Chaos::FGeometryParticle* Particle : Interface->GetAllParticles(ComponentData->AllPhysicsObjects))
+		if (ChangedComponent->HasValidPhysicsState())
 		{
-			if (!Particle)
+			FLockedReadPhysicsObjectExternalInterface Interface = FPhysicsObjectExternalInterface::LockRead(ComponentData->AllPhysicsObjects);
+			for (Chaos::FGeometryParticle* Particle : Interface->GetAllParticles(ComponentData->AllPhysicsObjects))
 			{
-				continue;
-			}
+				if (!Particle)
+				{
+					continue;
+				}
 
-			UniqueIdxToComponent.Remove(Particle->UniqueIdx().Idx);
+				UniqueIdxToComponent.Remove(Particle->UniqueIdx().Idx);
+			}
 		}
 
 		if (IsAuthority())
