CollisionMode PCM_DistanceField support for mobile platforms that support DF

diff --git a/Engine/Plugins/FX/Niagara/Shaders/Private/NiagaraDataInterfaceCollisionQuery.ush b/Engine/Plugins/FX/Niagara/Shaders/Private/NiagaraDataInterfaceCollisionQuery.ush
--- a/Engine/Plugins/FX/Niagara/Shaders/Private/NiagaraDataInterfaceCollisionQuery.ush
+++ b/Engine/Plugins/FX/Niagara/Shaders/Private/NiagaraDataInterfaceCollisionQuery.ush
@@ -216,7 +216,7 @@ void NDICollisionQuery_QueryCustomDepthGPU(in float3 In_SamplePos, in float3 In_
 
 void NDICollisionQuery_QueryMeshDistanceFieldGPU(in float3 In_SamplePos, in float3 In_LWCTile, out float Out_DistanceToNearestSurface, out float3 Out_FieldGradient, out bool Out_IsDistanceFieldValid)
 {
-#if PLATFORM_SUPPORTS_DISTANCE_FIELDS && (FEATURE_LEVEL >= FEATURE_LEVEL_SM5)
+#if PLATFORM_SUPPORTS_DISTANCE_FIELDS
 	FLWCVector3 LwcSamplePos = MakeLWCVector3(In_LWCTile, In_SamplePos);
 	float3 SamplePosition = LWCToFloat(LWCAdd(LwcSamplePos, PrimaryView.PreViewTranslation));
 	
diff --git a/Engine/Shaders/Private/ParticleSimulationShader.usf b/Engine/Shaders/Private/ParticleSimulationShader.usf
--- a/Engine/Shaders/Private/ParticleSimulationShader.usf
+++ b/Engine/Shaders/Private/ParticleSimulationShader.usf
@@ -84,7 +84,7 @@ void VertexMain(
 /** Define to 1 to force no collision in the shader. */
 #define FORCE_NO_COLLISION 0
 
-#if FORCE_NO_COLLISION || (FEATURE_LEVEL < FEATURE_LEVEL_SM4 && !VULKAN_PROFILE)
+#if FORCE_NO_COLLISION
 	#undef DEPTH_BUFFER_COLLISION
 	#define DEPTH_BUFFER_COLLISION 0
 	#undef DISTANCE_FIELD_COLLISION
diff --git a/Engine/Source/Runtime/Engine/Private/Particles/FXSystem.cpp b/Engine/Source/Runtime/Engine/Private/Particles/FXSystem.cpp
--- a/Engine/Source/Runtime/Engine/Private/Particles/FXSystem.cpp
+++ b/Engine/Source/Runtime/Engine/Private/Particles/FXSystem.cpp
@@ -33,7 +33,7 @@ bool IsParticleCollisionModeSupported(EShaderPlatform InPlatform, EParticleColli
 	case PCM_DepthBuffer:
 		return IsFeatureLevelSupported(InPlatform, ERHIFeatureLevel::SM5);
 	case PCM_DistanceField:
-		return IsFeatureLevelSupported(InPlatform, ERHIFeatureLevel::SM5);
+		return FDataDrivenShaderPlatformInfo::GetSupportsDistanceFields(InPlatform);
 	}
 	check(0);
 	return IsFeatureLevelSupported(InPlatform, ERHIFeatureLevel::SM5);
@@ -224,7 +224,6 @@ namespace FXConsoleVariables
 FFXSystem::FFXSystem(ERHIFeatureLevel::Type InFeatureLevel, EShaderPlatform InShaderPlatform, FGPUSortManager* InGPUSortManager)
 	: ParticleSimulationResources(NULL)
 	, FeatureLevel(InFeatureLevel)
-	, ShaderPlatform(InShaderPlatform)
 	, GPUSortManager(nullptr)
 #if WITH_EDITOR
 	, bSuspended(false)
diff --git a/Engine/Source/Runtime/Engine/Private/Particles/FXSystemPrivate.h b/Engine/Source/Runtime/Engine/Private/Particles/FXSystemPrivate.h
--- a/Engine/Source/Runtime/Engine/Private/Particles/FXSystemPrivate.h
+++ b/Engine/Source/Runtime/Engine/Private/Particles/FXSystemPrivate.h
@@ -223,7 +223,7 @@ public:
 	/**
 	 * Retrieve shaderplatform that this FXSystem was created for
 	 */
-	EShaderPlatform GetShaderPlatform() const { return ShaderPlatform; }
+	EShaderPlatform GetShaderPlatform() const { return GShaderPlatformForFeatureLevel[FeatureLevel]; }
 
 	/**
 	 * Add a new GPU simulation to the system.
@@ -376,8 +376,6 @@ private:
 	FParticleSimulationResources* ParticleSimulationResources;
 	/** Feature level of this effects system */
 	ERHIFeatureLevel::Type FeatureLevel;
-	/** Shader platform that will be rendering this effects system */
-	EShaderPlatform ShaderPlatform;
 
 	/** The shared GPUSortManager, used to register GPU sort tasks in order to generate sorted particle indices per emitter. */
 	TRefCountPtr<FGPUSortManager> GPUSortManager;
diff --git a/Engine/Source/Runtime/Engine/Private/Particles/ParticleGpuSimulation.cpp b/Engine/Source/Runtime/Engine/Private/Particles/ParticleGpuSimulation.cpp
--- a/Engine/Source/Runtime/Engine/Private/Particles/ParticleGpuSimulation.cpp
+++ b/Engine/Source/Runtime/Engine/Private/Particles/ParticleGpuSimulation.cpp
@@ -1274,6 +1274,10 @@ void ExecuteSimulationCommands(
 			else if (CollisionMode == PCM_DistanceField)
 			{
 				PsParameters.GlobalDistanceFieldParameters = SetupGlobalDistanceFieldParameters_Minimal(*GlobalDistanceFieldParameterData);
+				// TODO Move to common place without adding extra dependencies
+				PsParameters.GlobalDistanceFieldParameters.GlobalDistanceFieldCoverageAtlasTextureSampler = TStaticSamplerState<SF_Trilinear, AM_Wrap, AM_Wrap, AM_Wrap>::GetRHI();
+				PsParameters.GlobalDistanceFieldParameters.GlobalDistanceFieldPageAtlasTextureSampler = TStaticSamplerState<SF_Trilinear, AM_Wrap, AM_Wrap, AM_Wrap>::GetRHI();
+				PsParameters.GlobalDistanceFieldParameters.GlobalDistanceFieldMipTextureSampler = TStaticSamplerState<SF_Trilinear, AM_Clamp, AM_Clamp, AM_Clamp>::GetRHI();
 			}
 		}
 	}
