Fixed mismatched copy constructor/assignment in FNavPathType.

diff --git a/Engine/Source/Runtime/Engine/Classes/AI/Navigation/NavigationTypes.h b/Engine/Source/Runtime/Engine/Classes/AI/Navigation/NavigationTypes.h
--- a/Engine/Source/Runtime/Engine/Classes/AI/Navigation/NavigationTypes.h
+++ b/Engine/Source/Runtime/Engine/Classes/AI/Navigation/NavigationTypes.h
@@ -215,8 +215,7 @@ struct FNavPathPoint : public FNavLocation
 struct ENGINE_API FNavPathType
 {
 	explicit FNavPathType(const FNavPathType* Parent = nullptr) : Id(++NextUniqueId), ParentType(Parent) {}
-	FNavPathType(const FNavPathType& Src) : Id(Src.Id), ParentType(Src.ParentType) {}
-	
+
 	bool operator==(const FNavPathType& Other) const
 	{
 		return Id == Other.Id || (ParentType != nullptr && *ParentType == Other);
