Fixed mismatched copy constructors/assignments in FNavAgentProperties, FMovementProperties and FNavDataConfig.

diff --git a/Engine/Source/Runtime/Engine/Classes/AI/Navigation/NavigationTypes.h b/Engine/Source/Runtime/Engine/Classes/AI/Navigation/NavigationTypes.h
--- a/Engine/Source/Runtime/Engine/Classes/AI/Navigation/NavigationTypes.h
+++ b/Engine/Source/Runtime/Engine/Classes/AI/Navigation/NavigationTypes.h
@@ -316,7 +316,6 @@ struct FMovementProperties
 		, bCanFly(false)
 	{
 	}
-	FMovementProperties(const FMovementProperties& Other) = default;
 };
 
 /** Properties of representation of an 'agent' (or Pawn) used by AI navigation/pathfinding. */
@@ -348,7 +347,9 @@ struct ENGINE_API FNavAgentProperties : public FMovementProperties
 	FNavAgentProperties(float Radius = -1.f, float Height = -1.f)
 		: AgentRadius(Radius), AgentHeight(Height), AgentStepHeight(-1), NavWalkingSearchHeightScale(0.5f)
 	{}
+
 	FNavAgentProperties(const FNavAgentProperties& Other);
+	FNavAgentProperties& operator=(const FNavAgentProperties& Other);
 
 	void UpdateWithCollisionComponent(class UShapeComponent* CollisionComponent);
 
@@ -414,7 +415,9 @@ protected:
 
 public:	
 	FNavDataConfig(float Radius = FNavigationSystem::FallbackAgentRadius, float Height = FNavigationSystem::FallbackAgentHeight);
+
 	FNavDataConfig(const FNavDataConfig& Other);
+	FNavDataConfig& operator=(const FNavDataConfig& Other);
 
 	bool IsValid() const 
 	{
diff --git a/Engine/Source/Runtime/Engine/Private/AI/Navigation/NavigationTypes.cpp b/Engine/Source/Runtime/Engine/Private/AI/Navigation/NavigationTypes.cpp
--- a/Engine/Source/Runtime/Engine/Private/AI/Navigation/NavigationTypes.cpp
+++ b/Engine/Source/Runtime/Engine/Private/AI/Navigation/NavigationTypes.cpp
@@ -88,7 +88,6 @@ uint32 FNavPathType::NextUniqueId = 0;
 //----------------------------------------------------------------------//
 // FNavDataConfig
 //----------------------------------------------------------------------//
-PRAGMA_DISABLE_DEPRECATION_WARNINGS
 FNavDataConfig::FNavDataConfig(float Radius, float Height)
 	: FNavAgentProperties(Radius, Height)
 	, Name(TEXT("Default"))
@@ -98,15 +97,8 @@ FNavDataConfig::FNavDataConfig(float Radius, float Height)
 {
 }
 
-FNavDataConfig::FNavDataConfig(const FNavDataConfig& Other)
-	: FNavAgentProperties(Other)
-	, Name(Other.Name)
-	, Color(Other.Color)
-	, DefaultQueryExtent(Other.DefaultQueryExtent)
-	, NavDataClass(Other.NavDataClass)
-{
-}
-PRAGMA_ENABLE_DEPRECATION_WARNINGS
+FNavDataConfig::FNavDataConfig(const FNavDataConfig& Other) = default;
+FNavDataConfig& FNavDataConfig::operator=(const FNavDataConfig& Other) = default;
 
 void FNavDataConfig::SetNavDataClass(UClass* InNavDataClass)
 {
@@ -313,16 +305,8 @@ void FNavHeightfieldSamples::Empty()
 //----------------------------------------------------------------------//
 const FNavAgentProperties FNavAgentProperties::DefaultProperties;
 
-FNavAgentProperties::FNavAgentProperties(const FNavAgentProperties& Other)
-	: Super(Other)
-	, AgentRadius(Other.AgentRadius)
-	, AgentHeight(Other.AgentHeight)
-	, AgentStepHeight(Other.AgentStepHeight)
-	, NavWalkingSearchHeightScale(Other.NavWalkingSearchHeightScale)
-	, PreferredNavData(Other.PreferredNavData)
-{
-
-}
+FNavAgentProperties::FNavAgentProperties(const FNavAgentProperties& Other) = default;
+FNavAgentProperties& FNavAgentProperties::operator=(const FNavAgentProperties& Other) = default;
 
 void FNavAgentProperties::UpdateWithCollisionComponent(UShapeComponent* CollisionComponent)
 {
