Fixed mismatched copy constructor/assignment operator in FMaterialLayersFunctions and FMemoryImageMaterialParameterInfo.

diff --git a/Engine/Source/Runtime/Engine/Classes/Materials/MaterialLayersFunctions.h b/Engine/Source/Runtime/Engine/Classes/Materials/MaterialLayersFunctions.h
--- a/Engine/Source/Runtime/Engine/Classes/Materials/MaterialLayersFunctions.h
+++ b/Engine/Source/Runtime/Engine/Classes/Materials/MaterialLayersFunctions.h
@@ -199,9 +199,8 @@ struct ENGINE_API FMaterialLayersFunctions : public FMaterialLayersFunctionsRunt
 #endif // WITH_EDITOR
 
 	static const FGuid BackgroundGuid;
-		
+
 	FMaterialLayersFunctions() = default;
-	FMaterialLayersFunctions(const FMaterialLayersFunctions& Rhs) = default;
 	FMaterialLayersFunctions(const FMaterialLayersFunctionsRuntimeData&) = delete;
 
 	FMaterialLayersFunctionsRuntimeData& GetRuntime() { return *this; }
diff --git a/Engine/Source/Runtime/Engine/Public/MaterialTypes.h b/Engine/Source/Runtime/Engine/Public/MaterialTypes.h
--- a/Engine/Source/Runtime/Engine/Public/MaterialTypes.h
+++ b/Engine/Source/Runtime/Engine/Public/MaterialTypes.h
@@ -121,8 +121,6 @@ public:
 		, Association(Rhs.Association)
 	{}
 
-	FMemoryImageMaterialParameterInfo(const FMemoryImageMaterialParameterInfo& Rhs) = default;
-
 	FORCEINLINE FName GetName() const { return ScriptNameToName(Name); }
 
 	friend FArchive& operator<<(FArchive& Ar, FMemoryImageMaterialParameterInfo& Ref)
