Fixed mismatched copy constructor/assignment operator in FAudioComponentCommandInfo.

diff --git a/Engine/Source/Runtime/Engine/Classes/Sound/QuartzQuantizationUtilities.h b/Engine/Source/Runtime/Engine/Classes/Sound/QuartzQuantizationUtilities.h
--- a/Engine/Source/Runtime/Engine/Classes/Sound/QuartzQuantizationUtilities.h
+++ b/Engine/Source/Runtime/Engine/Classes/Sound/QuartzQuantizationUtilities.h
@@ -672,12 +672,6 @@ struct ENGINE_API FAudioComponentCommandInfo
 {
 	FAudioComponentCommandInfo() {}
 
-	FAudioComponentCommandInfo(const FAudioComponentCommandInfo& Other)
-		: Subscriber(Other.Subscriber)
-		, AnticapatoryBoundary(Other.AnticapatoryBoundary)
-		, CommandID(Other.CommandID)
-	{}
-
  	FAudioComponentCommandInfo(Audio::FQuartzGameThreadSubscriber InSubscriber, FQuartzQuantizationBoundary InAnticaptoryBoundary)
 		: Subscriber(InSubscriber)
 		, AnticapatoryBoundary(InAnticaptoryBoundary)
