Substrate - when a new mateiral is created, a simple UE4 like lit shading model material function is used instead of a slab (more complexx to understand it all).

diff --git a/Engine/Source/Runtime/Engine/Private/Materials/Material.cpp b/Engine/Source/Runtime/Engine/Private/Materials/Material.cpp
--- a/Engine/Source/Runtime/Engine/Private/Materials/Material.cpp
+++ b/Engine/Source/Runtime/Engine/Private/Materials/Material.cpp
@@ -3334,9 +3334,22 @@ void UMaterial::ConvertMaterialToStrataMaterial()
 	else if (!bUseMaterialAttributes && !EditorOnly->FrontMaterial.IsConnected() && GetExpressions().IsEmpty())
 	{
 		// Empty material: Create by default a slab node
-		UMaterialExpressionStrataSlabBSDF* SlabNode = NewObject<UMaterialExpressionStrataSlabBSDF>(this);
-		SetPosXAndMoveReferenceToTheRight(SlabNode);
-		EditorOnly->FrontMaterial.Connect(0, SlabNode);
+		UMaterialFunction* DefaultMF = LoadObject<UMaterialFunction>(nullptr, TEXT("/Engine/Functions/Strata/SMF_UE4Disney.SMF_UE4Disney"));
+		if (DefaultMF)
+		{
+			UMaterialExpressionMaterialFunctionCall* MFCallNode = NewObject<UMaterialExpressionMaterialFunctionCall>(this);
+			MFCallNode->SetMaterialFunction(DefaultMF);
+
+			SetPosXAndMoveReferenceToTheRight(MFCallNode);
+			EditorOnly->FrontMaterial.Connect(0, MFCallNode);
+		}
+		else
+		{
+			// Or if it cannot be found, a slab node
+			UMaterialExpressionStrataSlabBSDF* SlabNode = NewObject<UMaterialExpressionStrataSlabBSDF>(this);
+			SetPosXAndMoveReferenceToTheRight(SlabNode);
+			EditorOnly->FrontMaterial.Connect(0, SlabNode);
+		}
 		bRelinkCustomOutputNodes = false;
 		bInvalidateShader = true;
 	}
