Fix the crash on opening non-Substrate materials.

diff --git a/Engine/Source/Runtime/Engine/Classes/Engine/EngineTypes.h b/Engine/Source/Runtime/Engine/Classes/Engine/EngineTypes.h
--- a/Engine/Source/Runtime/Engine/Classes/Engine/EngineTypes.h
+++ b/Engine/Source/Runtime/Engine/Classes/Engine/EngineTypes.h
@@ -724,6 +724,7 @@ public:
 	FStrataMaterialInfo(EStrataShadingModel InShadingModel) { AddShadingModel(InShadingModel); }
 
 	bool Serialize(FArchive& Ar);
+	void PostSerialize(const FArchive& Ar);
 
 	// Shading model
 	void AddShadingModel(EStrataShadingModel InShadingModel) { check(InShadingModel < SSM_NUM); ShadingModelField |= (uint16)(1 << (uint16)InShadingModel); }
@@ -801,6 +802,7 @@ struct TStructOpsTypeTraits<FStrataMaterialInfo> : public TStructOpsTypeTraitsBa
 	enum
 	{
 		WithSerializer = true,
+		WithPostSerialize = true,
 	};
 };
 
diff --git a/Engine/Source/Runtime/Engine/Private/Materials/MaterialCachedData.cpp b/Engine/Source/Runtime/Engine/Private/Materials/MaterialCachedData.cpp
--- a/Engine/Source/Runtime/Engine/Private/Materials/MaterialCachedData.cpp
+++ b/Engine/Source/Runtime/Engine/Private/Materials/MaterialCachedData.cpp
@@ -1103,7 +1103,17 @@ void FMaterialCachedExpressionData::GetParameterValueByIndex(EMaterialParameterT
 bool FStrataMaterialInfo::Serialize(FArchive& Ar)
 {
 	Ar.UsingCustomVersion(FUE5MainStreamObjectVersion::GUID);
+	return false;
+}
+
+bool FMaterialCachedExpressionData::Serialize(FArchive& Ar)
+{
+	Ar.UsingCustomVersion(FUE5MainStreamObjectVersion::GUID);
+	return false;
+}
 
+void FStrataMaterialInfo::PostSerialize(const FArchive& Ar)
+{
 	if (Ar.IsLoading())
 	{
 		if (Ar.CustomVer(FUE5MainStreamObjectVersion::GUID) < FUE5MainStreamObjectVersion::IncreaseMaterialAttributesInputMask)
@@ -1111,14 +1121,10 @@ bool FStrataMaterialInfo::Serialize(FArchive& Ar)
 			ConnectedPropertyMask = uint64(ConnectedProperties_DEPRECATED);
 		}
 	}
-
-	return true;
 }
 
-bool FMaterialCachedExpressionData::Serialize(FArchive& Ar)
+void FMaterialCachedExpressionData::PostSerialize(const FArchive& Ar)
 {
-	Ar.UsingCustomVersion(FUE5MainStreamObjectVersion::GUID);
-
 	if (Ar.IsLoading())
 	{
 		if (Ar.CustomVer(FUE5MainStreamObjectVersion::GUID) < FUE5MainStreamObjectVersion::IncreaseMaterialAttributesInputMask)
@@ -1127,11 +1133,6 @@ bool FMaterialCachedExpressionData::Serialize(FArchive& Ar)
 		}
 	}
 
-	return true;
-}
-
-void FMaterialCachedExpressionData::PostSerialize(const FArchive& Ar)
-{
 #if WITH_EDITORONLY_DATA
 	if(Ar.IsLoading())
 	{
diff --git a/Engine/Source/Runtime/Engine/Public/MaterialCachedData.h b/Engine/Source/Runtime/Engine/Public/MaterialCachedData.h
--- a/Engine/Source/Runtime/Engine/Public/MaterialCachedData.h
+++ b/Engine/Source/Runtime/Engine/Public/MaterialCachedData.h
@@ -367,5 +367,6 @@ struct TStructOpsTypeTraits<FMaterialCachedExpressionData> : public TStructOpsTy
 	enum
 	{
 		WithSerializer = true,
+		WithPostSerialize = true,
 	};
 };
