Fix crash calling UAnimInstance::GetAllCurveNames from worker threads

diff --git a/Engine/Source/Runtime/Engine/Private/Animation/AnimInstance.cpp b/Engine/Source/Runtime/Engine/Private/Animation/AnimInstance.cpp
--- a/Engine/Source/Runtime/Engine/Private/Animation/AnimInstance.cpp
+++ b/Engine/Source/Runtime/Engine/Private/Animation/AnimInstance.cpp
@@ -1396,7 +1396,7 @@ void UAnimInstance::UpdateCurvesToComponents(USkeletalMeshComponent* Component /
 
 void UAnimInstance::AppendAnimationCurveList(EAnimCurveType Type, TMap<FName, float>& InOutCurveList) const
 {
-	InOutCurveList.Append(GetProxyOnGameThread<FAnimInstanceProxy>().GetAnimationCurves(Type));
+	InOutCurveList.Append(GetProxyOnAnyThread<FAnimInstanceProxy>().GetAnimationCurves(Type));
 }
 
 void UAnimInstance::GetAnimationCurveList(EAnimCurveType Type, TMap<FName, float>& InOutCurveList) const
