Make overlap queries handle LWI managers consistently with line traces, by creating a managed FActorInstanceHandle that can be used to fetch the underlying actor.

diff --git a/Engine/Source/Runtime/Engine/Private/Collision/CollisionConversions.cpp b/Engine/Source/Runtime/Engine/Private/Collision/CollisionConversions.cpp
--- a/Engine/Source/Runtime/Engine/Private/Collision/CollisionConversions.cpp
+++ b/Engine/Source/Runtime/Engine/Private/Collision/CollisionConversions.cpp
@@ -610,7 +610,14 @@ void ConvertQueryOverlap(const FPhysicsShape& Shape, const FPhysicsActor& Actor,
         BodyInst = FPhysicsInterface::ShapeToOriginalBodyInstance(BodyInst, &Shape);
 		if (const UPrimitiveComponent* OwnerComponent = BodyInst->OwnerComponent.Get())
 		{
-			OutOverlap.OverlapObjectHandle = FActorInstanceHandle(OwnerComponent->GetOwner());
+			if (ALightWeightInstanceManager* LWIManager = Cast<ALightWeightInstanceManager>(OwnerComponent->GetOwner()))
+			{
+				OutOverlap.OverlapObjectHandle = FActorInstanceHandle(LWIManager, BodyInst->InstanceBodyIndex);
+			}
+			else
+			{
+				OutOverlap.OverlapObjectHandle = FActorInstanceHandle(OwnerComponent->GetOwner());
+			}
 			OutOverlap.Component = BodyInst->OwnerComponent; // Copying weak pointer is faster than assigning raw pointer.
 			OutOverlap.ItemIndex = OwnerComponent->bMultiBodyOverlap ? BodyInst->InstanceBodyIndex : INDEX_NONE;
 		}
