Substrate - updated preoject setting tool tip.

diff --git a/Engine/Source/Runtime/Engine/Classes/Engine/RendererSettings.h b/Engine/Source/Runtime/Engine/Classes/Engine/RendererSettings.h
--- a/Engine/Source/Runtime/Engine/Classes/Engine/RendererSettings.h
+++ b/Engine/Source/Runtime/Engine/Classes/Engine/RendererSettings.h
@@ -953,7 +953,7 @@ class ENGINE_API URendererSettings : public UDeveloperSettings
 	*/
 	UPROPERTY(config, EditAnywhere, Category = Materials, meta = (
 		ConsoleVariable = "r.Material.RoughDiffuse", DisplayName = "Enable Rough Diffuse Material",
-		ToolTip = "Enable Rough Diffuse Material.",
+		ToolTip = "Enable Rough Diffuse Material. Please note that when Substrate is enabled, energy conservation is forced to ENABLED.",
 		ConfigRestartRequired = true))
 		uint32 bMaterialRoughDiffuse : 1; 
 
@@ -962,7 +962,7 @@ class ENGINE_API URendererSettings : public UDeveloperSettings
 	*/
 	UPROPERTY(config, EditAnywhere, Category = Materials, meta = (
 		ConsoleVariable = "r.Material.EnergyConservation", DisplayName = "Enable Energy Conservation on Material",
-		ToolTip = "Enable Energy Conservation on Material. Please note that when Substrate is enabled, energy conservation is forced to enabled.",
+		ToolTip = "Enable Energy Conservation on Material. Please note that when Substrate is enabled, energy conservation is forced to ENABLED.",
 		ConfigRestartRequired = true))
 		uint32 bMaterialEnergyConservation : 1;
 
