Fix rename in HierarchicalInstancedStaticMesh.cpp in recent commit that was missed by prior work

diff --git a/Engine/Source/Runtime/Engine/Private/HierarchicalInstancedStaticMesh.cpp b/Engine/Source/Runtime/Engine/Private/HierarchicalInstancedStaticMesh.cpp
--- a/Engine/Source/Runtime/Engine/Private/HierarchicalInstancedStaticMesh.cpp
+++ b/Engine/Source/Runtime/Engine/Private/HierarchicalInstancedStaticMesh.cpp
@@ -1892,7 +1892,7 @@ void FHierarchicalStaticMeshSceneProxy::GetDynamicMeshElements(const TArray<cons
 						int MaxEndCullDistance = CVarFoliageMaxEndCullDistance.GetValueOnRenderThread();
 						const float LODRandom = CVarRandomLODRange.GetValueOnRenderThread();
 						const float MaxDrawDistanceScale = GetCachedScalabilityCVars().ViewDistanceScale;
-						const float SphereRadius = RenderData->Bounds.SphereRadius + GetMaxWorldPositionOffsetDisplacement();
+						const float SphereRadius = RenderData->Bounds.SphereRadius + GetMaxWorldPositionOffsetExtent();
 
 						checkSlow(NumLODs > 0);
 
