Remove "beta" label from virtual shadow maps

diff --git a/Engine/Source/Runtime/Engine/Classes/Engine/EngineTypes.h b/Engine/Source/Runtime/Engine/Classes/Engine/EngineTypes.h
--- a/Engine/Source/Runtime/Engine/Classes/Engine/EngineTypes.h
+++ b/Engine/Source/Runtime/Engine/Classes/Engine/EngineTypes.h
@@ -479,8 +479,8 @@ namespace EShadowMapMethod
 		/** Render geometry into shadow depth maps for shadowing.  Requires manual setup of shadowing distances and only culls per-component, causing poor performance with high poly scenes.  Required to enable stationary baked shadows (but which is incompatible with Nanite geometry). */
 		ShadowMaps UMETA(DisplayName = "Shadow Maps"),
 
-		/** Render geometry into virtualized shadow depth maps for shadowing.  Provides high-quality shadows for next-gen projects with simplified setup.  High efficiency culling when used with Nanite. This system is in development and thus has a number of performance pitfalls. */
-		VirtualShadowMaps UMETA(DisplayName = "Virtual Shadow Maps (Beta)")
+		/** Render geometry into virtualized shadow depth maps for shadowing.  Provides high-quality shadows for next-gen projects with simplified setup.  High efficiency culling when used with Nanite. */
+		VirtualShadowMaps UMETA(DisplayName = "Virtual Shadow Maps")
 	};
 }
 
