Few more missed renames from clashing commit

diff --git a/Engine/Source/Runtime/Engine/Private/HierarchicalInstancedStaticMesh.cpp b/Engine/Source/Runtime/Engine/Private/HierarchicalInstancedStaticMesh.cpp
--- a/Engine/Source/Runtime/Engine/Private/HierarchicalInstancedStaticMesh.cpp
+++ b/Engine/Source/Runtime/Engine/Private/HierarchicalInstancedStaticMesh.cpp
@@ -911,7 +911,7 @@ void FHierarchicalStaticMeshSceneProxy::SetupOcclusion(UHierarchicalInstancedSta
 		OcclusionBounds.Reserve(NumNodes);
 		FMatrix XForm = InComponent->GetRenderMatrix();
 
-		const float MaxWorldPositionOffset = GetMaxWorldPositionOffsetDisplacement();
+		const float MaxWorldPositionOffset = GetMaxWorldPositionOffsetExtent();
 
 		for (int32 Index = FirstOcclusionNode; Index <= LastOcclusionNode; Index++)
 		{
@@ -1671,7 +1671,7 @@ void FHierarchicalStaticMeshSceneProxy::GetDynamicMeshElements(const TArray<cons
 
 				InstanceParams.ViewOriginInLocalZero = WorldToLocal.TransformPosition(View->GetTemporalLODOrigin(0, bMultipleSections));
 				InstanceParams.ViewOriginInLocalOne = WorldToLocal.TransformPosition(View->GetTemporalLODOrigin(1, bMultipleSections));
-				InstanceParams.MaxWPODisplacement = GetMaxWorldPositionOffsetDisplacement();
+				InstanceParams.MaxWPODisplacement = GetMaxWorldPositionOffsetExtent();
 
 				float MinSize = bIsOrtho ? 0.0f : CVarFoliageMinimumScreenSize.GetValueOnRenderThread();
 				float LODScale = UserData_AllInstances.LODDistanceScale;
@@ -1982,7 +1982,7 @@ void FHierarchicalStaticMeshSceneProxy::GetDynamicMeshElements(const TArray<cons
 			if (View->Family->EngineShowFlags.HISMCClusterTree)
 			{
 				FColor StartingColor(100, 0, 0);
-				const float MaxWorldPositionOffset = GetMaxWorldPositionOffsetDisplacement();
+				const float MaxWorldPositionOffset = GetMaxWorldPositionOffsetExtent();
 
 				for (const FClusterNode& CulsterNode : ClusterTree)
 				{
