Fix error from CL 25042857

diff --git a/Engine/Source/Runtime/Engine/Private/HierarchicalInstancedStaticMesh.cpp b/Engine/Source/Runtime/Engine/Private/HierarchicalInstancedStaticMesh.cpp
--- a/Engine/Source/Runtime/Engine/Private/HierarchicalInstancedStaticMesh.cpp
+++ b/Engine/Source/Runtime/Engine/Private/HierarchicalInstancedStaticMesh.cpp
@@ -1986,7 +1986,7 @@ void FHierarchicalStaticMeshSceneProxy::GetDynamicMeshElements(const TArray<cons
 
 				for (const FClusterNode& CulsterNode : ClusterTree)
 				{
-					DrawWireBox(Collector.GetPDI(ViewIndex), GetLocalToWorld(), FBox(CulsterNode.BoundMin, CulsterNode.BoundMax).ExpandBy(MaxWPODisplacement), StartingColor, View->Family->EngineShowFlags.Game ? SDPG_World : SDPG_Foreground);
+					DrawWireBox(Collector.GetPDI(ViewIndex), GetLocalToWorld(), FBox(CulsterNode.BoundMin, CulsterNode.BoundMax).ExpandBy(MaxWorldPositionOffset), StartingColor, View->Family->EngineShowFlags.Game ? SDPG_World : SDPG_Foreground);
 					StartingColor.R += 5;
 					StartingColor.G += 5;
 					StartingColor.B += 5;
