Remove unused stats
#rnx

diff --git a/Engine/Source/Runtime/Engine/Private/UnrealEngine.cpp b/Engine/Source/Runtime/Engine/Private/UnrealEngine.cpp
--- a/Engine/Source/Runtime/Engine/Private/UnrealEngine.cpp
+++ b/Engine/Source/Runtime/Engine/Private/UnrealEngine.cpp
@@ -12620,7 +12620,6 @@ DEFINE_STAT(STAT_StaticMeshTriangles);
 /** Skeletal stats */
 DEFINE_STAT(STAT_SkinningTime);
 DEFINE_STAT(STAT_UpdateClothVertsTime);
-DEFINE_STAT(STAT_UpdateSoftBodyVertsTime);
 DEFINE_STAT(STAT_SkelMeshTriangles);
 DEFINE_STAT(STAT_SkelMeshDrawCalls);
 DEFINE_STAT(STAT_CPUSkinVertices);
diff --git a/Engine/Source/Runtime/Engine/Public/EngineStats.h b/Engine/Source/Runtime/Engine/Public/EngineStats.h
--- a/Engine/Source/Runtime/Engine/Public/EngineStats.h
+++ b/Engine/Source/Runtime/Engine/Public/EngineStats.h
@@ -17,7 +17,6 @@ DECLARE_CYCLE_STAT_EXTERN(TEXT("HUD Time"),STAT_HudTime,STATGROUP_Engine, );
 DECLARE_DWORD_COUNTER_STAT_EXTERN(TEXT("Static Mesh Tris"),STAT_StaticMeshTriangles,STATGROUP_Engine, ENGINE_API);
 DECLARE_CYCLE_STAT_EXTERN(TEXT("Skel Skin Time"),STAT_SkinningTime,STATGROUP_Engine, );
 DECLARE_CYCLE_STAT_EXTERN(TEXT("Update Cloth Verts Time"),STAT_UpdateClothVertsTime,STATGROUP_Engine, );
-DECLARE_CYCLE_STAT_EXTERN(TEXT("Update SoftBody Verts Time"),STAT_UpdateSoftBodyVertsTime,STATGROUP_Engine, );
 DECLARE_DWORD_COUNTER_STAT_EXTERN(TEXT("Skel Mesh Tris"),STAT_SkelMeshTriangles,STATGROUP_Engine, );
 DECLARE_DWORD_COUNTER_STAT_EXTERN(TEXT("Skel Mesh Draw Calls"),STAT_SkelMeshDrawCalls,STATGROUP_Engine, );
 DECLARE_DWORD_COUNTER_STAT_EXTERN(TEXT("Skel Verts CPU Skin"),STAT_CPUSkinVertices,STATGROUP_Engine, );
