Fix suffix calculation: handle prefixed packages (/Temp)

#rnx

diff --git a/Engine/Source/Runtime/Engine/Private/WorldPartition/WorldPartitionRuntimeSpatialHash.cpp b/Engine/Source/Runtime/Engine/Private/WorldPartition/WorldPartitionRuntimeSpatialHash.cpp
--- a/Engine/Source/Runtime/Engine/Private/WorldPartition/WorldPartitionRuntimeSpatialHash.cpp
+++ b/Engine/Source/Runtime/Engine/Private/WorldPartition/WorldPartitionRuntimeSpatialHash.cpp
@@ -1277,9 +1277,9 @@ FString UWorldPartitionRuntimeSpatialHash::GetCellNameString(UWorld* InOuterWorl
 			const FTopLevelAssetPath InstanceAssetPath(InstancedWorldPath);
 			const FString InstancePackageName = InstanceAssetPath.GetPackageName().ToString();
 
-			if (InstancePackageName.StartsWith(SourcePackageName))
+			if (int32 Index = InstancePackageName.Find(SourcePackageName); Index != INDEX_NONE)
 			{
-				*OutInstanceSuffix = InstancePackageName.Mid(SourcePackageName.Len());
+				*OutInstanceSuffix = InstancePackageName.Mid(Index + SourcePackageName.Len());
 			}
 		}
 	}
