Fixed tooltip for PostProcessComponent.MotionBlurAmount.

diff --git a/Engine/Source/Runtime/Engine/Classes/Engine/Scene.h b/Engine/Source/Runtime/Engine/Classes/Engine/Scene.h
--- a/Engine/Source/Runtime/Engine/Classes/Engine/Scene.h
+++ b/Engine/Source/Runtime/Engine/Classes/Engine/Scene.h
@@ -2076,7 +2076,7 @@ struct FPostProcessSettings
 	UPROPERTY(interp, EditAnywhere, BlueprintReadWrite, Category="Lens|Mobile Depth of Field", AdvancedDisplay, meta=(UIMin = "0.0", UIMax = "100.0", editcondition = "bOverride_DepthOfFieldVignetteSize", DisplayName = "Vignette Size"))
 	float DepthOfFieldVignetteSize;
 
-	/** Strength of motion blur, 0:off, should be renamed to intensity */
+	/** Strength of motion blur, 0:off */
 	UPROPERTY(interp, BlueprintReadWrite, Category="Rendering Features|Motion Blur", meta=(ClampMin = "0.0", ClampMax = "1.0", editcondition = "bOverride_MotionBlurAmount", DisplayName = "Amount"))
 	float MotionBlurAmount;
 	/** max distortion caused by motion blur, in percent of the screen width, 0:off */
