Fix d3d12 gpu crash on unaligned copy.
For virtual textures we generate an editor only image for later use with thumbnails etc.
This strips borders when compositing tiles to lay them out in a regular texture.
But if we have block compressed formats with border less than 4 then doing this will lead to an unaligned copy.
The fix is to keep the small unwanted borders for just that case.

diff --git a/Engine/Source/Runtime/Engine/Private/Texture2D.cpp b/Engine/Source/Runtime/Engine/Private/Texture2D.cpp
--- a/Engine/Source/Runtime/Engine/Private/Texture2D.cpp
+++ b/Engine/Source/Runtime/Engine/Private/Texture2D.cpp
@@ -1398,10 +1398,11 @@ void FVirtualTexture2DResource::InitializeEditorResources(IVirtualTexture* InVir
 		}
 
 		const EPixelFormat PixelFormat = VTData->LayerTypes[0];
+		const bool bCopyUnwantedBordersForAlignment = VTData->TileBorderSize <= 2 && IsBlockCompressedFormat(PixelFormat);
 		const uint32 MipScaleFactor = (1u << MipLevel);
 		const uint32 MipWidthInTiles = FMath::DivideAndRoundUp(GetNumTilesX(), MipScaleFactor);
 		const uint32 MipHeightInTiles = FMath::DivideAndRoundUp(GetNumTilesY(), MipScaleFactor);
-		const uint32 TileSizeInPixels = GetTileSize();
+		const uint32 TileSizeInPixels = bCopyUnwantedBordersForAlignment ? GetTileSize() + 2 * GetBorderSize() : GetTileSize();
 		const uint32 LayerMask = 1u; // FVirtualTexture2DResource should only have a single layer
 
 		TArray<FPageToProduce> PagesToProduce;
@@ -1447,6 +1448,10 @@ void FVirtualTexture2DResource::InitializeEditorResources(IVirtualTexture* InVir
 
 		FRHICommandListImmediate& RHICommandList = FRHICommandListExecutor::GetImmediateCommandList();
 
+		// We want to strip borders when compositing tiles since we're just laying out tiles in a regular texture.
+		// But if we have block compressed formats with border less than 4 then doing this will lead to an unaligned copy. We keep the small unwanted borders in that case.
+		const EVTProducePageFlags ProducePageFlags = bCopyUnwantedBordersForAlignment ? EVTProducePageFlags::None : EVTProducePageFlags::SkipPageBorders;
+
 		TArray<IVirtualTextureFinalizer*> Finalizers;
 		for (const FPageToProduce& Page : PagesToProduce)
 		{
@@ -1458,7 +1463,7 @@ void FVirtualTexture2DResource::InitializeEditorResources(IVirtualTexture* InVir
 
 			IVirtualTextureFinalizer* Finalizer = InVirtualTexture->ProducePageData(RHICommandList,
 				GMaxRHIFeatureLevel,
-				EVTProducePageFlags::SkipPageBorders, // don't want to produce page borders, since we're laying out tiles in a regular texture
+				ProducePageFlags,
 				ProducerHandle, LayerMask, MipLevel, vAddress,
 				Page.Handle,
 				&TargetLayer);
