Expose FPlatformTime::Seconds() to scripts in editor

diff --git a/Engine/Source/Runtime/Engine/Classes/Kismet/KismetSystemLibrary.h b/Engine/Source/Runtime/Engine/Classes/Kismet/KismetSystemLibrary.h
--- a/Engine/Source/Runtime/Engine/Classes/Kismet/KismetSystemLibrary.h
+++ b/Engine/Source/Runtime/Engine/Classes/Kismet/KismetSystemLibrary.h
@@ -209,6 +209,17 @@ class ENGINE_API UKismetSystemLibrary : public UBlueprintFunctionLibrary
 	UFUNCTION(BlueprintPure, Category = "Utilities")
 	static int64 GetFrameCount();
 
+#if WITH_EDITOR
+	/** 
+	 * Returns the current platform time in seconds. Not coupled to any gameplay or other containerization logic - this
+	 * function is useful for timing execution time or timestamping data. Marked as callable rather than pure because 
+	 * implicit evaluation may be confusing, both for blueprint authors and blueprint readers. For implicit execution 
+	 * simply wrap it in a blueprint pure function.
+	 */
+	UFUNCTION(BlueprintCallable, Category = "Utilities|Time")
+	static double GetPlatformTime_Seconds();
+#endif // WITH_EDITOR
+
 	/** Returns whether the world this object is in is the host or not */
 	UFUNCTION(BlueprintPure, Category="Networking", meta=(WorldContext="WorldContextObject") )
 	static bool IsServer(const UObject* WorldContextObject);
diff --git a/Engine/Source/Runtime/Engine/Private/KismetSystemLibrary.cpp b/Engine/Source/Runtime/Engine/Private/KismetSystemLibrary.cpp
--- a/Engine/Source/Runtime/Engine/Private/KismetSystemLibrary.cpp
+++ b/Engine/Source/Runtime/Engine/Private/KismetSystemLibrary.cpp
@@ -257,6 +257,13 @@ int64 UKismetSystemLibrary::GetFrameCount()
 	return (int64) GFrameCounter;
 }
 
+#if WITH_EDITOR
+double UKismetSystemLibrary::GetPlatformTime_Seconds()
+{
+	return FPlatformTime::Seconds();
+}
+#endif//WITH_EDITOR
+
 bool UKismetSystemLibrary::IsServer(const UObject* WorldContextObject)
 {
 	UWorld* World = GEngine->GetWorldFromContextObject(WorldContextObject, EGetWorldErrorMode::LogAndReturnNull);
