Skeletal Mesh: Allow creation, use and cooking of arbitrary per-vertex attributes on skeletal meshes.

diff --git a/Engine/Plugins/Experimental/SkeletalReduction/Source/Private/SkeletalMeshReductionPlugin.cpp b/Engine/Plugins/Experimental/SkeletalReduction/Source/Private/SkeletalMeshReductionPlugin.cpp
--- a/Engine/Plugins/Experimental/SkeletalReduction/Source/Private/SkeletalMeshReductionPlugin.cpp
+++ b/Engine/Plugins/Experimental/SkeletalReduction/Source/Private/SkeletalMeshReductionPlugin.cpp
@@ -154,6 +154,10 @@ private:
 	*/
 	struct FPrioritizedFeatures;
 
+	/*
+	 * Vertex attribute information
+	 */
+	struct FVertexAttributeInfo;
 private:
 
 	/**
@@ -202,7 +206,8 @@ private:
 	* @param SrcLODModel    - the original model
 	* @param BoneMatrices   - define the configuration of the model
 	* @param LODIndex       - target index for the result.
-	* @param OutSkinnedMesh - the posed mesh 
+	* @param OutSkinnedMesh - the posed mesh
+	* @param VertexAttributeInfos The vertex attribute information from the original and how it is packed in the skinned mesh.
 	* @param VerToSectionMap - if provided, on return this array will allow VertexID to SectionID lookup
 	*/
 	void ConvertToFSkinnedSkeletalMesh( const FSkeletalMeshLODModel& SrcLODModel,
@@ -210,6 +215,7 @@ private:
 		                                const int32 LODIndex,
 		                                SkeletalSimplifier::FSkinnedSkeletalMesh& OutSkinnedMesh,
 										const FString& SkeletalMeshName,
+										TArray<FVertexAttributeInfo>& VertexAttributeInfos, 
 										TArray<int32>* VertToSectionMap = nullptr) const;
 
 	/**
@@ -219,6 +225,7 @@ private:
 	* @param SrcLODModel     - Reference model - provided in case some closest per-vertex data needs to be transfered (e.g. alternate weights) 
 	* @param SkinnedMesh     - the source mesh
 	* @param RefSkeleton    - reference skeleton
+	* @param VertexAttributeInfos The vertex attribute information from the original and how it is packed in the skinned mesh.
 	* @param NewModel       - resulting MeshLODModel
 	*/
 	void ConvertToFSkeletalMeshLODModel( const FString& SkeletalMeshName,
@@ -226,6 +233,7 @@ private:
 		                                 const FSkeletalMeshLODModel& SrcLODModel,
 		                                 const SkeletalSimplifier::FSkinnedSkeletalMesh& SkinnedMesh,
 		                                 const FReferenceSkeleton& RefSkeleton,
+		                                 const TArray<FVertexAttributeInfo>& VertexAttributeInfos,
 		                                 FSkeletalMeshLODModel& NewModel,
 										 const bool bReducingSourceModel,
 										 const ITargetPlatform* TargetPlatform) const;
@@ -328,6 +336,19 @@ struct FQuadricSkeletalMeshReduction::FPrioritizedFeatures
 		return (BoneIds.Num() >0 || SectionIds.Num() > 0);
 	}
 };
+
+struct FQuadricSkeletalMeshReduction::FVertexAttributeInfo
+{
+	/* The name of the vertex attribute */
+	FName AttributeName;
+	
+	/* Number of float components per vertex for this attribute */
+	int32 ComponentCount;
+	
+	/* Offset into the reducers's compacted attribute storage for unpacking. */
+	int32 Offset;
+};
+
 /**
 *  Required MeshReduction Interface.
 */
@@ -477,12 +498,15 @@ bool FQuadricSkeletalMeshReduction::RemoveMeshSection(FSkeletalMeshLODModel& Mod
 }
 
 
-void FQuadricSkeletalMeshReduction::ConvertToFSkinnedSkeletalMesh( const FSkeletalMeshLODModel& SrcLODModel,
-	                                                               const TArray<FMatrix>& BoneMatrices,
-	                                                               const int32 LODIndex,
-	                                                               SkeletalSimplifier::FSkinnedSkeletalMesh& OutSkinnedMesh,
-																   const FString& SkeletalMeshName,
-																   TArray<int32>* VertToSectionMap) const
+void FQuadricSkeletalMeshReduction::ConvertToFSkinnedSkeletalMesh(
+	const FSkeletalMeshLODModel& SrcLODModel,
+	const TArray<FMatrix>& BoneMatrices,
+	const int32 LODIndex,
+	SkeletalSimplifier::FSkinnedSkeletalMesh& OutSkinnedMesh,
+	const FString& SkeletalMeshName,
+	TArray<FVertexAttributeInfo>& VertexAttributeInfos,
+	TArray<int32>* VertToSectionMap
+	) const
 {
 
 
@@ -610,6 +634,15 @@ void FQuadricSkeletalMeshReduction::ConvertToFSkinnedSkeletalMesh( const FSkelet
 		return (MaxLODIndex != -1 && MaxLODIndex < LODIndex);
 	};
 
+	// Fill in the attribute infos
+	VertexAttributeInfos.Reset();
+	int32 AttributeOffset = 0;
+	for (const TPair<FName, FSkeletalMeshModelVertexAttribute>& Attribute: SrcLODModel.VertexAttributes)
+	{
+		VertexAttributeInfos.Add({Attribute.Key, Attribute.Value.ComponentCount, AttributeOffset});
+		AttributeOffset += Attribute.Value.ComponentCount; 
+	}
+
 	// Count the total number of verts, but only the number of triangles that
 	// are used in sections we don't skip.
 	// NB: This could result zero triangles, but a non-zero number of verts.
@@ -780,6 +813,23 @@ void FQuadricSkeletalMeshReduction::ConvertToFSkinnedSkeletalMesh( const FSkelet
 				}
 			}
 		}
+
+		// Copy over the attributes. 
+		for (const FVertexAttributeInfo& AttributeInfo: VertexAttributeInfos)
+		{
+			const FSkeletalMeshModelVertexAttribute& VertexAttribute = SrcLODModel.VertexAttributes[AttributeInfo.AttributeName]; 			
+			
+			for (int32 Vertex = VertexRange.Begin; Vertex < VertexRange.End; ++Vertex)
+			{
+				auto& VertexAttrs  = OutVertexBuffer[Vertex].AdditionalAttributes;
+				const float* AttributeData = VertexAttribute.Values.GetData() + Vertex * VertexAttribute.ComponentCount; 
+				
+				for (int32 Component = 0; Component < VertexAttribute.ComponentCount; Component++)
+				{
+					VertexAttrs.SetElement(AttributeInfo.Offset + Component, AttributeData[Component]); 
+				}
+			}
+		}
 	}
 
 	// store sectionID or MaterialID in the material id (there is a one to one mapping between them).
@@ -1653,14 +1703,16 @@ void  FQuadricSkeletalMeshReduction::AddSourceModelInfluences( const FSkeletalMe
 
 
 
-void FQuadricSkeletalMeshReduction::ConvertToFSkeletalMeshLODModel( const FString& SkeletalMeshName,
-	                                                                const int32 MaxBonesPerVertex,
-	                                                                const FSkeletalMeshLODModel& SrcLODModel,
-	                                                                const SkeletalSimplifier::FSkinnedSkeletalMesh& SkinnedMesh,
-	                                                                const FReferenceSkeleton& RefSkeleton,
-	                                                                FSkeletalMeshLODModel& NewModel,
-																	const bool bReducingSourceModel,
-																    const ITargetPlatform* TargetPlatform) const
+void FQuadricSkeletalMeshReduction::ConvertToFSkeletalMeshLODModel(
+	const FString& SkeletalMeshName,
+	const int32 MaxBonesPerVertex,
+	const FSkeletalMeshLODModel& SrcLODModel,
+	const SkeletalSimplifier::FSkinnedSkeletalMesh& SkinnedMesh,
+	const FReferenceSkeleton& RefSkeleton,
+	const TArray<FVertexAttributeInfo>& VertexAttributeInfos,
+	FSkeletalMeshLODModel& NewModel,
+	const bool bReducingSourceModel,
+	const ITargetPlatform* TargetPlatform) const
 {
 	// We might be re-using this model - so clear it.
 
@@ -1679,8 +1731,6 @@ void FQuadricSkeletalMeshReduction::ConvertToFSkeletalMeshLODModel( const FStrin
 
 		AddSourceModelInfluences(SrcLODModel, SkinnedMesh, NewModel, SkeletalMeshName);
 	}
-	
-
 
 	// Create dummy map of 'point to original'
 
@@ -1723,6 +1773,27 @@ void FQuadricSkeletalMeshReduction::ConvertToFSkeletalMeshLODModel( const FStrin
 	// Set texture coordinate count on the new model.
 	NewModel.NumTexCoords = SkeletalMeshData.TexCoordCount;
 
+	// Copy back the vertex attributes, if any. Use the MeshToImportVertexMap on the NewModel to map from the new point location to the
+	// one in the reduction mesh.
+	for (const FVertexAttributeInfo& AttributeInfo: VertexAttributeInfos)
+	{
+		FSkeletalMeshModelVertexAttribute& VertexAttribute = NewModel.VertexAttributes.FindOrAdd(AttributeInfo.AttributeName);
+
+		VertexAttribute.ComponentCount = AttributeInfo.ComponentCount;
+		VertexAttribute.Values.Reserve(SkinnedMesh.NumVertices() * VertexAttribute.ComponentCount);
+
+		for (int32 Index = 0; Index < SkinnedMesh.NumVertices(); Index++)
+		{
+			const int32 OriginalIndex = NewModel.MeshToImportVertexMap[Index];
+			const auto& Vertex = SkinnedMesh.VertexBuffer[OriginalIndex];
+			for (int32 Component = 0; Component < VertexAttribute.ComponentCount; Component++)
+			{
+				VertexAttribute.Values.Add(static_cast<float>(Vertex.AdditionalAttributes.GetElement(AttributeInfo.Offset + Component)));
+			}
+		}
+	}
+
+	
 
 	//if (!bReducingSourceModel)
 	{
@@ -1776,9 +1847,10 @@ bool FQuadricSkeletalMeshReduction::ReduceSkeletalLODModel( const FSkeletalMeshL
 	}
 	
 	// Generate a single skinned mesh form the SrcModel.  This mesh has per-vertex tangent space.
-	TArray<int32> VertToOriginalSectionId; 
+	TArray<int32> VertToOriginalSectionId;
+	TArray<FVertexAttributeInfo> VertexAttributeInfos;
 	SkeletalSimplifier::FSkinnedSkeletalMesh SkinnedSkeletalMesh;
-	ConvertToFSkinnedSkeletalMesh(SrcModel, BoneMatrices, LODIndex, SkinnedSkeletalMesh, SkeletalMeshName, &VertToOriginalSectionId);
+	ConvertToFSkinnedSkeletalMesh(SrcModel, BoneMatrices, LODIndex, SkinnedSkeletalMesh, SkeletalMeshName, VertexAttributeInfos, &VertToOriginalSectionId);
 
 	int32 IterationNum = 0;
 	//We keep the original MaxNumVerts because if we iterate we want to still compare with the original request.
@@ -1819,8 +1891,9 @@ bool FQuadricSkeletalMeshReduction::ReduceSkeletalLODModel( const FSkeletalMeshL
 		}
 
 		// Convert to SkeletalMeshLODModel. 
-
-		ConvertToFSkeletalMeshLODModel(SkeletalMeshName, Settings.MaxBonesPerVertex, SrcModel, SkinnedSkeletalMesh, RefSkeleton, OutSkeletalMeshLODModel, bReducingSourceModel, TargetPlatform);
+		ConvertToFSkeletalMeshLODModel(
+			SkeletalMeshName, Settings.MaxBonesPerVertex, SrcModel, SkinnedSkeletalMesh, RefSkeleton, 
+			VertexAttributeInfos, OutSkeletalMeshLODModel, bReducingSourceModel, TargetPlatform);
 
 		// We may need to do additional simplification if the user specified a hard number limit for verts and
 		// the internal chunking during conversion split some verts.
@@ -1844,7 +1917,7 @@ bool FQuadricSkeletalMeshReduction::ReduceSkeletalLODModel( const FSkeletalMeshL
 				}
 
 				UE_LOG(LogSkeletalMeshReduction, Log, TEXT("Chunking to limit unique bones per section generated additional vertices - continuing simplification of LOD %d "), LODIndex);
-				ConvertToFSkinnedSkeletalMesh(SrcModel, BoneMatrices, LODIndex, SkinnedSkeletalMesh, SkeletalMeshName);
+				ConvertToFSkinnedSkeletalMesh(SrcModel, BoneMatrices, LODIndex, SkinnedSkeletalMesh, SkeletalMeshName, VertexAttributeInfos);
 			}
 			else
 			{
diff --git a/Engine/Plugins/Experimental/SkeletalReduction/Source/Private/SkeletalSimplifierLinearAlgebra.h b/Engine/Plugins/Experimental/SkeletalReduction/Source/Private/SkeletalSimplifierLinearAlgebra.h
--- a/Engine/Plugins/Experimental/SkeletalReduction/Source/Private/SkeletalSimplifierLinearAlgebra.h
+++ b/Engine/Plugins/Experimental/SkeletalReduction/Source/Private/SkeletalSimplifierLinearAlgebra.h
@@ -885,6 +885,12 @@ namespace SkeletalSimplifier
 				Size = size;
 			}
 
+			TDenseArrayWrapper(TArrayView<DataType> InArrayView)
+			{
+				Data = InArrayView.GetData();
+				Size = InArrayView.Num();
+			}
+
 			int32 Num() const { return Size; }
 
 			void SetElement(const int32 j, const DataType Value) { checkSlow(j < Size);  Data[j] = Value; }
@@ -895,7 +901,7 @@ namespace SkeletalSimplifier
 			DataType& operator() (const int32 j) { return operator[](j); }
 			DataType operator() (const int32 j) const { return operator[](j); }
 
-			TDenseArrayWrapper& operator*= (const double& Scalar)
+			TDenseArrayWrapper& operator*= (double Scalar)
 			{
 				int32 Num = Size;
 				for (int32 i = 0; i < Num; ++i) Data[i] *= Scalar;
diff --git a/Engine/Plugins/Experimental/SkeletalReduction/Source/Private/SkeletalSimplifierMeshManager.h b/Engine/Plugins/Experimental/SkeletalReduction/Source/Private/SkeletalSimplifierMeshManager.h
--- a/Engine/Plugins/Experimental/SkeletalReduction/Source/Private/SkeletalSimplifierMeshManager.h
+++ b/Engine/Plugins/Experimental/SkeletalReduction/Source/Private/SkeletalSimplifierMeshManager.h
@@ -23,9 +23,9 @@ namespace SkeletalSimplifier
 	*  Define the Simplifier types needed.
 	*/
  
-	typedef VertexTypes::TBasicVertexAttrs<MAX_TEXCOORDS>                      BasicAttrArray;
-	typedef VertexTypes::BoneSparseVertexAttrs                                 SparseAttrArray;
-	typedef VertexTypes::BoneSparseVertexAttrs                                 BoneAttrArray;
+	typedef VertexTypes::TBasicVertexAttrs<MAX_TEXCOORDS>	BasicAttrArray;
+	typedef VertexTypes::DynamicVertexAttrs                 SparseAttrArray;
+	typedef VertexTypes::BoneSparseVertexAttrs              BoneAttrArray;
 
 	typedef VertexTypes::TSkeletalSimpVert<BasicAttrArray, SparseAttrArray, BoneAttrArray>    MeshVertType;
 
diff --git a/Engine/Plugins/Experimental/SkeletalReduction/Source/Private/SkeletalSimplifierVertex.h b/Engine/Plugins/Experimental/SkeletalReduction/Source/Private/SkeletalSimplifierVertex.h
--- a/Engine/Plugins/Experimental/SkeletalReduction/Source/Private/SkeletalSimplifierVertex.h
+++ b/Engine/Plugins/Experimental/SkeletalReduction/Source/Private/SkeletalSimplifierVertex.h
@@ -256,6 +256,18 @@ namespace SkeletalSimplifier
 
 		};
 
+		/** Dynamic vertex attributes. Distinguished from the BoneSparseVertexAttrs. Based on SparseVecD, rather than
+		 *  using dense vector wrappers, because of requirements by the simplifier. 
+		 */
+		class DynamicVertexAttrs : public SparseVecD
+		{
+		public:
+			static void Correct()
+			{
+				// Do nothing.
+			}
+		};  
+
 		/**
 		* Sparse attributes.
 		* Used to hold bone weights where the bone ID is the attribute key.
diff --git a/Engine/Source/Developer/MeshBuilder/Private/SkeletalMeshBuilder.cpp b/Engine/Source/Developer/MeshBuilder/Private/SkeletalMeshBuilder.cpp
--- a/Engine/Source/Developer/MeshBuilder/Private/SkeletalMeshBuilder.cpp
+++ b/Engine/Source/Developer/MeshBuilder/Private/SkeletalMeshBuilder.cpp
@@ -163,6 +163,8 @@ bool FSkeletalMeshBuilder::Build(const FSkeletalMeshBuildParameters& SkeletalMes
 			FLODUtilities::UpdateAlternateSkinWeights(SkeletalMesh, ProfileInfo.Name, LODIndex, Options);
 		}
 
+		// Copy vertex attribute definitions and their values from the import model.
+		FLODUtilities::UpdateLODInfoVertexAttributes(SkeletalMesh, LODIndex, LODIndex, /*CopyAttributeValues*/true);
 		
 		FSkeletalMeshUpdateContext UpdateContext;
 		UpdateContext.SkeletalMesh = SkeletalMesh;
diff --git a/Engine/Source/Developer/SkeletalMeshUtilitiesCommon/Private/LODUtilities.cpp b/Engine/Source/Developer/SkeletalMeshUtilitiesCommon/Private/LODUtilities.cpp
--- a/Engine/Source/Developer/SkeletalMeshUtilitiesCommon/Private/LODUtilities.cpp
+++ b/Engine/Source/Developer/SkeletalMeshUtilitiesCommon/Private/LODUtilities.cpp
@@ -2,49 +2,47 @@
 
 #include "LODUtilities.h"
 
-#include "BoneWeights.h"
-
 #if WITH_EDITOR
 
-#include "Misc/MessageDialog.h"
-#include "Misc/FeedbackContext.h"
-#include "Modules/ModuleManager.h"
-#include "UObject/UObjectIterator.h"
-#include "Components/SkinnedMeshComponent.h"
-#include "Components/SkeletalMeshComponent.h"
 #include "Animation/MorphTarget.h"
-#include "UObject/GarbageCollection.h"
-#include "Rendering/SkeletalMeshModel.h"
-#include "Rendering/SkeletalMeshLODModel.h"
-#include "GenericQuadTree.h"
+#include "Animation/SkinWeightProfile.h"
+#include "Async/ParallelFor.h"
+#include "BoneWeights.h"
+#include "ClothingAsset.h"
+#include "ComponentReregisterContext.h"
+#include "Components/SkeletalMeshComponent.h"
+#include "Components/SkinnedMeshComponent.h"
+#include "EditorFramework/AssetImportData.h"
 #include "Engine/SkeletalMesh.h"
 #include "Engine/SkeletalMeshLODSettings.h"
 #include "Engine/SkeletalMeshSocket.h"
 #include "Engine/SkinnedAssetAsyncCompileUtils.h"
 #include "Engine/SkinnedAssetCommon.h"
 #include "Engine/Texture2D.h"
-#include "EditorFramework/AssetImportData.h"
-#include "MeshUtilities.h"
-#include "MeshUtilitiesCommon.h"
-#include "ClothingAsset.h"
-#include "OverlappingCorners.h"
 #include "Framework/Commands/UIAction.h"
+#include "Framework/Notifications/NotificationManager.h"
+#include "GenericQuadTree.h"
 #include "HAL/ThreadSafeBool.h"
+#include "IMeshReductionManagerModule.h"
 #include "ImageCore.h"
 #include "ImageCoreUtils.h"
-#include "Rendering/SkeletalMeshLODImporterData.h"
-#include "ObjectTools.h"
-
-#include "ComponentReregisterContext.h"
-#include "IMeshReductionManagerModule.h"
-#include "Animation/SkinWeightProfile.h"
-
-#include "Async/ParallelFor.h"
+#include "ImportUtils/SkelImport.h"
 #include "Interfaces/ITargetPlatform.h"
 #include "Interfaces/ITargetPlatformManagerModule.h"
+#include "MeshUtilities.h"
+#include "MeshUtilitiesCommon.h"
 #include "Misc/CoreMisc.h"
-
-#include "Framework/Notifications/NotificationManager.h"
+#include "Misc/FeedbackContext.h"
+#include "Misc/MessageDialog.h"
+#include "Modules/ModuleManager.h"
+#include "ObjectTools.h"
+#include "OverlappingCorners.h"
+#include "Rendering/SkeletalMeshLODImporterData.h"
+#include "Rendering/SkeletalMeshLODModel.h"
+#include "Rendering/SkeletalMeshModel.h"
+#include "Tasks/Task.h"
+#include "UObject/GarbageCollection.h"
+#include "UObject/UObjectIterator.h"
 #include "Widgets/Notifications/SNotificationList.h"
 
 #include <limits>
@@ -1643,11 +1641,14 @@
 			{
 				ClearGeneratedMorphTarget(SkeletalMesh, DesiredLOD);
 			}
+
+			// Update the vertex attribute information in the LOD info.
+			UpdateLODInfoVertexAttributes(SkeletalMesh, ReductionSettings.BaseLOD, DesiredLOD, false);
 		}
 
 		if (IsInGameThread())
 		{
-			SkeletalMesh->MarkPackageDirty();
+			(void)SkeletalMesh->MarkPackageDirty();
 		}
 		else if(OutNeedsPackageDirtied)
 		{
@@ -2531,6 +2532,7 @@
 	return bBuildSuccess;
 }
 
+
 bool FLODUtilities::UpdateAlternateSkinWeights(
 	FSkeletalMeshLODModel& LODModelDest,
 	FSkeletalMeshImportData& ImportDataDest,
@@ -2762,6 +2764,94 @@
 	}
 }
 
+
+
+bool FLODUtilities::UpdateLODInfoVertexAttributes(
+	USkeletalMesh *InSkeletalMesh,
+	int32 InSourceLODIndex, 
+	int32 InTargetLODIndex, 
+	bool bInCopyAttributeValues
+	)
+{
+	FSkeletalMeshImportData SkeletalMeshImportData;
+	InSkeletalMesh->LoadLODImportedData(InSourceLODIndex, SkeletalMeshImportData);
+	
+	FSkeletalMeshLODModel& TargetLODModel = InSkeletalMesh->GetImportedModel()->LODModels[InTargetLODIndex];
+	
+	TArray<FSkeletalMeshVertexAttributeInfo>& SkelMeshAttributeInfos = InSkeletalMesh->GetLODInfo(InTargetLODIndex)->VertexAttributes; 
+
+	// Retain any existing attribute infos and match based on names.
+	TMap<FName, FSkeletalMeshVertexAttributeInfo> ExistingAttributeInfos;
+	for (FSkeletalMeshVertexAttributeInfo& AttributeInfo: SkelMeshAttributeInfos)
+	{
+		ExistingAttributeInfos.Add(AttributeInfo.Name, MoveTemp(AttributeInfo));
+	}
+
+	SkelMeshAttributeInfos.Reset(SkeletalMeshImportData.VertexAttributes.Num());
+
+	// If we're not copying the values, leave the existing data in place.
+	if (bInCopyAttributeValues)
+	{
+		TargetLODModel.VertexAttributes.Reset();
+	}
+	
+	TArray<UE::Tasks::FTask> ConversionTasks;
+	
+	for (int32 AttributeIndex = 0; AttributeIndex < SkeletalMeshImportData.VertexAttributes.Num(); AttributeIndex++)
+	{
+		const SkeletalMeshImportData::FVertexAttribute& ImportAttribute = SkeletalMeshImportData.VertexAttributes[AttributeIndex];
+		const FName AttributeName(SkeletalMeshImportData.VertexAttributeNames[AttributeIndex]);
+
+		// Did this definition already exist? Try to retain as much of the existing information as possible.
+		FSkeletalMeshVertexAttributeInfo Info;
+		
+		if (ExistingAttributeInfos.Contains(AttributeName))
+		{
+			Info = MoveTemp(ExistingAttributeInfos[AttributeName]);
+		}
+		else
+		{
+			Info.Name = AttributeName;
+		}
+
+		SkelMeshAttributeInfos.Add(Info);
+
+		if (Info.IsEnabledForRender() && bInCopyAttributeValues)
+		{
+			FSkeletalMeshModelVertexAttribute& ModelAttribute = TargetLODModel.VertexAttributes.FindOrAdd(AttributeName);
+
+			ModelAttribute.DataType = Info.DataType;
+			ModelAttribute.ComponentCount = ImportAttribute.ComponentCount;
+
+			if (InTargetLODIndex == InSourceLODIndex)
+			{
+				ConversionTasks.Add(
+					UE::Tasks::Launch(UE_SOURCE_LOCATION, [&TargetLODModel, &ModelAttribute, &ImportAttribute]()
+					{
+						ModelAttribute.Values.SetNumUninitialized(TargetLODModel.NumVertices);
+						for(uint32 VertexIndex = 0; VertexIndex < TargetLODModel.NumVertices; VertexIndex++)
+						{
+							const int32 ImportVertexIndex = TargetLODModel.MeshToImportVertexMap[VertexIndex];
+							ModelAttribute.Values[VertexIndex] = ImportAttribute.AttributeValues[ImportVertexIndex];
+						}
+					})
+				);
+			}
+			else
+			{
+				// Initialize with zero data, matching the vertex count.
+				ModelAttribute.Values.SetNumZeroed(TargetLODModel.NumVertices);
+			}
+		}
+	}
+
+	// Wait for all the attribute conversion tasks to complete.
+	UE::Tasks::Wait(ConversionTasks);
+
+	return true;
+}
+
+
 void FLODUtilities::RegenerateDependentLODs(USkeletalMesh* SkeletalMesh, int32 LODIndex, const ITargetPlatform* TargetPlatform)
 {
 	int32 LODNumber = SkeletalMesh->GetLODNum();
diff --git a/Engine/Source/Developer/SkeletalMeshUtilitiesCommon/Public/LODUtilities.h b/Engine/Source/Developer/SkeletalMeshUtilitiesCommon/Public/LODUtilities.h
--- a/Engine/Source/Developer/SkeletalMeshUtilitiesCommon/Public/LODUtilities.h
+++ b/Engine/Source/Developer/SkeletalMeshUtilitiesCommon/Public/LODUtilities.h
@@ -211,6 +211,10 @@ public:
 		bool bUseMikkTSpace,
 		bool bComputeWeightedNormals);
 
+	
+	/** Build the vertex attributes */
+	static bool UpdateLODInfoVertexAttributes(USkeletalMesh *InSkeletalMesh, int32 InSourceLODIndex, int32 InTargetLODIndex, bool bInCopyAttributeValues);
+	
 	/**
 	 * Re-generate all (editor-only) skin weight profile, used whenever we rebuild the skeletal mesh data which could change the chunking and bone indices
 	 * 
diff --git a/Engine/Source/Editor/UnrealEd/Public/ImportUtils/SkelImport.h b/Engine/Source/Editor/UnrealEd/Public/ImportUtils/SkelImport.h
--- a/Engine/Source/Editor/UnrealEd/Public/ImportUtils/SkelImport.h
+++ b/Engine/Source/Editor/UnrealEd/Public/ImportUtils/SkelImport.h
@@ -12,6 +12,7 @@
 #include "ReferenceSkeleton.h"
 #include "Engine/SkeletalMesh.h"
 #include "Rendering/SkeletalMeshLODImporterData.h"
+#include "SkeletalMeshTypes.h"
 
 class UAssetImportData;
 class UMorphTarget;
@@ -20,6 +21,7 @@ class USkeletalMeshSocket;
 class USkeleton;
 class UThumbnailInfo;
 class FSkeletalMeshLODModel;
+enum class ESkinVertexColorChannel : uint8;
 
 struct FExistingMeshLodSectionData
 {
diff --git a/Engine/Source/Runtime/Core/Private/UObject/UE5MainStreamObjectVersion.cpp b/Engine/Source/Runtime/Core/Private/UObject/UE5MainStreamObjectVersion.cpp
--- a/Engine/Source/Runtime/Core/Private/UObject/UE5MainStreamObjectVersion.cpp
+++ b/Engine/Source/Runtime/Core/Private/UObject/UE5MainStreamObjectVersion.cpp
@@ -15,7 +15,7 @@ TMap<FGuid, FGuid> FUE5MainStreamObjectVersion::GetSystemGuids()
 	SystemGuids.Add(DevGuids.NANITE_DERIVEDDATA_VER, FGuid("C283995F5F1346C3B6F7FB1DDF01ED7F"));
 	SystemGuids.Add(DevGuids.NIAGARASHADERMAP_DERIVEDDATA_VER, FGuid("8A37C45D24F2423CBE5F8F371DE33575"));
 	SystemGuids.Add(DevGuids.Niagara_LatestScriptCompileVersion, FGuid("4E3E29C28655E162A05986724B41711B"));
-	SystemGuids.Add(DevGuids.SkeletalMeshDerivedDataVersion, FGuid("4084BE94509A8C3FA7CF410DDF8ABF04"));
+	SystemGuids.Add(DevGuids.SkeletalMeshDerivedDataVersion, FGuid("4F04FE8848F36B4F2AEB5D156092E08F"));
 	SystemGuids.Add(DevGuids.STATICMESH_DERIVEDDATA_VER, FGuid("98F7D79A5811013825E75BBCC41ED3E9"));
 	SystemGuids.Add(DevGuids.POSESEARCHDB_DERIVEDDATA_VER, FGuid("4E595C2AC5E947D6BA9ABC874353E5BC"));
 	SystemGuids.Add(DevGuids.GROOM_BINDING_DERIVED_DATA_VERSION, FGuid("156678D4F9084D7CAE98FADC6AB93573"));
diff --git a/Engine/Source/Runtime/Core/Public/UObject/UE5MainStreamObjectVersion.h b/Engine/Source/Runtime/Core/Public/UObject/UE5MainStreamObjectVersion.h
--- a/Engine/Source/Runtime/Core/Public/UObject/UE5MainStreamObjectVersion.h
+++ b/Engine/Source/Runtime/Core/Public/UObject/UE5MainStreamObjectVersion.h
@@ -342,7 +342,7 @@ struct CORE_API FUE5MainStreamObjectVersion
 
 		// Change the default for facing & alignment to be automatic
 		NiagaraSpriteRendererFacingAlignmentAutoDefault,
-
+		
 		// Change the default for facing & alignment to be automatic
 		GroomAssetRemoveInAssetSerialization,
 
@@ -352,6 +352,9 @@ struct CORE_API FUE5MainStreamObjectVersion
 		// Combines proprties into a new binding so users can select constant or binding
 		NiagaraSimStageNumIterationsBindings,
 
+		// Skeletal vertex attributes
+		SkeletalVertexAttributes,
+		
 		// -----<new versions can be added above this line>-------------------------------------------------
 		VersionPlusOne,
 		LatestVersion = VersionPlusOne - 1
diff --git a/Engine/Source/Runtime/Engine/Classes/Engine/SkinnedAssetCommon.h b/Engine/Source/Runtime/Engine/Classes/Engine/SkinnedAssetCommon.h
--- a/Engine/Source/Runtime/Engine/Classes/Engine/SkinnedAssetCommon.h
+++ b/Engine/Source/Runtime/Engine/Classes/Engine/SkinnedAssetCommon.h
@@ -12,6 +12,8 @@
 #include "Engine/EngineTypes.h"
 #include "PerPlatformProperties.h"
 #include "SkeletalMeshReductionSettings.h"
+#include "Animation/SkeletalMeshVertexAttribute.h"
+
 #include "SkinnedAssetCommon.generated.h"
 
 
@@ -182,6 +184,10 @@ struct FSkeletalMeshLODInfo
 	UPROPERTY(EditAnywhere, Category = SkeletalMeshLODInfo)
 	uint8 bAllowCPUAccess : 1;
 
+	/** List of vertex attributes to include for rendering and what type they should be */
+	UPROPERTY(EditAnywhere, Category = SkeletalMeshLODInfo, AdvancedDisplay, EditFixedSize, Meta=(NoResetToDefault))
+	TArray<FSkeletalMeshVertexAttributeInfo> VertexAttributes;	
+	
 	/**
 	Mesh supports uniformly distributed sampling in constant time.
 	Memory cost is 8 bytes per triangle.
diff --git a/Engine/Source/Runtime/Engine/Private/Animation/SkeletalMeshVertexAttribute.cpp b/Engine/Source/Runtime/Engine/Private/Animation/SkeletalMeshVertexAttribute.cpp
new file mode 100644
--- /dev/null
+++ b/Engine/Source/Runtime/Engine/Private/Animation/SkeletalMeshVertexAttribute.cpp
@@ -0,0 +1,154 @@
+﻿// Copyright Epic Games, Inc. All Rights Reserved.
+
+#include "Animation/SkeletalMeshVertexAttribute.h"
+
+#include "Rendering/SkeletalMeshAttributeVertexBuffer.h"
+
+#include UE_INLINE_GENERATED_CPP_BY_NAME(SkeletalMeshVertexAttribute)
+
+
+FName FSkeletalMeshVertexAttributeInfo::GetRequirementName() const
+{
+	return FName(FString::Printf(TEXT("Vertex.%s"), *Name.ToString()));
+}
+
+bool FSkeletalMeshVertexAttributeInfo::IsEnabledForRender() const
+{
+	return EnabledForRender.GetValue();
+}
+
+
+#if WITH_EDITOR
+FArchive& operator<<(FArchive& Ar, FSkeletalMeshModelVertexAttribute& Attribute)
+{
+	Ar << Attribute.DataType;
+	Ar << Attribute.Values;
+	Ar << Attribute.ComponentCount;
+	return Ar;
+}
+#endif
+
+
+FSkeletalMeshVertexAttributeRenderData::~FSkeletalMeshVertexAttributeRenderData()
+{
+	for (const TTuple<FName, FSkeletalMeshAttributeVertexBuffer*>& BufferItem: Buffers)
+	{
+		delete BufferItem.Value;
+	}
+}
+
+
+bool FSkeletalMeshVertexAttributeRenderData::AddAttribute(
+	const FName InName,
+	const ESkeletalMeshVertexAttributeDataType InDataType,
+	const int32 InNumVertices,
+	const int32 InComponentCount,
+	const TArray<float>& InValues
+	)
+{
+	check(InComponentCount >= 1 && InComponentCount <= 4);
+	
+	FSkeletalMeshAttributeVertexBuffer* Buffer;
+	if (Buffers.Contains(InName))
+	{
+		Buffer = Buffers[InName];
+	}
+	else
+	{
+		Buffer = Buffers.Add(InName, new FSkeletalMeshAttributeVertexBuffer);
+	}
+
+	EPixelFormat PixelFormat = PF_Unknown;
+	switch(InDataType)
+	{
+	case ESkeletalMeshVertexAttributeDataType::NUInt8:
+		
+		PixelFormat = PF_R8;
+		break;
+	case ESkeletalMeshVertexAttributeDataType::Float:
+		PixelFormat = PF_R32_FLOAT;
+		break;
+	case ESkeletalMeshVertexAttributeDataType::HalfFloat:
+		PixelFormat = PF_R16F;
+		break;
+	}
+
+	Buffer->Init(PixelFormat, InNumVertices, InComponentCount, InValues);
+	return true;
+}
+
+
+FSkeletalMeshAttributeVertexBuffer* FSkeletalMeshVertexAttributeRenderData::GetAttributeBuffer(FName InName) const
+{
+	FSkeletalMeshAttributeVertexBuffer* const* BufferPtr = Buffers.Find(InName);
+	return BufferPtr ? *BufferPtr : nullptr;
+}
+
+void FSkeletalMeshVertexAttributeRenderData::InitResources()
+{
+	check(IsInGameThread());
+	for (const TTuple<FName, FSkeletalMeshAttributeVertexBuffer*>& BufferItem: Buffers)
+	{
+		BeginInitResource(BufferItem.Value);
+	}
+}
+
+void FSkeletalMeshVertexAttributeRenderData::ReleaseResources()
+{
+	check(IsInGameThread());
+	for (const TTuple<FName, FSkeletalMeshAttributeVertexBuffer*>& BufferItem: Buffers)
+	{
+		BeginReleaseResource(BufferItem.Value);
+	}
+}
+
+
+void FSkeletalMeshVertexAttributeRenderData::CleanUp()
+{
+	check(IsInGameThread());
+	for (const TTuple<FName, FSkeletalMeshAttributeVertexBuffer*>& BufferItem: Buffers)
+	{
+		BufferItem.Value->CleanUp();
+	}
+}
+
+
+int32 FSkeletalMeshVertexAttributeRenderData::GetResourceSize() const
+{
+	int32 Size = 0;
+	for (const TTuple<FName, FSkeletalMeshAttributeVertexBuffer*>& BufferItem: Buffers)
+	{
+		Size += BufferItem.Value->GetResourceSize();
+	}
+	return Size;
+}
+
+
+void FSkeletalMeshVertexAttributeRenderData::Serialize(FArchive& Ar)
+{
+	if (Ar.IsLoading())
+	{
+		int32 Count;
+		Ar << Count;
+		Buffers.Reset();
+		for (int32 Index = 0; Index < Count; Index++)
+		{
+			FName Name;
+			Ar << Name;
+			FSkeletalMeshAttributeVertexBuffer* Buffer = new FSkeletalMeshAttributeVertexBuffer;
+			Buffer->Serialize(Ar);
+			Buffers.Add(Name, Buffer);
+		}
+	}
+	else
+	{
+		int32 Count = Buffers.Num();
+		Ar << Count;
+		for (const TTuple<FName, FSkeletalMeshAttributeVertexBuffer*>& BufferItem: Buffers)
+		{
+			FName Name = BufferItem.Key;
+			Ar << Name;
+			BufferItem.Value->Serialize(Ar);
+		}
+	}
+}
diff --git a/Engine/Source/Runtime/Engine/Private/Rendering/SkeletalMeshAttributeVertexBuffer.cpp b/Engine/Source/Runtime/Engine/Private/Rendering/SkeletalMeshAttributeVertexBuffer.cpp
new file mode 100644
--- /dev/null
+++ b/Engine/Source/Runtime/Engine/Private/Rendering/SkeletalMeshAttributeVertexBuffer.cpp
@@ -0,0 +1,201 @@
+﻿// Copyright Epic Games, Inc. All Rights Reserved.
+
+#include "Rendering/SkeletalMeshAttributeVertexBuffer.h"
+
+#include "Rendering/SkeletalMeshVertexBuffer.h"
+
+
+class FSkeletalMeshVertexAttributeData :
+	public FSkeletalMeshVertexDataInterface,
+	public TResourceArray<uint8, VERTEXBUFFER_ALIGNMENT>
+{
+	using FResourceArray = TResourceArray<uint8, VERTEXBUFFER_ALIGNMENT>;
+public:
+	explicit FSkeletalMeshVertexAttributeData(
+		const int32 InStride,
+		const bool InNeedsCPUAccess = false
+		) :
+		FResourceArray(InNeedsCPUAccess),
+		Stride(InStride)
+	{
+	}
+	
+	// FSkeletalMeshVertexDataInterface implementations
+	virtual void ResizeBuffer(uint32 InNumVertices) override
+	{
+		const int32 Count = static_cast<int32>(InNumVertices) * Stride;
+		
+		if (Num() < Count)
+		{
+			Reserve(Count);
+			AddUninitialized(Count - Num());
+		}
+		else if (Num() > Count)
+		{
+			RemoveAt(Count, Num() - Count);
+		}
+	}
+
+	virtual uint32 GetStride() const override
+	{
+		return Stride;
+	}
+
+	virtual uint8* GetDataPointer() override
+	{
+		return GetData();
+	}
+	
+
+	virtual uint32 GetNumVertices() override
+	{
+		return Num() / Stride;
+	}
+	
+	virtual FResourceArrayInterface* GetResourceArray() override
+	{
+		return this;
+	}
+
+	virtual void Serialize(FArchive& Ar) override
+	{
+		Ar << Stride;
+		BulkSerialize(Ar);
+	}
+
+private:
+	int32 Stride;
+};
+
+
+FSkeletalMeshAttributeVertexBuffer::~FSkeletalMeshAttributeVertexBuffer()
+{
+	CleanUp();
+}
+
+
+void FSkeletalMeshAttributeVertexBuffer::Init(
+	const EPixelFormat InPixelFormat, 
+	const int32 InNumVertices,
+	const int32 InComponentCount,
+	TConstArrayView<float> InValues
+	)
+{
+	check(InNumVertices * InComponentCount == InValues.Num());
+	check(InComponentCount >= 1 && InComponentCount <= 4);
+
+	ComponentCount = InComponentCount;
+	PixelFormat = InPixelFormat;
+
+	Allocate();
+
+	ValueData->ResizeBuffer(InNumVertices);
+
+	SetData(InValues);
+}
+
+void FSkeletalMeshAttributeVertexBuffer::Allocate()
+{
+	CleanUp();
+
+	switch(PixelFormat)
+	{
+	case PF_R8:
+		ComponentStride = 1;
+		break;
+	case PF_R16F:
+		ComponentStride = 2;
+		break;
+	case PF_R32_FLOAT:
+		ComponentStride = 4;
+		break;
+	default:
+		checkNoEntry();
+		return;
+	}
+
+	ValueData = new FSkeletalMeshVertexAttributeData(ComponentStride * ComponentCount);
+}
+
+
+/** Convenience helpers to convert the float value data to various target types */
+template<typename T> T ConvertValue(const float In);
+template<> uint8 ConvertValue(const float In) { return FMath::Clamp(In, 0.0f, 1.0f) * MAX_uint8; }
+template<> int8 ConvertValue(const float In) { return FMath::Clamp(In, -1.0f, 1.0f) * MAX_int8; }
+template<> uint16 ConvertValue(const float In) { return FMath::Clamp(In, 0.0f, 1.0f) * MAX_uint16; }
+template<> int16 ConvertValue(const float In) { return FMath::Clamp(In, -1.0f, 1.0f) * MAX_int16; }
+template<> FFloat16 ConvertValue(const float In) { return FFloat16(In); }
+
+template<typename T>
+static void ConvertBatch(uint8* InDstValuesRawPtr, TConstArrayView<float> InSrcValues)
+{
+	T* DstValues = reinterpret_cast<T*>(InDstValuesRawPtr);
+	for (int32 Index = 0; Index < InSrcValues.Num(); Index++)
+	{
+		*DstValues++ = ConvertValue<T>(InSrcValues[Index]);
+	}
+}
+
+
+void FSkeletalMeshAttributeVertexBuffer::SetData(TConstArrayView<float> InValues)
+{
+	uint8* DstValuePtr = ValueData->GetDataPointer();
+
+	switch (PixelFormat)
+	{
+	case PF_R8:
+		ConvertBatch<uint8>(DstValuePtr, InValues);
+		break;
+	case PF_R16F:
+		ConvertBatch<FFloat16>(DstValuePtr, InValues);
+		break;
+	case PF_R32_FLOAT:
+		FMemory::Memcpy(DstValuePtr, InValues.GetData(), InValues.Num() * sizeof(float));
+		break;
+	default:
+		checkNoEntry();
+	}
+}
+
+
+void FSkeletalMeshAttributeVertexBuffer::CleanUp()
+{
+	delete ValueData;
+	ValueData = nullptr;
+}
+
+
+int32 FSkeletalMeshAttributeVertexBuffer::GetResourceSize() const
+{
+	return ValueData ? ValueData->GetNumVertices() * ValueData->GetStride() : 0;
+}
+
+
+void FSkeletalMeshAttributeVertexBuffer::InitRHI()
+{
+	FResourceArrayInterface* ResourceArray = ValueData->GetResourceArray();
+	FRHIResourceCreateInfo CreateInfo(TEXT("SkeletalMeshAttributeVertexBuffer"), ResourceArray);
+	ValueBuffer.VertexBufferRHI = RHICreateVertexBuffer(ResourceArray->GetResourceDataSize(), BUF_Static | BUF_ShaderResource, CreateInfo);
+	ValueBuffer.VertexBufferSRV = RHICreateShaderResourceView(ValueBuffer.VertexBufferRHI, ComponentStride, PixelFormat);
+}
+
+
+void FSkeletalMeshAttributeVertexBuffer::ReleaseRHI()
+{
+	ValueBuffer.SafeRelease();
+}
+
+
+void FSkeletalMeshAttributeVertexBuffer::Serialize(FArchive& Ar)
+{
+	Ar << ComponentCount;
+	Ar << PixelFormat;
+	Ar << ComponentStride;
+
+	if (Ar.IsLoading())
+	{
+		Allocate();
+	}
+
+	ValueData->Serialize(Ar);
+}
diff --git a/Engine/Source/Runtime/Engine/Private/Rendering/SkeletalMeshLODImporterData.cpp b/Engine/Source/Runtime/Engine/Private/Rendering/SkeletalMeshLODImporterData.cpp
--- a/Engine/Source/Runtime/Engine/Private/Rendering/SkeletalMeshLODImporterData.cpp
+++ b/Engine/Source/Runtime/Engine/Private/Rendering/SkeletalMeshLODImporterData.cpp
@@ -1,9 +1,9 @@
 // Copyright Epic Games, Inc. All Rights Reserved.
 
-#if WITH_EDITOR
-
 #include "Rendering/SkeletalMeshLODImporterData.h"
 
+#if WITH_EDITOR
+
 #include "Engine/SkeletalMesh.h"
 #include "Engine/SkinnedAssetCommon.h"
 #include "ImportUtils/SkeletalMeshImportUtils.h"
@@ -46,6 +46,9 @@ void FSkeletalMeshImportData::KeepAlternateSkinningBuildDataOnly()
 	//Remove material array and PointToRawMap
 	Materials.Empty();
 	PointToRawMap.Empty();
+	
+	VertexAttributes.Empty();
+	VertexAttributeNames.Empty();
 }
 
 /**
@@ -165,6 +168,8 @@ bool FSkeletalMeshImportData::ReplaceSkeletalMeshGeometryImportData(const USkele
 	ImportData->MorphTargets.Reset();
 	ImportData->MorphTargetModifiedPoints.Reset();
 	ImportData->MeshInfos.Reset();
+	ImportData->VertexAttributes.Reset();
+	ImportData->VertexAttributeNames.Reset();
 
 	//Material is a special case since we cannot serialize the UMaterialInstance when saving the RawSkeletalMeshBulkData
 	//So it has to be reconstructed.
@@ -189,6 +194,8 @@ bool FSkeletalMeshImportData::ReplaceSkeletalMeshGeometryImportData(const USkele
 	ImportData->MorphTargets += OriginalSkeletalMeshImportData.MorphTargets;
 	ImportData->MorphTargetModifiedPoints += OriginalSkeletalMeshImportData.MorphTargetModifiedPoints;
 	ImportData->MeshInfos += OriginalSkeletalMeshImportData.MeshInfos;
+	ImportData->VertexAttributes += OriginalSkeletalMeshImportData.VertexAttributes;
+	ImportData->VertexAttributeNames += OriginalSkeletalMeshImportData.VertexAttributeNames;
 
 	return ImportData->ApplyRigToGeo(NewGeometryAndRigData);
 }
@@ -651,6 +658,7 @@ enum
 	RAW_SKELETAL_MESH_BULKDATA_VER_AlternateInfluence = 1,
 	RAW_SKELETAL_MESH_BULKDATA_VER_RebuildSystem = 2,
 	RAW_SKELETAL_MESH_BULKDATA_VER_CompressMorphTargetData = 3,
+	RAW_SKELETAL_MESH_BULKDATA_VER_VertexAttributes = 4,
 	// Add new raw mesh versions here.
 
 	RAW_SKELETAL_MESH_BULKDATA_VER_PLUS_ONE,
@@ -717,6 +725,17 @@ FArchive& operator<<(FArchive& Ar, FSkeletalMeshImportData& RawMesh)
 		RawMesh.AlternateInfluenceProfileNames.Empty();
 	}
 
+	if (Version >= RAW_SKELETAL_MESH_BULKDATA_VER_VertexAttributes)
+	{
+		Ar << RawMesh.VertexAttributes;
+		Ar << RawMesh.VertexAttributeNames;
+	}
+	else
+	{
+		RawMesh.VertexAttributes.Empty();
+		RawMesh.VertexAttributeNames.Empty();
+	}
+
 	if (Ar.IsLoading() && Version < RAW_SKELETAL_MESH_BULKDATA_VER_CompressMorphTargetData)
 	{
 		if (RawMesh.MorphTargetModifiedPoints.Num() != 0)
@@ -1743,6 +1762,48 @@ bool FSkeletalMeshImportData::GetMeshDescription(FMeshDescription& OutMeshDescri
 	FSkeletalMeshAttributes::FBoneParentIndexAttributesRef BoneParentIndices = MeshAttributes.GetBoneParentIndices();
 	FSkeletalMeshAttributes::FBonePoseAttributesRef BonePoses = MeshAttributes.GetBonePoses();
 
+	TMap<FString, FName> ValidAttributes;
+	for (int32 AttributeIndex = 0; AttributeIndex < VertexAttributes.Num(); AttributeIndex++)
+	{
+		const FString& VertexAttributeName = VertexAttributeNames[AttributeIndex];
+		const SkeletalMeshImportData::FVertexAttribute& VertexAttribute = VertexAttributes[AttributeIndex];
+		if (!ensure(VertexAttribute.AttributeValues.Num() == (Points.Num() * VertexAttribute.ComponentCount)))
+		{
+			continue;
+		}
+
+		EMeshAttributeFlags DefaultAttributeFlags = EMeshAttributeFlags::Mergeable | EMeshAttributeFlags::Lerpable;
+		
+		FName RegisteredName(VertexAttributeName);
+
+		// Ignore attributes with reserved names. This should have been handled at import time or when the attribute
+		// was created/renamed.
+		if (!ensure(!FSkeletalMeshAttributes::IsReservedAttributeName(RegisteredName)))
+		{
+			continue;
+		}
+		
+		switch(VertexAttribute.ComponentCount)
+		{
+		case 1:
+			OutMeshDescription.VertexAttributes().RegisterAttribute<float>(RegisteredName, 1, 0.0f, DefaultAttributeFlags);
+			break;
+		case 2:
+			OutMeshDescription.VertexAttributes().RegisterAttribute<FVector2f>(RegisteredName, 1, FVector2f::Zero(), DefaultAttributeFlags);
+			break;
+		case 3:
+			OutMeshDescription.VertexAttributes().RegisterAttribute<FVector3f>(RegisteredName, 1, FVector3f::Zero(), DefaultAttributeFlags);
+			break;
+		case 4:
+			OutMeshDescription.VertexAttributes().RegisterAttribute<FVector4f>(RegisteredName, 1, FVector4f::Zero(), DefaultAttributeFlags);
+			break;
+		default:
+			continue;
+		}
+		
+		ValidAttributes.Add(VertexAttributeName, RegisteredName);
+	}
+	
 	VertexInstanceUVs.SetNumChannels(NumTexCoords);
 	
 	// Avoid repeated allocations and reserve the target buffers right off the bat.
@@ -1896,12 +1957,167 @@ bool FSkeletalMeshImportData::GetMeshDescription(FMeshDescription& OutMeshDescri
 			FaceSmoothingMasks[PolygonID.GetValue()] = Triangle.SmoothingGroups;
 		}
 	}
+	
+	for (int32 AttributeIndex = 0; AttributeIndex < VertexAttributes.Num(); AttributeIndex++)
+	{
+		const FString& VertexAttributeName = VertexAttributeNames[AttributeIndex];
+		const SkeletalMeshImportData::FVertexAttribute& VertexAttribute = VertexAttributes[AttributeIndex];
+		
+		if (!ValidAttributes.Contains(VertexAttributeName))
+		{
+			continue;
+		}
+		
+		FName RegisteredName(ValidAttributes[VertexAttributeName]);
+		switch(VertexAttribute.ComponentCount)
+		{
+		case 1:
+			{
+				TVertexAttributesRef<float> AttributeRef = OutMeshDescription.VertexAttributes().GetAttributesRef<float>(RegisteredName);
+				for (int32 Index = 0; Index < VertexAttribute.AttributeValues.Num(); Index++)
+				{
+					AttributeRef.Set(VertexIDMap[Index], VertexAttribute.AttributeValues[Index]);
+				}
+				break;
+			}
+		case 2:
+			{
+				TVertexAttributesRef<FVector2f> AttributeRef = OutMeshDescription.VertexAttributes().GetAttributesRef<FVector2f>(RegisteredName);
+				for (int32 Index = 0; Index < VertexAttribute.AttributeValues.Num(); Index += 2)
+				{
+					AttributeRef.Set(VertexIDMap[Index / 2],
+						FVector2f(VertexAttribute.AttributeValues[Index], VertexAttribute.AttributeValues[Index + 1]));
+				}
+				break;
+			}
+		case 3:
+			{
+				TVertexAttributesRef<FVector3f> AttributeRef = OutMeshDescription.VertexAttributes().GetAttributesRef<FVector3f>(RegisteredName);
+				for (int32 Index = 0; Index < VertexAttribute.AttributeValues.Num(); Index += 3)
+				{
+					AttributeRef.Set(VertexIDMap[Index / 3],
+						FVector3f(VertexAttribute.AttributeValues[Index], VertexAttribute.AttributeValues[Index + 1], VertexAttribute.AttributeValues[Index + 2]));
+				}
+				break;
+			}
+		case 4:
+			{
+				TVertexAttributesRef<FVector4f> AttributeRef = OutMeshDescription.VertexAttributes().GetAttributesRef<FVector4f>(RegisteredName);
+				for (int32 Index = 0; Index < VertexAttribute.AttributeValues.Num(); Index += 4)
+				{
+					AttributeRef.Set(VertexIDMap[Index / 4],
+						FVector4f(
+							VertexAttribute.AttributeValues[Index], VertexAttribute.AttributeValues[Index + 1],
+							VertexAttribute.AttributeValues[Index + 2], VertexAttribute.AttributeValues[Index + 3]));
+				}
+				break;
+			}
+		default:
+			checkNoEntry();
+		}
+	}
 
 	FSkeletalMeshOperations::ConvertSmoothGroupToHardEdges(FaceSmoothingMasks, OutMeshDescription);
 	
 	return true;
 }
 
+template<typename T>
+struct FCopyAttributeElement {};
+
+template<>
+struct FCopyAttributeElement<float>
+{
+	static const int32 ComponentCount = 1;
+	static void CopyElement(const float InValue, TArrayView<float> OutValue)
+	{
+		OutValue[0] = InValue;
+	}
+};
+
+template<>
+struct FCopyAttributeElement<FVector2f>
+{
+	static const int32 ComponentCount = 2;
+	static void CopyElement(const FVector2f& InValue, TArrayView<float> OutValue)
+	{
+		OutValue[0] = InValue.X;
+		OutValue[1] = InValue.Y;
+	}
+};
+
+template<>
+struct FCopyAttributeElement<FVector3f>
+{
+	static const int32 ComponentCount = 3;
+	static void CopyElement(const FVector3f& InValue, TArrayView<float> OutValue)
+	{
+		OutValue[0] = InValue.X;
+		OutValue[1] = InValue.Y;
+		OutValue[2] = InValue.Z;
+	}
+};
+
+template<>
+struct FCopyAttributeElement<FVector4f>
+{
+	static const int32 ComponentCount = 4;
+	static void CopyElement(const FVector4f& InValue, TArrayView<float> OutValue)
+	{
+		OutValue[0] = InValue.X;
+		OutValue[1] = InValue.Y;
+		OutValue[2] = InValue.Z;
+		OutValue[3] = InValue.W;
+	}
+};
+
+
+template<typename T>
+struct FCreateAndCopyAttributeValues
+{
+	FCreateAndCopyAttributeValues(
+		FSkeletalMeshImportData& InImportData,
+		const FVertexArray& InVertices ) :
+	ImportData(InImportData), Vertices(InVertices) {}
+	
+	void operator()(const FName InAttributeName, TVertexAttributesConstRef<T> InSrcAttribute)
+	{
+		if (FSkeletalMeshAttributes::IsReservedAttributeName(InAttributeName))
+		{
+			return;
+		}
+
+		// We should not end up with a naming clash, but err on the side of caution,
+		if (!ensure(!ImportData.VertexAttributeNames.Contains(InAttributeName.ToString())))
+		{
+			return;
+		}
+
+		ImportData.VertexAttributeNames.Add(InAttributeName.ToString());
+		SkeletalMeshImportData::FVertexAttribute& DstAttribute = ImportData.VertexAttributes.AddDefaulted_GetRef();
+		
+		DstAttribute.ComponentCount = FCopyAttributeElement<T>::ComponentCount;
+		DstAttribute.AttributeValues.SetNumUninitialized(ImportData.Points.Num());
+		for (int32 VertexIndex = 0; VertexIndex < ImportData.Points.Num(); VertexIndex++)
+		{
+			// We go via the PointToRawMap because the number of points could have increased after a call to 
+			// SplitVerticesBySmoothingGroups.
+			FVertexID VertexID(ImportData.PointToRawMap[VertexIndex]);
+			FCopyAttributeElement<T>::CopyElement(
+				InSrcAttribute.Get(VertexID),
+				TArrayView<float>(DstAttribute.AttributeValues.GetData() + VertexIndex * DstAttribute.ComponentCount, DstAttribute.ComponentCount));
+		}
+	}
+	
+	// Unhandled sub-types.
+	void operator()(const FName, TVertexAttributesConstRef<TArrayAttribute<T>>) { }
+	void operator()(const FName, TVertexAttributesConstRef<TArrayView<T>>) { }
+
+private:
+	FSkeletalMeshImportData& ImportData;
+	const FVertexArray& Vertices;
+};
+
 FSkeletalMeshImportData FSkeletalMeshImportData::CreateFromMeshDescription(const FMeshDescription& InMeshDescription)
 {
 	FSkeletalMeshImportData SkelMeshImportData;
@@ -1917,6 +2133,7 @@ FSkeletalMeshImportData FSkeletalMeshImportData::CreateFromMeshDescription(const
 	TVertexInstanceAttributesConstRef<FVector4f> VertexInstanceColors = Attributes.GetVertexInstanceColors();
 
 	TPolygonGroupAttributesConstRef<FName> PolygonGroupMaterialSlotNames = Attributes.GetPolygonGroupMaterialSlotNames();
+
 	//Get the per face smoothing
 	TArray<uint32> FaceSmoothingMasks;
 	FaceSmoothingMasks.AddZeroed(InMeshDescription.Triangles().Num());
@@ -2054,6 +2271,12 @@ FSkeletalMeshImportData FSkeletalMeshImportData::CreateFromMeshDescription(const
 
 	SkelMeshImportData.CleanUpUnusedMaterials();
 	SkelMeshImportData.SplitVerticesBySmoothingGroups();
+	
+	// Copy any non-reserved float vertex attributes.
+	InMeshDescription.VertexAttributes().ForEachByType<float>(FCreateAndCopyAttributeValues<float>(SkelMeshImportData, InMeshDescription.Vertices()));
+	InMeshDescription.VertexAttributes().ForEachByType<FVector2f>(FCreateAndCopyAttributeValues<FVector2f>(SkelMeshImportData, InMeshDescription.Vertices()));
+	InMeshDescription.VertexAttributes().ForEachByType<FVector3f>(FCreateAndCopyAttributeValues<FVector3f>(SkelMeshImportData, InMeshDescription.Vertices()));
+	InMeshDescription.VertexAttributes().ForEachByType<FVector4f>(FCreateAndCopyAttributeValues<FVector4f>(SkelMeshImportData, InMeshDescription.Vertices()));
 
 	return SkelMeshImportData;
 }
diff --git a/Engine/Source/Runtime/Engine/Private/SkeletalMeshLODRenderData.cpp b/Engine/Source/Runtime/Engine/Private/SkeletalMeshLODRenderData.cpp
--- a/Engine/Source/Runtime/Engine/Private/SkeletalMeshLODRenderData.cpp
+++ b/Engine/Source/Runtime/Engine/Private/SkeletalMeshLODRenderData.cpp
@@ -281,6 +281,8 @@ void FSkeletalMeshLODRenderData::InitResources(bool bNeedsVertexColors, int32 LO
 		BeginInitResource(&MorphTargetVertexInfoBuffers);
 	}
 
+	VertexAttributeBuffers.InitResources();
+
 #if RHI_RAYTRACING
 	if (IsRayTracingAllowed())
 	{
@@ -313,6 +315,7 @@ void FSkeletalMeshLODRenderData::ReleaseResources()
 	DEC_DWORD_STAT_BY(STAT_SkeletalMeshVertexMemory, SkinWeightProfilesData.GetResourcesSize());
 	SkinWeightProfilesData.ReleaseResources();
 
+	VertexAttributeBuffers.ReleaseResources();
 #if RHI_RAYTRACING
 	if (IsRayTracingAllowed())
 	{
@@ -353,19 +356,23 @@ void FSkeletalMeshLODRenderData::DecrementMemoryStats()
 }
 
 #if WITH_EDITOR
-void FSkeletalMeshLODRenderData::BuildFromLODModel(const FSkeletalMeshLODModel* ImportedModel, ESkeletalMeshVertexFlags BuildFlags)
+void FSkeletalMeshLODRenderData::BuildFromLODModel(
+	const FSkeletalMeshLODModel* InLODModel,
+	TConstArrayView<FSkeletalMeshVertexAttributeInfo> InVertexAttributeInfos,
+	ESkeletalMeshVertexFlags InBuildFlags
+	)
 {
-	const bool bUseFullPrecisionUVs = EnumHasAllFlags(BuildFlags, ESkeletalMeshVertexFlags::UseFullPrecisionUVs);
-	const bool bUseHighPrecisionTangentBasis = EnumHasAllFlags(BuildFlags, ESkeletalMeshVertexFlags::UseHighPrecisionTangentBasis);
-	const bool bHasVertexColors = EnumHasAllFlags(BuildFlags, ESkeletalMeshVertexFlags::HasVertexColors);
-	const bool bUseBackwardsCompatibleF16TruncUVs = EnumHasAllFlags(BuildFlags, ESkeletalMeshVertexFlags::UseBackwardsCompatibleF16TruncUVs);
-	const bool bUseHighPrecisionWeights = EnumHasAllFlags(BuildFlags, ESkeletalMeshVertexFlags::UseHighPrecisionWeights);
+	const bool bUseFullPrecisionUVs = EnumHasAllFlags(InBuildFlags, ESkeletalMeshVertexFlags::UseFullPrecisionUVs);
+	const bool bUseHighPrecisionTangentBasis = EnumHasAllFlags(InBuildFlags, ESkeletalMeshVertexFlags::UseHighPrecisionTangentBasis);
+	const bool bHasVertexColors = EnumHasAllFlags(InBuildFlags, ESkeletalMeshVertexFlags::HasVertexColors);
+	const bool bUseBackwardsCompatibleF16TruncUVs = EnumHasAllFlags(InBuildFlags, ESkeletalMeshVertexFlags::UseBackwardsCompatibleF16TruncUVs);
+	const bool bUseHighPrecisionWeights = EnumHasAllFlags(InBuildFlags, ESkeletalMeshVertexFlags::UseHighPrecisionWeights);
 
 	// Copy required info from source sections
 	RenderSections.Empty();
-	for (int32 SectionIndex = 0; SectionIndex < ImportedModel->Sections.Num(); SectionIndex++)
+	for (int32 SectionIndex = 0; SectionIndex < InLODModel->Sections.Num(); SectionIndex++)
 	{
-		const FSkelMeshSection& ModelSection = ImportedModel->Sections[SectionIndex];
+		const FSkelMeshSection& ModelSection = InLODModel->Sections[SectionIndex];
 
 		FSkelMeshRenderSection NewRenderSection;
 		NewRenderSection.MaterialIndex = ModelSection.MaterialIndex;
@@ -388,7 +395,7 @@ void FSkeletalMeshLODRenderData::BuildFromLODModel(const FSkeletalMeshLODModel*
 	}
 
 	TArray<FSoftSkinVertex> Vertices;
-	ImportedModel->GetVertices(Vertices);
+	InLODModel->GetVertices(Vertices);
 
 	// match UV and tangent precision for mesh vertex buffer to setting from parent mesh
 	StaticVertexBuffers.StaticMeshVertexBuffer.SetUseFullPrecisionUVs(bUseFullPrecisionUVs);
@@ -396,19 +403,19 @@ void FSkeletalMeshLODRenderData::BuildFromLODModel(const FSkeletalMeshLODModel*
 
 	// init vertex buffer with the vertex array
 	StaticVertexBuffers.PositionVertexBuffer.Init(Vertices.Num());
-	StaticVertexBuffers.StaticMeshVertexBuffer.Init(Vertices.Num(), ImportedModel->NumTexCoords);
+	StaticVertexBuffers.StaticMeshVertexBuffer.Init(Vertices.Num(), InLODModel->NumTexCoords);
 
 	for (int32 i = 0; i < Vertices.Num(); i++)
 	{
 		StaticVertexBuffers.PositionVertexBuffer.VertexPosition(i) = Vertices[i].Position;
 		StaticVertexBuffers.StaticMeshVertexBuffer.SetVertexTangents(i, Vertices[i].TangentX, Vertices[i].TangentY, Vertices[i].TangentZ);
-		for (uint32 j = 0; j < ImportedModel->NumTexCoords; j++)
+		for (uint32 j = 0; j < InLODModel->NumTexCoords; j++)
 		{
 			StaticVertexBuffers.StaticMeshVertexBuffer.SetVertexUV(i, j, Vertices[i].UVs[j], bUseBackwardsCompatibleF16TruncUVs);
 		}
 	}
 
-	int32 MaxBoneInfluences = ImportedModel->GetMaxBoneInfluences();
+	int32 MaxBoneInfluences = InLODModel->GetMaxBoneInfluences();
 	if (!bUseHighPrecisionWeights)
 	{
 		// Re-normalize the weights for 8-bits to ensure that they are distributed using the old algorithm
@@ -445,7 +452,7 @@ void FSkeletalMeshLODRenderData::BuildFromLODModel(const FSkeletalMeshLODModel*
 	// Init skin weight buffer
 	SkinWeightVertexBuffer.SetNeedsCPUAccess(true);
 	SkinWeightVertexBuffer.SetMaxBoneInfluences(MaxBoneInfluences);
-	SkinWeightVertexBuffer.SetUse16BitBoneIndex(ImportedModel->DoSectionsUse16BitBoneIndex());
+	SkinWeightVertexBuffer.SetUse16BitBoneIndex(InLODModel->DoSectionsUse16BitBoneIndex());
 	SkinWeightVertexBuffer.SetUse16BitBoneWeight(bUseHighPrecisionWeights);
 	SkinWeightVertexBuffer.Init(Vertices);
 
@@ -455,17 +462,17 @@ void FSkeletalMeshLODRenderData::BuildFromLODModel(const FSkeletalMeshLODModel*
 		StaticVertexBuffers.ColorVertexBuffer.InitFromColorArray(&Vertices[0].Color, Vertices.Num(), sizeof(FSoftSkinVertex));
 	}
 
-	if (ImportedModel->HasClothData())
+	if (InLODModel->HasClothData())
 	{
 		TArray<FMeshToMeshVertData> MappingData;
 		TArray<FClothBufferIndexMapping> ClothIndexMapping;
-		ImportedModel->GetClothMappingData(MappingData, ClothIndexMapping);
+		InLODModel->GetClothMappingData(MappingData, ClothIndexMapping);
 		ClothVertexBuffer.Init(MappingData, ClothIndexMapping);
 	}
 
-	const uint8 DataTypeSize = (ImportedModel->NumVertices < MAX_uint16) ? sizeof(uint16) : sizeof(uint32);
+	const uint8 DataTypeSize = (InLODModel->NumVertices < MAX_uint16) ? sizeof(uint16) : sizeof(uint32);
 
-	MultiSizeIndexContainer.RebuildIndexBuffer(DataTypeSize, ImportedModel->IndexBuffer);
+	MultiSizeIndexContainer.RebuildIndexBuffer(DataTypeSize, InLODModel->IndexBuffer);
 	
 	IMeshUtilities& MeshUtilities = FModuleManager::Get().LoadModuleChecked<IMeshUtilities>("MeshUtilities");
 
@@ -473,14 +480,33 @@ void FSkeletalMeshLODRenderData::BuildFromLODModel(const FSkeletalMeshLODModel*
 
 	SkinWeightProfilesData.Init(&SkinWeightVertexBuffer);
 	// Generate runtime version of skin weight profile data, containing all required per-skin weight override data
-	for (const auto& Pair : ImportedModel->SkinWeightProfiles)
+	for (const auto& Pair : InLODModel->SkinWeightProfiles)
 	{
 		FRuntimeSkinWeightProfileData& Override = SkinWeightProfilesData.AddOverrideData(Pair.Key);
-		MeshUtilities.GenerateRuntimeSkinWeightData(ImportedModel, Pair.Value.SkinWeights, bUseHighPrecisionWeights, Override);
+		MeshUtilities.GenerateRuntimeSkinWeightData(InLODModel, Pair.Value.SkinWeights, bUseHighPrecisionWeights, Override);
 	}
 
-	ActiveBoneIndices = ImportedModel->ActiveBoneIndices;
-	RequiredBones = ImportedModel->RequiredBones;
+	for (const FSkeletalMeshVertexAttributeInfo& VertexAttributeInfo: InVertexAttributeInfos)
+	{
+		if (!VertexAttributeInfo.IsEnabledForRender())
+		{
+			continue;
+		}
+
+		const FSkeletalMeshModelVertexAttribute* VertexAttribute = InLODModel->VertexAttributes.Find(VertexAttributeInfo.Name);
+		if (ensure(VertexAttribute))
+		{
+			VertexAttributeBuffers.AddAttribute(
+				VertexAttributeInfo.Name,
+				VertexAttributeInfo.DataType,
+				InLODModel->NumVertices,
+				VertexAttribute->ComponentCount,
+				VertexAttribute->Values);
+		}
+	}
+
+	ActiveBoneIndices = InLODModel->ActiveBoneIndices;
+	RequiredBones = InLODModel->RequiredBones;
 }
 
 #endif // WITH_EDITOR
@@ -503,6 +529,7 @@ void FSkeletalMeshLODRenderData::ReleaseCPUResources(bool bForStreaming)
 			ClothVertexBuffer.CleanUp();
 			StaticVertexBuffers.ColorVertexBuffer.CleanUp();
 			SkinWeightProfilesData.ReleaseCPUResources();
+			VertexAttributeBuffers.CleanUp();
 		}
 	}
 }
@@ -524,7 +551,8 @@ void FSkeletalMeshLODRenderData::GetResourceSizeEx(FResourceSizeEx& CumulativeRe
 	CumulativeResourceSize.AddUnknownMemoryBytes(TEXT("SkinWeightVertexBuffer"), SkinWeightVertexBuffer.GetVertexDataSize());
 	CumulativeResourceSize.AddUnknownMemoryBytes(TEXT("ColorVertexBuffer"), StaticVertexBuffers.ColorVertexBuffer.GetAllocatedSize());
 	CumulativeResourceSize.AddUnknownMemoryBytes(TEXT("ClothVertexBuffer"), ClothVertexBuffer.GetVertexDataSize());
-	CumulativeResourceSize.AddUnknownMemoryBytes(TEXT("SkinWeightProfilesData"), SkinWeightProfilesData.GetResourcesSize());	
+	CumulativeResourceSize.AddUnknownMemoryBytes(TEXT("SkinWeightProfilesData"), SkinWeightProfilesData.GetResourcesSize());
+	CumulativeResourceSize.AddUnknownMemoryBytes(TEXT("VertexAttributeData"), VertexAttributeBuffers.GetResourceSize());
 }
 
 SIZE_T FSkeletalMeshLODRenderData::GetCPUAccessMemoryOverhead() const
@@ -796,6 +824,10 @@ void FSkeletalMeshLODRenderData::SerializeStreamedData(FArchive& Ar, USkinnedAss
 			}
 		}
 	}
+	if (Ar.CustomVer(FUE5MainStreamObjectVersion::GUID) >= FUE5MainStreamObjectVersion::SkeletalVertexAttributes)
+	{
+		Ar << VertexAttributeBuffers;
+	}
 }
 
 void FSkeletalMeshLODRenderData::SerializeAvailabilityInfo(FArchive& Ar, int32 LODIdx, bool bAdjacencyDataStripped, bool bNeedsCPUAccess)
diff --git a/Engine/Source/Runtime/Engine/Private/SkeletalMeshRenderData.cpp b/Engine/Source/Runtime/Engine/Private/SkeletalMeshRenderData.cpp
--- a/Engine/Source/Runtime/Engine/Private/SkeletalMeshRenderData.cpp
+++ b/Engine/Source/Runtime/Engine/Private/SkeletalMeshRenderData.cpp
@@ -331,8 +331,8 @@ void FSkeletalMeshRenderData::Cache(const ITargetPlatform* TargetPlatform, USkin
 			{
 				Owner->BuildLODModel(TargetPlatform, LODIndex);
 
-				FSkeletalMeshLODModel* LODModel = &(SkelMeshModel->LODModels[LODIndex]);
-				FSkeletalMeshLODInfo* LODInfo = Owner->GetLODInfo(LODIndex);
+				const FSkeletalMeshLODModel* LODModel = &(SkelMeshModel->LODModels[LODIndex]);
+				const FSkeletalMeshLODInfo* LODInfo = Owner->GetLODInfo(LODIndex);
 				check(LODInfo);
 
 				FSkeletalMeshLODRenderData* LODData = new FSkeletalMeshLODRenderData();
@@ -362,7 +362,7 @@ void FSkeletalMeshRenderData::Cache(const ITargetPlatform* TargetPlatform, USkin
 						VertexBufferBuildFlags |= ESkeletalMeshVertexFlags::UseHighPrecisionWeights;
 					}
 				}
-				LODData->BuildFromLODModel(LODModel, VertexBufferBuildFlags);
+				LODData->BuildFromLODModel(LODModel, LODInfo->VertexAttributes, VertexBufferBuildFlags);
 			}
 
 			FLargeMemoryWriter Ar(0, /*bIsPersistent=*/ true);
diff --git a/Engine/Source/Runtime/Engine/Private/SkinnedAssetCommon.cpp b/Engine/Source/Runtime/Engine/Private/SkinnedAssetCommon.cpp
--- a/Engine/Source/Runtime/Engine/Private/SkinnedAssetCommon.cpp
+++ b/Engine/Source/Runtime/Engine/Private/SkinnedAssetCommon.cpp
@@ -118,6 +118,20 @@ FGuid FSkeletalMeshLODInfo::ComputeDeriveDataCacheKey(const FSkeletalMeshLODGrou
 	SerializeBuildSettingsForDDC(Ar, BuildSettings);
 	SerializeReductionSettingsForDDC(Ar, RealReductionSettings);
 
+	int32 NumRenderableAttributes = 0;
+	for (const FSkeletalMeshVertexAttributeInfo& VertexAttributeInfo: VertexAttributes)
+	{
+		if (VertexAttributeInfo.IsEnabledForRender())
+		{
+			FString AttributeName(VertexAttributeInfo.Name.ToString());
+			Ar << AttributeName;
+			Ar << VertexAttributeInfo.DataType;
+			NumRenderableAttributes++;
+		}
+	}
+	Ar << NumRenderableAttributes;
+	
+
 	FSHA1 Sha;
 	Sha.Update(TempBytes.GetData(), TempBytes.Num() * TempBytes.GetTypeSize());
 	Sha.Final();
diff --git a/Engine/Source/Runtime/Engine/Public/Animation/SkeletalMeshVertexAttribute.h b/Engine/Source/Runtime/Engine/Public/Animation/SkeletalMeshVertexAttribute.h
new file mode 100644
--- /dev/null
+++ b/Engine/Source/Runtime/Engine/Public/Animation/SkeletalMeshVertexAttribute.h
@@ -0,0 +1,102 @@
+﻿// Copyright Epic Games, Inc. All Rights Reserved.
+
+#pragma once
+
+#include "PerPlatformProperties.h"
+
+#include "SkeletalMeshVertexAttribute.generated.h"
+
+
+class FSkeletalMeshAttributeVertexBuffer;
+
+
+UENUM()
+enum class ESkeletalMeshVertexAttributeDataType
+{
+	Float UMETA(DisplayName="Float", ToolTip="Store the vertex attribute values as a 32-bit floating point (float)"),
+	HalfFloat UMETA(DisplayName="Half-float", ToolTip="Store the vertex attribute values as a 16-bit floating point (half)"),
+	NUInt8 UMETA(DisplayName="8-bit Unsigned Normalized", ToolTip="Quantize and store the vertex attribute values as an unsigned normalized 8-bit integer. Values outside the [0.0 - 1.0] range are clamped."),
+	// Commented out until we have PixelFormat support for these types.
+	// NSInt8 UMETA(DisplayName="8-bit Signed Normalized", ToolTip="Quantize and store the vertex attribute values as a signed normalized 8-bit integer. Values outside the [-1.0 - 1.0] range are clamped."),
+	// NUInt16 UMETA(DisplayName="16-bit Unsigned Normalized", ToolTip="Quantize and store the vertex attribute values as an unsigned normalized 16-bit integer. Values outside the [0.0 - 1.0] range are clamped."),
+	// NSInt16 UMETA(DisplayName="16-bit Signed Normalized", ToolTip="Quantize and store the vertex attribute values as a signed normalized 16-bit integer. Values outside the [-1.0 - 1.0] range are clamped.")
+};
+
+
+/** A structure to store user-controllable settings for attributes */  
+USTRUCT()
+struct ENGINE_API FSkeletalMeshVertexAttributeInfo
+{
+	GENERATED_BODY()
+
+	/** Whether this vertex attribute should be included in the render data. Requires a rebuild of the render mesh. */
+	UPROPERTY(EditAnywhere, Category="Vertex Attributes")
+	FPerPlatformBool EnabledForRender = false;
+	
+	/** The name of the vertex attribute */
+	UPROPERTY(VisibleAnywhere, Category="Vertex Attributes")
+	FName Name;
+
+	/** The data type to store the vertex data as for rendering */
+	UPROPERTY(EditAnywhere, Category="Vertex Attributes")
+	ESkeletalMeshVertexAttributeDataType DataType = ESkeletalMeshVertexAttributeDataType::Float;
+
+	/** Returns the name to use for this attribute when */ 
+	FName GetRequirementName() const;
+
+	/** Returns true if this attribute is enabled for rendering on this platform */
+	bool IsEnabledForRender() const;
+};
+
+#if WITH_EDITORONLY_DATA
+
+/** Editor-only representation, stored in FSkeletalMeshLODModel */
+struct FSkeletalMeshModelVertexAttribute
+{
+	/** How the data should be encoded for rendering */
+	ESkeletalMeshVertexAttributeDataType DataType = ESkeletalMeshVertexAttributeDataType::Float;
+	
+	/** The per-vertex value data, where each vertex has ComponentCount number of values.
+	 */
+	TArray<float> Values;
+
+	/** The per vertex component count. Only a range of 1-4 is allowed. */
+	int32 ComponentCount;
+
+	friend FArchive& operator<<(FArchive& Ar, FSkeletalMeshModelVertexAttribute& Attribute);
+};
+
+#endif
+
+
+struct ENGINE_API FSkeletalMeshVertexAttributeRenderData
+{
+	~FSkeletalMeshVertexAttributeRenderData();
+	
+	bool AddAttribute(
+		const FName InName,
+		const ESkeletalMeshVertexAttributeDataType InDataType,
+		const int32 InNumVertices,
+		const int32 InComponentCount,
+		const TArray<float>& InValues
+		);
+
+	FSkeletalMeshAttributeVertexBuffer* GetAttributeBuffer(FName InName) const; 
+
+	void InitResources();
+	void ReleaseResources();
+	void CleanUp();
+	int32 GetResourceSize() const;
+
+	// void SerializeMetaData(FArchive& Ar);
+	void Serialize(FArchive& Ar);
+
+	friend FArchive& operator<<(FArchive& Ar, FSkeletalMeshVertexAttributeRenderData& RenderData)
+	{
+		RenderData.Serialize(Ar);
+		return Ar;
+	}
+	
+private:
+	TMap<FName, FSkeletalMeshAttributeVertexBuffer*> Buffers;
+};
diff --git a/Engine/Source/Runtime/Engine/Public/Rendering/SkeletalMeshAttributeVertexBuffer.h b/Engine/Source/Runtime/Engine/Public/Rendering/SkeletalMeshAttributeVertexBuffer.h
new file mode 100644
--- /dev/null
+++ b/Engine/Source/Runtime/Engine/Public/Rendering/SkeletalMeshAttributeVertexBuffer.h
@@ -0,0 +1,53 @@
+﻿// Copyright Epic Games, Inc. All Rights Reserved.
+
+#pragma once
+
+
+#include "GPUSkinVertexFactory.h"
+#include "SkeletalMeshVertexBuffer.h"
+
+
+class FSkeletalMeshAttributeVertexBuffer :
+	public FRenderResource
+{
+public:
+	virtual ~FSkeletalMeshAttributeVertexBuffer() override;
+	
+	void Init(
+		const EPixelFormat InPixelFormat,
+		const int32 InNumVertices,
+		const int32 InComponentCount,
+		TConstArrayView<float> InValues
+		);
+
+	FRHIShaderResourceView* GetSRV() const
+	{
+		return ValueBuffer.VertexBufferSRV;
+	}
+
+	int32 GetResourceSize() const;
+	void CleanUp();
+
+	virtual void InitRHI() override;
+	virtual void ReleaseRHI() override;
+	virtual FString GetFriendlyName() const override { return TEXT("SkeletalMeshAttributeVertexBuffer"); }
+
+	friend FArchive& operator<<(FArchive& Ar, FSkeletalMeshAttributeVertexBuffer& Data)
+	{
+		Data.Serialize(Ar);
+		return Ar;
+	}
+
+	void Serialize(FArchive& Ar);
+	
+private:
+	void Allocate();
+	void SetData(TConstArrayView<float> InValues);
+	
+	FVertexBufferAndSRV ValueBuffer;
+	FSkeletalMeshVertexDataInterface *ValueData = nullptr;
+
+	int32 ComponentCount = 0;
+	int32 PixelFormat = PF_Unknown;
+	int32 ComponentStride = 0;
+};
diff --git a/Engine/Source/Runtime/Engine/Public/Rendering/SkeletalMeshLODImporterData.h b/Engine/Source/Runtime/Engine/Public/Rendering/SkeletalMeshLODImporterData.h
--- a/Engine/Source/Runtime/Engine/Public/Rendering/SkeletalMeshLODImporterData.h
+++ b/Engine/Source/Runtime/Engine/Public/Rendering/SkeletalMeshLODImporterData.h
@@ -4,27 +4,18 @@
 
 #if WITH_EDITOR
 
-#include "CoreMinimal.h"
 #include "Engine/EngineTypes.h"
 #include "BoneIndices.h"
-#include "SkeletalMeshTypes.h"
 #include "Serialization/BulkData.h"
 #include "Components.h"
 #include "Math/GenericOctree.h"
 #include "Animation/MorphTarget.h"
 #include "Templates/DontCopy.h"
 
-#include "SkeletalMeshLODImporterData.generated.h"
-
 struct FMeshDescription;
 class FSkeletalMeshLODModel;
 
-#endif
-
-//////////////////////////////////////////////////////////////////////////
-//uenum class cannot be inside a preprocessor like #if WITH_EDITOR
 
-UENUM()
 enum class ESkeletalMeshGeoImportVersions : uint8
 {
 	Before_Versionning = 0,
@@ -35,7 +26,6 @@ enum class ESkeletalMeshGeoImportVersions : uint8
 	LatestVersion = VersionPlusOne - 1
 };
 
-UENUM()
 enum class ESkeletalMeshSkinningImportVersions : uint8
 {
 	Before_Versionning = 0,
@@ -46,10 +36,6 @@ enum class ESkeletalMeshSkinningImportVersions : uint8
 	LatestVersion = VersionPlusOne - 1
 };
 
-// End of enum declaration
-//////////////////////////////////////////////////////////////////////////
-
-#if WITH_EDITOR
 
 namespace SkeletalMeshImportData
 {
@@ -328,6 +314,19 @@ namespace SkeletalMeshImportData
 		}
 	};
 
+	struct FVertexAttribute
+	{
+		TArray<float> AttributeValues;
+		int32 ComponentCount;
+		
+		friend FArchive &operator<<(FArchive& Ar, FVertexAttribute& A)
+		{
+			Ar << A.AttributeValues;
+			Ar << A.ComponentCount;
+
+			return Ar;
+		}
+	};
 }
 
 template <> struct TIsPODType<SkeletalMeshImportData::FMeshWedge> { enum { Value = true }; };
@@ -367,6 +366,8 @@ public:
 	TArray<FSkeletalMeshImportData> AlternateInfluences;
 	TArray<FString> AlternateInfluenceProfileNames;
 	
+	TArray<SkeletalMeshImportData::FVertexAttribute> VertexAttributes;
+	TArray<FString> VertexAttributeNames;
 
 	//////////////////////////////////////////////////////////////////////////
 
diff --git a/Engine/Source/Runtime/Engine/Public/Rendering/SkeletalMeshLODModel.h b/Engine/Source/Runtime/Engine/Public/Rendering/SkeletalMeshLODModel.h
--- a/Engine/Source/Runtime/Engine/Public/Rendering/SkeletalMeshLODModel.h
+++ b/Engine/Source/Runtime/Engine/Public/Rendering/SkeletalMeshLODModel.h
@@ -1,6 +1,7 @@
 // Copyright Epic Games, Inc. All Rights Reserved.
 
 #pragma once
+#include "Animation/SkeletalMeshVertexAttribute.h"
 
 #if WITH_EDITOR
 
@@ -381,6 +382,9 @@ public:
 	/** Set of skin weight profile, identified by a FName which matches FSkinWeightProfileInfo.Name in the owning Skeletal Mesh*/
 	TMap<FName, FImportedSkinWeightProfileData> SkinWeightProfiles;
 
+	/** Set of vertex attributes identified by an FName */
+	TMap<FName, FSkeletalMeshModelVertexAttribute> VertexAttributes;
+
 	/** Mapping from final mesh vertex index to raw import vertex index. Needed for vertex animation, which only stores positions for import verts. */
 	TArray<int32>				MeshToImportVertexMap;
 	/** The max index in MeshToImportVertexMap, ie. the number of imported (raw) verts. */
diff --git a/Engine/Source/Runtime/Engine/Public/Rendering/SkeletalMeshLODRenderData.h b/Engine/Source/Runtime/Engine/Public/Rendering/SkeletalMeshLODRenderData.h
--- a/Engine/Source/Runtime/Engine/Public/Rendering/SkeletalMeshLODRenderData.h
+++ b/Engine/Source/Runtime/Engine/Public/Rendering/SkeletalMeshLODRenderData.h
@@ -14,11 +14,13 @@
 #include "BoneIndices.h"
 #include "StaticMeshResources.h"
 #include "GPUSkinVertexFactory.h"
+#include "Animation/SkeletalMeshVertexAttribute.h"
 #include "Animation/SkinWeightProfile.h"
 
 class USkinnedAsset;
 #if WITH_EDITOR
 class FSkeletalMeshLODModel;
+struct FSkeletalMeshVertexAttributeInfo;
 #endif // WITH_EDITOR
 
 struct FSkelMeshRenderSection
@@ -141,6 +143,8 @@ public:
 	/** Skin weight profile data structures, can contain multiple profiles and their runtime FSkinWeightVertexBuffer */
 	FSkinWeightProfilesData SkinWeightProfilesData;
 
+	FSkeletalMeshVertexAttributeRenderData VertexAttributeBuffers;
+
 	TArray<FBoneIndexType> ActiveBoneIndices;
 
 	TArray<FBoneIndexType> RequiredBones;
@@ -229,9 +233,15 @@ public:
 #if WITH_EDITOR
 	/**
 	 * Initialize render data (e.g. vertex buffers) from model info
-	 * @param BuildFlags See ESkeletalMeshVertexFlags.
+	 * @param InLODModel The model to build the render data from.
+	 * @param InVertexAttributeInfos The vertex attributes to possibly include and their stored data type.
+	 * @param InBuildFlags See ESkeletalMeshVertexFlags.
 	 */
-	void ENGINE_API BuildFromLODModel(const FSkeletalMeshLODModel* LODModel, ESkeletalMeshVertexFlags BuildFlags = ESkeletalMeshVertexFlags::None);
+	void ENGINE_API BuildFromLODModel(
+		const FSkeletalMeshLODModel* InLODModel,
+		TConstArrayView<FSkeletalMeshVertexAttributeInfo> InVertexAttributeInfos = {},
+		ESkeletalMeshVertexFlags InBuildFlags = ESkeletalMeshVertexFlags::None
+		);
 #endif // WITH_EDITOR
 
 	uint32 GetNumVertices() const
