fix issue where it'd be possible for a few frames to happen on the GT when tracing against a cluster union to be missing a newly added geometry collection

- one potential issue arose from  FClusterUnionPhysicsProxy::PullFromPhysicsState settings its geometry from the previous pull data instead of the next/current pull data potentially causing some delay.
- modify all trace/sweep/overlap functions on the cluster union component to also consider components that are pending
- modify the GT geometry of cluster unions immediately on the GT when adding or removing components
- add a way to acquire a locked external interface by passing in the physics scene directly

diff --git a/Engine/Source/Runtime/Engine/Classes/PhysicsEngine/ClusterUnionComponent.h b/Engine/Source/Runtime/Engine/Classes/PhysicsEngine/ClusterUnionComponent.h
--- a/Engine/Source/Runtime/Engine/Classes/PhysicsEngine/ClusterUnionComponent.h
+++ b/Engine/Source/Runtime/Engine/Classes/PhysicsEngine/ClusterUnionComponent.h
@@ -34,6 +34,9 @@ struct FClusteredComponentData
 	// Set of physics objects that we actually added into the cluster union.
 	TSet<Chaos::FPhysicsObjectHandle> PhysicsObjects;
 
+	// Set of bone Ids that we actually added into the cluster union.
+	TSet<int32> BoneIds;
+
 	// Every physics object associated with this particular component.
 	TArray<Chaos::FPhysicsObjectHandle> AllPhysicsObjects;
 
@@ -44,6 +47,9 @@ struct FClusteredComponentData
 
 	UPROPERTY()
 	bool bWasReplicating = true;
+
+	UPROPERTY()
+	bool bPendingDeletion = false;
 };
 
 USTRUCT()
@@ -160,6 +166,8 @@ protected:
 	UFUNCTION()
 	void ForceSetChildToParent(UPrimitiveComponent* InComponent, const TArray<int32>& BoneIds, const TArray<FTransform>& ChildToParent);
 
+	TArray<int32> GetAddedBoneIdsForComponent(UPrimitiveComponent* Component) const;
+
 private:
 	// These are the statically clustered components. These should
 	// be specified in the editor and never change.
@@ -178,9 +186,15 @@ private:
 	TMap<TObjectKey<AActor>, FClusteredActorData> ActorToComponents;
 
 	// Sometimes we might be in the process of waiting for a component to create it physics state before adding to the cluster.
-	// Make sure we don't try to add the component multiples times while the add is pending.
+	// Make sure we don't try to add the component multiples times while the add is pending. Gets removed finally when PT syncs
+	// back to the GT with this component.
 	TMap<TObjectKey<UPrimitiveComponent>, FClusterUnionPendingAddData> PendingComponentsToAdd;
 
+	// After we add to a cluster union, we need to wait for the sync from the PT back to the GT before removing the component from PendingComponentSync.
+	// Before that happens, we need to perform operations on the GT assuming that the component was added already otherwise there'll be a few frames
+	// where the component hasn't been added to the cluster union on the GT causing a mismatch in behavior.
+	TMap<TObjectKey<UPrimitiveComponent>, FClusterUnionPendingAddData> PendingComponentSync;
+
 	// Given a unique index of a particle that we're adding to the cluster union - map it back to the component that owns it.
 	// This works decently because we assume that when we're using a cluster union component, we will only try to add to the
 	// cluster union via the GT so we can guarantee to have a decent mapping here.
@@ -215,10 +229,16 @@ private:
 	void HandleAddOrModifiedClusteredComponent(UPrimitiveComponent* ChangedComponent, const TMap<int32, FTransform>& PerBoneChildToParent);
 	void HandleRemovedClusteredComponent(UPrimitiveComponent* ChangedComponent, bool bDestroyReplicatedProxy);
 
+	TArray<UPrimitiveComponent*> GetAllCurrentChildComponents() const;
+
 	// Whether or not this code is running on the server.
 	UFUNCTION()
 	bool IsAuthority() const;
 
+	// Force a rebuild of the GT geometry. This needs to happen immediately when we add/remove on the GT so that the SQ is up to date
+	// and doesn't need to wait for the next OnSyncBodies.
+	void ForceRebuildGTParticleGeometry();
+
 	//~ Begin UActorComponent Interface
 public:
 	virtual void OnCreatePhysicsState() override;
diff --git a/Engine/Source/Runtime/Engine/Classes/PhysicsEngine/PhysicsObjectExternalInterface.h b/Engine/Source/Runtime/Engine/Classes/PhysicsEngine/PhysicsObjectExternalInterface.h
--- a/Engine/Source/Runtime/Engine/Classes/PhysicsEngine/PhysicsObjectExternalInterface.h
+++ b/Engine/Source/Runtime/Engine/Classes/PhysicsEngine/PhysicsObjectExternalInterface.h
@@ -47,8 +47,10 @@ using FLockedWritePhysicsObjectExternalInterface = FLockedPhysicsObjectExternalI
 class ENGINE_API FPhysicsObjectExternalInterface : public Chaos::FPhysicsObjectInterface, public PhysicsObjectPhysicsCoreInterface
 {
 public:
+	static FLockedReadPhysicsObjectExternalInterface LockRead(FChaosScene* Scene);
 	static FLockedReadPhysicsObjectExternalInterface LockRead(Chaos::FConstPhysicsObjectHandle InObject);
 	static FLockedReadPhysicsObjectExternalInterface LockRead(TArrayView<const Chaos::FConstPhysicsObjectHandle> InObjects);
+	static FLockedWritePhysicsObjectExternalInterface LockWrite(FChaosScene* Scene);
 	static FLockedWritePhysicsObjectExternalInterface LockWrite(TArrayView<const Chaos::FPhysicsObjectHandle> InObjects);
 
 	static Chaos::FReadPhysicsObjectInterface_External GetRead_AssumesLocked();
diff --git a/Engine/Source/Runtime/Engine/Private/PhysicsEngine/ClusterUnionComponent.cpp b/Engine/Source/Runtime/Engine/Private/PhysicsEngine/ClusterUnionComponent.cpp
--- a/Engine/Source/Runtime/Engine/Private/PhysicsEngine/ClusterUnionComponent.cpp
+++ b/Engine/Source/Runtime/Engine/Private/PhysicsEngine/ClusterUnionComponent.cpp
@@ -43,14 +43,15 @@ void UClusterUnionComponent::AddComponentToCluster(UPrimitiveComponent* InCompon
 		return;
 	}
 
+	FClusterUnionPendingAddData PendingData;
+	PendingData.BoneIds = BoneIds;
+
 	if (!InComponent->HasValidPhysicsState())
 	{
 		if (!PendingComponentsToAdd.Contains(InComponent))
 		{
 			// Early out - defer adding the component to the cluster until the component has a valid physics state.
-			FClusterUnionPendingAddData Data;
-			Data.BoneIds = BoneIds;
-			PendingComponentsToAdd.Add(InComponent, Data);
+			PendingComponentsToAdd.Add(InComponent, PendingData);
 			InComponent->OnComponentPhysicsStateChanged.AddDynamic(this, &UClusterUnionComponent::HandleComponentPhysicsStateChange);
 		}
 		return;
@@ -105,7 +106,15 @@ void UClusterUnionComponent::AddComponentToCluster(UPrimitiveComponent* InCompon
 		return;
 	}
 
+	PendingData.BoneIds.Empty(Objects.Num());
+	for (Chaos::FPhysicsObjectHandle Object : Objects)
+	{
+		PendingData.BoneIds.Add(Chaos::FPhysicsObjectInterface::GetId(Object));
+	}
+
+	PendingComponentSync.Add(InComponent, PendingData);
 	PhysicsProxy->AddPhysicsObjects_External(Objects);
+	ForceRebuildGTParticleGeometry();
 }
 
 void UClusterUnionComponent::RemoveComponentFromCluster(UPrimitiveComponent* InComponent)
@@ -129,6 +138,7 @@ void UClusterUnionComponent::RemoveComponentFromCluster(UPrimitiveComponent* InC
 		// We need to mark the replicated proxy as pending deletion.
 		// This way anyone who tries to use the replicated proxy component knows that it
 		// doesn't actually denote a meaningful cluster union relationship.
+		ComponentData->bPendingDeletion = true;
 		if (IsAuthority())
 		{
 			if (UClusterUnionReplicatedProxyComponent* Component = ComponentData->ReplicatedProxyComponent.Get())
@@ -150,6 +160,55 @@ void UClusterUnionComponent::RemoveComponentFromCluster(UPrimitiveComponent* InC
 	}
 
 	PhysicsProxy->RemovePhysicsObjects_External(PhysicsObjectsToRemove);
+	ForceRebuildGTParticleGeometry();
+}
+
+void UClusterUnionComponent::ForceRebuildGTParticleGeometry()
+{
+	if (!PhysicsProxy)
+	{
+		return;
+	}
+
+	// This is an assumption that all the physics objects are part of this scene and thus this is the right thing to lock.
+	FLockedReadPhysicsObjectExternalInterface Interface = FPhysicsObjectExternalInterface::LockRead(GetWorld()->GetPhysicsScene());
+
+	const FTransform ClusterWorldTM = GetComponentTransform();
+	TArray<UPrimitiveComponent*> Components = GetAllCurrentChildComponents();
+	// This code is similar-ish to FClusterUnionManager::ForceRecreateClusterUnionSharedGeometry but not enough to make it necessary
+	// to reshare exactly the same code.
+	TArray<TUniquePtr<Chaos::FImplicitObject>> Objects;
+	TArray<Chaos::FPBDRigidParticle*> Particles;
+
+	// Should be a good number to reserve for now since it's a generally safe assumption we'll be working with 1 particle per component added.
+	Objects.Reserve(Components.Num());
+	Particles.Reserve(Components.Num());
+
+	for (UPrimitiveComponent* Child : Components)
+	{
+		check(Child != nullptr);
+		TArray<int32> BoneIds = GetAddedBoneIdsForComponent(Child);
+		TArray<Chaos::FPhysicsObjectHandle> PhysicsObjects;
+		PhysicsObjects.Reserve(BoneIds.Num());
+		for (int32 Id : BoneIds)
+		{
+			PhysicsObjects.Add(Child->GetPhysicsObjectById(Id));
+		}
+
+		for (Chaos::FPBDRigidParticle* Particle : Interface->GetAllRigidParticles(PhysicsObjects))
+		{
+			if (Particle && Particle->Geometry())
+			{
+				const FTransform ChildWorldTM{ Particle->R(), Particle->X() };
+				const FTransform Frame = ChildWorldTM.GetRelativeTransform(ClusterWorldTM);
+				Objects.Add(TUniquePtr<Chaos::FImplicitObject>(Chaos::FClusterUnionManager::CreateTransformGeometryForClusterUnion<Chaos::EThreadContext::External>(Particle, Frame)));
+				Particles.Add(Particle);
+			}
+		}
+	}
+
+	TSharedPtr<Chaos::FImplicitObject, ESPMode::ThreadSafe> NewGeometry = Objects.IsEmpty() ? MakeShared<Chaos::FImplicitObjectUnionClustered>() : MakeShared<Chaos::FImplicitObjectUnion>(MoveTemp(Objects));
+	PhysicsProxy->SetSharedGeometry_External(NewGeometry, Particles);
 }
 
 TArray<UPrimitiveComponent*> UClusterUnionComponent::GetPrimitiveComponents()
@@ -408,6 +467,7 @@ void UClusterUnionComponent::HandleAddOrModifiedClusteredComponent(UPrimitiveCom
 
 	const bool bIsNew = !ComponentToPhysicsObjects.Contains(ChangedComponent);
 	FClusteredComponentData& ComponentData = ComponentToPhysicsObjects.FindOrAdd(ChangedComponent);
+	PendingComponentSync.Remove(ChangedComponent);
 
 	// If this is a *new* component that we're keeping track of then there's additional book-keeping
 	// we need to do to make sure we don't forget what exactly we're tracking. Additionally, we need to
@@ -459,14 +519,13 @@ void UClusterUnionComponent::HandleAddOrModifiedClusteredComponent(UPrimitiveCom
 		}
 	}
 
+	PerBoneChildToParent.GetKeys(ComponentData.BoneIds);
+
 	if (IsAuthority() && ComponentData.ReplicatedProxyComponent.IsValid())
 	{
 		// We really only need to do modifications on the server since that's where we're changing the replicated proxy to broadcast this data change.
-		TSet<int32> BoneIds;
-		PerBoneChildToParent.GetKeys(BoneIds);
-
 		TObjectPtr<UClusterUnionReplicatedProxyComponent> ReplicatedProxy = ComponentData.ReplicatedProxyComponent.Get();
-		ReplicatedProxy->SetParticleBoneIds(BoneIds.Array());
+		ReplicatedProxy->SetParticleBoneIds(ComponentData.BoneIds.Array());
 		for (const TPair<int32, FTransform>& Kvp : PerBoneChildToParent)
 		{
 			ReplicatedProxy->SetParticleChildToParent(Kvp.Key, Kvp.Value);
@@ -610,9 +669,9 @@ bool UClusterUnionComponent::LineTraceComponent(FHitResult& OutHit, const FVecto
 {
 	bool bHasHit = false;
 	OutHit.Distance = TNumericLimits<float>::Max();
-	for (const TPair<TObjectKey<UPrimitiveComponent>, FClusteredComponentData>& Kvp : ComponentToPhysicsObjects)
+	for (UPrimitiveComponent* Component : GetAllCurrentChildComponents())
 	{
-		if (UPrimitiveComponent* Component = Kvp.Key.ResolveObjectPtr())
+		if (Component)
 		{
 			FHitResult ComponentHit;
 			if (Component->LineTraceComponent(ComponentHit, Start, End, Params) && ComponentHit.Distance < OutHit.Distance)
@@ -629,9 +688,9 @@ bool UClusterUnionComponent::SweepComponent(FHitResult& OutHit, const FVector St
 {
 	bool bHasHit = false;
 	OutHit.Distance = TNumericLimits<float>::Max();
-	for (const TPair<TObjectKey<UPrimitiveComponent>, FClusteredComponentData>& Kvp : ComponentToPhysicsObjects)
+	for (UPrimitiveComponent* Component : GetAllCurrentChildComponents())
 	{
-		if (UPrimitiveComponent* Component = Kvp.Key.ResolveObjectPtr())
+		if (Component)
 		{
 			FHitResult ComponentHit;
 			if (Component->SweepComponent(ComponentHit, Start, End, ShapeWorldRotation, CollisionShape, bTraceComplex) && ComponentHit.Distance < OutHit.Distance)
@@ -647,9 +706,9 @@ bool UClusterUnionComponent::SweepComponent(FHitResult& OutHit, const FVector St
 bool UClusterUnionComponent::OverlapComponentWithResult(const FVector& Pos, const FQuat& Rot, const FCollisionShape& CollisionShape, TArray<FOverlapResult>& OutOverlap) const
 {
 	bool bHasOverlap = false;
-	for (const TPair<TObjectKey<UPrimitiveComponent>, FClusteredComponentData>& Kvp : ComponentToPhysicsObjects)
+	for (UPrimitiveComponent* Component : GetAllCurrentChildComponents())
 	{
-		if (UPrimitiveComponent* Component = Kvp.Key.ResolveObjectPtr())
+		if (Component)
 		{
 			TArray<FOverlapResult> SubOverlaps;
 			if (Component->OverlapComponentWithResult(Pos, Rot, CollisionShape, SubOverlaps))
@@ -680,9 +739,9 @@ bool UClusterUnionComponent::ComponentOverlapComponentWithResultImpl(const class
 		IgnoredComponents.Add(Id);
 	}
 
-	for (const TPair<TObjectKey<UPrimitiveComponent>, FClusteredComponentData>& Kvp : ComponentToPhysicsObjects)
+	for (UPrimitiveComponent* Component : GetAllCurrentChildComponents())
 	{
-		if (UPrimitiveComponent* Component = Kvp.Key.ResolveObjectPtr())
+		if (Component)
 		{
 			if (IgnoredComponents.Contains(Component->GetUniqueID()))
 			{
@@ -708,6 +767,45 @@ bool UClusterUnionComponent::ComponentOverlapComponentWithResultImpl(const class
 	return bHasOverlap;
 }
 
+TArray<UPrimitiveComponent*> UClusterUnionComponent::GetAllCurrentChildComponents() const
+{
+	TArray<UPrimitiveComponent*> Components;
+	Components.Reserve(ComponentToPhysicsObjects.Num() + PendingComponentSync.Num());
+
+	for (const TPair<TObjectKey<UPrimitiveComponent>, FClusterUnionPendingAddData>& Kvp : PendingComponentSync)
+	{
+		if (UPrimitiveComponent* Component = Kvp.Key.ResolveObjectPtr(); Component && Component->HasValidPhysicsState())
+		{
+			Components.Add(Component);
+		}
+	}
+
+	for (const TPair<TObjectKey<UPrimitiveComponent>, FClusteredComponentData>& Kvp : ComponentToPhysicsObjects)
+	{
+		if (UPrimitiveComponent* Component = Kvp.Key.ResolveObjectPtr(); Component && Component->HasValidPhysicsState() && !Kvp.Value.bPendingDeletion)
+		{
+			Components.Add(Component);
+		}
+	}
+
+	return Components;
+}
+
+TArray<int32> UClusterUnionComponent::GetAddedBoneIdsForComponent(UPrimitiveComponent* Component) const
+{
+	if (const FClusteredComponentData* Data = ComponentToPhysicsObjects.Find(Component))
+	{
+		return Data->BoneIds.Array();
+	}
+
+	if (const FClusterUnionPendingAddData* Data = PendingComponentSync.Find(Component))
+	{
+		return Data->BoneIds;
+	}
+
+	return {};
+}
+
 void UClusterUnionComponent::AddReferencedObjects(UObject* InThis, FReferenceCollector& Collector)
 {
 	UPrimitiveComponent::AddReferencedObjects(InThis, Collector);
@@ -737,4 +835,12 @@ void UClusterUnionComponent::AddReferencedObjects(UObject* InThis, FReferenceCol
 			Collector.AddReferencedObjects(ScriptStruct, reinterpret_cast<void*>(&Kvp.Value), This, nullptr);
 		}
 	}
+
+	{
+		const UScriptStruct* ScriptStruct = FClusterUnionPendingAddData::StaticStruct();
+		for (TPair<TObjectKey<UPrimitiveComponent>, FClusterUnionPendingAddData>& Kvp : This->PendingComponentSync)
+		{
+			Collector.AddReferencedObjects(ScriptStruct, reinterpret_cast<void*>(&Kvp.Value), This, nullptr);
+		}
+	}
 }
\ No newline at end of file
diff --git a/Engine/Source/Runtime/Engine/Private/PhysicsEngine/PhysicsObjectExternalInterface.cpp b/Engine/Source/Runtime/Engine/Private/PhysicsEngine/PhysicsObjectExternalInterface.cpp
--- a/Engine/Source/Runtime/Engine/Private/PhysicsEngine/PhysicsObjectExternalInterface.cpp
+++ b/Engine/Source/Runtime/Engine/Private/PhysicsEngine/PhysicsObjectExternalInterface.cpp
@@ -1,6 +1,11 @@
 // Copyright Epic Games, Inc. All Rights Reserved.
 #include "PhysicsEngine/PhysicsObjectExternalInterface.h"
 
+FLockedReadPhysicsObjectExternalInterface FPhysicsObjectExternalInterface::LockRead(FChaosScene* Scene)
+{
+	return FLockedReadPhysicsObjectExternalInterface{ Scene, CreateReadInterface<Chaos::EThreadContext::External>() };
+}
+
 FLockedReadPhysicsObjectExternalInterface FPhysicsObjectExternalInterface::LockRead(Chaos::FConstPhysicsObjectHandle InObject)
 {
 	return LockRead(TArray<Chaos::FConstPhysicsObjectHandle>{ InObject });
@@ -8,12 +13,17 @@ FLockedReadPhysicsObjectExternalInterface FPhysicsObjectExternalInterface::LockR
 
 FLockedReadPhysicsObjectExternalInterface FPhysicsObjectExternalInterface::LockRead(TArrayView<const Chaos::FConstPhysicsObjectHandle> InObjects)
 {
-	return FLockedReadPhysicsObjectExternalInterface{ GetScene(InObjects), CreateReadInterface<Chaos::EThreadContext::External>() };
+	return LockRead(GetScene(InObjects));
+}
+
+FLockedWritePhysicsObjectExternalInterface FPhysicsObjectExternalInterface::LockWrite(FChaosScene* Scene)
+{
+	return FLockedWritePhysicsObjectExternalInterface{ Scene, CreateWriteInterface<Chaos::EThreadContext::External>() };
 }
 
 FLockedWritePhysicsObjectExternalInterface FPhysicsObjectExternalInterface::LockWrite(TArrayView<const Chaos::FPhysicsObjectHandle> InObjects)
 {
-	return FLockedWritePhysicsObjectExternalInterface{ GetScene(InObjects), CreateWriteInterface<Chaos::EThreadContext::External>() };
+	return LockWrite(GetScene(InObjects));
 }
 
 Chaos::FReadPhysicsObjectInterface_External FPhysicsObjectExternalInterface::GetRead_AssumesLocked()
diff --git a/Engine/Source/Runtime/Experimental/Chaos/Private/Chaos/ClusterUnionManager.cpp b/Engine/Source/Runtime/Experimental/Chaos/Private/Chaos/ClusterUnionManager.cpp
--- a/Engine/Source/Runtime/Experimental/Chaos/Private/Chaos/ClusterUnionManager.cpp
+++ b/Engine/Source/Runtime/Experimental/Chaos/Private/Chaos/ClusterUnionManager.cpp
@@ -33,7 +33,6 @@ namespace Chaos
 		// Some parameters aren't relevant after creation.
 		NewUnion.ClusterUnionParameters.UniqueIndex = nullptr;
 
-
 		if (ensure(NewUnion.InternalCluster != nullptr))
 		{
 			NewUnion.InternalCluster->SetInternalCluster(true);
@@ -94,7 +93,6 @@ namespace Chaos
 		}
 
 		// TODO: Can we do something better than a union?
-
 		const FRigidTransform3 ClusterWorldTM(Union.InternalCluster->X(), Union.InternalCluster->R());
 		TArray<TUniquePtr<FImplicitObject>> Objects;
 		Objects.Reserve(Union.ChildParticles.Num());
@@ -115,9 +113,7 @@ namespace Chaos
 
 			if (Child->Geometry())
 			{
-				FImplicitObjectTransformed* TransformedChildGeometry = new TImplicitObjectTransformed<FReal, 3>(Child->Geometry(), Frame);
-				TransformedChildGeometry->SetSharedObject(Child->SharedGeometry());
-				Objects.Add(TUniquePtr<FImplicitObject>(TransformedChildGeometry));
+				Objects.Add(TUniquePtr<FImplicitObject>(CreateTransformGeometryForClusterUnion<EThreadContext::Internal>(Child, Frame)));
 			}
 		}
 
diff --git a/Engine/Source/Runtime/Experimental/Chaos/Private/PhysicsProxy/ClusterUnionPhysicsProxy.cpp b/Engine/Source/Runtime/Experimental/Chaos/Private/PhysicsProxy/ClusterUnionPhysicsProxy.cpp
--- a/Engine/Source/Runtime/Experimental/Chaos/Private/PhysicsProxy/ClusterUnionPhysicsProxy.cpp
+++ b/Engine/Source/Runtime/Experimental/Chaos/Private/PhysicsProxy/ClusterUnionPhysicsProxy.cpp
@@ -245,6 +245,57 @@ namespace Chaos
 		);
 	}
 
+	void FClusterUnionPhysicsProxy::SetSharedGeometry_External(const TSharedPtr<Chaos::FImplicitObject, ESPMode::ThreadSafe>& Geometry, const TArray<FPBDRigidParticle*>& ShapeParticles)
+	{
+		if (!ensure(Particle_External))
+		{
+			return;
+		}
+
+		// In the cases where this is necessary, the SQ should have a valid state - it's just the geometry itself that isn't valid.
+		Particle_External->SetGeometry(Geometry);
+
+		// Need to fill in query/sim data because the input geometry will not have it set properly.
+		// TODO: This duplicates some code in Chaos::FClusterUnionManager and has the same assumptions.
+		if (ShapeParticles.Num() == Particle_External->ShapesArray().Num())
+		{
+			int32 Index = 0;
+			for (const TUniquePtr<FPerShapeData>& ShapeData : Particle_External->ShapesArray())
+			{
+				if (Index >= ShapeParticles.Num() || ShapeParticles[Index]->ShapesArray().IsEmpty())
+				{
+					++Index;
+					continue;
+				}
+
+				const TUniquePtr<Chaos::FPerShapeData>& TemplateShape = ShapeParticles[Index]->ShapesArray()[0];
+				if (ShapeData && TemplateShape)
+				{
+					{
+						FCollisionData Data = TemplateShape->GetCollisionData();
+						Data.UserData = nullptr;
+						ShapeData->SetCollisionData(Data);
+					}
+
+					{
+						FCollisionFilterData Data = TemplateShape->GetQueryData();
+						Data.Word0 = InitData.ActorId;
+						ShapeData->SetQueryData(Data);
+					}
+
+					{
+						FCollisionFilterData Data = TemplateShape->GetSimData();
+						Data.Word0 = 0;
+						Data.Word2 = InitData.ComponentId;
+						ShapeData->SetSimData(Data);
+					}
+				}
+
+				++Index;
+			}
+		}
+	}
+
 	DECLARE_CYCLE_STAT(TEXT("FClusterUnionPhysicsProxy::PushToPhysicsState"), STAT_ClusterUnionPhysicsProxyPushToPhysicsState, STATGROUP_Chaos);
 	void FClusterUnionPhysicsProxy::PushToPhysicsState(const FDirtyPropertiesManager& Manager, int32 DataIdx, const FDirtyProxy& Dirty)
 	{
@@ -324,16 +375,16 @@ namespace Chaos
 			return false;
 		}
 
-		const FClusterUnionProxyTimestamp* ProxyTimestamp = PullData.GetTimestamp();
+		const FDirtyClusterUnionData& CurrentPullData = NextPullData ? *NextPullData : PullData;
+		const FClusterUnionProxyTimestamp* ProxyTimestamp = CurrentPullData.GetTimestamp();
 		if (!ProxyTimestamp)
 		{
 			return false;
 		}
 
-		
-		SyncedData_External.bIsAnchored = PullData.bIsAnchored;
-		SyncedData_External.ChildParticles.Empty(PullData.ChildParticles.Num());
-		for (const FDirtyClusterUnionParticleData& InData : PullData.ChildParticles)
+		SyncedData_External.bIsAnchored = CurrentPullData.bIsAnchored;
+		SyncedData_External.ChildParticles.Empty(CurrentPullData.ChildParticles.Num());
+		for (const FDirtyClusterUnionParticleData& InData : CurrentPullData.ChildParticles)
 		{
 			FClusterUnionChildData ConvertedData;
 			ConvertedData.ParticleIdx = InData.ParticleIdx;
@@ -341,27 +392,27 @@ namespace Chaos
 			SyncedData_External.ChildParticles.Add(ConvertedData);
 		}
 
-		Particle_External->SetGeometry(PullData.SharedGeometry);
-		Particle_External->SetObjectState(PullData.ObjectState, true, /*bInvalidate=*/false);
+		Particle_External->SetGeometry(CurrentPullData.SharedGeometry);
+		Particle_External->SetObjectState(CurrentPullData.ObjectState, true, /*bInvalidate=*/false);
 
 		const FShapesArray& ShapeArray = Particle_External->ShapesArray();
 		for (int32 ShapeIndex = 0; ShapeIndex < ShapeArray.Num(); ++ShapeIndex)
 		{
 			if (const TUniquePtr<Chaos::FPerShapeData>& ShapeData = ShapeArray[ShapeIndex])
 			{
-				if (PullData.CollisionData.IsValidIndex(ShapeIndex))
+				if (CurrentPullData.CollisionData.IsValidIndex(ShapeIndex))
 				{
-					ShapeData->SetCollisionData(PullData.CollisionData[ShapeIndex]);
+					ShapeData->SetCollisionData(CurrentPullData.CollisionData[ShapeIndex]);
 				}
 
-				if (PullData.QueryData.IsValidIndex(ShapeIndex))
+				if (CurrentPullData.QueryData.IsValidIndex(ShapeIndex))
 				{
-					ShapeData->SetQueryData(PullData.QueryData[ShapeIndex]);
+					ShapeData->SetQueryData(CurrentPullData.QueryData[ShapeIndex]);
 				}
 
-				if (PullData.SimData.IsValidIndex(ShapeIndex))
+				if (CurrentPullData.SimData.IsValidIndex(ShapeIndex))
 				{
-					ShapeData->SetSimData(PullData.SimData[ShapeIndex]);
+					ShapeData->SetSimData(CurrentPullData.SimData[ShapeIndex]);
 				}
 			}
 		}
diff --git a/Engine/Source/Runtime/Experimental/Chaos/Public/Chaos/ClusterUnionManager.h b/Engine/Source/Runtime/Experimental/Chaos/Public/Chaos/ClusterUnionManager.h
--- a/Engine/Source/Runtime/Experimental/Chaos/Public/Chaos/ClusterUnionManager.h
+++ b/Engine/Source/Runtime/Experimental/Chaos/Public/Chaos/ClusterUnionManager.h
@@ -78,6 +78,22 @@ namespace Chaos
 	class CHAOS_API FClusterUnionManager
 	{
 	public:
+
+		template<EThreadContext Id>
+		static FImplicitObjectTransformed* CreateTransformGeometryForClusterUnion(TThreadRigidParticle<Id>* Child, const FTransform& Frame)
+		{
+			FImplicitObjectTransformed* TransformedChildGeometry = new TImplicitObjectTransformed<FReal, 3>(Child->Geometry(), Frame);
+			if constexpr (Id == EThreadContext::External)
+			{
+				TransformedChildGeometry->SetSharedObject(Child->GeometrySharedLowLevel());
+			}
+			else
+			{
+				TransformedChildGeometry->SetSharedObject(Child->SharedGeometry());
+			}
+			return TransformedChildGeometry;
+		}
+
 		FClusterUnionManager(FRigidClustering& InClustering, FPBDRigidsEvolutionGBF& InEvolution);
 
 		// Creates a new cluster union with an automatically assigned cluster union index.
diff --git a/Engine/Source/Runtime/Experimental/Chaos/Public/PhysicsProxy/ClusterUnionPhysicsProxy.h b/Engine/Source/Runtime/Experimental/Chaos/Public/PhysicsProxy/ClusterUnionPhysicsProxy.h
--- a/Engine/Source/Runtime/Experimental/Chaos/Public/PhysicsProxy/ClusterUnionPhysicsProxy.h
+++ b/Engine/Source/Runtime/Experimental/Chaos/Public/PhysicsProxy/ClusterUnionPhysicsProxy.h
@@ -69,6 +69,9 @@ namespace Chaos
 		EObjectStateType GetObjectState_External() const;
 		void SetObjectState_External(EObjectStateType State);
 
+		// Set GT geometry - this is only for smoothing over any changes until the PT syncs back to the GT.
+		void SetSharedGeometry_External(const TSharedPtr<Chaos::FImplicitObject, ESPMode::ThreadSafe>& Geometry, const TArray<FPBDRigidParticle*>& ShapeParticles);
+
 		// Cluster union proxy initialization happens in two first on the game thread (external) then on the
 		// physics thread (internal). Cluster unions are a primarily physics concept so the things exposed to
 		// an external context is just there to be able to safely query the state on the physics thread and to
