Undo changelist 25084812 ( causing stack overflow )

diff --git a/Engine/Source/Runtime/Engine/Private/Collision/CollisionConversions.cpp b/Engine/Source/Runtime/Engine/Private/Collision/CollisionConversions.cpp
--- a/Engine/Source/Runtime/Engine/Private/Collision/CollisionConversions.cpp
+++ b/Engine/Source/Runtime/Engine/Private/Collision/CollisionConversions.cpp
@@ -610,14 +610,7 @@ void ConvertQueryOverlap(const FPhysicsShape& Shape, const FPhysicsActor& Actor,
         BodyInst = FPhysicsInterface::ShapeToOriginalBodyInstance(BodyInst, &Shape);
 		if (const UPrimitiveComponent* OwnerComponent = BodyInst->OwnerComponent.Get())
 		{
-			if (ALightWeightInstanceManager* LWIManager = Cast<ALightWeightInstanceManager>(OwnerComponent->GetOwner()))
-			{
-				OutOverlap.OverlapObjectHandle = FActorInstanceHandle(LWIManager, BodyInst->InstanceBodyIndex);
-			}
-			else
-			{
-				OutOverlap.OverlapObjectHandle = FActorInstanceHandle(OwnerComponent->GetOwner());
-			}
+			OutOverlap.OverlapObjectHandle = FActorInstanceHandle(OwnerComponent->GetOwner());
 			OutOverlap.Component = BodyInst->OwnerComponent; // Copying weak pointer is faster than assigning raw pointer.
 			OutOverlap.ItemIndex = OwnerComponent->bMultiBodyOverlap ? BodyInst->InstanceBodyIndex : INDEX_NONE;
 		}
