Substrate - fixed crash when assigning material function onto new material reopened later.

diff --git a/Engine/Source/Runtime/Engine/Private/Materials/Material.cpp b/Engine/Source/Runtime/Engine/Private/Materials/Material.cpp
--- a/Engine/Source/Runtime/Engine/Private/Materials/Material.cpp
+++ b/Engine/Source/Runtime/Engine/Private/Materials/Material.cpp
@@ -3337,11 +3337,23 @@ void UMaterial::ConvertMaterialToStrataMaterial()
 		UMaterialFunction* DefaultMF = LoadObject<UMaterialFunction>(nullptr, TEXT("/Engine/Functions/Strata/SMF_UE4Disney.SMF_UE4Disney"));
 		if (DefaultMF)
 		{
+			DefaultMF->UpdateFromFunctionResource();
+			DefaultMF->PostEditChange();
+			DefaultMF->PostLoad();
+
 			UMaterialExpressionMaterialFunctionCall* MFCallNode = NewObject<UMaterialExpressionMaterialFunctionCall>(this);
-			MFCallNode->SetMaterialFunction(DefaultMF);
+			if (MFCallNode->SetMaterialFunction(DefaultMF))
+			{
+				// This is needed for input/output expressions to be set correctly, otherwise compilation will fail.
+				GetExpressionCollection().AddExpression(MFCallNode);
 
-			SetPosXAndMoveReferenceToTheRight(MFCallNode);
-			EditorOnly->FrontMaterial.Connect(0, MFCallNode);
+				SetPosXAndMoveReferenceToTheRight(MFCallNode);
+				EditorOnly->FrontMaterial.Connect(0, MFCallNode);
+
+				MFCallNode->UpdateFromFunctionResource();
+				MFCallNode->PostEditChange();
+				MFCallNode->PostLoad();
+			}
 		}
 		else
 		{
