Substrate - using conidtional post load.

fyi charles.derousiers

diff --git a/Engine/Source/Runtime/Engine/Private/Materials/Material.cpp b/Engine/Source/Runtime/Engine/Private/Materials/Material.cpp
--- a/Engine/Source/Runtime/Engine/Private/Materials/Material.cpp
+++ b/Engine/Source/Runtime/Engine/Private/Materials/Material.cpp
@@ -3339,7 +3339,7 @@ void UMaterial::ConvertMaterialToStrataMaterial()
 		{
 			DefaultMF->UpdateFromFunctionResource();
 			DefaultMF->PostEditChange();
-			DefaultMF->PostLoad();
+			DefaultMF->ConditionalPostLoad();
 
 			UMaterialExpressionMaterialFunctionCall* MFCallNode = NewObject<UMaterialExpressionMaterialFunctionCall>(this);
 			if (MFCallNode->SetMaterialFunction(DefaultMF))
@@ -3352,7 +3352,7 @@ void UMaterial::ConvertMaterialToStrataMaterial()
 
 				MFCallNode->UpdateFromFunctionResource();
 				MFCallNode->PostEditChange();
-				MFCallNode->PostLoad();
+				MFCallNode->ConditionalPostLoad();
 			}
 		}
 		else
