WorldPartition: Fix Cell DebugName

#rnx

diff --git a/Engine/Source/Runtime/Engine/Private/WorldPartition/WorldPartitionRuntimeSpatialHash.cpp b/Engine/Source/Runtime/Engine/Private/WorldPartition/WorldPartitionRuntimeSpatialHash.cpp
--- a/Engine/Source/Runtime/Engine/Private/WorldPartition/WorldPartitionRuntimeSpatialHash.cpp
+++ b/Engine/Source/Runtime/Engine/Private/WorldPartition/WorldPartitionRuntimeSpatialHash.cpp
@@ -1420,7 +1420,7 @@ bool UWorldPartitionRuntimeSpatialHash::CreateStreamingGrid(const FSpatialHashRu
 					CellDataSpatialHash->Position = FVector(Bounds.GetCenter(), 0.f);
 					CellDataSpatialHash->Extent = (float)CellExtent;
 					CellDataSpatialHash->GridName = RuntimeGrid.GridName;			
-					CellDataSpatialHash->DebugName = StreamingCell->GetName();
+					CellDataSpatialHash->DebugName = CellName + WorldInstanceSuffix;
 
 					PopulateRuntimeCell(StreamingCell, FilteredActors, OutPackagesToGenerate);
 
