Upgrade error in skeletal meshes
#fix change load error to info

diff --git a/Engine/Source/Runtime/Engine/Private/Animation/Skeleton.cpp b/Engine/Source/Runtime/Engine/Private/Animation/Skeleton.cpp
--- a/Engine/Source/Runtime/Engine/Private/Animation/Skeleton.cpp
+++ b/Engine/Source/Runtime/Engine/Private/Animation/Skeleton.cpp
@@ -758,7 +758,7 @@ int32 USkeleton::BuildLinkup(const USkinnedAsset* InSkinnedAsset)
 			static FName NAME_LoadErrors("LoadErrors");
 			FMessageLog LoadErrors(NAME_LoadErrors);
 
-			TSharedRef<FTokenizedMessage> Message = LoadErrors.Error();
+			TSharedRef<FTokenizedMessage> Message = LoadErrors.Info();
 			Message->AddToken(FTextToken::Create(LOCTEXT("SkeletonBuildLinkupMissingBones1", "The Skeleton ")));
 			Message->AddToken(FAssetNameToken::Create(GetPathName(), FText::FromString( GetNameSafe(this) ) ));
 			Message->AddToken(FTextToken::Create(LOCTEXT("SkeletonBuildLinkupMissingBones2", " is missing bones that SkeletalMesh ")));
