Additional debug information for cooking anim sequence data soft-lock

diff --git a/Engine/Source/Runtime/Engine/Private/Animation/AnimCompressionTypes.cpp b/Engine/Source/Runtime/Engine/Private/Animation/AnimCompressionTypes.cpp
--- a/Engine/Source/Runtime/Engine/Private/Animation/AnimCompressionTypes.cpp
+++ b/Engine/Source/Runtime/Engine/Private/Animation/AnimCompressionTypes.cpp
@@ -1477,7 +1477,7 @@ void FCompressedAnimSequence::ClearCompressedCurveData()
 	CurveCompressionCodec = nullptr;
 }
 
-bool FCompressedAnimSequence::IsValid(const UAnimSequence* AnimSequence) const
+bool FCompressedAnimSequence::IsValid(const UAnimSequence* AnimSequence, bool bLogInformation/*=false*/) const
 {
 #if WITH_EDITOR
 	const bool bHasBoneTracks = AnimSequence->GetDataModelInterface()->GetNumBoneTracks() != 0;
@@ -1487,6 +1487,11 @@ bool FCompressedAnimSequence::IsValid(const UAnimSequence* AnimSequence) const
 	const bool bHasValidCompressedCurveData = CompressedCurveByteStream.Num() != 0;
 	const bool bValidCompressedData = (!bHasBoneTracks || bHasValidCompressedBoneData) && (!bHasCurveData || bHasValidCompressedCurveData || AnimSequence->IsValidAdditive());
 
+	if(bLogInformation)
+	{
+		UE_LOG(LogAnimation, Warning, TEXT("Num Bone Tracks: %i\nNum Curves: %i\nCompressed bone data: %i (%i)\nCompressed curve data %i\nValid additive: %i"), AnimSequence->GetDataModelInterface()->GetNumBoneTracks(), AnimSequence->GetDataModelInterface()->GetNumberOfFloatCurves(), CompressedDataStructure != nullptr ? 1 : 0, CompressedTrackToSkeletonMapTable.Num(), CompressedCurveByteStream.Num() != 0 ? 1 : 0, AnimSequence->IsValidAdditive() ? 1 : 0);
+	}
+
 	return bValidCompressedData;
 #else
 	return true;
diff --git a/Engine/Source/Runtime/Engine/Private/Animation/AnimSequence.cpp b/Engine/Source/Runtime/Engine/Private/Animation/AnimSequence.cpp
--- a/Engine/Source/Runtime/Engine/Private/Animation/AnimSequence.cpp
+++ b/Engine/Source/Runtime/Engine/Private/Animation/AnimSequence.cpp
@@ -1090,15 +1090,41 @@ bool UAnimSequence::IsCachedCookedPlatformDataLoaded(const ITargetPlatform* Targ
 		EndCacheDerivedData(KeyHash);
 	}
 
+	auto LogCookedDataWarning = [this, TargetPlatform, KeyHash]()
+	{
+		// This means cooked data is expected, but it it currently invalid and no task is running to generate it
+		UE_LOG(LogAnimation, Warning, TEXT("Expecting either valid compressed Animation Data, or an in-flight compression task for asset %s\nPlatform: %s\nHash: %s"), *GetPathName(), *TargetPlatform->DisplayName().ToString(), *LexToString(KeyHash));
+	};
+
 	if (KeyHash == DataKeyHash)
 	{
-		return CompressedData.IsValid(this) && !CacheTasksByKeyHash.Contains(KeyHash);
+		if (CompressedData.IsValid(this) && !CacheTasksByKeyHash.Contains(KeyHash))
+		{
+			return true;
+		}
+		else if (!CacheTasksByKeyHash.Contains(KeyHash))
+		{
+			LogCookedDataWarning();
+			CompressedData.IsValid(this, true);
+		}
 	}
 	else
 	{
 		if(const TUniquePtr<FCompressedAnimSequence>* CompressedDataPtr = DataByPlatformKeyHash.Find(KeyHash))
 		{
-			return (*CompressedDataPtr)->IsValid(this) && !CacheTasksByKeyHash.Contains(KeyHash);
+			if ((*CompressedDataPtr)->IsValid(this) && !CacheTasksByKeyHash.Contains(KeyHash))
+			{
+				return true;
+			}
+			else if (!CacheTasksByKeyHash.Contains(KeyHash))
+			{
+				LogCookedDataWarning();
+				(*CompressedDataPtr)->IsValid(this, true);
+			}
+		}
+		else if (!CacheTasksByKeyHash.Contains(KeyHash))
+		{
+			LogCookedDataWarning();
 		}
 	}
 
diff --git a/Engine/Source/Runtime/Engine/Public/Animation/AnimCompressionTypes.h b/Engine/Source/Runtime/Engine/Public/Animation/AnimCompressionTypes.h
--- a/Engine/Source/Runtime/Engine/Public/Animation/AnimCompressionTypes.h
+++ b/Engine/Source/Runtime/Engine/Public/Animation/AnimCompressionTypes.h
@@ -843,7 +843,7 @@ public:
 	void ClearCompressedBoneData();
 	void ClearCompressedCurveData();
 
-	bool IsValid(const UAnimSequence* AnimSequence) const;
+	bool IsValid(const UAnimSequence* AnimSequence, bool bLogInformation=false) const;
 };
 
 struct FRootMotionReset
