Allow VTs to be built with LeaveExisting if all blocks have complete mip chains.

diff --git a/Engine/Source/Runtime/Engine/Private/Texture.cpp b/Engine/Source/Runtime/Engine/Private/Texture.cpp
--- a/Engine/Source/Runtime/Engine/Private/Texture.cpp
+++ b/Engine/Source/Runtime/Engine/Private/Texture.cpp
@@ -594,9 +594,9 @@ void UTexture::ValidateSettingsAfterImportOrEdit(bool * pRequiresNotifyMaterials
 				VirtualTextureStreaming = false;
 				bRequiresNotifyMaterials = true;
 			}
-			if (MipGenSettings == TMGS_LeaveExistingMips || MipGenSettings == TMGS_NoMipmaps)
+			if (MipGenSettings == TMGS_NoMipmaps)
 			{
-				UE_LOG(LogTexture, Warning, TEXT("Virtual textures require mips and MipGenSettings is NoMipmaps or LeaveExistingMips: Forcing to SimpleAverage (%s)"), *GetName());
+				UE_LOG(LogTexture, Warning, TEXT("Virtual textures require mips and MipGenSettings is NoMipmaps: Forcing to SimpleAverage (%s)"), *GetName());
 				MipGenSettings = TMGS_SimpleAverage;
 			}
 		}
diff --git a/Engine/Source/Runtime/Engine/Private/TextureDerivedData.cpp b/Engine/Source/Runtime/Engine/Private/TextureDerivedData.cpp
--- a/Engine/Source/Runtime/Engine/Private/TextureDerivedData.cpp
+++ b/Engine/Source/Runtime/Engine/Private/TextureDerivedData.cpp
@@ -31,6 +31,7 @@
 #include "DerivedDataCache.h"
 #include "DerivedDataCacheInterface.h"
 #include "DerivedDataRequestOwner.h"
+#include "ImageCoreUtils.h"
 #include "Interfaces/ITargetPlatform.h"
 #include "Interfaces/ITargetPlatformManagerModule.h"
 #include "Interfaces/ITextureFormat.h"
@@ -899,12 +900,30 @@ static void GetTextureBuildSettings(
 
 	// Virtual textures must have mips as VT memory management relies on a 1:1 texel/pixel mapping, which in turn
 	// requires that we be able to swap in lower mips when that density gets too high for a given texture.
-	if (bVirtualTextureStreaming &&
-		(MipGenSettings == TMGS_NoMipmaps || MipGenSettings == TMGS_LeaveExistingMips)
-		)
+	if (bVirtualTextureStreaming && MipGenSettings == TMGS_NoMipmaps)
 	{
 		MipGenSettings = TMGS_SimpleAverage;
-		UE_LOG(LogTexture, Warning, TEXT("Texture %s is virtual and has NoMips / LeaveExistingMips - forcing to SimpleAverage."), *Texture.GetPathName());
+		UE_LOG(LogTexture, Warning, TEXT("Texture %s is virtual and has NoMips - forcing to SimpleAverage."), *Texture.GetPathName());
+	}
+	if (bVirtualTextureStreaming && MipGenSettings == TMGS_LeaveExistingMips)
+	{
+		for (int32 BlockIndex = 0; BlockIndex < Texture.Source.GetNumBlocks(); BlockIndex++)
+		{
+			FTextureSourceBlock Block;
+			Texture.Source.GetBlock(BlockIndex, Block);
+
+			int32 ExpectedNumMips = FImageCoreUtils::GetMipCountFromDimensions(Block.SizeX, Block.SizeY, 0, false);
+			if (Block.NumMips != ExpectedNumMips)
+			{
+				MipGenSettings = TMGS_SimpleAverage;
+				UE_LOG(LogTexture, Warning, TEXT("Texture %s is virtual and has LeaveExistingMips with an incomplete mip chain - forcing to SimpleAverage (Block %d has %d mips, expected %d)."), 
+					*Texture.GetPathName(),
+					BlockIndex,
+					Block.NumMips,
+					ExpectedNumMips
+					);
+			}
+		}
 	}
 
 	const FIntPoint SourceSize = Texture.Source.GetLogicalSize();
