Fix truncation warnings in FInstancedPlacementHash

diff --git a/Engine/Source/Runtime/Engine/Public/Instances/InstancedPlacementHash.h b/Engine/Source/Runtime/Engine/Public/Instances/InstancedPlacementHash.h
--- a/Engine/Source/Runtime/Engine/Public/Instances/InstancedPlacementHash.h
+++ b/Engine/Source/Runtime/Engine/Public/Instances/InstancedPlacementHash.h
@@ -32,7 +32,7 @@ private:
 			// compute distance to the cell cube
 			int32 CellSize = (1 << LocalHashCellBits);
 			int32 HalfCellSize = CellSize >> 1;
-			FVector CellCenter(X * CellSize + HalfCellSize, Y * CellSize + HalfCellSize, Z * CellSize + HalfCellSize);
+			FVector CellCenter(static_cast<FVector::FReal>(X * CellSize + HalfCellSize), static_cast<FVector::FReal>(Y * CellSize + HalfCellSize), static_cast<FVector::FReal>(Z * CellSize + HalfCellSize));
 			FVector AbsRelativePosition = (Position - CellCenter).GetAbs();			
 			FVector CubeDeltaVector = (AbsRelativePosition - HalfCellSize).ComponentMax(FVector::ZeroVector);
 			return CubeDeltaVector.SizeSquared();
@@ -69,7 +69,9 @@ private:
 
 	FVector MakeLocation(FKey CellKey) const
 	{
-		return FVector(FVector::FReal(CellKey.X << HashCellBits), FVector::FReal(CellKey.Y << HashCellBits), FVector::FReal(CellKey.Z << HashCellBits));
+		return FVector(static_cast<FVector::FReal>(CellKey.X << HashCellBits),
+		               static_cast<FVector::FReal>(CellKey.Y << HashCellBits),
+		               static_cast<FVector::FReal>(CellKey.Z << HashCellBits));
 	}
 
 public:
