[UE] Change category of missing bone logging in FBoneReference::Initialize from 'Warning' to 'Log'
- Assumption is that sometimes we might expect to attempt to initialize a bone reference to a bone that doesn't exist (for example, when a blueprint is used on multiple assets) and that calling code should cope with this. In this situation, Anim nodes should pass through the input pose.

NOTE: This is a duplicate of CL-25104483 which was made on a different branch when the target function was in a different file. I made a mess of mergeing that change so I need to re-submit it here.

diff --git a/Engine/Source/Runtime/Engine/Private/Animation/BoneReference.cpp b/Engine/Source/Runtime/Engine/Private/Animation/BoneReference.cpp
--- a/Engine/Source/Runtime/Engine/Private/Animation/BoneReference.cpp
+++ b/Engine/Source/Runtime/Engine/Private/Animation/BoneReference.cpp
@@ -1,4 +1,4 @@
-﻿// Copyright Epic Games, Inc. All Rights Reserved.
+// Copyright Epic Games, Inc. All Rights Reserved.
 
 #include "Animation/BoneReference.h"
 #include "BoneContainer.h"
@@ -17,14 +17,14 @@ bool FBoneReference::Initialize(const FBoneContainer& RequiredBones)
 	bUseSkeletonIndex = false;
 	// If bone name is not found, look into the leader skeleton to see if it's found there.
 	// SkeletalMeshes can exclude bones from the leader skeleton, and that's OK.
-	// If it's not found in the leader skeleton, the bone does not exist at all! so we should report it as a warning.
+	// If it's not found in the leader skeleton, the bone does not exist at all! so we should log it.
 	if (BoneIndex == INDEX_NONE && BoneName != NAME_None)
 	{
 		if (USkeleton* SkeletonAsset = RequiredBones.GetSkeletonAsset())
 		{
 			if (SkeletonAsset->GetReferenceSkeleton().FindBoneIndex(BoneName) == INDEX_NONE)
 			{
-				UE_LOG(LogAnimation, Warning, TEXT("FBoneReference::Initialize BoneIndex for Bone '%s' does not exist in Skeleton '%s'"),
+				UE_LOG(LogAnimation, Log, TEXT("FBoneReference::Initialize BoneIndex for Bone '%s' does not exist in Skeleton '%s'"),
 					*BoneName.ToString(), *GetNameSafe(SkeletonAsset));
 			}
 		}
