Fixed mismatched copy constructor/assignment operator in FLocalPlayerContext.

diff --git a/Engine/Source/Runtime/Engine/Classes/Engine/LocalPlayer.h b/Engine/Source/Runtime/Engine/Classes/Engine/LocalPlayer.h
--- a/Engine/Source/Runtime/Engine/Classes/Engine/LocalPlayer.h
+++ b/Engine/Source/Runtime/Engine/Classes/Engine/LocalPlayer.h
@@ -44,7 +44,6 @@ struct ENGINE_API FLocalPlayerContext
 	FLocalPlayerContext();
 	FLocalPlayerContext(const class ULocalPlayer* InLocalPlayer, UWorld* InWorld = nullptr);
 	FLocalPlayerContext(const class APlayerController* InPlayerController);
-	FLocalPlayerContext(const FLocalPlayerContext& InPlayerContext);
 
 	/** Is this context initialized and still valid? */
 	bool IsValid() const;
diff --git a/Engine/Source/Runtime/Engine/Private/LocalPlayer.cpp b/Engine/Source/Runtime/Engine/Private/LocalPlayer.cpp
--- a/Engine/Source/Runtime/Engine/Private/LocalPlayer.cpp
+++ b/Engine/Source/Runtime/Engine/Private/LocalPlayer.cpp
@@ -88,13 +88,6 @@ FLocalPlayerContext::FLocalPlayerContext( const class APlayerController* InPlaye
 	SetPlayerController( InPlayerController );
 }
 
-FLocalPlayerContext::FLocalPlayerContext( const FLocalPlayerContext& InPlayerContext )
-	: World(InPlayerContext.World)
-{
-	check(InPlayerContext.GetLocalPlayer());
-	SetLocalPlayer(InPlayerContext.GetLocalPlayer());
-}
-
 bool FLocalPlayerContext::IsValid() const
 {
 	if (ULocalPlayer* LocalPlayerPtr = LocalPlayer.Get())
