Fixed mismatched copy constructor/assignment operator in FStreamableRenderResourceState.

diff --git a/Engine/Source/Runtime/Engine/Public/Streaming/StreamableRenderResourceState.h b/Engine/Source/Runtime/Engine/Public/Streaming/StreamableRenderResourceState.h
--- a/Engine/Source/Runtime/Engine/Public/Streaming/StreamableRenderResourceState.h
+++ b/Engine/Source/Runtime/Engine/Public/Streaming/StreamableRenderResourceState.h
@@ -26,9 +26,15 @@ struct FStreamableRenderResourceState
 
 	FStreamableRenderResourceState() : Data(0) {}
 
-	FORCEINLINE void operator=(const FStreamableRenderResourceState& Other) 
-	{ 
-		Data = Other.Data; 
+	// Manually implement copy operations, because MSVC produces bad codegen for bitfields in unions
+	// (as of the 14.35.32217 toolchain)
+	FORCEINLINE FStreamableRenderResourceState(const FStreamableRenderResourceState& Other)
+		: Data(Other.Data)
+	{
+	}
+	FORCEINLINE void operator=(const FStreamableRenderResourceState& Other)
+	{
+		Data = Other.Data;
 	}
 
 	union
