[UE-183205] Add the ability to disable/skip post-process AnimBPs for lower LOD levels

* Added a post-process AnimBP LOD threshold, similar to several of the anim graph nodes have.
* Disables evaluation of the PP AnimBP at its entirety.
* We are seeing performance improvements in animation evaluation time ranging from 3x to 117x on Meta Humans with this.
* Speed improvements are bigger the more complex the PP AnimBP is and the less nodes support or have setup node-based LOD.

diff --git a/Engine/Source/Runtime/Engine/Classes/Engine/SkeletalMesh.h b/Engine/Source/Runtime/Engine/Classes/Engine/SkeletalMesh.h
--- a/Engine/Source/Runtime/Engine/Classes/Engine/SkeletalMesh.h
+++ b/Engine/Source/Runtime/Engine/Classes/Engine/SkeletalMesh.h
@@ -2108,6 +2108,20 @@ public:
 	UPROPERTY(EditAnywhere, BlueprintReadOnly, Category = SkeletalMesh)
 	TSubclassOf<UAnimInstance> PostProcessAnimBlueprint;
 
+private:
+	/*
+	 * Max LOD level that post-process AnimBPs are evaluated.
+	 * For example if you have the threshold set to 2, it will evaluate until including LOD 2 (based on 0 index). In case the LOD level gets set to 3, it will stop evaluating the post-process AnimBP.
+	 * Setting it to -1 will always evaluate it and disable LODing.
+	 */
+	UPROPERTY(EditAnywhere, Category = SkeletalMesh, meta = (DisplayName = "Post-Process AnimBP LOD Threshold"))
+	int32 PostProcessAnimBPLODThreshold = INDEX_NONE;
+
+public:
+	int32 GetPostProcessAnimBPLODThreshold() const;
+	void SetPostProcessAnimBPLODThreshold(int32 InPostProcessAnimBPLODThreshold);
+	bool ShouldEvaluatePostProcessAnimBP(int32 LODLevel) const;
+
 	static FName GetPostProcessAnimBlueprintMemberName()
 	{
 		PRAGMA_DISABLE_DEPRECATION_WARNINGS
diff --git a/Engine/Source/Runtime/Engine/Private/Components/SkeletalMeshComponent.cpp b/Engine/Source/Runtime/Engine/Private/Components/SkeletalMeshComponent.cpp
--- a/Engine/Source/Runtime/Engine/Private/Components/SkeletalMeshComponent.cpp
+++ b/Engine/Source/Runtime/Engine/Private/Components/SkeletalMeshComponent.cpp
@@ -4331,16 +4331,40 @@ void USkeletalMeshComponent::FinalizeBoneTransform()
 
 bool USkeletalMeshComponent::ShouldUpdatePostProcessInstance() const
 {
+	if (USkeletalMesh* SkelMesh = GetSkeletalMeshAsset())
+	{
+		if (!SkelMesh->ShouldEvaluatePostProcessAnimBP(GetPredictedLODLevel()))
+		{
+			return false;
+		}
+	}
+
 	return PostProcessAnimInstance && !bDisablePostProcessBlueprint;
 }
 
 bool USkeletalMeshComponent::ShouldPostUpdatePostProcessInstance() const
 {
+	if (USkeletalMesh* SkelMesh = GetSkeletalMeshAsset())
+	{
+		if (!SkelMesh->ShouldEvaluatePostProcessAnimBP(GetPredictedLODLevel()))
+		{
+			return false;
+		}
+	}
+
 	return PostProcessAnimInstance && PostProcessAnimInstance->NeedsUpdate() && !bDisablePostProcessBlueprint;
 }
 
 bool USkeletalMeshComponent::ShouldEvaluatePostProcessInstance() const
 {
+	if (USkeletalMesh* SkelMesh = GetSkeletalMeshAsset())
+	{
+		if (!SkelMesh->ShouldEvaluatePostProcessAnimBP(GetPredictedLODLevel()))
+		{
+			return false;
+		}
+	}
+
 	return PostProcessAnimInstance && !bDisablePostProcessBlueprint;
 }
 
diff --git a/Engine/Source/Runtime/Engine/Private/SkeletalMesh.cpp b/Engine/Source/Runtime/Engine/Private/SkeletalMesh.cpp
--- a/Engine/Source/Runtime/Engine/Private/SkeletalMesh.cpp
+++ b/Engine/Source/Runtime/Engine/Private/SkeletalMesh.cpp
@@ -5673,6 +5673,21 @@ FString USkeletalMesh::GetAsyncPropertyName(uint64 Property) const
 	return StaticEnum<ESkeletalMeshAsyncProperties>()->GetNameByValue(Property).ToString();
 }
 
+int32 USkeletalMesh::GetPostProcessAnimBPLODThreshold() const
+{
+	return PostProcessAnimBPLODThreshold;
+}
+
+void USkeletalMesh::SetPostProcessAnimBPLODThreshold(int32 InPostProcessAnimBPLODThreshold)
+{
+	PostProcessAnimBPLODThreshold = InPostProcessAnimBPLODThreshold;
+}
+
+bool USkeletalMesh::ShouldEvaluatePostProcessAnimBP(int32 LODLevel) const
+{
+	return ((PostProcessAnimBPLODThreshold == INDEX_NONE) || (LODLevel <= PostProcessAnimBPLODThreshold));
+}
+
 /*-----------------------------------------------------------------------------
 USkeletalMeshSocket
 -----------------------------------------------------------------------------*/
