Optimized World Partition UpdateStreamingSources. Saves 1ms on project using >300 instances of partitioned worlds.

diff --git a/Engine/Source/Runtime/Engine/Private/WorldPartition/WorldPartitionStreamingPolicy.cpp b/Engine/Source/Runtime/Engine/Private/WorldPartition/WorldPartitionStreamingPolicy.cpp
--- a/Engine/Source/Runtime/Engine/Private/WorldPartition/WorldPartitionStreamingPolicy.cpp
+++ b/Engine/Source/Runtime/Engine/Private/WorldPartition/WorldPartitionStreamingPolicy.cpp
@@ -19,24 +19,12 @@
 #include "Engine/Level.h"
 #include "Engine/NetConnection.h"
 #include "WorldPartition/WorldPartitionRuntimeHash.h"
-#include "GameFramework/PlayerController.h"
+#include "Misc/HashBuilder.h"
 
 #include UE_INLINE_GENERATED_CPP_BY_NAME(WorldPartitionStreamingPolicy)
 
-#if WITH_EDITOR
-#include "LevelEditorViewport.h"
-#else
-#include "Misc/ScopeExit.h"
-#endif
-
 #define LOCTEXT_NAMESPACE "WorldPartitionStreamingPolicy"
 
-static int32 GUpdateStreamingSources = 1;
-static FAutoConsoleVariableRef CVarUpdateStreamingSources(
-	TEXT("wp.Runtime.UpdateStreamingSources"),
-	GUpdateStreamingSources,
-	TEXT("Set to 0 to stop updating (freeze) world partition streaming sources."));
-
 static int32 GMaxLoadingStreamingCells = 4;
 static FAutoConsoleVariableRef CMaxLoadingStreamingCells(
 	TEXT("wp.Runtime.MaxLoadingStreamingCells"),
@@ -64,20 +52,6 @@ FAutoConsoleVariableRef UWorldPartitionStreamingPolicy::CVarUpdateOptimEnabled(
 	TEXT("if nothing relevant changed since last update."),
 	ECVF_Default);
 
-int32 UWorldPartitionStreamingPolicy::LocationQuantization = 400;
-FAutoConsoleVariableRef UWorldPartitionStreamingPolicy::CVarLocationQuantization(
-	TEXT("wp.Runtime.UpdateStreaming.LocationQuantization"),
-	UWorldPartitionStreamingPolicy::LocationQuantization,
-	TEXT("Distance (in Unreal units) used to quantize the streaming sources location to determine if a world partition streaming update is necessary."),
-	ECVF_Default);
-
-int32 UWorldPartitionStreamingPolicy::RotationQuantization = 10;
-FAutoConsoleVariableRef UWorldPartitionStreamingPolicy::CVarRotationQuantization(
-	TEXT("wp.Runtime.UpdateStreaming.RotationQuantization"),
-	UWorldPartitionStreamingPolicy::RotationQuantization,
-	TEXT("Angle (in degrees) used to quantize the streaming sources rotation to determine if a world partition streaming update is necessary."),
-	ECVF_Default);
-
 int32 UWorldPartitionStreamingPolicy::ForceUpdateFrameCount = 30;
 FAutoConsoleVariableRef UWorldPartitionStreamingPolicy::CVarForceUpdateFrameCount(
 	TEXT("wp.Runtime.UpdateStreaming.ForceUpdateFrameCount"),
@@ -103,6 +77,7 @@ UWorldPartitionStreamingPolicy::UWorldPartitionStreamingPolicy(const FObjectInit
 	, ContentBundleServerEpoch(INT_MIN)
 	, ServerStreamingEnabledEpoch(INT_MIN)
 	, UpdateStreamingHash(0)
+	, UpdateStreamingSourcesHash(0)
 	, UpdateStreamingStateCalls(0)
 	, StreamingPerformance(EWorldPartitionStreamingPerformance::Good)
 #if !UE_BUILD_SHIPPING
@@ -117,93 +92,25 @@ UWorldPartitionStreamingPolicy::UWorldPartitionStreamingPolicy(const FObjectInit
 	}
 }
 
-void UWorldPartitionStreamingPolicy::UpdateStreamingSources()
+void UWorldPartitionStreamingPolicy::UpdateStreamingSources(bool bCanOptimizeUpdate)
 {
-	if (!GUpdateStreamingSources)
-	{
-		return;
-	}
-
 	TRACE_CPUPROFILER_EVENT_SCOPE(UWorldPartitionStreamingPolicy::UpdateStreamingSources);
-
-	StreamingSources.Reset();
-
 	if (!WorldPartition->CanStream())
 	{
+		StreamingSources.Reset();
 		return;
 	}
 
-	const FTransform WorldToLocal = WorldPartition->GetInstanceTransform().Inverse();
-	UWorld* World = GetWorld();
-
-	const bool bIsServer = WorldPartition->IsServer();
-	const bool bIsServerStreamingEnabled = WorldPartition->IsServerStreamingEnabled();
-
-	if (!AWorldPartitionReplay::IsPlaybackEnabled(World) || !WorldPartition->Replay->GetReplayStreamingSources(StreamingSources))
+	const UWorldPartitionSubsystem* WorldPartitionSubsystem = GetWorld()->GetSubsystem<UWorldPartitionSubsystem>();
+	const uint32 NewUpdateStreamingSourcesHash = WorldPartitionSubsystem->GetStreamingSourcesHash();
+	if (bCanOptimizeUpdate && (UpdateStreamingSourcesHash == NewUpdateStreamingSourcesHash))
 	{
-		bool bAllowPlayerControllerStreamingSources = true;
-
-#if WITH_EDITOR
-		if (UWorldPartition::IsSimulating())
-		{
-			// We are in the SIE
-			const FVector ViewLocation = WorldToLocal.TransformPosition(GCurrentLevelEditingViewportClient->GetViewLocation());
-			const FRotator ViewRotation = WorldToLocal.TransformRotation(GCurrentLevelEditingViewportClient->GetViewRotation().Quaternion()).Rotator();
-			static const FName NAME_SIE(TEXT("SIE"));
-			StreamingSources.Add(FWorldPartitionStreamingSource(NAME_SIE, ViewLocation, ViewRotation, EStreamingSourceTargetState::Activated, /*bBlockOnSlowLoading=*/false, EStreamingSourcePriority::Default, false));
-			bAllowPlayerControllerStreamingSources = false;
-		}
-#endif
-		if (!bIsServer || bIsServerStreamingEnabled || AWorldPartitionReplay::IsRecordingEnabled(World))
-		{
-			UWorldPartitionSubsystem* WorldPartitionSubsystem = GetWorld()->GetSubsystem<UWorldPartitionSubsystem>();
-			check(WorldPartitionSubsystem);
-
-			TArray<FWorldPartitionStreamingSource> ProviderStreamingSources;
-			for (IWorldPartitionStreamingSourceProvider* StreamingSourceProvider : WorldPartitionSubsystem->GetStreamingSourceProviders())
-			{
-				if (bAllowPlayerControllerStreamingSources || !Cast<APlayerController>(StreamingSourceProvider->GetStreamingSourceOwner()))
-				{
-					ProviderStreamingSources.Reset();
-					if (StreamingSourceProvider->GetStreamingSources(ProviderStreamingSources))
-					{
-						for (FWorldPartitionStreamingSource& ProviderStreamingSource : ProviderStreamingSources)
-						{
-							// Transform to Local
-							ProviderStreamingSource.Location = WorldToLocal.TransformPosition(ProviderStreamingSource.Location);
-							ProviderStreamingSource.Rotation = WorldToLocal.TransformRotation(ProviderStreamingSource.Rotation.Quaternion()).Rotator();
-							StreamingSources.Add(MoveTemp(ProviderStreamingSource));
-						}
-					}
-				}
-			}
-		}
-	}
-
-	for (auto& Pair : StreamingSourcesVelocity)
-	{
-		Pair.Value.Invalidate();
-	}
-
-	// Update streaming sources velocity
-	const float CurrentTime = GetWorld()->GetTimeSeconds();
-	for (FWorldPartitionStreamingSource& StreamingSource : StreamingSources)
-	{
-		if (!StreamingSource.Name.IsNone())
-		{
-			FStreamingSourceVelocity& SourceVelocity = StreamingSourcesVelocity.FindOrAdd(StreamingSource.Name, FStreamingSourceVelocity(StreamingSource.Name));
-			StreamingSource.Velocity = SourceVelocity.GetAverageVelocity(StreamingSource.Location, CurrentTime);
-		}
+		return;
 	}
 
-	// Cleanup StreamingSourcesVelocity
-	for (auto It(StreamingSourcesVelocity.CreateIterator()); It; ++It)
-	{
-		if (!It.Value().IsValid())
-		{
-			It.RemoveCurrent();
-		}
-	}
+	StreamingSources.Reset();
+	WorldPartitionSubsystem->GetStreamingSources(WorldPartition, StreamingSources);
+	UpdateStreamingSourcesHash = NewUpdateStreamingSourcesHash;
 }
 
 #define WORLDPARTITION_LOG_UPDATESTREAMINGSTATE(Verbosity)\
@@ -243,79 +150,32 @@ int32 UWorldPartitionStreamingPolicy::ComputeServerStreamingEnabledEpoch() const
 bool UWorldPartitionStreamingPolicy::IsUpdateStreamingOptimEnabled()
 {
 	return UWorldPartitionStreamingPolicy::IsUpdateOptimEnabled &&
-		(UWorldPartitionStreamingPolicy::LocationQuantization > 0) &&
-		(UWorldPartitionStreamingPolicy::RotationQuantization > 0);
+		(FWorldPartitionStreamingSource::GetLocationQuantization() > 0) &&
+		(FWorldPartitionStreamingSource::GetRotationQuantization() > 0);
 };
 
-uint32 UWorldPartitionStreamingPolicy::ComputeUpdateStreamingHash() const
+uint32 UWorldPartitionStreamingPolicy::ComputeUpdateStreamingHash(bool bCanOptimizeUpdate) const
 {
-	const UWorldPartitionRuntimeHash* RuntimeHash = WorldPartition->RuntimeHash;
-	const bool bForceFrameUpdate = (UWorldPartitionStreamingPolicy::ForceUpdateFrameCount > 0) ? ((UpdateStreamingStateCalls % UWorldPartitionStreamingPolicy::ForceUpdateFrameCount) == 0) : false;
-
-	const bool bCanOptimize = 
-		RuntimeHash &&
-		IsUpdateStreamingOptimEnabled() &&							// Check CVars to see if optimization is enabled
-		!WorldPartition->IsServer() &&								// Don't optimize on the server
-		bLastUpdateCompletedLoadingAndActivation &&					// Don't optimize if last frame didn't process all cells to load/activate
-		!IsInBlockTillLevelStreamingCompleted() &&					// Don't optimize when inside UWorld::BlockTillLevelStreamingCompleted
-		ActivatedCells.GetPendingAddToWorldCells().IsEmpty() &&		// Don't optimize when remaining cells to add to world
-		!bForceFrameUpdate;											// We garantee to update every N frame to force some internal updates like UpdateStreamingPerformance
-
-	if (bCanOptimize)
+	if (bCanOptimizeUpdate)
 	{
+		const bool bIsStreaming3D = WorldPartition->RuntimeHash && WorldPartition->RuntimeHash->IsStreaming3D();
+
 		// Build hash that will be used to detect relevant changes
-		uint32 NewHash = GetTypeHash(ComputeServerStreamingEnabledEpoch());
-		NewHash = HashCombine(NewHash, GetTypeHash(FContentBundle::GetContentBundleEpoch()));
-		NewHash = HashCombine(NewHash, GetTypeHash(AWorldDataLayers::GetDataLayersStateEpoch()));
-		NewHash = HashCombine(NewHash, GetTypeHash(RuntimeHash->IsStreaming3D()));
+		FHashBuilder HashBuilder;
+		HashBuilder << ComputeServerStreamingEnabledEpoch();
+		HashBuilder << FContentBundle::GetContentBundleEpoch();
+		HashBuilder << AWorldDataLayers::GetDataLayersStateEpoch();
+		HashBuilder << bIsStreaming3D;
 		for (const FWorldPartitionStreamingSource& Source : StreamingSources)
 		{
-			NewHash = HashCombine(NewHash, ComputeStreamingSourceHash(Source));
+			HashBuilder << Source.GetHash(bIsStreaming3D);
 		}
-		return NewHash;
+		return HashBuilder.GetHash();
 	}
 
 	return 0;
 };
 
-uint32 UWorldPartitionStreamingPolicy::ComputeStreamingSourceHash(const FWorldPartitionStreamingSource& Source) const
-{
-	uint32 Hash = (GetTypeHash(Source.Name));
-	Hash = HashCombine(Hash, GetTypeHash(FMath::FloorToInt(Source.Location.X / UWorldPartitionStreamingPolicy::LocationQuantization)));
-	Hash = HashCombine(Hash, GetTypeHash(FMath::FloorToInt(Source.Location.Y / UWorldPartitionStreamingPolicy::LocationQuantization)));
-	Hash = HashCombine(Hash, GetTypeHash(FMath::FloorToInt(Source.Rotation.Yaw / UWorldPartitionStreamingPolicy::RotationQuantization)));
-	// Only consider Z position and pitch/roll rotations when hash is streaming in 3D
-	if (WorldPartition->RuntimeHash && WorldPartition->RuntimeHash->IsStreaming3D())
-	{
-		Hash = HashCombine(Hash, GetTypeHash(FMath::FloorToInt(Source.Location.Z / UWorldPartitionStreamingPolicy::LocationQuantization)));
-		Hash = HashCombine(Hash, GetTypeHash(FMath::FloorToInt(Source.Rotation.Pitch / UWorldPartitionStreamingPolicy::RotationQuantization)));
-		Hash = HashCombine(Hash, GetTypeHash(FMath::FloorToInt(Source.Rotation.Roll / UWorldPartitionStreamingPolicy::RotationQuantization)));
-	}
-	Hash = HashCombine(Hash, GetTypeHash(Source.TargetState));
-	Hash = HashCombine(Hash, GetTypeHash(Source.bBlockOnSlowLoading));
-	Hash = HashCombine(Hash, GetTypeHash(Source.bReplay));
-	Hash = HashCombine(Hash, GetTypeHash(Source.bRemote));
-	Hash = HashCombine(Hash, GetTypeHash(Source.Priority));
-	Hash = HashCombine(Hash, GetTypeHash(Source.TargetBehavior));
-	for (FName TargetGrid : Source.TargetGrids)
-	{
-		Hash = HashCombine(Hash, GetTypeHash(TargetGrid));
-	}
-	for (const FSoftObjectPath& TargetHLODLayer : Source.TargetHLODLayers)
-	{
-		Hash = HashCombine(Hash, GetTypeHash(TargetHLODLayer));
-	}
-	PRAGMA_DISABLE_DEPRECATION_WARNINGS
-	Hash = HashCombine(Hash, GetTypeHash(Source.TargetGrid));
-	Hash = HashCombine(Hash, GetTypeHash(Source.TargetHLODLayer));
-	PRAGMA_ENABLE_DEPRECATION_WARNINGS
-	for (const FStreamingSourceShape& Shape : Source.Shapes)
-	{
-		Hash = HashCombine(Hash, GetTypeHash(Shape));
-	}
-	return Hash;
-};
-
 bool UWorldPartitionStreamingPolicy::GetIntersectingCells(const TArray<FWorldPartitionStreamingQuerySource>& InSources, TArray<const IWorldPartitionCell*>& OutCells) const
 {
 	if (!WorldPartition || !WorldPartition->RuntimeHash)
@@ -408,12 +268,22 @@ void UWorldPartitionStreamingPolicy::UpdateStreamingState()
 		bCriticalPerformanceRequestedBlockTillOnWorld = false;
 		CriticalPerformanceBlockTillLevelStreamingCompletedEpoch = World->GetBlockTillLevelStreamingCompletedEpoch();
 	}
+
+	const bool bForceFrameUpdate = (UWorldPartitionStreamingPolicy::ForceUpdateFrameCount > 0) ? ((UpdateStreamingStateCalls % UWorldPartitionStreamingPolicy::ForceUpdateFrameCount) == 0) : false;
+	const bool bCanOptimizeUpdate =
+		WorldPartition->RuntimeHash &&
+		!bForceFrameUpdate &&										// We garantee to update every N frame to force some internal updates like UpdateStreamingPerformance
+		IsUpdateStreamingOptimEnabled() &&							// Check CVars to see if optimization is enabled
+		!WorldPartition->IsServer() &&								// Don't optimize on the server
+		bLastUpdateCompletedLoadingAndActivation &&					// Don't optimize if last frame didn't process all cells to load/activate
+		!IsInBlockTillLevelStreamingCompleted() &&					// Don't optimize when inside UWorld::BlockTillLevelStreamingCompleted
+		ActivatedCells.GetPendingAddToWorldCells().IsEmpty(); 		// Don't optimize when remaining cells to add to world
 	
 	// Update streaming sources
-	UpdateStreamingSources();
+	UpdateStreamingSources(bCanOptimizeUpdate);
 
 	// Detect if nothing relevant changed and early out
-	const uint32 NewUpdateStreamingHash = ComputeUpdateStreamingHash();
+	const uint32 NewUpdateStreamingHash = ComputeUpdateStreamingHash(bCanOptimizeUpdate);
 	const bool bShouldSkipUpdate = NewUpdateStreamingHash && (UpdateStreamingHash == NewUpdateStreamingHash);
 	if (bShouldSkipUpdate)
 	{
@@ -1137,58 +1007,4 @@ void UWorldPartitionStreamingPolicy::FActivatedCells::OnRemovedFromWorld(const U
 	}
 }
 
-/*
- * FStreamingSourceVelocity Implementation
- */
-
-FStreamingSourceVelocity::FStreamingSourceVelocity(const FName& InSourceName)
-: bIsValid(false)
-, SourceName(InSourceName)
-, LastIndex(INDEX_NONE)
-, LastUpdateTime(-1.0)
-, VelocitiesHistorySum(0.f)
-{
-	VelocitiesHistory.SetNumZeroed(VELOCITY_HISTORY_SAMPLE_COUNT);
-}
-
-float FStreamingSourceVelocity::GetAverageVelocity(const FVector& NewPosition, const float CurrentTime)
-{
-	bIsValid = true;
-
-	const double TeleportDistance = 100;
-	const float MaxDeltaSeconds = 5.f;
-	const bool bIsFirstCall = (LastIndex == INDEX_NONE);
-	const float DeltaSeconds = bIsFirstCall ? 0.f : (CurrentTime - LastUpdateTime);
-	const double Distance = bIsFirstCall ? 0.f : ((NewPosition - LastPosition) * 0.01).Size();
-	if (bIsFirstCall)
-	{
-		UE_LOG(LogWorldPartition, Verbose, TEXT("New Streaming Source: %s -> Position: %s"), *SourceName.ToString(), *NewPosition.ToString());
-		LastIndex = 0;
-	}
-
-	ON_SCOPE_EXIT
-	{
-		LastUpdateTime = CurrentTime;
-		LastPosition = NewPosition;
-	};
-
-	// Handle invalid cases
-	if (bIsFirstCall || (DeltaSeconds <= 0.f) || (DeltaSeconds > MaxDeltaSeconds) || (Distance > TeleportDistance))
-	{
-		UE_CLOG(Distance > TeleportDistance, LogWorldPartition, Verbose, TEXT("Detected Streaming Source Teleport: %s -> Last Position: %s -> New Position: %s"), *SourceName.ToString(), *LastPosition.ToString(), *NewPosition.ToString());
-		return 0.f;
-	}
-
-	// Compute velocity (m/s)
-	check(Distance < MAX_flt);
-	const float Velocity = (float)Distance / DeltaSeconds;
-	// Update velocities history buffer and sum
-	LastIndex = (LastIndex + 1) % VELOCITY_HISTORY_SAMPLE_COUNT;
-	VelocitiesHistorySum = FMath::Max<float>(0.f, (VelocitiesHistorySum + Velocity - VelocitiesHistory[LastIndex]));
-	VelocitiesHistory[LastIndex] = Velocity;
-
-	// return average
-	return (VelocitiesHistorySum / (float)VELOCITY_HISTORY_SAMPLE_COUNT);
-}
-
 #undef LOCTEXT_NAMESPACE
diff --git a/Engine/Source/Runtime/Engine/Private/WorldPartition/WorldPartitionStreamingSource.cpp b/Engine/Source/Runtime/Engine/Private/WorldPartition/WorldPartitionStreamingSource.cpp
--- a/Engine/Source/Runtime/Engine/Private/WorldPartition/WorldPartitionStreamingSource.cpp
+++ b/Engine/Source/Runtime/Engine/Private/WorldPartition/WorldPartitionStreamingSource.cpp
@@ -1,9 +1,62 @@
 // Copyright Epic Games, Inc. All Rights Reserved.
 
 #include "WorldPartition/WorldPartitionStreamingSource.h"
+#include "Misc/HashBuilder.h"
 
 #include UE_INLINE_GENERATED_CPP_BY_NAME(WorldPartitionStreamingSource)
 
+int32 FWorldPartitionStreamingSource::LocationQuantization = 400;
+FAutoConsoleVariableRef FWorldPartitionStreamingSource::CVarLocationQuantization(
+	TEXT("wp.Runtime.UpdateStreaming.LocationQuantization"),
+	FWorldPartitionStreamingSource::LocationQuantization,
+	TEXT("Distance (in Unreal units) used to quantize the streaming sources location to determine if a world partition streaming update is necessary."),
+	ECVF_Default);
+
+int32 FWorldPartitionStreamingSource::RotationQuantization = 10;
+FAutoConsoleVariableRef FWorldPartitionStreamingSource::CVarRotationQuantization(
+	TEXT("wp.Runtime.UpdateStreaming.RotationQuantization"),
+	FWorldPartitionStreamingSource::RotationQuantization,
+	TEXT("Angle (in degrees) used to quantize the streaming sources rotation to determine if a world partition streaming update is necessary."),
+	ECVF_Default);
+
+void FWorldPartitionStreamingSource::UpdateHash()
+{
+	FHashBuilder HashBuilder;
+	HashBuilder << Name;
+	HashBuilder << FMath::FloorToInt(Location.X / FWorldPartitionStreamingSource::LocationQuantization);
+	HashBuilder << FMath::FloorToInt(Location.Y / FWorldPartitionStreamingSource::LocationQuantization);
+	HashBuilder << FMath::FloorToInt(Rotation.Yaw / FWorldPartitionStreamingSource::RotationQuantization);
+	HashBuilder << TargetState;
+	HashBuilder << bBlockOnSlowLoading;
+	HashBuilder << bReplay;
+	HashBuilder << bRemote;
+	HashBuilder << Priority;
+	HashBuilder << TargetBehavior;
+	for (FName Grid : TargetGrids)
+	{
+		HashBuilder << Grid;
+	}
+	for (const FSoftObjectPath& HLODLayer : TargetHLODLayers)
+	{
+		HashBuilder << HLODLayer;
+	}
+	PRAGMA_DISABLE_DEPRECATION_WARNINGS
+	HashBuilder << TargetGrid;
+	HashBuilder << TargetHLODLayer;
+	PRAGMA_ENABLE_DEPRECATION_WARNINGS
+	for (const FStreamingSourceShape& Shape : Shapes)
+	{
+		HashBuilder << Shape;
+	}
+	Hash2D = HashBuilder.GetHash();
+	
+	// Hash3D
+	HashBuilder << FMath::FloorToInt(Location.Z / FWorldPartitionStreamingSource::LocationQuantization);
+	HashBuilder << FMath::FloorToInt(Rotation.Pitch / FWorldPartitionStreamingSource::RotationQuantization);
+	HashBuilder << FMath::FloorToInt(Rotation.Roll / FWorldPartitionStreamingSource::RotationQuantization);
+	Hash3D = HashBuilder.GetHash();
+};
+
 FORCEINLINE static bool IsClockWise(const FVector2D& V1, const FVector2D& V2)
 {
 	return (V1 ^ V2) < 0;
diff --git a/Engine/Source/Runtime/Engine/Private/WorldPartition/WorldPartitionSubsystem.cpp b/Engine/Source/Runtime/Engine/Private/WorldPartition/WorldPartitionSubsystem.cpp
--- a/Engine/Source/Runtime/Engine/Private/WorldPartition/WorldPartitionSubsystem.cpp
+++ b/Engine/Source/Runtime/Engine/Private/WorldPartition/WorldPartitionSubsystem.cpp
@@ -6,6 +6,10 @@
 #include "WorldPartition/WorldPartition.h"
 #include "WorldPartition/WorldPartitionActorDesc.h"
 #include "WorldPartition/WorldPartitionDebugHelper.h"
+#include "WorldPartition/WorldPartitionLog.h"
+#include "WorldPartition/WorldPartitionRuntimeHash.h"
+#include "WorldPartition/WorldPartitionStreamingSource.h"
+#include "WorldPartition/WorldPartitionReplay.h"
 #include "WorldPartition/DataLayer/WorldDataLayersActorDesc.h"
 #include "WorldPartition/DataLayer/DataLayerManager.h"
 #include "WorldPartition/DataLayer/DataLayerAsset.h"
@@ -14,9 +18,12 @@
 #include "Engine/CoreSettings.h"
 #include "Engine/LevelBounds.h"
 #include "Debug/DebugDrawService.h"
-#include "WorldPartition/WorldPartitionLog.h"
-#include "WorldPartition/WorldPartitionRuntimeHash.h"
-#include "WorldPartition/WorldPartitionStreamingSource.h"
+#include "GameFramework/PlayerController.h"
+
+#if WITH_EDITOR
+#include "Editor.h"
+#include "LevelEditorViewport.h"
+#endif
 
 #include UE_INLINE_GENERATED_CPP_BY_NAME(WorldPartitionSubsystem)
 
@@ -111,6 +118,7 @@ static FAutoConsoleCommandWithOutputDevice GDumpStreamingSourcesCmd(
 );
 
 UWorldPartitionSubsystem::UWorldPartitionSubsystem()
+: StreamingSourcesHash(0)
 {}
 
 UWorldPartition* UWorldPartitionSubsystem::GetWorldPartition()
@@ -552,12 +560,12 @@ bool UWorldPartitionSubsystem::IsStreamingCompleted(const IWorldPartitionStreami
 {
 	// Convert specified/optional streaming source provider to a world partition 
 	// streaming source and pass it along to each registered world partition
-	TArray<FWorldPartitionStreamingSource> StreamingSources;
+	TArray<FWorldPartitionStreamingSource> LocalStreamingSources;
 	TArray<FWorldPartitionStreamingSource>* StreamingSourcesPtr = nullptr;
 	if (InStreamingSourceProvider)
 	{
-		StreamingSourcesPtr = &StreamingSources;
-		if (!InStreamingSourceProvider->GetStreamingSources(StreamingSources))
+		StreamingSourcesPtr = &LocalStreamingSources;
+		if (!InStreamingSourceProvider->GetStreamingSources(LocalStreamingSources))
 		{
 			return true;
 		}
@@ -588,22 +596,145 @@ bool UWorldPartitionSubsystem::IsStreamingCompleted(EWorldPartitionRuntimeCellSt
 
 void UWorldPartitionSubsystem::DumpStreamingSources(FOutputDevice& OutputDevice) const
 {
-	if (const UWorldPartition* WorldPartition = GetWorldPartition())
+	if (StreamingSources.Num() > 0)
+	{
+		OutputDevice.Logf(TEXT("Streaming Sources:"));
+		for (const FWorldPartitionStreamingSource& StreamingSource : StreamingSources)
+		{
+			OutputDevice.Logf(TEXT("  - %s: %s"), *StreamingSource.Name.ToString(), *StreamingSource.ToString());
+		}
+	}
+}
+
+static int32 GUpdateStreamingSources = 1;
+static FAutoConsoleVariableRef CVarUpdateStreamingSources(
+	TEXT("wp.Runtime.UpdateStreamingSources"),
+	GUpdateStreamingSources,
+	TEXT("Set to 0 to stop updating (freeze) world partition streaming sources."));
+
+#if WITH_EDITOR
+static const FName NAME_SIEStreamingSource(TEXT("SIE"));
+#endif
+
+void UWorldPartitionSubsystem::UpdateStreamingSources()
+{
+	if (!GUpdateStreamingSources)
+	{
+		return;
+	}
+
+	TRACE_CPUPROFILER_EVENT_SCOPE(UWorldPartitionSubsystem::UpdateStreamingSources);
+
+	StreamingSources.Reset();
+
+	UWorld* World = GetWorld();
+	bool bIsUsingReplayStreamingSources = false;
+	if (AWorldPartitionReplay::IsPlaybackEnabled(World))
 	{
-		const TArray<FWorldPartitionStreamingSource>* StreamingSources = &WorldPartition->GetStreamingSources();
-		if (StreamingSources && (StreamingSources->Num() > 0))
+		if (const UWorldPartition* WorldPartition = World->GetWorldPartition())
+		{
+			bIsUsingReplayStreamingSources = WorldPartition->Replay->GetReplayStreamingSources(StreamingSources);
+		}
+	}
+
+	if (!bIsUsingReplayStreamingSources)
+	{
+		bool bAllowPlayerControllerStreamingSources = true;
+#if WITH_EDITOR
+		if (UWorldPartition::IsSimulating())
+		{
+			// We are in the SIE
+			const FVector ViewLocation = GCurrentLevelEditingViewportClient->GetViewLocation();
+			const FRotator ViewRotation = GCurrentLevelEditingViewportClient->GetViewRotation();
+			StreamingSources.Add(FWorldPartitionStreamingSource(NAME_SIEStreamingSource, ViewLocation, ViewRotation, EStreamingSourceTargetState::Activated, /*bBlockOnSlowLoading=*/false, EStreamingSourcePriority::Default, false));
+			bAllowPlayerControllerStreamingSources = false;
+		}
+#endif
+		TArray<FWorldPartitionStreamingSource> ProviderStreamingSources;
+		for (IWorldPartitionStreamingSourceProvider* StreamingSourceProvider : GetStreamingSourceProviders())
 		{
-			OutputDevice.Logf(TEXT("Streaming Sources:"));
-			for (const FWorldPartitionStreamingSource& StreamingSource : *StreamingSources)
+			if (bAllowPlayerControllerStreamingSources || !Cast<APlayerController>(StreamingSourceProvider->GetStreamingSourceOwner()))
 			{
-				OutputDevice.Logf(TEXT("  - %s: %s"), *StreamingSource.Name.ToString(), *StreamingSource.ToString());
+				ProviderStreamingSources.Reset();
+				if (StreamingSourceProvider->GetStreamingSources(ProviderStreamingSources))
+				{
+					for (FWorldPartitionStreamingSource& ProviderStreamingSource : ProviderStreamingSources)
+					{
+						StreamingSources.Add(MoveTemp(ProviderStreamingSource));
+					}
+				}
 			}
 		}
 	}
+
+	for (auto& Pair : StreamingSourcesVelocity)
+	{
+		Pair.Value.Invalidate();
+	}
+
+	StreamingSourcesHash = 0;
+	const float CurrentTime = World->GetTimeSeconds();
+	for (FWorldPartitionStreamingSource& StreamingSource : StreamingSources)
+	{
+		// Update streaming sources velocity
+		if (!StreamingSource.Name.IsNone())
+		{
+			FStreamingSourceVelocity& SourceVelocity = StreamingSourcesVelocity.FindOrAdd(StreamingSource.Name, FStreamingSourceVelocity(StreamingSource.Name));
+			StreamingSource.Velocity = SourceVelocity.GetAverageVelocity(StreamingSource.Location, CurrentTime);
+		}
+
+		// Update streaming source hash
+		StreamingSource.UpdateHash();
+		// Build hash for all streaming sources
+		StreamingSourcesHash = HashCombine(StreamingSourcesHash, StreamingSource.GetHash());
+	}
+
+	// Cleanup StreamingSourcesVelocity
+	for (auto It(StreamingSourcesVelocity.CreateIterator()); It; ++It)
+	{
+		if (!It.Value().IsValid())
+		{
+			It.RemoveCurrent();
+		}
+	}
+}
+
+void UWorldPartitionSubsystem::GetStreamingSources(const UWorldPartition* InWorldPartition, TArray<FWorldPartitionStreamingSource>& OutStreamingSources) const
+{
+	const bool bIsServer = InWorldPartition->IsServer();
+	const bool bIsServerStreamingEnabled = InWorldPartition->IsServerStreamingEnabled();
+	const bool bIncludeStreamingSources = (!bIsServer || bIsServerStreamingEnabled || AWorldPartitionReplay::IsRecordingEnabled(GetWorld()));
+
+	if (bIncludeStreamingSources)
+	{
+		OutStreamingSources.Append(StreamingSources);
+	}
+#if WITH_EDITOR
+	else if (UWorldPartition::IsSimulating())
+	{
+		if (const FWorldPartitionStreamingSource* SIEStreamingSource = (StreamingSources.Num() > 0) && (StreamingSources[0].Name == NAME_SIEStreamingSource) ? &StreamingSources[0] : nullptr)
+		{
+			OutStreamingSources.Add(*SIEStreamingSource);
+		}
+	}
+#endif
+
+	// Transform to Local
+	if (OutStreamingSources.Num())
+	{
+		const FTransform WorldToLocal = InWorldPartition->GetInstanceTransform().Inverse();
+		for (FWorldPartitionStreamingSource& StreamingSource : OutStreamingSources)
+		{
+			StreamingSource.Location = WorldToLocal.TransformPosition(StreamingSource.Location);
+			StreamingSource.Rotation = WorldToLocal.TransformRotation(StreamingSource.Rotation.Quaternion()).Rotator();
+		}
+	}
 }
 
 void UWorldPartitionSubsystem::UpdateStreamingState()
 {
+	UpdateStreamingSources();
+
 	//make temp copy of array as UpdateStreamingState may FlushAsyncLoading, which may add a new world partition to RegisteredWorldPartitions while iterating
 	const TArray<UWorldPartition*> RegisteredWorldPartitionsCopy = RegisteredWorldPartitions;
 	for (UWorldPartition* RegisteredWorldParition : RegisteredWorldPartitionsCopy)
@@ -699,15 +830,15 @@ void UWorldPartitionSubsystem::Draw(UCanvas* Canvas, class APlayerController* PC
 	{
 		TRACE_CPUPROFILER_EVENT_SCOPE(UWorldPartitionSubsystem::DrawStreamingSources);
 
-		const TArray<FWorldPartitionStreamingSource>& StreamingSources = WorldPartition->GetStreamingSources();
-		if (StreamingSources.Num() > 0)
+		const TArray<FWorldPartitionStreamingSource>& LocalStreamingSources = WorldPartition->GetStreamingSources();
+		if (LocalStreamingSources.Num() > 0)
 		{
 			FString Title(TEXT("Streaming Sources"));
 			FWorldPartitionDebugHelper::DrawText(Canvas, Title, GEngine->GetSmallFont(), FColor::Yellow, CurrentOffset);
 
 			FVector2D Pos = CurrentOffset;
 			float MaxTextWidth = 0;
-			for (const FWorldPartitionStreamingSource& StreamingSource : StreamingSources)
+			for (const FWorldPartitionStreamingSource& StreamingSource : LocalStreamingSources)
 			{
 				FString StreamingSourceDisplay = StreamingSource.Name.ToString();
 				if (StreamingSource.bReplay)
@@ -717,7 +848,7 @@ void UWorldPartitionSubsystem::Draw(UCanvas* Canvas, class APlayerController* PC
 				FWorldPartitionDebugHelper::DrawText(Canvas, StreamingSourceDisplay, GEngine->GetSmallFont(), StreamingSource.GetDebugColor(), Pos, &MaxTextWidth);
 			}
 			Pos = CurrentOffset + FVector2D(MaxTextWidth + 10, 0.f);
-			for (const FWorldPartitionStreamingSource& StreamingSource : StreamingSources)
+			for (const FWorldPartitionStreamingSource& StreamingSource : LocalStreamingSources)
 			{
 				FWorldPartitionDebugHelper::DrawText(Canvas, *StreamingSource.ToString(), GEngine->GetSmallFont(), FColor::White, Pos);
 			}
@@ -750,3 +881,57 @@ void UWorldPartitionSubsystem::Draw(UCanvas* Canvas, class APlayerController* PC
 		WorldPartition->DrawRuntimeCellsDetails(Canvas, CurrentOffset);
 	}
 }
+
+/*
+ * FStreamingSourceVelocity Implementation
+ */
+
+FStreamingSourceVelocity::FStreamingSourceVelocity(const FName& InSourceName)
+	: bIsValid(false)
+	, SourceName(InSourceName)
+	, LastIndex(INDEX_NONE)
+	, LastUpdateTime(-1.0)
+	, VelocityHistorySum(0.f)
+{
+	VelocityHistory.SetNumZeroed(VELOCITY_HISTORY_SAMPLE_COUNT);
+}
+
+float FStreamingSourceVelocity::GetAverageVelocity(const FVector& NewPosition, const float CurrentTime)
+{
+	bIsValid = true;
+
+	const double TeleportDistance = 100;
+	const float MaxDeltaSeconds = 5.f;
+	const bool bIsFirstCall = (LastIndex == INDEX_NONE);
+	const float DeltaSeconds = bIsFirstCall ? 0.f : (CurrentTime - LastUpdateTime);
+	const double Distance = bIsFirstCall ? 0.f : ((NewPosition - LastPosition) * 0.01).Size();
+	if (bIsFirstCall)
+	{
+		UE_LOG(LogWorldPartition, Verbose, TEXT("New Streaming Source: %s -> Position: %s"), *SourceName.ToString(), *NewPosition.ToString());
+		LastIndex = 0;
+	}
+
+	ON_SCOPE_EXIT
+	{
+		LastUpdateTime = CurrentTime;
+		LastPosition = NewPosition;
+	};
+
+	// Handle invalid cases
+	if (bIsFirstCall || (DeltaSeconds <= 0.f) || (DeltaSeconds > MaxDeltaSeconds) || (Distance > TeleportDistance))
+	{
+		UE_CLOG(Distance > TeleportDistance, LogWorldPartition, Verbose, TEXT("Detected Streaming Source Teleport: %s -> Last Position: %s -> New Position: %s"), *SourceName.ToString(), *LastPosition.ToString(), *NewPosition.ToString());
+		return 0.f;
+	}
+
+	// Compute velocity (m/s)
+	check(Distance < MAX_flt);
+	const float Velocity = (float)Distance / DeltaSeconds;
+	// Update velocity history buffer and sum
+	LastIndex = (LastIndex + 1) % VELOCITY_HISTORY_SAMPLE_COUNT;
+	VelocityHistorySum = FMath::Max<float>(0.f, (VelocityHistorySum + Velocity - VelocityHistory[LastIndex]));
+	VelocityHistory[LastIndex] = Velocity;
+
+	// return average
+	return (VelocityHistorySum / (float)VELOCITY_HISTORY_SAMPLE_COUNT);
+}
\ No newline at end of file
diff --git a/Engine/Source/Runtime/Engine/Public/WorldPartition/WorldPartition.h b/Engine/Source/Runtime/Engine/Public/WorldPartition/WorldPartition.h
--- a/Engine/Source/Runtime/Engine/Public/WorldPartition/WorldPartition.h
+++ b/Engine/Source/Runtime/Engine/Public/WorldPartition/WorldPartition.h
@@ -519,5 +519,5 @@ private:
 	class AWorldPartitionReplay* Replay;
 
 	friend class AWorldPartitionReplay;
-	friend class UWorldPartitionStreamingPolicy;
+	friend class UWorldPartitionSubsystem;
 };
\ No newline at end of file
diff --git a/Engine/Source/Runtime/Engine/Public/WorldPartition/WorldPartitionStreamingPolicy.h b/Engine/Source/Runtime/Engine/Public/WorldPartition/WorldPartitionStreamingPolicy.h
--- a/Engine/Source/Runtime/Engine/Public/WorldPartition/WorldPartitionStreamingPolicy.h
+++ b/Engine/Source/Runtime/Engine/Public/WorldPartition/WorldPartitionStreamingPolicy.h
@@ -18,27 +18,6 @@
 
 class UWorldPartition;
 
-/**
- * Helper to compute streaming source velocity based on position history.
- */
-struct FStreamingSourceVelocity
-{
-	FStreamingSourceVelocity(const FName& InSourceName);
-	void Invalidate() { bIsValid = false; }
-	bool IsValid() { return bIsValid; }
-	float GetAverageVelocity(const FVector& NewPosition, const float CurrentTime);
-
-private:
-	enum { VELOCITY_HISTORY_SAMPLE_COUNT = 16 };
-	bool bIsValid;
-	FName SourceName;
-	int32 LastIndex;
-	float LastUpdateTime;
-	FVector LastPosition;
-	float VelocitiesHistorySum;
-	TArray<float, TInlineAllocator<VELOCITY_HISTORY_SAMPLE_COUNT>> VelocitiesHistory;
-};
-
 UCLASS(Abstract, Within = WorldPartition)
 class UWorldPartitionStreamingPolicy : public UObject
 {
@@ -85,7 +64,7 @@ protected:
 	virtual void SetCellsStateToUnloaded(const TArray<const UWorldPartitionRuntimeCell*>& ToUnloadCells);
 	virtual int32 GetCellLoadingCount() const { return 0; }
 	virtual int32 GetMaxCellsToLoad() const;
-	virtual void UpdateStreamingSources();
+	virtual void UpdateStreamingSources(bool bCanOptimizeUpdate);
 	void UpdateStreamingPerformance(const TSet<const UWorldPartitionRuntimeCell*>& InCells);
 	bool ShouldSkipCellForPerformance(const UWorldPartitionRuntimeCell* Cell) const;
 	bool IsInBlockTillLevelStreamingCompleted(bool bIsCausedByBadStreamingPerformance = false) const;
@@ -115,15 +94,13 @@ protected:
 
 	// Streaming Sources
 	TArray<FWorldPartitionStreamingSource> StreamingSources;
-	TMap<FName, FStreamingSourceVelocity> StreamingSourcesVelocity;
 
 	TSet<const UWorldPartitionRuntimeCell*> FrameActivateCells;
 	TSet<const UWorldPartitionRuntimeCell*> FrameLoadCells;
 	
 private:
 	// Update optimization
-	uint32 ComputeUpdateStreamingHash() const;
-	uint32 ComputeStreamingSourceHash(const FWorldPartitionStreamingSource& Source) const;
+	uint32 ComputeUpdateStreamingHash(bool bCanOptimizeUpdate) const;
 	int32 ComputeServerStreamingEnabledEpoch() const;
 
 	const TSet<FName>& GetServerDisallowedStreamingOutDataLayers();
@@ -132,12 +109,8 @@ private:
 
 	// CVars to control update optimization
 	static bool IsUpdateOptimEnabled;
-	static int32 LocationQuantization;
-	static int32 RotationQuantization;
 	static int32 ForceUpdateFrameCount;
 	static FAutoConsoleVariableRef CVarUpdateOptimEnabled;
-	static FAutoConsoleVariableRef CVarLocationQuantization;
-	static FAutoConsoleVariableRef CVarRotationQuantization;
 	static FAutoConsoleVariableRef CVarForceUpdateFrameCount;
 
 	bool bLastUpdateCompletedLoadingAndActivation;
@@ -147,6 +120,7 @@ private:
 	int32 ContentBundleServerEpoch;
 	int32 ServerStreamingEnabledEpoch;
 	uint32 UpdateStreamingHash;
+	uint32 UpdateStreamingSourcesHash;
 	uint32 UpdateStreamingStateCalls;
 
 	TOptional<TSet<FName>> CachedServerDisallowStreamingOutDataLayers;
diff --git a/Engine/Source/Runtime/Engine/Public/WorldPartition/WorldPartitionStreamingSource.h b/Engine/Source/Runtime/Engine/Public/WorldPartition/WorldPartitionStreamingSource.h
--- a/Engine/Source/Runtime/Engine/Public/WorldPartition/WorldPartitionStreamingSource.h
+++ b/Engine/Source/Runtime/Engine/Public/WorldPartition/WorldPartitionStreamingSource.h
@@ -3,6 +3,7 @@
 
 #include "Math/Color.h"
 #include "Math/RandomStream.h"
+#include "HAL/IConsoleManager.h"
 #include "WorldPartitionStreamingSource.generated.h"
 
 /** See https://en.wikipedia.org/wiki/Spherical_sector. */
@@ -382,6 +383,8 @@ struct ENGINE_API FWorldPartitionStreamingSource
 		, TargetBehavior(EStreamingSourceTargetBehavior::Include)
 		, bReplay(false)
 		, bRemote(false)
+		, Hash2D(0)
+		, Hash3D(0)
 	{}
 
 	FWorldPartitionStreamingSource(FName InName, const FVector& InLocation, const FRotator& InRotation, EStreamingSourceTargetState InTargetState, bool bInBlockOnSlowLoading, EStreamingSourcePriority InPriority, bool bRemote, float InVelocity = 0.f)
@@ -396,6 +399,8 @@ struct ENGINE_API FWorldPartitionStreamingSource
 		, TargetBehavior(EStreamingSourceTargetBehavior::Include)
 		, bReplay(false)
 		, bRemote(bRemote)
+		, Hash2D(0)
+		, Hash3D(0)
 	{}
 
 	// Define Copy Constructor to avoid deprecation warnings
@@ -420,6 +425,8 @@ struct ENGINE_API FWorldPartitionStreamingSource
 		Shapes = Other.Shapes;
 		bReplay = Other.bReplay;
 		bRemote = Other.bRemote;
+		Hash2D = Other.Hash2D;
+		Hash3D = Other.Hash3D;
 
 		PRAGMA_DISABLE_DEPRECATION_WARNINGS
 		TargetGrid = Other.TargetGrid;
@@ -439,6 +446,9 @@ struct ENGINE_API FWorldPartitionStreamingSource
 		return FColor(DebugColor.R, DebugColor.G, DebugColor.B, 255);
 	}
 
+	void UpdateHash();
+	uint32 GetHash(bool bInclude3DInformation = true) const { return bInclude3DInformation ? Hash3D : Hash2D; }
+
 	/** Source unique name. */
 	FName Name;
 
@@ -532,6 +542,20 @@ struct ENGINE_API FWorldPartitionStreamingSource
 			(int32)(Velocity*2.23694f)
 		);
 	}
+
+	static int32 GetLocationQuantization() { return LocationQuantization; }
+	static int32 GetRotationQuantization() { return RotationQuantization; }
+
+private:
+
+	static int32 LocationQuantization;
+	static int32 RotationQuantization;
+	static FAutoConsoleVariableRef CVarLocationQuantization;
+	static FAutoConsoleVariableRef CVarRotationQuantization;
+
+	/** Hash of streaming source (used to detect changes) */
+	uint32 Hash2D;
+	uint32 Hash3D;
 };
 
 /**
diff --git a/Engine/Source/Runtime/Engine/Public/WorldPartition/WorldPartitionSubsystem.h b/Engine/Source/Runtime/Engine/Public/WorldPartition/WorldPartitionSubsystem.h
--- a/Engine/Source/Runtime/Engine/Public/WorldPartition/WorldPartitionSubsystem.h
+++ b/Engine/Source/Runtime/Engine/Public/WorldPartition/WorldPartitionSubsystem.h
@@ -17,6 +17,27 @@ class FWorldPartitionActorDesc;
 
 enum class EWorldPartitionRuntimeCellState : uint8;
 
+/**
+ * Helper to compute streaming source velocity based on position history.
+ */
+struct FStreamingSourceVelocity
+{
+	FStreamingSourceVelocity(const FName& InSourceName);
+	void Invalidate() { bIsValid = false; }
+	bool IsValid() { return bIsValid; }
+	float GetAverageVelocity(const FVector& NewPosition, const float CurrentTime);
+
+private:
+	enum { VELOCITY_HISTORY_SAMPLE_COUNT = 16 };
+	bool bIsValid;
+	FName SourceName;
+	int32 LastIndex;
+	float LastUpdateTime;
+	FVector LastPosition;
+	float VelocityHistorySum;
+	TArray<float, TInlineAllocator<VELOCITY_HISTORY_SAMPLE_COUNT>> VelocityHistory;
+};
+
 /**
  * UWorldPartitionSubsystem
  */
@@ -132,6 +153,11 @@ protected:
 
 private:
 
+	// Streaming Sources
+	void UpdateStreamingSources();
+	void GetStreamingSources(const UWorldPartition* InWorldPartition, TArray<FWorldPartitionStreamingSource>& OutStreamingSources) const;
+	uint32 GetStreamingSourcesHash() const { return StreamingSourcesHash; }
+
 	void OnWorldPartitionInitialized(UWorldPartition* InWorldPartition);
 	void OnWorldPartitionUninitialized(UWorldPartition* InWorldPartition);
 
@@ -139,6 +165,7 @@ private:
 	const UWorldPartition* GetWorldPartition() const;
 	void Draw(class UCanvas* Canvas, class APlayerController* PC);
 	friend class UWorldPartition;
+	friend class UWorldPartitionStreamingPolicy;
 
 	TArray<TObjectPtr<UWorldPartition>> RegisteredWorldPartitions;
 
@@ -146,6 +173,11 @@ private:
 
 	FWorldPartitionStreamingSourceProviderFilter IsStreamingSourceProviderFiltered;
 
+	// Streaming Sources
+	TArray<FWorldPartitionStreamingSource> StreamingSources;
+	TMap<FName, FStreamingSourceVelocity> StreamingSourcesVelocity;
+	uint32 StreamingSourcesHash;
+
 	FDelegateHandle	DrawHandle;
 
 	// GC backup values
