OnlyEditorOnly: Move ICookInfo.h into CoreUObject so new scopes that will be added into it can be used by LinkerLoad and AsyncLoading.
#rnx

diff --git a/Engine/Source/Editor/UnrealEd/Classes/CookOnTheSide/CookOnTheFlyServer.h b/Engine/Source/Editor/UnrealEd/Classes/CookOnTheSide/CookOnTheFlyServer.h
--- a/Engine/Source/Editor/UnrealEd/Classes/CookOnTheSide/CookOnTheFlyServer.h
+++ b/Engine/Source/Editor/UnrealEd/Classes/CookOnTheSide/CookOnTheFlyServer.h
@@ -6,7 +6,6 @@
 #include "Containers/RingBuffer.h"
 #include "CookOnTheSide/CookLog.h"
 #include "CookPackageSplitter.h"
-#include "Engine/ICookInfo.h"
 #include "HAL/PlatformMemory.h"
 #include "INetworkFileSystemModule.h"
 #include "IPlatformFileSandboxWrapper.h"
@@ -17,6 +16,7 @@
 #include "Templates/RefCounting.h"
 #include "Templates/UniquePtr.h"
 #include "TickableEditorObject.h"
+#include "UObject/ICookInfo.h"
 #include "UObject/Object.h"
 #include "UObject/ObjectHandleTracking.h"
 #include "UObject/ObjectMacros.h"
diff --git a/Engine/Source/Editor/UnrealEd/Private/Cooker/CookPackageData.h b/Engine/Source/Editor/UnrealEd/Private/Cooker/CookPackageData.h
--- a/Engine/Source/Editor/UnrealEd/Private/Cooker/CookPackageData.h
+++ b/Engine/Source/Editor/UnrealEd/Private/Cooker/CookPackageData.h
@@ -11,7 +11,6 @@
 #include "Containers/SortedMap.h"
 #include "Containers/StringFwd.h"
 #include "Containers/UnrealString.h"
-#include "Engine/ICookInfo.h"
 #include "HAL/CriticalSection.h"
 #include "HAL/Platform.h"
 #include "Misc/EnumClassFlags.h"
@@ -21,6 +20,7 @@
 #include "Templates/UniquePtr.h"
 #include "TypedBlockAllocator.h"
 #include "UObject/GCObject.h"
+#include "UObject/ICookInfo.h"
 #include "UObject/NameTypes.h"
 #include "UObject/PackageResourceManager.h"
 #include "UObject/WeakObjectPtr.h"
diff --git a/Engine/Source/Editor/UnrealEd/Private/Cooker/CookRequestCluster.h b/Engine/Source/Editor/UnrealEd/Private/Cooker/CookRequestCluster.h
--- a/Engine/Source/Editor/UnrealEd/Private/Cooker/CookRequestCluster.h
+++ b/Engine/Source/Editor/UnrealEd/Private/Cooker/CookRequestCluster.h
@@ -8,7 +8,6 @@
 #include "Containers/Set.h"
 #include "Containers/StringFwd.h"
 #include "Containers/UnrealString.h"
-#include "Engine/ICookInfo.h"
 #include "HAL/Event.h"
 #include "HAL/Platform.h"
 #include "HAL/PlatformCrt.h"
@@ -21,6 +20,7 @@
 #include "HAL/CriticalSection.h"
 #include "TargetDomain/TargetDomainUtils.h"
 #include "Templates/UniquePtr.h"
+#include "UObject/ICookInfo.h"
 #include "UObject/NameTypes.h"
 
 #include <atomic>
diff --git a/Engine/Source/Editor/UnrealEd/Private/Cooker/CookRequests.h b/Engine/Source/Editor/UnrealEd/Private/Cooker/CookRequests.h
--- a/Engine/Source/Editor/UnrealEd/Private/Cooker/CookRequests.h
+++ b/Engine/Source/Editor/UnrealEd/Private/Cooker/CookRequests.h
@@ -9,11 +9,11 @@
 #include "Containers/RingBuffer.h"
 #include "Containers/UnrealString.h"
 #include "CookTypes.h"
-#include "Engine/ICookInfo.h"
 #include "HAL/CriticalSection.h"
 #include "HAL/PlatformCrt.h"
 #include "HAL/PlatformProcess.h"
 #include "Templates/UniquePtr.h"
+#include "UObject/ICookInfo.h"
 #include "UObject/NameTypes.h"
 
 #include <atomic>
diff --git a/Engine/Source/Runtime/CoreUObject/Private/UObject/ICookInfo.cpp b/Engine/Source/Runtime/CoreUObject/Private/UObject/ICookInfo.cpp
new file mode 100644
--- /dev/null
+++ b/Engine/Source/Runtime/CoreUObject/Private/UObject/ICookInfo.cpp
@@ -0,0 +1,50 @@
+// Copyright Epic Games, Inc. All Rights Reserved.
+
+#include "UObject/ICookInfo.h"
+
+#if WITH_EDITOR
+
+
+namespace UE::Cook
+{
+
+const TCHAR* LexToString(EInstigator Value)
+{
+	switch (Value)
+	{
+#define EINSTIGATOR_VALUE_CALLBACK(Name, bAllowUnparameterized) case EInstigator::Name: return TEXT(#Name);
+		EINSTIGATOR_VALUES(EINSTIGATOR_VALUE_CALLBACK)
+#undef EINSTIGATOR_VALUE_CALLBACK
+	default: return TEXT("OutOfRangeCategory");
+	}
+}
+
+FString FInstigator::ToString() const
+{
+	TStringBuilder<256> Result;
+	Result << LexToString(Category);
+	if (!Referencer.IsNone())
+	{
+		Result << TEXT(": ") << Referencer;
+	}
+	else
+	{
+		bool bCategoryAllowsUnparameterized = false;
+		switch (Category)
+		{
+#define EINSTIGATOR_VALUE_CALLBACK(Name, bAllowUnparameterized) case EInstigator::Name: bCategoryAllowsUnparameterized = bAllowUnparameterized; break;
+			EINSTIGATOR_VALUES(EINSTIGATOR_VALUE_CALLBACK)
+#undef EINSTIGATOR_VALUE_CALLBACK
+		default: break;
+		}
+		if (!bCategoryAllowsUnparameterized)
+		{
+			Result << TEXT(": <NoReferencer>");
+		}
+	}
+	return FString(Result);
+}
+
+}
+
+#endif
diff --git a/Engine/Source/Runtime/CoreUObject/Public/UObject/ICookInfo.h b/Engine/Source/Runtime/CoreUObject/Public/UObject/ICookInfo.h
new file mode 100644
--- /dev/null
+++ b/Engine/Source/Runtime/CoreUObject/Public/UObject/ICookInfo.h
@@ -0,0 +1,99 @@
+// Copyright Epic Games, Inc. All Rights Reserved.
+
+#pragma once
+
+#include "Containers/Array.h"
+#include "Containers/UnrealString.h"
+#include "HAL/Platform.h"
+#include "UObject/NameTypes.h"
+
+#if WITH_EDITOR
+
+namespace UE::Cook
+{
+
+/**
+ * List the keywords that should be enum values of the EInstigator enum
+ * The list is declared as a macro instead of ordinary C keywords in the enum declaration
+ * so that we can reduce duplication in the functions that specify the string name and other properties for each
+ * enum value; see the uses in ICookInfo.cpp.
+ */
+#define EINSTIGATOR_VALUES(CallbackMacro) \
+	/** CallbackMacro(CPPToken Name, bool bAllowUnparameterized) */ \
+	CallbackMacro(InvalidCategory, true) \
+	CallbackMacro(NotYetRequested, true) \
+	CallbackMacro(Unspecified, false) \
+	CallbackMacro(StartupPackage, true) \
+	CallbackMacro(AlwaysCookMap, true) \
+	CallbackMacro(IniMapSection, false) \
+	CallbackMacro(IniAllMaps, true) \
+	CallbackMacro(CommandLinePackage, true) \
+	CallbackMacro(CommandLineDirectory, true) \
+	CallbackMacro(DirectoryToAlwaysCook, false) \
+	CallbackMacro(FullDepotSearch, true) \
+	CallbackMacro(GameDefaultObject, false) \
+	CallbackMacro(InputSettingsIni, true) \
+	CallbackMacro(StartupSoftObjectPath, true) \
+	CallbackMacro(PackagingSettingsMapToCook, true) \
+	CallbackMacro(CookModificationDelegate, true) \
+	CallbackMacro(ModifyCookDelegate, true) \
+	CallbackMacro(AssetManagerModifyCook, true) \
+	CallbackMacro(AssetManagerModifyDLCCook, true) \
+	CallbackMacro(TargetPlatformExtraPackagesToCook, true) \
+	CallbackMacro(ConsoleCommand, true) \
+	CallbackMacro(CookOnTheFly, true) \
+	CallbackMacro(IterativeCook, true) \
+	CallbackMacro(PreviousAssetRegistry, true) \
+	CallbackMacro(RequestPackageFunction, true) \
+	CallbackMacro(Dependency, false) \
+	CallbackMacro(HardDependency, false) \
+	CallbackMacro(SoftDependency, false) \
+	CallbackMacro(Unsolicited, false) \
+	CallbackMacro(GeneratedPackage, false) \
+
+
+/** The different ways a package can be discovered by the cooker. */
+enum class EInstigator : uint8
+{
+#define EINSTIGATOR_VALUE_CALLBACK(Name, bAllowUnparameterized) Name,
+	EINSTIGATOR_VALUES(EINSTIGATOR_VALUE_CALLBACK)
+#undef EINSTIGATOR_VALUE_CALLBACK
+	Count,
+};
+COREUOBJECT_API const TCHAR* LexToString(EInstigator Value);
+
+/** Category and referencer for how a package was discovered by the cooker. */
+struct FInstigator
+{
+	FName Referencer;
+	EInstigator Category;
+
+	FInstigator() : Category(EInstigator::InvalidCategory)
+	{
+	}
+	FInstigator(EInstigator InCategory, FName InReferencer = NAME_None) : Referencer(InReferencer), Category(InCategory)
+	{
+	}
+
+	COREUOBJECT_API FString ToString() const;
+};
+
+/** Engine interface for information provided by UCookOnTheFlyServer in cook callbacks. */
+class ICookInfo
+{
+public:
+	/**
+	 * Return the instigator that first caused the package to be requested by the cook.
+	 * Returns category EInstigator::NotYetRequested if package is not yet known to the cook.
+	 */
+	virtual FInstigator GetInstigator(FName PackageName) = 0;
+	/**
+	 * Return the chain of instigators that caused the package to be requested by the cook.
+	 * First element is the direct instigator of the package, last is the root instigator that started the chain.
+	 */
+	virtual TArray<FInstigator> GetInstigatorChain(FName PackageName) = 0;
+};
+
+}
+
+#endif
\ No newline at end of file
diff --git a/Engine/Source/Runtime/Engine/Private/AssetManager.cpp b/Engine/Source/Runtime/Engine/Private/AssetManager.cpp
--- a/Engine/Source/Runtime/Engine/Private/AssetManager.cpp
+++ b/Engine/Source/Runtime/Engine/Private/AssetManager.cpp
@@ -13,7 +13,6 @@
 #include "Engine/BlueprintCore.h"
 #include "Engine/BlueprintGeneratedClass.h"
 #include "Engine/GameInstance.h"
-#include "Engine/ICookInfo.h"
 #include "Engine/StreamableManager.h"
 #include "HAL/PlatformApplicationMisc.h"
 #include "HAL/PlatformFileManager.h"
@@ -31,6 +30,7 @@
 #include "Stats/StatsMisc.h"
 #include "String/Find.h"
 #include "UObject/ConstructorHelpers.h"
+#include "UObject/ICookInfo.h"
 #include "UObject/LinkerLoad.h"
 #include "UObject/ObjectSaveContext.h"
 
diff --git a/Engine/Source/Runtime/Engine/Private/Engine/ICookInfo.cpp b/Engine/Source/Runtime/Engine/Private/Engine/ICookInfo.cpp
deleted file mode 100644
diff --git a/Engine/Source/Runtime/Engine/Public/Engine/ICookInfo.h b/Engine/Source/Runtime/Engine/Public/Engine/ICookInfo.h
--- a/Engine/Source/Runtime/Engine/Public/Engine/ICookInfo.h
+++ b/Engine/Source/Runtime/Engine/Public/Engine/ICookInfo.h
@@ -2,98 +2,7 @@
 
 #pragma once
 
-#include "Containers/Array.h"
-#include "Containers/UnrealString.h"
-#include "HAL/Platform.h"
-#include "UObject/NameTypes.h"
+// HEADER_UNIT_SKIP - Deprecated
 
-#if WITH_EDITOR
-
-namespace UE::Cook
-{
-
-/**
- * List the keywords that should be enum values of the EInstigator enum
- * The list is declared as a macro instead of ordinary C keywords in the enum declaration
- * so that we can reduce duplication in the functions that specify the string name and other properties for each
- * enum value; see the uses in ICookInfo.cpp.
- */
-#define EINSTIGATOR_VALUES(CallbackMacro) \
-	/** CallbackMacro(CPPToken Name, bool bAllowUnparameterized) */ \
-	CallbackMacro(InvalidCategory, true) \
-	CallbackMacro(NotYetRequested, true) \
-	CallbackMacro(Unspecified, false) \
-	CallbackMacro(StartupPackage, true) \
-	CallbackMacro(AlwaysCookMap, true) \
-	CallbackMacro(IniMapSection, false) \
-	CallbackMacro(IniAllMaps, true) \
-	CallbackMacro(CommandLinePackage, true) \
-	CallbackMacro(CommandLineDirectory, true) \
-	CallbackMacro(DirectoryToAlwaysCook, false) \
-	CallbackMacro(FullDepotSearch, true) \
-	CallbackMacro(GameDefaultObject, false) \
-	CallbackMacro(InputSettingsIni, true) \
-	CallbackMacro(StartupSoftObjectPath, true) \
-	CallbackMacro(PackagingSettingsMapToCook, true) \
-	CallbackMacro(CookModificationDelegate, true) \
-	CallbackMacro(ModifyCookDelegate, true) \
-	CallbackMacro(AssetManagerModifyCook, true) \
-	CallbackMacro(AssetManagerModifyDLCCook, true) \
-	CallbackMacro(TargetPlatformExtraPackagesToCook, true) \
-	CallbackMacro(ConsoleCommand, true) \
-	CallbackMacro(CookOnTheFly, true) \
-	CallbackMacro(IterativeCook, true) \
-	CallbackMacro(PreviousAssetRegistry, true) \
-	CallbackMacro(RequestPackageFunction, true) \
-	CallbackMacro(Dependency, false) \
-	CallbackMacro(HardDependency, false) \
-	CallbackMacro(SoftDependency, false) \
-	CallbackMacro(Unsolicited, false) \
-	CallbackMacro(GeneratedPackage, false) \
-
-
-/** The different ways a package can be discovered by the cooker. */
-enum class EInstigator : uint8
-{
-#define EINSTIGATOR_VALUE_CALLBACK(Name, bAllowUnparameterized) Name,
-	EINSTIGATOR_VALUES(EINSTIGATOR_VALUE_CALLBACK)
-#undef EINSTIGATOR_VALUE_CALLBACK
-	Count,
-};
-ENGINE_API const TCHAR* LexToString(EInstigator Value);
-
-/** Category and referencer for how a package was discovered by the cooker. */
-struct FInstigator
-{
-	FName Referencer;
-	EInstigator Category;
-
-	FInstigator() : Category(EInstigator::InvalidCategory)
-	{
-	}
-	FInstigator(EInstigator InCategory, FName InReferencer = NAME_None) : Referencer(InReferencer), Category(InCategory)
-	{
-	}
-
-	ENGINE_API FString ToString() const;
-};
-
-/** Engine interface for information provided by UCookOnTheFlyServer in cook callbacks. */
-class ICookInfo
-{
-public:
-	/**
-	 * Return the instigator that first caused the package to be requested by the cook.
-	 * Returns category EInstigator::NotYetRequested if package is not yet known to the cook.
-	 */
-	virtual FInstigator GetInstigator(FName PackageName) = 0;
-	/**
-	 * Return the chain of instigators that caused the package to be requested by the cook.
-	 * First element is the direct instigator of the package, last is the root instigator that started the chain.
-	 */
-	virtual TArray<FInstigator> GetInstigatorChain(FName PackageName) = 0;
-};
-
-}
-
-#endif
\ No newline at end of file
+// Deprecated in UE 5.3.
+#pragma message (__FILE__"(8): warning: #include UObject/ICookInfo.h instead of Engine/ICookInfo.h")
