fix bug causing synchronous texture builds to fail
was wrongly releasing memory of source blocks even if build was not tried
add concept of tried and failed so we don't try to build twice

#rnx

diff --git a/Engine/Source/Runtime/Engine/Private/TextureDerivedData.cpp b/Engine/Source/Runtime/Engine/Private/TextureDerivedData.cpp
--- a/Engine/Source/Runtime/Engine/Private/TextureDerivedData.cpp
+++ b/Engine/Source/Runtime/Engine/Private/TextureDerivedData.cpp
@@ -1388,6 +1388,10 @@ void FTexturePlatformData::Cache(
 			OutResultMetadataPerLayerFetchFirst, 
 			OutResultMetadataPerLayerFetchOrBuild,
 			Flags);
+
+		// LocalTask->TextureData Init may have failed and have bValid = false
+		//	but we still go ahead and create the async task, perhaps wrongly so
+
 		AsyncTask = LocalTask;
 		LocalTask->StartBackgroundTask(TextureThreadPool, BasePriority, EQueuedWorkFlags::DoNotRunInsideBusyWait, LocalTask->GetTask().GetRequiredMemoryEstimate(), TEXT("TextureDerivedData"));
 	}
diff --git a/Engine/Source/Runtime/Engine/Private/TextureDerivedDataTask.cpp b/Engine/Source/Runtime/Engine/Private/TextureDerivedDataTask.cpp
--- a/Engine/Source/Runtime/Engine/Private/TextureDerivedDataTask.cpp
+++ b/Engine/Source/Runtime/Engine/Private/TextureDerivedDataTask.cpp
@@ -175,6 +175,8 @@ static bool ValidateTexture2DPlatformData(const FTexturePlatformData& TextureDat
 
 void FTextureSourceData::Init(UTexture& InTexture, TextureMipGenSettings InMipGenSettings, bool bInCubeMap, bool bInTextureArray, bool bInVolumeTexture, bool bAllowAsyncLoading)
 {
+	check( bValid == false ); // we set to true at the end, acts as our return value
+
 	const int32 NumBlocks = InTexture.Source.GetNumBlocks();
 	const int32 NumLayers = InTexture.Source.GetNumLayers();
 	if (NumBlocks < 1 || NumLayers < 1)
@@ -230,11 +232,13 @@ void FTextureSourceData::Init(UTexture& InTexture, TextureMipGenSettings InMipGe
 
 	for (FTextureSourceBlockData& Block : Blocks)
 	{
+		// for the common case of NumBlocks == 1, BlockSizeX == Block.SizeX, MipBiasX/Y will both be zero
 		const int32 MipBiasX = FMath::CeilLogTwo(BlockSizeX / Block.SizeX);
 		const int32 MipBiasY = FMath::CeilLogTwo(BlockSizeY / Block.SizeY);
 		if (MipBiasX != MipBiasY)
 		{
-			UE_LOG(LogTexture, Warning, TEXT("Texture has blocks with mismatched aspect ratios"), *InTexture.GetPathName());
+			// @todo Oodle: this is failing even if "pad to pow2 square" is set, can we allow it through in that case?
+			UE_LOG(LogTexture, Warning, TEXT("VT has blocks with mismatched aspect ratios, cannot build."), *InTexture.GetPathName());  // <- should be an Error, not a Warning
 			return;
 		}
 
@@ -624,7 +628,7 @@ static void DDC1_BuildTexture(
 
 	FTexturePlatformData* DerivedData,
 	int64& BytesCached,
-	bool& bSucceeded	// currently completely ignored in the main build path, but used in the synchronous path
+	bool& bSucceeded
 	)
 {
 	TRACE_CPUPROFILER_EVENT_SCOPE(FTextureCacheDerivedDataWorker::BuildTexture);
@@ -1924,6 +1928,8 @@ void FTextureCacheDerivedDataWorker::DoWork()
 		}
 	}
 
+	check( ! bTriedAndFailed );
+
 	if (!bSucceeded && bAllowAsyncBuild)
 	{
 		if (DDC1_LoadAndValidateTextureData(Texture, TextureData, CompositeTextureData, ImageWrapper, bAllowAsyncLoading))
@@ -1952,17 +1958,31 @@ void FTextureCacheDerivedDataWorker::DoWork()
 					Chunk.bCorruptDataLoadedFromDDC = true;
 				}
 			}
+
+			if ( ! bSucceeded )
+			{
+				bTriedAndFailed = true;
+			}
 		}
 		else
 		{
 			bSucceeded = false;
+
+			// bTriedAndFailed = true; // no retry in Finalize ?  @todo ?
 		}
 	}
 
-	TextureData.ReleaseMemory();
-	CompositeTextureData.ReleaseMemory();
+	// there are actually 3 states to bSucceeded
+	//	tried & succeeded, tried & failed, not tried yet
+	// we may try the build again in Finalize (eg. if !bAllowAsyncBuild)
+
+	if (bSucceeded || bTriedAndFailed)
+	{
+		TextureData.ReleaseMemory();
+		CompositeTextureData.ReleaseMemory();
+	}
 
-	if (bSucceeded)
+	if ( bSucceeded )
 	{
 		// Populate the VT DDC Cache now if we're asynchronously loading to avoid too many high prio/synchronous request on the render thread
 		if (!IsInGameThread() && DerivedData->VTData && !DerivedData->VTData->Chunks.Last().DerivedDataKey.IsEmpty())
@@ -1980,6 +2000,13 @@ void FTextureCacheDerivedDataWorker::Finalize()
 	//		in the editor and is missing for some reason.
 	// --	always with a lighting build, as the async light/shadowmap tasks will disallow async builds. 
 	// --	if the texture compiler cvar disallows async texture compilation
+
+	if ( bTriedAndFailed )
+	{
+		UE_LOG(LogTexture, Warning, TEXT("Texture build failed for %s.  Will not retry in Finalize."), *TexturePathName);
+		return;
+	}
+
 	if (!bSucceeded)
 	{
 		if (!TextureData.HasPayload() && !Texture.Source.HasPayloadData())
@@ -2025,6 +2052,11 @@ void FTextureCacheDerivedDataWorker::Finalize()
 					TextureData, CompositeTextureData, BuildSettingsPerLayerFetchOrBuild, KeySuffix, false /* currently corrupt vt data is not routed out of DoWork() */,
 					RequiredMemoryEstimate, DerivedData, BytesCached, bSucceeded);
 			}
+
+			if ( ! bSucceeded )
+			{
+				bTriedAndFailed = true;
+			}
 		}
 	}
 		
diff --git a/Engine/Source/Runtime/Engine/Private/TextureDerivedDataTask.h b/Engine/Source/Runtime/Engine/Private/TextureDerivedDataTask.h
--- a/Engine/Source/Runtime/Engine/Private/TextureDerivedDataTask.h
+++ b/Engine/Source/Runtime/Engine/Private/TextureDerivedDataTask.h
@@ -166,6 +166,8 @@ class FTextureCacheDerivedDataWorker : public FNonAbandonableTask
 
 	/** true if caching has succeeded. */
 	bool bSucceeded = false;
+	/** true if caching was tried and failed; if bSucceeded & bTriedAndFailed are both false, it has not been tried yet */
+	bool bTriedAndFailed = false;
 	/** true if the derived data was pulled from DDC */
 	bool bLoadedFromDDC = false;
 public:
