Guard HISM AcceptOcclusionResults with a mutex.

diff --git a/Engine/Source/Runtime/Engine/Private/HierarchicalInstancedStaticMesh.cpp b/Engine/Source/Runtime/Engine/Private/HierarchicalInstancedStaticMesh.cpp
--- a/Engine/Source/Runtime/Engine/Private/HierarchicalInstancedStaticMesh.cpp
+++ b/Engine/Source/Runtime/Engine/Private/HierarchicalInstancedStaticMesh.cpp
@@ -2007,6 +2007,9 @@ void FHierarchicalStaticMeshSceneProxy::AcceptOcclusionResults(const FSceneView*
 	// Don't accept subprimitive occlusion results from a previously-created sceneproxy - the tree may have been different
 	if (OcclusionBounds.Num() == NumResults && SceneProxyCreatedFrameNumberRenderThread < GFrameNumberRenderThread)
 	{
+		// This lock is necessary to guard against access from multiple views.
+		OcclusionResultsMutex.Lock();
+
 		uint32 ViewId = View->GetViewKey();
 		FFoliageOcclusionResults* OldResults = OcclusionResults.Find(ViewId);
 		if (OldResults)
@@ -2028,6 +2031,8 @@ void FHierarchicalStaticMeshSceneProxy::AcceptOcclusionResults(const FSceneView*
 			}
 			OcclusionResults.Add(ViewId, FFoliageOcclusionResults(Results, ResultsStart, NumResults));
 		}
+
+		OcclusionResultsMutex.Unlock();
 	}
 }
 
diff --git a/Engine/Source/Runtime/Engine/Public/HierarchicalStaticMeshSceneProxy.h b/Engine/Source/Runtime/Engine/Public/HierarchicalStaticMeshSceneProxy.h
--- a/Engine/Source/Runtime/Engine/Public/HierarchicalStaticMeshSceneProxy.h
+++ b/Engine/Source/Runtime/Engine/Public/HierarchicalStaticMeshSceneProxy.h
@@ -9,6 +9,7 @@
 #include "Components/InstancedStaticMeshComponent.h"
 #include "StaticMeshResources.h"
 #include "Components/HierarchicalInstancedStaticMeshComponent.h"
+#include "Async/Mutex.h"
 
 struct FClusterNode;
 struct FFoliageElementParams;
@@ -44,6 +45,7 @@ class ENGINE_API FHierarchicalStaticMeshSceneProxy final : public FInstancedStat
 	int32 LastOcclusionNode;
 	TArray<FBoxSphereBounds> OcclusionBounds;
 	TMap<uint32, FFoliageOcclusionResults> OcclusionResults;
+	UE::FMutex OcclusionResultsMutex;
 	EHISMViewRelevanceType ViewRelevance;
 	bool bDitheredLODTransitions;
 	uint32 SceneProxyCreatedFrameNumberRenderThread;
