Allow the option to not change the StateId for UMaterial and UMaterialFunction during duplication.  This helps provide more sharing of DDC cache keys for identical materials.

diff --git a/Engine/Source/Runtime/Engine/Classes/Materials/MaterialFunctionInterface.h b/Engine/Source/Runtime/Engine/Classes/Materials/MaterialFunctionInterface.h
--- a/Engine/Source/Runtime/Engine/Classes/Materials/MaterialFunctionInterface.h
+++ b/Engine/Source/Runtime/Engine/Classes/Materials/MaterialFunctionInterface.h
@@ -80,6 +80,7 @@ public:
 
 	//~ Begin UObject Interface.
 	virtual void PostInitProperties() override;
+	virtual void PostDuplicate(bool bDuplicateForPIE) override;
 	virtual void PostLoad() override;
 #if WITH_EDITORONLY_DATA
 	ENGINE_API static void DeclareConstructClasses(TArray<FTopLevelAssetPath>& OutConstructClasses, const UClass* SpecificSubclass);
@@ -93,7 +94,7 @@ public:
 	//~ End UObject Interface.
 
 	/** Used by materials using this function to know when to recompile. */
-	UPROPERTY(duplicatetransient)
+	UPROPERTY()
 	FGuid StateId;
 
 protected:
diff --git a/Engine/Source/Runtime/Engine/Private/Materials/Material.cpp b/Engine/Source/Runtime/Engine/Private/Materials/Material.cpp
--- a/Engine/Source/Runtime/Engine/Private/Materials/Material.cpp
+++ b/Engine/Source/Runtime/Engine/Private/Materials/Material.cpp
@@ -129,6 +129,13 @@ static TAutoConsoleVariable<bool> CVarMaterialLogErrorOnFailure(
 	TEXT("When enabled, when a material fails to compile it will issue an Error instead of a Warning.\n")
 	TEXT("Default: false"),
 	ECVF_RenderThreadSafe);
+	
+static TAutoConsoleVariable<bool> CVarMaterialsDuplicateVerbatim(
+	TEXT("r.MaterialsDuplicateVerbatim"),
+	false,
+	TEXT("When enabled, when a material or material function is duplicated, it will not change StateId (which influences DDC keys) pre-emptively.\n")
+	TEXT("Default: false"),
+	ECVF_Default);
 
 namespace MaterialImpl
 {
@@ -2870,8 +2877,11 @@ void UMaterial::PostDuplicate(bool bDuplicateForPIE)
 {
 	Super::PostDuplicate(bDuplicateForPIE);
 
-	// Reset the StateId on duplication since it needs to be unique for each material.
-	FPlatformMisc::CreateGuid(StateId);
+	if (!CVarMaterialsDuplicateVerbatim.GetValueOnAnyThread())
+	{
+		// Reset the StateId on duplication since it needs to be unique for each material.
+		FPlatformMisc::CreateGuid(StateId);
+	}
 }
 
 void UMaterial::BackwardsCompatibilityInputConversion()
diff --git a/Engine/Source/Runtime/Engine/Private/Materials/MaterialExpressions.cpp b/Engine/Source/Runtime/Engine/Private/Materials/MaterialExpressions.cpp
--- a/Engine/Source/Runtime/Engine/Private/Materials/MaterialExpressions.cpp
+++ b/Engine/Source/Runtime/Engine/Private/Materials/MaterialExpressions.cpp
@@ -13806,6 +13806,19 @@ void UMaterialFunctionInterface::PostInitProperties()
 	StateId = FGuid::NewGuid();
 }
 
+void UMaterialFunctionInterface::PostDuplicate(bool bDuplicateForPIE)
+{
+	Super::PostDuplicate(bDuplicateForPIE);
+
+	static const auto CVarDuplicateVerbatim = IConsoleManager::Get().FindTConsoleVariableDataBool(TEXT("r.MaterialsDuplicateVerbatim"));
+	const bool bKeepStateId = StateId.IsValid() && HasAnyFlags(RF_WasLoaded) && CVarDuplicateVerbatim->GetValueOnAnyThread();
+	if (!bKeepStateId)
+	{
+		// Initialize StateId to something unique, in case this is a new function
+		StateId = FGuid::NewGuid();
+	}
+}
+
 void UMaterialFunctionInterface::PostLoad()
 {
 	Super::PostLoad();
