Static analysis logging fix

diff --git a/Engine/Source/Runtime/Engine/Private/UnrealEngine.cpp b/Engine/Source/Runtime/Engine/Private/UnrealEngine.cpp
--- a/Engine/Source/Runtime/Engine/Private/UnrealEngine.cpp
+++ b/Engine/Source/Runtime/Engine/Private/UnrealEngine.cpp
@@ -15795,7 +15795,7 @@ UWorld* UEngine::GetCurrentPlayWorld(UWorld* PossiblePlayWorld) const
 				else
 				{
 					// If you see this warning, game code is traversing PIE world boundaries in an unsafe way. That should be fixed or GPlayInEditorID needs to be set properly
-					UE_LOG(LogEngine, Warning, TEXT("GetCurrentPlayWorld failed with ambiguous PIE world! GPlayInEditorID %d does not match %s"), GPlayInEditorID, *PossiblePlayWorld->GetPathName());
+					UE_LOG(LogEngine, Warning, TEXT("GetCurrentPlayWorld failed with ambiguous PIE world! GPlayInEditorID %d does not match %s"), (int32)GPlayInEditorID, *PossiblePlayWorld->GetPathName());
 				}
 			}
 #endif
