More additional debug information for cooking anim sequence data soft-lock, logging out timestamps for BeginCache/ClearCache/WillNeverCach/ApplyCompressedData on a per-hash basis

diff --git a/Engine/Source/Runtime/Engine/Classes/Animation/AnimSequence.h b/Engine/Source/Runtime/Engine/Classes/Animation/AnimSequence.h
--- a/Engine/Source/Runtime/Engine/Classes/Animation/AnimSequence.h
+++ b/Engine/Source/Runtime/Engine/Classes/Animation/AnimSequence.h
@@ -882,9 +882,20 @@ protected:
 	bool PollCacheDerivedData(const FIoHash& KeyHash) const;
 	void EndCacheDerivedData(const FIoHash& KeyHash);
 
+	void LogCompressionRequestTimings() const;
+
 	FIoHash DataKeyHash;
 	TMap<FIoHash, TUniquePtr<FCompressedAnimSequence>> DataByPlatformKeyHash;
 	TMap<FIoHash, TPimplPtr<UE::Anim::FAnimationSequenceAsyncCacheTask>> CacheTasksByKeyHash;
+
+	struct FCompressionRequestTimings
+	{
+		FDateTime LastBeginCacheRequest;
+		FDateTime LastApplyCacheData;
+		FDateTime LastClearCache;
+		FDateTime LastWillNeverCook;
+	};
+	TMap<FIoHash, FCompressionRequestTimings> TimingsByKeyHash;
 protected:
 	bool TryCancelAsyncTasks();
 	void FinishAsyncTasks();
diff --git a/Engine/Source/Runtime/Engine/Private/Animation/AnimSequence.cpp b/Engine/Source/Runtime/Engine/Private/Animation/AnimSequence.cpp
--- a/Engine/Source/Runtime/Engine/Private/Animation/AnimSequence.cpp
+++ b/Engine/Source/Runtime/Engine/Private/Animation/AnimSequence.cpp
@@ -458,6 +458,11 @@ void UAnimSequence::WillNeverCacheCookedPlatformDataAgain()
 
 	// Clear out current platform, and any target platform data
 	CompressedData.Reset();
+	
+	for (TPair<FIoHash, FCompressionRequestTimings>& TimingsPair : TimingsByKeyHash)
+	{
+		TimingsPair.Value.LastWillNeverCook = FDateTime::Now();	
+	}
 		
 	CacheTasksByKeyHash.Empty();
 	DataByPlatformKeyHash.Empty();
@@ -476,6 +481,11 @@ void UAnimSequence::ClearAllCachedCookedPlatformData()
 	DataByPlatformKeyHash.Empty();
 	CompressedData.Reset();
 	DataKeyHash = FIoHash::Zero;
+
+	for (TPair<FIoHash, FCompressionRequestTimings>& TimingsPair : TimingsByKeyHash)
+	{
+		TimingsPair.Value.LastClearCache = FDateTime::Now();	
+	}
 	
 	PRAGMA_DISABLE_DEPRECATION_WARNINGS
 	bUseRawDataOnly = true;
@@ -1094,6 +1104,8 @@ bool UAnimSequence::IsCachedCookedPlatformDataLoaded(const ITargetPlatform* Targ
 	{
 		// This means cooked data is expected, but it it currently invalid and no task is running to generate it
 		UE_LOG(LogAnimation, Warning, TEXT("Expecting either valid compressed Animation Data, or an in-flight compression task for asset %s\nPlatform: %s\nHash: %s"), *GetPathName(), *TargetPlatform->DisplayName().ToString(), *LexToString(KeyHash));
+
+		LogCompressionRequestTimings();
 	};
 
 	if (KeyHash == DataKeyHash)
@@ -5088,6 +5100,8 @@ FIoHash UAnimSequence::BeginCacheDerivedData(const ITargetPlatform* TargetPlatfo
     {
     	TargetData = DataByPlatformKeyHash.Emplace(KeyHash, MakeUnique<FCompressedAnimSequence>()).Get();
     }
+	
+	TimingsByKeyHash.FindOrAdd(KeyHash).LastBeginCacheRequest = FDateTime::Now();
 	check(TargetData);
 	TargetData->Reset();
 
@@ -5159,7 +5173,20 @@ void UAnimSequence::EndCacheDerivedData(const FIoHash& KeyHash)
 	
 	UE::Anim::FAnimSequenceCompilingManager::Get().FinishCompilation({this});
 }
-	
+
+void UAnimSequence::LogCompressionRequestTimings() const
+{
+	for (const TPair<FIoHash, FCompressionRequestTimings>& TimingsPair : TimingsByKeyHash)
+	{
+		UE_LOG(LogAnimation, Warning, TEXT("Hash: %s\nBeginCache: %s\nApplyCache: %s\nClearCache: %s\nWillNeverCook: %s"),
+			*LexToString(TimingsPair.Key),
+			*TimingsPair.Value.LastBeginCacheRequest.ToString(),
+			*TimingsPair.Value.LastApplyCacheData.ToString(),
+			*TimingsPair.Value.LastClearCache.ToString(),
+			*TimingsPair.Value.LastWillNeverCook.ToString());
+	}
+}
+
 FCompressedAnimSequence& UAnimSequence::CacheDerivedData(const ITargetPlatform* TargetPlatform)
 {
 	const FIoHash KeyHash = BeginCacheDerivedData(TargetPlatform);
@@ -5265,6 +5292,8 @@ void UAnimSequence::FinishAsyncTasks()
 					bUseRawDataOnly = false;
                     PRAGMA_ENABLE_DEPRECATION_WARNINGS
 				}
+
+				TimingsByKeyHash.FindChecked(It->Key).LastApplyCacheData = FDateTime::Now();
 			}
 			else
 			{
