Ignore HLOD settings when setting up the cull distances

diff --git a/Engine/Source/Runtime/Engine/Classes/Components/HierarchicalInstancedStaticMeshComponent.h b/Engine/Source/Runtime/Engine/Classes/Components/HierarchicalInstancedStaticMeshComponent.h
--- a/Engine/Source/Runtime/Engine/Classes/Components/HierarchicalInstancedStaticMeshComponent.h
+++ b/Engine/Source/Runtime/Engine/Classes/Components/HierarchicalInstancedStaticMeshComponent.h
@@ -273,6 +273,7 @@ public:
 	bool IsAsyncBuilding() const { return bIsAsyncBuilding; }
 	bool IsTreeFullyBuilt() const { return !bIsOutOfDate; }
 	void GetTree(TArray<FClusterNode>& OutClusterTree);
+	FVector GetAverageScale() const;
 
 	/** Heuristic for the number of leaves in the tree **/
 	int32 DesiredInstancesPerLeaf();
diff --git a/Engine/Source/Runtime/Engine/Private/HierarchicalInstancedStaticMesh.cpp b/Engine/Source/Runtime/Engine/Private/HierarchicalInstancedStaticMesh.cpp
--- a/Engine/Source/Runtime/Engine/Private/HierarchicalInstancedStaticMesh.cpp
+++ b/Engine/Source/Runtime/Engine/Private/HierarchicalInstancedStaticMesh.cpp
@@ -3521,6 +3521,19 @@ void UHierarchicalInstancedStaticMeshComponent::GetTree(TArray<FClusterNode>& Ou
 	OutClusterTree = *ClusterTreePtr;
 }
 
+FVector UHierarchicalInstancedStaticMeshComponent::GetAverageScale() const
+{
+	const TArray<FClusterNode>& ClusterTree = *ClusterTreePtr;
+	if (ClusterTree.Num())
+	{
+		return FVector(ClusterTree[0].MinInstanceScale + (ClusterTree[0].MaxInstanceScale - ClusterTree[0].MinInstanceScale) / 2.0f);
+	}
+	else
+	{
+		return FVector::Zero();
+	}
+}
+
 
 void UHierarchicalInstancedStaticMeshComponent::GetNavigationPerInstanceTransforms(const FBox& AreaBox, TArray<FTransform>& InstanceData) const
 {
