Deprecated legacy gather mesh elements parallel tasks.

diff --git a/Engine/Source/Runtime/Engine/Private/Particles/ParticleSystemRender.cpp b/Engine/Source/Runtime/Engine/Private/Particles/ParticleSystemRender.cpp
--- a/Engine/Source/Runtime/Engine/Private/Particles/ParticleSystemRender.cpp
+++ b/Engine/Source/Runtime/Engine/Private/Particles/ParticleSystemRender.cpp
@@ -891,29 +891,6 @@ void FDynamicSpriteEmitterData::GetDynamicMeshElementsEmitter(const FParticleSys
 				}
 				{
 					SCOPE_CYCLE_COUNTER(STAT_FDynamicSpriteEmitterData_PerParticleWorkOrTasks);
-					if (Collector.ShouldUseTasks())
-					{
-						Collector.AddTask(
-							[this, SourceData, View, Proxy, Allocation, DynamicParameterAllocation, bSort, ParticleCount, NumVerticesPerParticleInBuffer, InstanceFactor]()
-							{
-								SCOPE_CYCLE_COUNTER(STAT_FDynamicSpriteEmitterData_GetDynamicMeshElementsEmitter_Task);
-								SCOPE_CYCLE_COUNTER(STAT_ParticlesOverview_RT_CNC);
-
-								FMemMark Mark(FMemStack::Get());
-								FParticleOrder* ParticleOrder = NULL;
-								if (bSort)
-								{
-									ParticleOrder = (FParticleOrder*)FMemStack::Get().Alloc(sizeof(FParticleOrder)* ParticleCount, alignof(FParticleOrder));
-									SortSpriteParticles(SourceData->SortMode, SourceData->bUseLocalSpace, SourceData->ActiveParticleCount, 
-										SourceData->DataContainer.ParticleData, SourceData->ParticleStride, SourceData->DataContainer.ParticleIndices,
-										View, Proxy->GetLocalToWorld(), ParticleOrder);
-								}
-								// Fill vertex buffers.
-								GetVertexAndIndexData(Allocation.Buffer, DynamicParameterAllocation.Buffer, NULL, ParticleOrder, View->ViewMatrices.GetViewOrigin(), Proxy->GetLocalToWorld(), InstanceFactor);
-							}
-						);
-					}
-					else
 					{
 						FParticleOrder* ParticleOrder = NULL;
 
@@ -1448,20 +1425,7 @@ void FDynamicMeshEmitterData::GetDynamicMeshElementsEmitter(const FParticleSyste
 
 			if (Allocation.IsValid() && (!bUsesDynamicParameter || DynamicParameterAllocation.IsValid()))
 			{
-				// Fill instance buffer.
-				if (Collector.ShouldUseTasks())
-				{
-					Collector.AddTask(
-						[this, View, Proxy, Allocation, DynamicParameterAllocation, PrevTransformBuffer, InstanceFactor]()
-						{
-							GetInstanceData(Allocation.Buffer, DynamicParameterAllocation.Buffer, PrevTransformBuffer, Proxy, View, InstanceFactor);
-						}
-					);
-				}
-				else
-				{
-					GetInstanceData(Allocation.Buffer, DynamicParameterAllocation.Buffer, PrevTransformBuffer, Proxy, View, InstanceFactor);
-				}
+				GetInstanceData(Allocation.Buffer, DynamicParameterAllocation.Buffer, PrevTransformBuffer, Proxy, View, InstanceFactor);
 			}
 
 			if (bGeneratePrevTransformBuffer)
diff --git a/Engine/Source/Runtime/Engine/Private/SceneManagement.cpp b/Engine/Source/Runtime/Engine/Private/SceneManagement.cpp
--- a/Engine/Source/Runtime/Engine/Private/SceneManagement.cpp
+++ b/Engine/Source/Runtime/Engine/Private/SceneManagement.cpp
@@ -313,19 +313,13 @@ FMeshBatchAndRelevance::FMeshBatchAndRelevance(const FMeshBatch& InMesh, const F
 	bRenderInMainPass = PrimitiveSceneProxy->ShouldRenderInMainPass();
 }
 
-static TAutoConsoleVariable<int32> CVarUseParallelGetDynamicMeshElementsTasks(
-	TEXT("r.UseParallelGetDynamicMeshElementsTasks"),
-	0,
-	TEXT("If > 0, and if FApp::ShouldUseThreadingForPerformance(), then parts of GetDynamicMeshElements will be done in parallel."));
-
 FMeshElementCollector::FMeshElementCollector(ERHIFeatureLevel::Type InFeatureLevel, FSceneRenderingBulkObjectAllocator& InBulkAllocator) :
 	OneFrameResources(InBulkAllocator),
 	PrimitiveSceneProxy(NULL),
 	DynamicIndexBuffer(nullptr),
 	DynamicVertexBuffer(nullptr),
 	DynamicReadBuffer(nullptr),
-	FeatureLevel(InFeatureLevel),
-	bUseAsyncTasks(FApp::ShouldUseThreadingForPerformance() && CVarUseParallelGetDynamicMeshElementsTasks.GetValueOnAnyThread() > 0)
+	FeatureLevel(InFeatureLevel)
 {	
 }
 
@@ -336,7 +330,6 @@ FMeshElementCollector::~FMeshElementCollector()
 
 void FMeshElementCollector::DeleteTemporaryProxies()
 {
-	check(!ParallelTasks.Num()); // We should have blocked on this already
 	for (int32 ProxyIndex = 0; ProxyIndex < TemporaryProxies.Num(); ProxyIndex++)
 	{
 		delete TemporaryProxies[ProxyIndex];
@@ -420,25 +413,6 @@ void FMeshElementCollector::AddViewMeshArrays(
 #endif
 }
 
-void FMeshElementCollector::ProcessTasks()
-{
-	check(IsInRenderingThread());
-	check(!ParallelTasks.Num() || bUseAsyncTasks);
-
-	if (ParallelTasks.Num())
-	{
-		QUICK_SCOPE_CYCLE_COUNTER(STAT_FMeshElementCollector_ProcessTasks);
-		TArray<TFunction<void()>, SceneRenderingAllocator>& LocalParallelTasks(ParallelTasks);
-		ParallelFor(ParallelTasks.Num(), 
-			[&LocalParallelTasks](int32 Index)
-			{
-				LocalParallelTasks[Index]();
-				LocalParallelTasks[Index] = {};
-			});
-		ParallelTasks.Empty();
-	}
-}
-
 void FMeshElementCollector::AddMesh(int32 ViewIndex, FMeshBatch& MeshBatch)
 {
 	DEFINE_LOG_CATEGORY_STATIC(FMeshElementCollector_AddMesh, Warning, All);
diff --git a/Engine/Source/Runtime/Engine/Public/SceneManagement.h b/Engine/Source/Runtime/Engine/Public/SceneManagement.h
--- a/Engine/Source/Runtime/Engine/Public/SceneManagement.h
+++ b/Engine/Source/Runtime/Engine/Public/SceneManagement.h
@@ -1931,23 +1931,21 @@ public:
 	{
 		return *OneFrameResources.Create<T>(Forward<ARGS>(Args)...);
 	}
-
+	
+	UE_DEPRECATED(5.3, "ShouldUseTasks has been deprecated.")
 	FORCEINLINE bool ShouldUseTasks() const
 	{
-		return bUseAsyncTasks;
-	}
-
-	FORCEINLINE void AddTask(TFunction<void()>&& Task)
-	{
-		ParallelTasks.Emplace(MoveTemp(Task));
+		return false;
 	}
+	
+	UE_DEPRECATED(5.3, "AddTask has been deprecated.")
+	FORCEINLINE void AddTask(TFunction<void()>&& Task) {}
 
-	FORCEINLINE void AddTask(const TFunction<void()>& Task)
-	{
-		ParallelTasks.Emplace(Task);
-	}
+	UE_DEPRECATED(5.3, "AddTask has been deprecated.")
+	FORCEINLINE void AddTask(const TFunction<void()>& Task) {}
 
-	void ProcessTasks();
+	UE_DEPRECATED(5.3, "ProcessTasks has been deprecated.")
+	void ProcessTasks() {}
 
 	FORCEINLINE ERHIFeatureLevel::Type GetFeatureLevel() const
 	{
@@ -2021,12 +2019,6 @@ protected:
 
 	ERHIFeatureLevel::Type FeatureLevel;
 
-	/** This is related to some cvars and FApp stuff and if true means calling code should use async tasks. */
-	const bool bUseAsyncTasks;
-
-	/** Tasks to wait for at the end of gathering dynamic mesh elements. */
-	TArray<TFunction<void()>, SceneRenderingAllocator> ParallelTasks;
-
 	/** Tracks dynamic primitive data for upload to GPU Scene for every view, when enabled. */
 	TArray<FGPUScenePrimitiveCollector*, TInlineAllocator<2, SceneRenderingAllocator>> DynamicPrimitiveCollectorPerView;
 
diff --git a/Engine/Source/Runtime/Renderer/Private/SceneVisibility.cpp b/Engine/Source/Runtime/Renderer/Private/SceneVisibility.cpp
--- a/Engine/Source/Runtime/Renderer/Private/SceneVisibility.cpp
+++ b/Engine/Source/Runtime/Renderer/Private/SceneVisibility.cpp
@@ -3656,7 +3656,6 @@ void FSceneRenderer::GatherDynamicMeshElements(
 			}
 		}
 	}
-	MeshCollector.ProcessTasks();
 }
 
 /**
diff --git a/Engine/Source/Runtime/Renderer/Private/ShadowSetup.cpp b/Engine/Source/Runtime/Renderer/Private/ShadowSetup.cpp
--- a/Engine/Source/Runtime/Renderer/Private/ShadowSetup.cpp
+++ b/Engine/Source/Runtime/Renderer/Private/ShadowSetup.cpp
@@ -1275,7 +1275,7 @@ struct FDynamicShadowsTaskData
 		, bStaticSceneOnly(AreAnyViewsStaticSceneOnly(Views))
 		, bRunningEarly(bInRunningEarly)
 		, bMultithreaded((FApp::ShouldUseThreadingForPerformance() || FForkProcessHelper::IsForkedMultithreadInstance()) && CVarParallelGatherShadowPrimitives.GetValueOnRenderThread() > 0)
-		, bMultithreadedCreateAndFilterShadows(bRunningEarly&& bMultithreaded&& GRHISupportsAsyncGetRenderQueryResult&& GParallelInitDynamicShadows > 0)
+		, bMultithreadedCreateAndFilterShadows(bRunningEarly && bMultithreaded && GRHISupportsAsyncGetRenderQueryResult && GParallelInitDynamicShadows > 0)
 	{
 		if (bMultithreadedCreateAndFilterShadows)
 		{
@@ -2507,8 +2507,6 @@ void FProjectedShadowInfo::GatherDynamicMeshElements(FSceneRenderer& Renderer, F
 		ShadowDepthView->SetDynamicMeshElementsShadowCullFrustum(&CasterOuterFrustum);
 		GatherDynamicMeshElementsArray(Renderer, DynamicIndexBuffer, DynamicVertexBuffer, DynamicReadBuffer, 
 			SubjectTranslucentPrimitives, ReusedViewsArray, DynamicSubjectTranslucentMeshElements, NumDynamicSubjectTranslucentMeshElements);
-
-		Renderer.MeshCollector.ProcessTasks();
 	}
 
 	SetupMeshDrawCommandsForProjectionStenciling(Renderer, InstanceCullingManager);
