OnlyEditorOnly support: Add system-specific declaration of dependencies and FCookLoadScopes.
#rnx

diff --git a/Engine/Plugins/Experimental/ChaosCloth/Source/Chaos/Private/ChaosCloth/ChaosClothVisualization.cpp b/Engine/Plugins/Experimental/ChaosCloth/Source/Chaos/Private/ChaosCloth/ChaosClothVisualization.cpp
--- a/Engine/Plugins/Experimental/ChaosCloth/Source/Chaos/Private/ChaosCloth/ChaosClothVisualization.cpp
+++ b/Engine/Plugins/Experimental/ChaosCloth/Source/Chaos/Private/ChaosCloth/ChaosClothVisualization.cpp
@@ -34,6 +34,7 @@
 #include "Engine/Canvas.h"  // For draw text
 #include "CanvasItem.h"     //
 #include "Engine/Engine.h"  //
+#include "UObject/ICookInfo.h"
 #endif  // #if WITH_EDITOR
 
 static int DrawSkinnedLattice = 0;
@@ -45,6 +46,7 @@ namespace Chaos
 		: Solver(InSolver)
 	{
 #if WITH_EDITOR
+		FCookLoadScope CookLoadScope(ECookLoadType::EditorOnly);
 		ClothMaterial = LoadObject<UMaterial>(nullptr, TEXT("/Engine/EditorMaterials/Cloth/CameraLitDoubleSided.CameraLitDoubleSided"), nullptr, LOAD_None, nullptr);  // LOAD_EditorOnly
 		ClothMaterialVertex = LoadObject<UMaterial>(nullptr, TEXT("/Engine/EditorMaterials/WidgetVertexColorMaterial"), nullptr, LOAD_None, nullptr);  // LOAD_EditorOnly
 		CollisionMaterial = LoadObject<UMaterial>(nullptr, TEXT("/Engine/EditorMaterials/PhAT_UnselectedMaterial"), nullptr, LOAD_None, nullptr);
diff --git a/Engine/Plugins/Experimental/Water/Source/Runtime/Private/WaterBodyComponent.cpp b/Engine/Plugins/Experimental/Water/Source/Runtime/Private/WaterBodyComponent.cpp
--- a/Engine/Plugins/Experimental/Water/Source/Runtime/Private/WaterBodyComponent.cpp
+++ b/Engine/Plugins/Experimental/Water/Source/Runtime/Private/WaterBodyComponent.cpp
@@ -28,6 +28,8 @@
 #include "Engine/StaticMesh.h"
 #include "DynamicMesh/DynamicMeshAABBTree3.h"
 #include "Operations/MeshPlaneCut.h"
+#include "UObject/ICookInfo.h"
+#include "UObject/ObjectSaveContext.h"
 
 #include UE_INLINE_GENERATED_CPP_BY_NAME(WaterBodyComponent)
 
@@ -1471,6 +1473,15 @@ void UWaterBodyComponent::Serialize(FArchive& Ar)
 
 	Ar.UsingCustomVersion(FWaterCustomVersion::GUID);
 	Ar.UsingCustomVersion(FFortniteMainBranchObjectVersion::GUID);
+
+#if WITH_EDITOR
+	if (Ar.IsObjectReferenceCollector() && Ar.IsSaving())
+	{
+		FSoftObjectPathSerializationScope EditorOnlyScope(ESoftObjectPathCollectType::EditorOnlyCollect);
+		FSoftObjectPath WaterSpriteTexture(GetWaterSpriteTextureName());
+		Ar << WaterSpriteTexture;
+	}
+#endif
 }
 
 void UWaterBodyComponent::PostLoad()
@@ -2162,7 +2173,11 @@ void UWaterBodyComponent::UpdateNonTessellatedMeshSections()
 
 void UWaterBodyComponent::CreateWaterSpriteComponent()
 {
-	UTexture2D* Texture = LoadObject<UTexture2D>(nullptr, GetWaterSpriteTextureName());
+	UTexture2D* Texture = nullptr;
+	{
+		FCookLoadScope EditorOnlyScope(ECookLoadType::EditorOnly);
+		Texture = LoadObject<UTexture2D>(nullptr, GetWaterSpriteTextureName());
+	}
 
 	IWaterModuleInterface& WaterModule = FModuleManager::GetModuleChecked<IWaterModuleInterface>(TEXT("Water"));
 	if (IWaterEditorServices* WaterEditorServices = WaterModule.GetWaterEditorServices())
diff --git a/Engine/Source/Runtime/Engine/Classes/Components/ForceFeedbackComponent.h b/Engine/Source/Runtime/Engine/Classes/Components/ForceFeedbackComponent.h
--- a/Engine/Source/Runtime/Engine/Classes/Components/ForceFeedbackComponent.h
+++ b/Engine/Source/Runtime/Engine/Classes/Components/ForceFeedbackComponent.h
@@ -126,6 +126,7 @@ public:
 	//~ Begin UObject Interface.
 #if WITH_EDITORONLY_DATA
 	virtual void PostEditChangeProperty(FPropertyChangedEvent& PropertyChangedEvent) override;
+	virtual void Serialize(FArchive& Ar) override;
 #endif // WITH_EDITOR
 	//~ End UObject Interface.
 
diff --git a/Engine/Source/Runtime/Engine/Classes/PhysicsEngine/PhysicsConstraintComponent.h b/Engine/Source/Runtime/Engine/Classes/PhysicsEngine/PhysicsConstraintComponent.h
--- a/Engine/Source/Runtime/Engine/Classes/PhysicsEngine/PhysicsConstraintComponent.h
+++ b/Engine/Source/Runtime/Engine/Classes/PhysicsEngine/PhysicsConstraintComponent.h
@@ -74,6 +74,9 @@ public:
 	virtual void PostEditChangeChainProperty(FPropertyChangedChainEvent& PropertyChangedEvent) override;
 	virtual void PostEditChangeProperty(FPropertyChangedEvent& PropertyChangedEvent) override;
 #endif // WITH_EDITOR
+#if WITH_EDITORONLY_DATA
+	virtual void Serialize(FArchive& Ar) override;
+#endif
 	//End UObject interface
 
 	//Begin ActorComponent interface
diff --git a/Engine/Source/Runtime/Engine/Private/Components/AudioComponent.cpp b/Engine/Source/Runtime/Engine/Private/Components/AudioComponent.cpp
--- a/Engine/Source/Runtime/Engine/Private/Components/AudioComponent.cpp
+++ b/Engine/Source/Runtime/Engine/Private/Components/AudioComponent.cpp
@@ -14,6 +14,8 @@
 #include "Sound/SoundNodeAttenuation.h"
 #include "Stats/StatsTrace.h"
 #include "UObject/FrameworkObjectVersion.h"
+#include "UObject/ICookInfo.h"
+#include "UObject/SoftObjectPath.h"
 
 #include UE_INLINE_GENERATED_CPP_BY_NAME(AudioComponent)
 
@@ -65,6 +67,10 @@ uint64 UAudioComponent::AudioComponentIDCounter = 0;
 TMap<uint64, UAudioComponent*> UAudioComponent::AudioIDToComponentMap;
 FCriticalSection UAudioComponent::AudioIDToComponentMapLock;
 
+#if WITH_EDITORONLY_DATA
+static const TCHAR* GAudioSpriteAssetNameAutoActivate = TEXT("/Engine/EditorResources/AudioIcons/S_AudioComponent_AutoActivate.S_AudioComponent_AutoActivate");
+static const TCHAR* GAudioSpriteAssetName = TEXT("/Engine/EditorResources/AudioIcons/S_AudioComponent.S_AudioComponent");
+#endif
 
 UInitialActiveSoundParams::UInitialActiveSoundParams(const FObjectInitializer& ObjectInitializer)
 	: Super(ObjectInitializer)
@@ -191,6 +197,15 @@ void UAudioComponent::Serialize(FArchive& Ar)
 
 		ModulationRouting.VersionModulators();
 	}
+	if (Ar.IsSaving() && Ar.IsObjectReferenceCollector() && !Ar.IsCooking())
+	{
+		FSoftObjectPathSerializationScope EditorOnlyScope(ESoftObjectPathCollectType::EditorOnlyCollect);
+		FSoftObjectPath SpriteAssets[]{ FSoftObjectPath(GAudioSpriteAssetNameAutoActivate), FSoftObjectPath(GAudioSpriteAssetName) };
+		for (FSoftObjectPath& SpriteAsset : SpriteAssets)
+		{
+			Ar << SpriteAsset;
+		}
+	}
 #endif // WITH_EDITORONLY_DATA
 }
 
@@ -1195,13 +1210,14 @@ void UAudioComponent::UpdateSpriteTexture()
 		SpriteComponent->SpriteInfo.Category = TEXT("Sounds");
 		SpriteComponent->SpriteInfo.DisplayName = NSLOCTEXT("SpriteCategory", "Sounds", "Sounds");
 
+		FCookLoadScope EditorOnlyScope(ECookLoadType::EditorOnly);
 		if (bAutoActivate)
 		{
-			SpriteComponent->SetSprite(LoadObject<UTexture2D>(nullptr, TEXT("/Engine/EditorResources/AudioIcons/S_AudioComponent_AutoActivate.S_AudioComponent_AutoActivate")));
+			SpriteComponent->SetSprite(LoadObject<UTexture2D>(nullptr, GAudioSpriteAssetNameAutoActivate));
 		}
 		else
 		{
-			SpriteComponent->SetSprite(LoadObject<UTexture2D>(nullptr, TEXT("/Engine/EditorResources/AudioIcons/S_AudioComponent.S_AudioComponent")));
+			SpriteComponent->SetSprite(LoadObject<UTexture2D>(nullptr, GAudioSpriteAssetName));
 		}
 	}
 }
diff --git a/Engine/Source/Runtime/Engine/Private/Components/ForceFeedbackComponent.cpp b/Engine/Source/Runtime/Engine/Private/Components/ForceFeedbackComponent.cpp
--- a/Engine/Source/Runtime/Engine/Private/Components/ForceFeedbackComponent.cpp
+++ b/Engine/Source/Runtime/Engine/Private/Components/ForceFeedbackComponent.cpp
@@ -7,12 +7,19 @@
 #include "Engine/Canvas.h"
 #include "GenericPlatform/IInputInterface.h"
 #include "Engine/Texture2D.h"
+#include "UObject/ICookInfo.h"
+#include "UObject/SoftObjectPath.h"
 
 #include UE_INLINE_GENERATED_CPP_BY_NAME(ForceFeedbackComponent)
 
 TArray<FForceFeedbackManager*> FForceFeedbackManager::PerWorldForceFeedbackManagers;
 FDelegateHandle FForceFeedbackManager::OnWorldCleanupHandle;
 
+#if WITH_EDITORONLY_DATA
+static const TCHAR* GForceFeedbackSpriteAssetNameAutoActivate = TEXT("/Engine/EditorResources/S_ForceFeedbackComponent_AutoActivate.S_ForceFeedbackComponent_AutoActivate");
+static const TCHAR* GForceFeedbackSpriteAssetName = TEXT("/Engine/EditorResources/S_ForceFeedbackComponent.S_ForceFeedbackComponent");
+#endif
+
 FForceFeedbackManager* FForceFeedbackManager::Get(UWorld* World, bool bCreateIfMissing)
 {
 	FForceFeedbackManager* ManagerForWorld = nullptr;
@@ -170,6 +177,21 @@ void UForceFeedbackComponent::PostEditChangeProperty(FPropertyChangedEvent& Prop
 	Super::PostEditChangeProperty(PropertyChangedEvent);
 }
 
+void UForceFeedbackComponent::Serialize(FArchive& Ar)
+{
+	Super::Serialize(Ar);
+	if (Ar.IsSaving() && Ar.IsObjectReferenceCollector() && !Ar.IsCooking())
+	{
+		FSoftObjectPathSerializationScope EditorOnlyScope(ESoftObjectPathCollectType::EditorOnlyCollect);
+		FSoftObjectPath SpriteAssets[]{ FSoftObjectPath(GForceFeedbackSpriteAssetNameAutoActivate), FSoftObjectPath(GForceFeedbackSpriteAssetName) };
+		for (FSoftObjectPath& AssetPath : SpriteAssets)
+		{
+			Ar << AssetPath;
+		}
+	}
+}
+
+
 void UForceFeedbackComponent::UpdateSpriteTexture()
 {
 	if (SpriteComponent)
@@ -177,13 +199,14 @@ void UForceFeedbackComponent::UpdateSpriteTexture()
 		SpriteComponent->SpriteInfo.Category = TEXT("Misc");
 		SpriteComponent->SpriteInfo.DisplayName = NSLOCTEXT("SpriteCategory", "Misc", "Misc");
 
+		FCookLoadScope EditorOnlyScope(ECookLoadType::EditorOnly);
 		if (bAutoActivate)
 		{
-			SpriteComponent->SetSprite(LoadObject<UTexture2D>(nullptr, TEXT("/Engine/EditorResources/S_ForceFeedbackComponent_AutoActivate.S_ForceFeedbackComponent_AutoActivate")));
+			SpriteComponent->SetSprite(LoadObject<UTexture2D>(nullptr, GForceFeedbackSpriteAssetNameAutoActivate));
 		}
 		else
 		{
-			SpriteComponent->SetSprite(LoadObject<UTexture2D>(nullptr, TEXT("/Engine/EditorResources/S_ForceFeedbackComponent.S_ForceFeedbackComponent")));
+			SpriteComponent->SetSprite(LoadObject<UTexture2D>(nullptr, GForceFeedbackSpriteAssetName));
 		}
 	}
 }
diff --git a/Engine/Source/Runtime/Engine/Private/Components/LightComponent.cpp b/Engine/Source/Runtime/Engine/Private/Components/LightComponent.cpp
--- a/Engine/Source/Runtime/Engine/Private/Components/LightComponent.cpp
+++ b/Engine/Source/Runtime/Engine/Private/Components/LightComponent.cpp
@@ -27,6 +27,8 @@
 #include "Components/DirectionalLightComponent.h"
 #include "Components/BillboardComponent.h"
 #include "ComponentRecreateRenderStateContext.h"
+#include "UObject/ICookInfo.h"
+#include "UObject/SoftObjectPath.h"
 #include "UObject/UnrealType.h"
 
 #include UE_INLINE_GENERATED_CPP_BY_NAME(LightComponent)
@@ -35,6 +37,10 @@
 #include "Rendering/StaticLightingSystemInterface.h"
 #endif
 
+#if WITH_EDITOR
+static const TCHAR* GLightSpriteAssetName = TEXT("/Engine/EditorResources/LightIcons/S_LightError.S_LightError");
+#endif
+
 void FStaticShadowDepthMap::InitRHI()
 {
 	if (FApp::CanEverRender() && Data && Data->ShadowMapSizeX > 0 && Data->ShadowMapSizeY > 0 && GMaxRHIFeatureLevel >= ERHIFeatureLevel::SM5)
@@ -167,6 +173,15 @@ void ULightComponentBase::Serialize(FArchive& Ar)
 	{
 		CastRaytracedShadow = bCastRaytracedShadow_DEPRECATED == 0? ECastRayTracedShadow::Disabled : ECastRayTracedShadow::UseProjectSetting;
 	}
+
+#if WITH_EDITOR
+	if (Ar.IsSaving() && Ar.IsObjectReferenceCollector() && !Ar.IsCooking())
+	{
+		FSoftObjectPathSerializationScope EditorOnlyScope(ESoftObjectPathCollectType::EditorOnlyCollect);
+		FSoftObjectPath SpritePath(GLightSpriteAssetName);
+		Ar << SpritePath;
+	}
+#endif
 }
 
 /**
@@ -764,7 +779,8 @@ void ULightComponent::UpdateLightSpriteTexture()
 			(GetWorld() && !GetWorld()->IsPreviewWorld()))
 		{
 			UTexture2D* SpriteTexture = NULL;
-			SpriteTexture = LoadObject<UTexture2D>(NULL, TEXT("/Engine/EditorResources/LightIcons/S_LightError.S_LightError"));
+			FCookLoadScope EditorOnlyScope(ECookLoadType::EditorOnly);
+			SpriteTexture = LoadObject<UTexture2D>(NULL, GLightSpriteAssetName);
 			SpriteComponent->SetSprite(SpriteTexture);
 			SpriteComponent->SetRelativeScale3D(FVector(0.5f));
 		}
diff --git a/Engine/Source/Runtime/Engine/Private/LevelInstance/LevelInstanceComponent.cpp b/Engine/Source/Runtime/Engine/Private/LevelInstance/LevelInstanceComponent.cpp
--- a/Engine/Source/Runtime/Engine/Private/LevelInstance/LevelInstanceComponent.cpp
+++ b/Engine/Source/Runtime/Engine/Private/LevelInstance/LevelInstanceComponent.cpp
@@ -8,9 +8,15 @@
 #include "Engine/Texture2D.h"
 #include "Engine/Level.h"
 #include "Engine/World.h"
+#include "UObject/ICookInfo.h"
+#include "UObject/SoftObjectPath.h"
 
 #include UE_INLINE_GENERATED_CPP_BY_NAME(LevelInstanceComponent)
 
+#if WITH_EDITORONLY_DATA
+static const TCHAR* GLevelSpriteAssetName = TEXT("/Engine/EditorResources/LevelInstance");
+#endif
+
 ULevelInstanceComponent::ULevelInstanceComponent(const FObjectInitializer& ObjectInitializer)
 	: Super(ObjectInitializer)
 {
@@ -19,6 +25,19 @@ ULevelInstanceComponent::ULevelInstanceComponent(const FObjectInitializer& Objec
 #endif
 }
 
+#if WITH_EDITORONLY_DATA
+void ULevelInstanceComponent::Serialize(FArchive& Ar)
+{
+	Super::Serialize(Ar);
+	if (Ar.IsSaving() && Ar.IsObjectReferenceCollector() && !Ar.IsCooking())
+	{
+		FSoftObjectPathSerializationScope EditorOnlyScope(ESoftObjectPathCollectType::EditorOnlyCollect);
+		FSoftObjectPath SpriteAsset(GLevelSpriteAssetName);
+		Ar << SpriteAsset;
+	}
+}
+#endif
+
 #if WITH_EDITOR
 void ULevelInstanceComponent::OnRegister()
 {
@@ -36,7 +55,8 @@ void ULevelInstanceComponent::OnRegister()
 	{
 		// Re-enable before calling CreateSpriteComponent
 		bVisualizeComponent = true;
-		CreateSpriteComponent(LoadObject<UTexture2D>(nullptr, TEXT("/Engine/EditorResources/LevelInstance")), false);
+		FCookLoadScope CookLoadScope(ECookLoadType::EditorOnly);
+		CreateSpriteComponent(LoadObject<UTexture2D>(nullptr, GLevelSpriteAssetName), false);
 		if (SpriteComponent)
 		{
 			SpriteComponent->bShowLockedLocation = false;
diff --git a/Engine/Source/Runtime/Engine/Private/PhysicsEngine/PhysicsConstraintComponent.cpp b/Engine/Source/Runtime/Engine/Private/PhysicsEngine/PhysicsConstraintComponent.cpp
--- a/Engine/Source/Runtime/Engine/Private/PhysicsEngine/PhysicsConstraintComponent.cpp
+++ b/Engine/Source/Runtime/Engine/Private/PhysicsEngine/PhysicsConstraintComponent.cpp
@@ -11,6 +11,8 @@
 #include "PhysicsEngine/PhysicsObjectExternalInterface.h"
 #include "PhysicsProxy/SingleParticlePhysicsProxy.h"
 #include "Components/BillboardComponent.h"
+#include "UObject/ICookInfo.h"
+#include "UObject/SoftObjectPath.h"
 
 #include UE_INLINE_GENERATED_CPP_BY_NAME(PhysicsConstraintComponent)
 
@@ -599,26 +601,56 @@ FConstraintInstanceAccessor UPhysicsConstraintComponent::GetConstraint()
 	return FConstraintInstanceAccessor(this);
 }
 
+#if WITH_EDITOR
+static const TCHAR* GPhysicsConstraintHingeSpriteAssetName = TEXT("/Engine/EditorResources/S_KHinge.S_KHinge");
+static const TCHAR* GPhysicsConstraintPrismaticSpriteAssetName = TEXT("/Engine/EditorResources/S_KPrismatic.S_KPrismatic");
+static const TCHAR* GPhysicsConstraintJointSpriteAssetName = TEXT("/Engine/EditorResources/S_KBSJoint.S_KBSJoint");
+#endif
+
 #if WITH_EDITOR
 void UPhysicsConstraintComponent::UpdateSpriteTexture()
 {
 	if (SpriteComponent)
 	{
+		FCookLoadScope CookLoadScope(ECookLoadType::EditorOnly);
 		if (ConstraintUtils::IsHinge(ConstraintInstance))
 		{
-			SpriteComponent->SetSprite(LoadObject<UTexture2D>(NULL, TEXT("/Engine/EditorResources/S_KHinge.S_KHinge")));
+			SpriteComponent->SetSprite(LoadObject<UTexture2D>(NULL, GPhysicsConstraintHingeSpriteAssetName));
 		}
 		else if (ConstraintUtils::IsPrismatic(ConstraintInstance))
 		{
-			SpriteComponent->SetSprite(LoadObject<UTexture2D>(NULL, TEXT("/Engine/EditorResources/S_KPrismatic.S_KPrismatic")));
+			SpriteComponent->SetSprite(LoadObject<UTexture2D>(NULL, GPhysicsConstraintPrismaticSpriteAssetName));
 		}
 		else
 		{
-			SpriteComponent->SetSprite(LoadObject<UTexture2D>(NULL, TEXT("/Engine/EditorResources/S_KBSJoint.S_KBSJoint")));
+			SpriteComponent->SetSprite(LoadObject<UTexture2D>(NULL, GPhysicsConstraintJointSpriteAssetName));
 		}
 	}
 }
 #endif // WITH_EDITOR
+#if WITH_EDITORONLY_DATA
+void UPhysicsConstraintComponent::Serialize(FArchive& Ar)
+{
+	Super::Serialize(Ar);
+#if WITH_EDITOR
+	if (Ar.IsSaving() && Ar.IsObjectReferenceCollector() && !Ar.IsCooking())
+	{
+		FSoftObjectPathSerializationScope EditorOnlyScope(ESoftObjectPathCollectType::EditorOnlyCollect);
+		FSoftObjectPath SpriteAssets[]
+		{
+			FSoftObjectPath(GPhysicsConstraintHingeSpriteAssetName),
+			FSoftObjectPath(GPhysicsConstraintPrismaticSpriteAssetName),
+			FSoftObjectPath(GPhysicsConstraintJointSpriteAssetName)
+		};
+		for (FSoftObjectPath& SpriteAsset : SpriteAssets)
+		{
+			Ar << SpriteAsset;
+		}
+	}
+#endif
+}
+
+#endif
 
 void UPhysicsConstraintComponent::SetLinearPositionDrive( bool bEnableDriveX, bool bEnableDriveY, bool bEnableDriveZ )
 {
diff --git a/Engine/Source/Runtime/Engine/Public/LevelInstance/LevelInstanceComponent.h b/Engine/Source/Runtime/Engine/Public/LevelInstance/LevelInstanceComponent.h
--- a/Engine/Source/Runtime/Engine/Public/LevelInstance/LevelInstanceComponent.h
+++ b/Engine/Source/Runtime/Engine/Public/LevelInstance/LevelInstanceComponent.h
@@ -25,6 +25,9 @@ class ENGINE_API ULevelInstanceComponent : public USceneComponent
 {
 	GENERATED_UCLASS_BODY()
 public:
+#if WITH_EDITORONLY_DATA
+	virtual void Serialize(FArchive& Ar) override;
+#endif
 #if WITH_EDITOR
 	// Those are the methods that need overriding to be able to properly update the AttachComponent
 	virtual void OnRegister() override;
