[StaticMesh] Fix crash in static mesh build, trying to complete a build that hasn't started yet, causing us to access null render data in UStaticMesh::FinishBuildInternal().

diff --git a/Engine/Source/Runtime/Engine/Private/StaticMeshBuild.cpp b/Engine/Source/Runtime/Engine/Private/StaticMeshBuild.cpp
--- a/Engine/Source/Runtime/Engine/Private/StaticMeshBuild.cpp
+++ b/Engine/Source/Runtime/Engine/Private/StaticMeshBuild.cpp
@@ -235,6 +235,7 @@ void UStaticMesh::BatchBuild(const TArray<UStaticMesh*>& InStaticMeshes, const F
 				// Only launch async compile if errors are not required
 				if (BuildParameters.OutErrors == nullptr && FStaticMeshCompilingManager::Get().IsAsyncCompilationAllowed(StaticMesh))
 				{
+					const int64 BuildRequiredMemory = StaticMesh->GetBuildRequiredMemory();
 					TUniquePtr<FStaticMeshBuildContext> Context = MakeUnique<FStaticMeshBuildContext>(BuildParameters);
 					StaticMesh->BeginBuildInternal(Context.Get());
 
@@ -242,7 +243,7 @@ void UStaticMesh::BatchBuild(const TArray<UStaticMesh*>& InStaticMeshes, const F
 					EQueuedWorkPriority BasePriority = FStaticMeshCompilingManager::Get().GetBasePriority(StaticMesh);
 					check(StaticMesh->AsyncTask == nullptr);
 					StaticMesh->AsyncTask = MakeUnique<FStaticMeshAsyncBuildTask>(StaticMesh, MoveTemp(Context));
-					StaticMesh->AsyncTask->StartBackgroundTask(StaticMeshThreadPool, BasePriority, EQueuedWorkFlags::DoNotRunInsideBusyWait, StaticMesh->GetBuildRequiredMemory());
+					StaticMesh->AsyncTask->StartBackgroundTask(StaticMeshThreadPool, BasePriority, EQueuedWorkFlags::DoNotRunInsideBusyWait, BuildRequiredMemory);
 					FStaticMeshCompilingManager::Get().AddStaticMeshes({ StaticMesh });
 					return true;
 				}
