Add a "Get All Action and Axis Names" function to UInputSettings. This is useful if you want some FName UPROPERTY where the options are only valid Action/Axis mappings.

#rnx

diff --git a/Engine/Source/Runtime/Engine/Classes/GameFramework/InputSettings.h b/Engine/Source/Runtime/Engine/Classes/GameFramework/InputSettings.h
--- a/Engine/Source/Runtime/Engine/Classes/GameFramework/InputSettings.h
+++ b/Engine/Source/Runtime/Engine/Classes/GameFramework/InputSettings.h
@@ -224,6 +224,19 @@ public:
 	UFUNCTION(BlueprintCallable, Category = Settings)
 	void RemoveAxisMapping(const FInputAxisKeyMapping& KeyMapping, bool bForceRebuildKeymaps = true);
 
+#if WITH_EDITOR	
+	/**
+	 * Returns all known legacy action and axis names that is useful for
+	 * properties that you want a drop down selection of the available names
+	 * on a UPROPERTY
+	 * i.e.
+	 *
+	 * meta=(GetOptions="Engine.InputSettings.GetAllActionAndAxisNames")
+	 */
+	UFUNCTION()
+	static const TArray<FName>& GetAllActionAndAxisNames();
+#endif// WITH_EDITOR
+
 	/** Flush the current mapping values to the config file */
 	UFUNCTION(BlueprintCallable, Category=Settings)
 	void SaveKeyMappings();
diff --git a/Engine/Source/Runtime/Engine/Private/UserInterface/InputSettings.cpp b/Engine/Source/Runtime/Engine/Private/UserInterface/InputSettings.cpp
--- a/Engine/Source/Runtime/Engine/Private/UserInterface/InputSettings.cpp
+++ b/Engine/Source/Runtime/Engine/Private/UserInterface/InputSettings.cpp
@@ -336,6 +336,26 @@ void UInputSettings::GetAxisNames(TArray<FName>& AxisNames) const
 	}
 }
 
+#if WITH_EDITOR
+const TArray<FName>& UInputSettings::GetAllActionAndAxisNames()
+{
+	static TArray<FName> OutNames;
+	OutNames.Reset();
+
+	const UInputSettings* InputSettings = GetDefault<UInputSettings>();
+	check(InputSettings);
+
+	TArray<FName> ActionNames;
+	InputSettings->GetActionNames(ActionNames);
+	OutNames.Append(ActionNames);
+	
+	InputSettings->GetAxisNames(ActionNames);
+	OutNames.Append(ActionNames);
+	
+	return OutNames;
+}
+#endif // WITH_EDITOR
+
 void UInputSettings::ForceRebuildKeymaps()
 {
 	for (TObjectIterator<UPlayerInput> It; It; ++It)
