LWI - make templated FActorInstanceHandle::FetchActor not spawn actor if template class is not the right type ( to prevent unnecessary actor spawns ).

diff --git a/Engine/Source/Runtime/Engine/Classes/Engine/ActorInstanceHandle.h b/Engine/Source/Runtime/Engine/Classes/Engine/ActorInstanceHandle.h
--- a/Engine/Source/Runtime/Engine/Classes/Engine/ActorInstanceHandle.h
+++ b/Engine/Source/Runtime/Engine/Classes/Engine/ActorInstanceHandle.h
@@ -26,6 +26,9 @@ struct ENGINE_API FActorInstanceHandle
 
 	bool DoesRepresentClass(const UClass* OtherClass) const;
 
+	template<typename T>
+	bool DoesRepresentClass() const;
+
 	UClass* GetRepresentedClass() const;
 
 	FVector GetLocation() const;
@@ -93,8 +96,25 @@ private:
 	uint32 InstanceUID;
 };
 
+template<typename T>
+bool FActorInstanceHandle::DoesRepresentClass() const
+{
+	if constexpr (TIsIInterface<T>::Value)
+	{
+		return DoesRepresentClass(T::UClassType::StaticClass());
+	}
+	else
+	{
+		return DoesRepresentClass(T::StaticClass());
+	}
+}
+
 template <typename T>
 T* FActorInstanceHandle::FetchActor() const
 {
-	return Cast<T>(FetchActor());
+	if (DoesRepresentClass<T>())
+	{
+		return Cast<T>(FetchActor());
+	}
+	return nullptr;
 }
