Prevent stack overflows if an actor spawned by an LWI manager triggers the spawn of another actor instance, and vice versa.
( eg. two instances on an LWI manager are overlapping each other. One is triggered to spawn an actor, which does some overlap check that triggers the spawn of the other instance actor, which does some overlap check that triggers the spawn of the original actor, etc. ).

diff --git a/Engine/Source/Runtime/Engine/Classes/GameFramework/LightWeightInstanceManager.h b/Engine/Source/Runtime/Engine/Classes/GameFramework/LightWeightInstanceManager.h
--- a/Engine/Source/Runtime/Engine/Classes/GameFramework/LightWeightInstanceManager.h
+++ b/Engine/Source/Runtime/Engine/Classes/GameFramework/LightWeightInstanceManager.h
@@ -259,4 +259,8 @@ protected:
 protected:
 	virtual void PostEditChangeProperty(struct FPropertyChangedEvent& PropertyChangedEvent) override;	
 #endif // WITH_EDITOR
+
+#if !UE_BUILD_SHIPPING
+	TArray<int32> InstanceToActorConversionsInProgress;
+#endif //!UE_BUILD_SHIPPING
 };
diff --git a/Engine/Source/Runtime/Engine/Private/LightWeightInstanceManager.cpp b/Engine/Source/Runtime/Engine/Private/LightWeightInstanceManager.cpp
--- a/Engine/Source/Runtime/Engine/Private/LightWeightInstanceManager.cpp
+++ b/Engine/Source/Runtime/Engine/Private/LightWeightInstanceManager.cpp
@@ -70,6 +70,17 @@ AActor* ALightWeightInstanceManager::FetchActorFromHandle(const FActorInstanceHa
 		}
 		else
 		{
+#if !UE_BUILD_SHIPPING
+			const bool bSpawnInProgress = InstanceToActorConversionsInProgress.Contains(Handle.GetInstanceIndex());
+			if (!ensure(!bSpawnInProgress))
+			{
+				UE_LOG(LogLightWeightInstance, Error, 
+					TEXT("Calling FetchActorFromHandle on ActorInstance [ %s ] index [ %d ] - actor is spawned but may not be fully setup"),
+					*BaseInstanceName,
+					Handle.GetInstanceIndex());
+			}
+#endif //!UE_BUILD_SHIPPING
+
 			Handle.Actor = *FoundActor;
 		}
 	}
@@ -82,6 +93,15 @@ AActor* ALightWeightInstanceManager::FetchActorFromHandle(const FActorInstanceHa
 
 AActor* ALightWeightInstanceManager::ConvertInstanceToActor(const FActorInstanceHandle& Handle)
 {
+#if !UE_BUILD_SHIPPING
+	check(!InstanceToActorConversionsInProgress.Contains(Handle.GetInstanceIndex()));
+	InstanceToActorConversionsInProgress.Add(Handle.GetInstanceIndex());
+	ON_SCOPE_EXIT
+	{
+		InstanceToActorConversionsInProgress.RemoveSingleSwap(Handle.GetInstanceIndex());
+	};
+#endif //!UE_BUILD_SHIPPING
+
 	// we shouldn't be calling this on indices that already have an actor representing them
 	if (Actors.Contains(Handle.GetInstanceIndex()) && Actors[Handle.GetInstanceIndex()] != nullptr)
 	{
@@ -99,12 +119,20 @@ AActor* ALightWeightInstanceManager::ConvertInstanceToActor(const FActorInstance
 	{
 		FActorSpawnParameters SpawnParams;
 		SetSpawnParameters(SpawnParams);
+		SpawnParams.CustomPreSpawnInitalization = [this, &Handle](AActor* SpawnedActor)
+		{
+			Handle.Actor = SpawnedActor;
+			Actors.Add(Handle.GetInstanceIndex(), SpawnedActor);
+		};
+
 		NewActor = GetLevel()->GetWorld()->SpawnActor<AActor>(GetActorClassToSpawn(Handle), GetTransform(Handle), SpawnParams);
 		check(NewActor);
-		NewActor->OnDestroyed.AddUniqueDynamic(this, &ALightWeightInstanceManager::OnSpawnedActorDestroyed);
 
-		Handle.Actor = NewActor;
-		Actors.Add(Handle.GetInstanceIndex(), NewActor);
+		//should have been assigned in CustomPreSpawnInitialization
+		check(Handle.Actor == NewActor);
+		check(NewActor == Actors.FindRef(Handle.GetInstanceIndex()));
+
+		NewActor->OnDestroyed.AddUniqueDynamic(this, &ALightWeightInstanceManager::OnSpawnedActorDestroyed);
 
 		PostActorSpawn(Handle);
 	}
