Fix non-unity build compile error

diff --git a/Engine/Source/Runtime/Engine/Private/LightWeightInstanceManager.cpp b/Engine/Source/Runtime/Engine/Private/LightWeightInstanceManager.cpp
--- a/Engine/Source/Runtime/Engine/Private/LightWeightInstanceManager.cpp
+++ b/Engine/Source/Runtime/Engine/Private/LightWeightInstanceManager.cpp
@@ -4,6 +4,7 @@
 #include "GameFramework/LightWeightInstanceSubsystem.h"
 #include "Engine/Level.h"
 #include "Engine/World.h"
+#include "Misc/ScopeExit.h"
 #include "Net/UnrealNetwork.h"
 
 #include UE_INLINE_GENERATED_CPP_BY_NAME(LightWeightInstanceManager)
