World Partition Streaming:
- Each cell now keeps a mapping of referenced editor only actors and set the instancing context to skip loading them.
- Don't generate world partition streaming logs for sub partitioned worlds.

Tested with and without -zenloader.

diff --git a/Engine/Plugins/Editor/WorldPartitionHLODUtilities/Source/Private/WorldPartition/HLOD/Utilities/WorldPartitionHLODUtilities.cpp b/Engine/Plugins/Editor/WorldPartitionHLODUtilities/Source/Private/WorldPartition/HLOD/Utilities/WorldPartitionHLODUtilities.cpp
--- a/Engine/Plugins/Editor/WorldPartitionHLODUtilities/Source/Private/WorldPartition/HLOD/Utilities/WorldPartitionHLODUtilities.cpp
+++ b/Engine/Plugins/Editor/WorldPartitionHLODUtilities/Source/Private/WorldPartition/HLOD/Utilities/WorldPartitionHLODUtilities.cpp
@@ -60,7 +60,8 @@ static UWorldPartitionLevelStreamingDynamic* CreateLevelStreamingFromHLODActor(A
 			SubActor.ContainerPackage, 
 			World->GetPackage()->GetFName(), 
 			InHLODActor->GetContentBundleGuid(),
-			SubActor.ActorGuid
+			SubActor.ActorGuid,
+			false
 		);
 	});
 
diff --git a/Engine/Source/Runtime/Engine/Private/WorldPartition/WorldPartitionActorDescView.cpp b/Engine/Source/Runtime/Engine/Private/WorldPartition/WorldPartitionActorDescView.cpp
--- a/Engine/Source/Runtime/Engine/Private/WorldPartition/WorldPartitionActorDescView.cpp
+++ b/Engine/Source/Runtime/Engine/Private/WorldPartition/WorldPartitionActorDescView.cpp
@@ -154,6 +154,11 @@ const TArray<FGuid>& FWorldPartitionActorDescView::GetReferences() const
 	return RuntimeReferences.IsSet() ? RuntimeReferences.GetValue() : ActorDesc->GetReferences();
 }
 
+const TArray<FGuid>& FWorldPartitionActorDescView::GetEditorReferences() const
+{
+	return EditorReferences;
+}
+
 FString FWorldPartitionActorDescView::ToString() const
 {
 	return ActorDesc->ToString();
@@ -240,16 +245,21 @@ void FWorldPartitionActorDescView::SetInvalidDataLayers()
 	}
 }
 
-void FWorldPartitionActorDescView::SetRuntimeDataLayerInstanceNames(TArray<FName>& InRuntimeDataLayerInstanceNames)
+void FWorldPartitionActorDescView::SetRuntimeDataLayerInstanceNames(const TArray<FName>& InRuntimeDataLayerInstanceNames)
 {
 	RuntimeDataLayerInstanceNames = InRuntimeDataLayerInstanceNames;
 }
 
-void FWorldPartitionActorDescView::SetRuntimeReferences(TArray<FGuid>& InRuntimeReferences)
+void FWorldPartitionActorDescView::SetRuntimeReferences(const TArray<FGuid>& InRuntimeReferences)
 {
 	RuntimeReferences = InRuntimeReferences;
 }
 
+void FWorldPartitionActorDescView::SetEditorReferences(const TArray<FGuid>& InEditorReferences)
+{
+	EditorReferences = InEditorReferences;
+}
+
 AActor* FWorldPartitionActorDescView::GetActor() const
 {
 	return ActorDesc->GetActor();
diff --git a/Engine/Source/Runtime/Engine/Private/WorldPartition/WorldPartitionLevelHelper.cpp b/Engine/Source/Runtime/Engine/Private/WorldPartition/WorldPartitionLevelHelper.cpp
--- a/Engine/Source/Runtime/Engine/Private/WorldPartition/WorldPartitionLevelHelper.cpp
+++ b/Engine/Source/Runtime/Engine/Private/WorldPartition/WorldPartitionLevelHelper.cpp
@@ -365,12 +365,16 @@ bool FWorldPartitionLevelHelper::LoadActors(UWorld* InOuterWorld, ULevel* InDest
 		const FName ContainerPackageInstanceName = Context->RemapPackage(PackageObjectMapping.ContainerPackage);
 		if (PackageObjectMapping.ContainerPackage != ContainerPackageInstanceName)
 		{
-			const FString ActorPackageName = FPackageName::ObjectPathToPackageName(PackageObjectMapping.Package.ToString());
-			const FString ActorPackageInstanceName = ULevel::GetExternalActorPackageInstanceName(ContainerPackageInstanceName.ToString(), ActorPackageName);
+			const FName ActorPackageName = *FPackageName::ObjectPathToPackageName(PackageObjectMapping.Package.ToString());
+			const FName ActorPackageInstanceName = PackageObjectMapping.bIsEditorOnly ? NAME_None : FName(*ULevel::GetExternalActorPackageInstanceName(ContainerPackageInstanceName.ToString(), ActorPackageName.ToString()));
 
-			Context->AddPackageMapping(FName(*ActorPackageName), FName(*ActorPackageInstanceName));
+			Context->AddPackageMapping(ActorPackageName, ActorPackageInstanceName);
+		}
+
+		if (!PackageObjectMapping.bIsEditorOnly)
+		{
+			ActorPackages.Add(&PackageObjectMapping);
 		}
-		ActorPackages.Add(&PackageObjectMapping);
 	}
 
 	LoadProgress->NumPendingLoadRequests = ActorPackages.Num();
diff --git a/Engine/Source/Runtime/Engine/Private/WorldPartition/WorldPartitionRuntimeLevelStreamingCell.cpp b/Engine/Source/Runtime/Engine/Private/WorldPartition/WorldPartitionRuntimeLevelStreamingCell.cpp
--- a/Engine/Source/Runtime/Engine/Private/WorldPartition/WorldPartitionRuntimeLevelStreamingCell.cpp
+++ b/Engine/Source/Runtime/Engine/Private/WorldPartition/WorldPartitionRuntimeLevelStreamingCell.cpp
@@ -205,14 +205,37 @@ void UWorldPartitionRuntimeLevelStreamingCell::AddActorToCell(const FWorldPartit
 {
 	check(!ActorDescView.GetActorIsEditorOnly());
 
+	if (!IsRunningCookCommandlet())
+	{
+		for (const FGuid& EditorReferenceGuid : ActorDescView.GetEditorReferences())
+		{
+			const FWorldPartitionActorDesc& ReferenceActorDesc = InContainer->GetActorDescChecked(EditorReferenceGuid);
+			check(ReferenceActorDesc.GetActorIsEditorOnly());
+
+			Packages.Emplace(
+				ReferenceActorDesc.GetActorPackage(), 
+				*ReferenceActorDesc.GetActorSoftPath().ToString(), 
+				InContainerID, 
+				InContainerTransform, 
+				InContainer->GetContainerPackage(), 
+				GetWorld()->GetPackage()->GetFName(), 
+				ReferenceActorDesc.GetContentBundleGuid(),
+				InContainerID.GetActorGuid(ReferenceActorDesc.GetGuid()),
+				true
+			);
+		}
+	}
+
 	Packages.Emplace(
 		ActorDescView.GetActorPackage(), 
 		*ActorDescView.GetActorSoftPath().ToString(), 
-		InContainerID, InContainerTransform, 
+		InContainerID, 
+		InContainerTransform, 
 		InContainer->GetContainerPackage(), 
 		GetWorld()->GetPackage()->GetFName(), 
 		ActorDescView.GetContentBundleGuid(),
-		InContainerID.GetActorGuid(ActorDescView.GetGuid())
+		InContainerID.GetActorGuid(ActorDescView.GetGuid()),
+		false
 	);
 }
 
diff --git a/Engine/Source/Runtime/Engine/Private/WorldPartition/WorldPartitionStreamingGeneration.cpp b/Engine/Source/Runtime/Engine/Private/WorldPartition/WorldPartitionStreamingGeneration.cpp
--- a/Engine/Source/Runtime/Engine/Private/WorldPartition/WorldPartitionStreamingGeneration.cpp
+++ b/Engine/Source/Runtime/Engine/Private/WorldPartition/WorldPartitionStreamingGeneration.cpp
@@ -6,6 +6,7 @@
 #include "WorldPartition/WorldPartitionStreamingGenerationContext.h"
 
 #include "Editor.h"
+#include "Algo/Transform.h"
 #include "Engine/LevelScriptBlueprint.h"
 #include "ActorReferencesUtils.h"
 #include "Misc/PackageName.h"
@@ -244,7 +245,7 @@ class FWorldPartitionStreamingGenerator
 
 		const UActorDescContainer* Container;
 		FActorDescViewMap ActorDescViewMap;
-		FActorDescViewMap FilteredActorDescViewMap;
+		TSet<FGuid> EditorOnlyActorDescMap;
 		TArray<FWorldPartitionActorDescView> ContainerInstanceViews;
 		TArray<TArray<FGuid>> Clusters;
 	};
@@ -306,7 +307,7 @@ class FWorldPartitionStreamingGenerator
 		}
 	}
 
-	void CreateActorDescViewMap(const UActorDescContainer* InContainer, FActorDescViewMap& OutActorDescViewMap, FActorDescViewMap& OutFilteredActorDescViewMap, const FActorContainerID& InContainerID, TArray<FWorldPartitionActorDescView>& OutContainerInstances)
+	void CreateActorDescViewMap(const UActorDescContainer* InContainer, FActorDescViewMap& OutActorDescViewMap, TSet<FGuid>& OutEditorOnlyActorDescMap, const FActorContainerID& InContainerID, TArray<FWorldPartitionActorDescView>& OutContainerInstances)
 	{
 		// Should we handle unsaved or newly created actors?
 		const bool bHandleUnsavedActors = ModifiedActorsDescList && InContainerID.IsMainContainer();
@@ -398,8 +399,7 @@ class FWorldPartitionStreamingGenerator
 			}
 			else
 			{
-				FWorldPartitionActorDescView ActorDescView(*ActorDescIt);
-				OutFilteredActorDescViewMap.Emplace(ActorDescIt->GetGuid(), ActorDescView);
+				OutEditorOnlyActorDescMap.Add(ActorDescIt->GetGuid());
 			}
 		}
 
@@ -434,7 +434,7 @@ class FWorldPartitionStreamingGenerator
 				ContainerDescriptor.Container = ContainerInstanceDescriptor.Container;
 
 				// Gather actor descriptor views for this container
-				CreateActorDescViewMap(ContainerInstanceDescriptor.Container, ContainerDescriptor.ActorDescViewMap, ContainerDescriptor.FilteredActorDescViewMap, ContainerInstanceDescriptor.ID, ContainerDescriptor.ContainerInstanceViews);
+				CreateActorDescViewMap(ContainerInstanceDescriptor.Container, ContainerDescriptor.ActorDescViewMap, ContainerDescriptor.EditorOnlyActorDescMap, ContainerInstanceDescriptor.ID, ContainerDescriptor.ContainerInstanceViews);
 
 				// Resolve actor descriptor views before validation
 				ResolveContainerDescriptor(ContainerDescriptor);
@@ -780,9 +780,12 @@ PRAGMA_ENABLE_DEPRECATION_WARNINGS
 				}
 
 				TArray<FGuid> RuntimeReferences;
+				TArray<FGuid> EditorReferences;
+
 				if (PassType == EPassType::Fixup)
 				{
 					RuntimeReferences.Reserve(ActorDescView.GetReferences().Num());
+					EditorReferences.Reserve(ActorDescView.GetReferences().Num());
 				}
 
 				for (FActorReferenceInfo& Info : References)
@@ -845,7 +848,7 @@ PRAGMA_ENABLE_DEPRECATION_WARNINGS
 					}
 					else
 					{
-						if (!ContainerDescriptor.FilteredActorDescViewMap.FindByGuid(Info.ReferenceGuid))
+						if (!ContainerDescriptor.EditorOnlyActorDescMap.Contains(Info.ReferenceGuid))
 						{
 							if (PassType == EPassType::ErrorReporting)
 							{
@@ -853,7 +856,7 @@ PRAGMA_ENABLE_DEPRECATION_WARNINGS
 								FWorldPartitionActorDescView* ReferendceActorDescViewPtr = nullptr;
 
 								if (const UActorDescContainer** ExistingReferenceContainerPtr = ActorGuidsToContainerMap.Find((Info.ReferenceGuid)))
-								{
+								{ 
 									if (const FWorldPartitionActorDesc* ReferenceActorDesc = (*ExistingReferenceContainerPtr)->GetActorDesc(Info.ReferenceGuid))
 									{
 										ReferendceActorDescView = FWorldPartitionActorDescView(ReferenceActorDesc);
@@ -864,6 +867,10 @@ PRAGMA_ENABLE_DEPRECATION_WARNINGS
 								ErrorHandler->OnInvalidReference(*RefererActorDescView, Info.ReferenceGuid, ReferendceActorDescViewPtr);
 							}
 						}
+						else if (PassType == EPassType::Fixup)
+						{
+							EditorReferences.Add(Info.ReferenceGuid);
+						}
 
 						NbErrorsDetected++;
 					}
@@ -874,7 +881,8 @@ PRAGMA_ENABLE_DEPRECATION_WARNINGS
 					if (RuntimeReferences.Num() != ActorDescView.GetReferences().Num())
 					{
 						ActorDescView.SetRuntimeReferences(RuntimeReferences);
-					}
+						ActorDescView.SetEditorReferences(EditorReferences);
+					}					
 				}
 			});		
 		}
@@ -1123,17 +1131,24 @@ bool UWorldPartition::GenerateContainerStreaming(const FGenerateStreamingParams&
 	UE_SCOPED_TIMER(*FString::Printf(TEXT("GenerateStreaming for '%s'"), *ContainerShortName), LogWorldPartition, Display);
 
 	// Dump state log
-	TStringBuilder<256> StateLogSuffix;
-	StateLogSuffix += bIsPIE ? TEXT("PIE") : (IsRunningGame() ? TEXT("Game") : (IsRunningCookCommandlet() ? TEXT("Cook") : (GIsAutomationTesting ? TEXT("UnitTest") : TEXT("Manual"))));
-	StateLogSuffix += TEXT("_");
-	StateLogSuffix += ContainerShortName;
-	TUniquePtr<FArchive> LogFileAr = FWorldPartitionStreamingGenerator::CreateDumpStateLogArchive(*StateLogSuffix, !InParams.OutputLogPath);
-	FHierarchicalLogArchive HierarchicalLogAr(*LogFileAr);
+	TUniquePtr<FArchive> LogFileAr;
 
-	if (LogFileAr)
+	if (IsMainWorldPartition())
 	{
+		TStringBuilder<256> StateLogSuffix;
+		StateLogSuffix += bIsPIE ? TEXT("PIE") : (IsRunningGame() ? TEXT("Game") : (IsRunningCookCommandlet() ? TEXT("Cook") : (GIsAutomationTesting ? TEXT("UnitTest") : TEXT("Manual"))));
+		StateLogSuffix += TEXT("_");
+		StateLogSuffix += ContainerShortName;
+		LogFileAr = FWorldPartitionStreamingGenerator::CreateDumpStateLogArchive(*StateLogSuffix, !InParams.OutputLogPath);
+		
 		InContext.OutputLogFilename = LogFileAr->GetArchiveName();
 	}
+	else
+	{
+		LogFileAr = MakeUnique<FArchive>();
+	}
+
+	FHierarchicalLogArchive HierarchicalLogAr(*LogFileAr);
 
 	FWorldPartitionStreamingGenerator::FWorldPartitionStreamingGeneratorParams StreamingGeneratorParams;
 	StreamingGeneratorParams.WorldPartitionContext = this;
@@ -1156,11 +1171,7 @@ bool UWorldPartition::GenerateContainerStreaming(const FGenerateStreamingParams&
 	check(RuntimeHash);
 	if (RuntimeHash->GenerateStreaming(StreamingPolicy, StreamingGenerator.GetStreamingGenerationContext(), InContext.PackagesToGenerate))
 	{
-		//if (IsRunningCookCommandlet())
-		{
-			RuntimeHash->DumpStateLog(HierarchicalLogAr);
-		}
-
+		RuntimeHash->DumpStateLog(HierarchicalLogAr);
 		StreamingPolicy->PrepareActorToCellRemapping();
 		return true;
 	}
diff --git a/Engine/Source/Runtime/Engine/Public/WorldPartition/WorldPartitionActorDescView.h b/Engine/Source/Runtime/Engine/Public/WorldPartition/WorldPartitionActorDescView.h
--- a/Engine/Source/Runtime/Engine/Public/WorldPartition/WorldPartitionActorDescView.h
+++ b/Engine/Source/Runtime/Engine/Public/WorldPartition/WorldPartitionActorDescView.h
@@ -53,6 +53,7 @@ public:
 	FBox GetRuntimeBounds() const;
 
 	const TArray<FGuid>& GetReferences() const;
+	const TArray<FGuid>& GetEditorReferences() const;
 	FString ToString() const;
 	const FGuid& GetParentActor() const;
 	FName GetActorName() const;
@@ -78,8 +79,9 @@ public:
 	void SetForcedNonSpatiallyLoaded();
 	void SetForcedNoRuntimeGrid();
 	void SetInvalidDataLayers();
-	void SetRuntimeDataLayerInstanceNames(TArray<FName>& InRuntimeDataLayerInstanceNames);
-	void SetRuntimeReferences(TArray<FGuid>& InRuntimeReferences);
+	void SetRuntimeDataLayerInstanceNames(const TArray<FName>& InRuntimeDataLayerInstanceNames);
+	void SetRuntimeReferences(const TArray<FGuid>& InRuntimeReferences);
+	void SetEditorReferences(const TArray<FGuid>& InEditorReferences);
 	void SetDataLayerInstanceNames(const TArray<FName>& InDataLayerInstanceNames);
 
 	AActor* GetActor() const;
@@ -107,5 +109,6 @@ protected:
 	TOptional<TArray<FName>> ResolvedDataLayerInstanceNames;
 	TOptional<TArray<FName>> RuntimeDataLayerInstanceNames;
 	TOptional<TArray<FGuid>> RuntimeReferences;
+	TArray<FGuid> EditorReferences;
 };
 #endif
\ No newline at end of file
diff --git a/Engine/Source/Runtime/Engine/Public/WorldPartition/WorldPartitionRuntimeCell.h b/Engine/Source/Runtime/Engine/Public/WorldPartition/WorldPartitionRuntimeCell.h
--- a/Engine/Source/Runtime/Engine/Public/WorldPartition/WorldPartitionRuntimeCell.h
+++ b/Engine/Source/Runtime/Engine/Public/WorldPartition/WorldPartitionRuntimeCell.h
@@ -42,7 +42,7 @@ struct FWorldPartitionRuntimeCellObjectMapping
 #endif
 	{}
 
-	FWorldPartitionRuntimeCellObjectMapping(FName InPackage, FName InPath, const FActorContainerID& InContainerID, const FTransform& InContainerTransform, FName InContainerPackage, FName InWorldPackage, const FGuid& InContentBundleGuid, const FGuid& InActorInstanceGuid)
+	FWorldPartitionRuntimeCellObjectMapping(FName InPackage, FName InPath, const FActorContainerID& InContainerID, const FTransform& InContainerTransform, FName InContainerPackage, FName InWorldPackage, const FGuid& InContentBundleGuid, const FGuid& InActorInstanceGuid, bool bInIsEditorOnly)
 #if WITH_EDITORONLY_DATA
 		: Package(InPackage)
 		, Path(InPath)
@@ -53,6 +53,7 @@ struct FWorldPartitionRuntimeCellObjectMapping
 		, ContentBundleGuid(InContentBundleGuid)
 		, ActorInstanceGuid(InActorInstanceGuid)
 		, LoadedPath(InPath)
+		, bIsEditorOnly(bInIsEditorOnly)
 #endif
 	{}
 
@@ -115,6 +116,9 @@ struct FWorldPartitionRuntimeCellObjectMapping
 	*/
 	UPROPERTY()
 	FName LoadedPath;
+
+	UPROPERTY()
+	bool bIsEditorOnly;
 #endif
 };
 
