Fix STAT_StaticMeshTotalMemory2 not taking into account depth-only index buffers.

#rnx

diff --git a/Engine/Source/Runtime/Engine/Private/StaticMesh.cpp b/Engine/Source/Runtime/Engine/Private/StaticMesh.cpp
--- a/Engine/Source/Runtime/Engine/Private/StaticMesh.cpp
+++ b/Engine/Source/Runtime/Engine/Private/StaticMesh.cpp
@@ -864,6 +864,7 @@ void FStaticMeshLODResources::GetResourceSizeEx(FResourceSizeEx& CumulativeResou
 
 	const int32 IndexStride = (IndexBuffer.Is32Bit() ? 4 : 2);
 	CumulativeResourceSize.AddUnknownMemoryBytes(TEXT("IndexBuffer"), IndexBuffer.GetNumIndices() * IndexStride);
+	CumulativeResourceSize.AddUnknownMemoryBytes(TEXT("DepthOnlyIndexBuffer"), DepthOnlyIndexBuffer.GetNumIndices() * IndexStride);
 
 	if (AdditionalIndexBuffers)
 	{
@@ -878,6 +879,7 @@ void FStaticMeshLODResources::GetResourceSizeEx(FResourceSizeEx& CumulativeResou
 SIZE_T FStaticMeshLODResources::GetCPUAccessMemoryOverhead() const
 {
 	int32 NumIndices = IndexBuffer.GetAllowCPUAccess() ? IndexBuffer.GetNumIndices() : 0;
+	NumIndices += DepthOnlyIndexBuffer.GetAllowCPUAccess() ? DepthOnlyIndexBuffer.GetNumIndices() : 0;
 	if (AdditionalIndexBuffers)
 	{
 		NumIndices += AdditionalIndexBuffers->ReversedDepthOnlyIndexBuffer.GetAllowCPUAccess() ? AdditionalIndexBuffers->ReversedDepthOnlyIndexBuffer.GetNumIndices() : 0;
