Substrate - bring back proper premultipled alpha that can be toggle on materials and material instances.
 - Coverage over the background can be override from the root node when AlphaCompiste blend mode is selected. If not plugged in, coverage of over brackground is off and regular blending is used.
 - Decals only support translucent, composite and modulate. Normal is ignored when using CompositeAlpha (seen somewhere)
 - Light functions can only be opaque so nothing to do in this case.
 - UI domain tested with all blend mode
 - Material attribute aslo tested
 - Fixed many of issues with blend mode defines, being forced, etc.

diff --git a/Engine/Shaders/Private/BasePassPixelShader.usf b/Engine/Shaders/Private/BasePassPixelShader.usf
--- a/Engine/Shaders/Private/BasePassPixelShader.usf
+++ b/Engine/Shaders/Private/BasePassPixelShader.usf
@@ -1825,11 +1825,10 @@ void FPixelShaderInOut_MainPS(
 		#endif
 
 	#elif STRATA_TRANSLUCENT_MATERIAL
-			
-		#if STRATA_LEGACY_PREMULT_ALPHA_OVERRIDE && !STRATA_OPTIMIZED_UNLIT
-			// This data patching is working because DualBlendSurfaceCoverage is always 1 when converting material from legacy to Strata with premultiplied alpha blending.
-			DualBlendSurfaceTransmittancePreCoverage= StrataData.PreMultipliedAlphaOverrideCoverage >= 0.0f ? 0.0f											: DualBlendSurfaceTransmittancePreCoverage;
-			DualBlendSurfaceCoverage				= StrataData.PreMultipliedAlphaOverrideCoverage >= 0.0f ? StrataData.PreMultipliedAlphaOverrideCoverage	: DualBlendSurfaceCoverage;
+
+		#if STRATA_USE_PREMULTALPHA_OVERRIDE // AlphaComposite - Premultiplied alpha blending
+			DualBlendSurfaceTransmittancePreCoverage= 0.0;
+			DualBlendSurfaceCoverage = PixelMaterialInputs.Opacity;
 		#endif
 
 			// Add fog luminance according to surfacecoverage and reduce surface luminance according to fog coverage.
diff --git a/Engine/Shaders/Private/DBufferDecalShared.ush b/Engine/Shaders/Private/DBufferDecalShared.ush
--- a/Engine/Shaders/Private/DBufferDecalShared.ush
+++ b/Engine/Shaders/Private/DBufferDecalShared.ush
@@ -129,9 +129,7 @@ FStrataDBuffer GetDBufferData(float2 BufferUV, uint RTMaskBit)
 #if MATERIAL_IS_STRATA
 void ApplyDBufferData(in FStrataDBuffer StrataDBufferData, inout FStrataPixelHeader StrataPixelHeader, inout FStrataData StrataData)
 {
-	// Early out if DBuffer has no coverage
-	if (StrataDBufferData.Coverage <= 0)
-		return;
+	// We cannot early out due when coverage = 0 due to the possibility of using premultipled alpha.
 
 	// For normal, apply decal data onto all basis, without consideration of top layers
 	#if STRATA_DBUFFER_RESPONSE_NORMAL && STRATA_INLINE_SHADING
diff --git a/Engine/Shaders/Private/DecalCommon.ush b/Engine/Shaders/Private/DecalCommon.ush
--- a/Engine/Shaders/Private/DecalCommon.ush
+++ b/Engine/Shaders/Private/DecalCommon.ush
@@ -17,10 +17,12 @@ float DecalCommonOutput(
 	inout FPixelShaderOut Out, 
 	in FStrataPixelHeader InHeader,
 	in FStrataBSDF InBSDF,
-	in float CoverageWeight)
+	in float FadeOutWeight,
+	in float CoverageOverBackground)
 {
 	FStrataDBuffer DBuffer = InHeader.StrataConvertToDBuffer(InBSDF);
-	DBuffer.Coverage *= CoverageWeight;
+
+	DBuffer.Coverage *= FadeOutWeight;
 
 	float3 Color = DBuffer.Emissive;
 
@@ -51,9 +53,11 @@ float DecalCommonOutput(
 	float4 PreMulEmissive_Coverage = 0;
 	StrataPackDBuffer(DBuffer, PreMulEmissive_Coverage, Out.MRT[0], Out.MRT[1], Out.MRT[2]);
 
-	#if DECAL_RENDERTARGET_COUNT == 4
-		Out.MRT[3] = 1.0f; // DBufferMask
-	#endif
+	// This is to set the coverage over the background to an arbitrary value in case premultiplied alpha is used.
+	Out.MRT[0].a = CoverageOverBackground;
+	Out.MRT[1].a = CoverageOverBackground;
+	Out.MRT[2].a = CoverageOverBackground;
+	Out.MRT[3].a = CoverageOverBackground;
 
 #endif
 	// Allow shader compiler to cull outputs that we won't use.
diff --git a/Engine/Shaders/Private/DeferredDecal.usf b/Engine/Shaders/Private/DeferredDecal.usf
--- a/Engine/Shaders/Private/DeferredDecal.usf
+++ b/Engine/Shaders/Private/DeferredDecal.usf
@@ -159,8 +159,15 @@ void FPixelShaderInOut_MainPS(inout FPixelShaderIn In, inout FPixelShaderOut Out
 		const FStrataIntegrationSettings Settings = InitStrataIntegrationSettings(false /*bForceFullyRough*/, false /*bRoughDiffuseEnabled*/, 0, false/*bRoughnessTracking*/);
 		StrataHeader.StrataUpdateTree(NullStrataAddressing, V, Settings);
 
+		// Either propagate the BSDF opacity, or use the one for AlphaComposite (premultiplied alpha)
+		float CoverageOverBackground = StrataHeader.StrataTree.BSDFs[0].Coverage;
+#if STRATA_USE_PREMULTALPHA_OVERRIDE // AlphaComposite - Premultiplied alpha blending
+		CoverageOverBackground = GetMaterialOpacity(PixelMaterialInputs);
+#endif 
+
 		// Since the Convert-To-Decal node forces parameter blending for the entire Strata tree, the BSDF we are interested in will be in 0.
-		Coverage = DecalCommonOutput(Out, StrataHeader, StrataHeader.StrataTree.BSDFs[0], DecalFading);
+		Coverage = DecalCommonOutput(Out, StrataHeader, StrataHeader.StrataTree.BSDFs[0], DecalFading, CoverageOverBackground);
+
 	}
 	const float Opacity = Coverage;
 #else // STRATA_ENABLED
diff --git a/Engine/Shaders/Private/MaterialTemplate.ush b/Engine/Shaders/Private/MaterialTemplate.ush
--- a/Engine/Shaders/Private/MaterialTemplate.ush
+++ b/Engine/Shaders/Private/MaterialTemplate.ush
@@ -143,13 +143,16 @@
 	#define ALLOW_CLIP 1
 #endif
 
-#define TEMPLATE_USES_STRATA			(STRATA_ENABLED && MATERIAL_IS_STRATA)
-#define STRATA_OPAQUE_MATERIAL			(TEMPLATE_USES_STRATA && ((!STRATA_MATERIAL_EXPORT_EXECUTED && MATERIALBLENDING_ANY_TRANSLUCENT<=0) || STRATA_MATERIAL_EXPORT_FROM_OPAQUE))
-#define STRATA_TRANSLUCENT_MATERIAL		(TEMPLATE_USES_STRATA && ((!STRATA_MATERIAL_EXPORT_EXECUTED && MATERIALBLENDING_ANY_TRANSLUCENT>0 ) || STRATA_MATERIAL_EXPORT_FROM_TRANSLUCENT))
-#define STRATA_OPAQUE_DEFERRED			(STRATA_OPAQUE_MATERIAL && !FORWARD_SHADING)
-#define STRATA_TRANSLUCENT_FORWARD		(STRATA_TRANSLUCENT_MATERIAL && !FORWARD_SHADING)
-#define STRATA_FORWARD_SHADING			(FORWARD_SHADING && (STRATA_OPAQUE_MATERIAL || STRATA_TRANSLUCENT_MATERIAL))
-#define STRATA_OUTPUT_ROUGH_REFRACTION	(STRATA_ENABLED && STRATA_OPAQUE_ROUGH_REFRACTION_ENABLED && STRATA_MATERIAL_OUTPUT_OPAQUE_ROUGH_REFRACTIONS)
+#define TEMPLATE_USES_STRATA				(STRATA_ENABLED && MATERIAL_IS_STRATA)
+#define STRATA_OPAQUE_MATERIAL				(TEMPLATE_USES_STRATA && ((!STRATA_MATERIAL_EXPORT_EXECUTED && MATERIALBLENDING_ANY_TRANSLUCENT<=0) || STRATA_MATERIAL_EXPORT_FROM_OPAQUE))
+#define STRATA_TRANSLUCENT_MATERIAL			(TEMPLATE_USES_STRATA && ((!STRATA_MATERIAL_EXPORT_EXECUTED && MATERIALBLENDING_ANY_TRANSLUCENT>0 ) || STRATA_MATERIAL_EXPORT_FROM_TRANSLUCENT))
+#define STRATA_OPAQUE_DEFERRED				(STRATA_OPAQUE_MATERIAL && !FORWARD_SHADING)
+#define STRATA_TRANSLUCENT_FORWARD			(STRATA_TRANSLUCENT_MATERIAL && !FORWARD_SHADING)
+#define STRATA_FORWARD_SHADING				(FORWARD_SHADING && (STRATA_OPAQUE_MATERIAL || STRATA_TRANSLUCENT_MATERIAL))
+#define STRATA_OUTPUT_ROUGH_REFRACTION		(STRATA_ENABLED && STRATA_OPAQUE_ROUGH_REFRACTION_ENABLED && STRATA_MATERIAL_OUTPUT_OPAQUE_ROUGH_REFRACTIONS)
+
+// Premultiplied alpha override is only enabled when alpha composite blending is used and the root node input is fed.
+#define STRATA_USE_PREMULTALPHA_OVERRIDE	(STRATA_ENABLED && MATERIALBLENDING_ALPHACOMPOSITE && STRATA_PREMULTIPLIED_ALPHA_OPACITY_OVERRIDEN)
 
 #if TEMPLATE_USES_STRATA
 #define STRATA_INLINE_SHADING 1
diff --git a/Engine/Shaders/Private/MeshDecals.usf b/Engine/Shaders/Private/MeshDecals.usf
--- a/Engine/Shaders/Private/MeshDecals.usf
+++ b/Engine/Shaders/Private/MeshDecals.usf
@@ -146,8 +146,14 @@ void FPixelShaderInOut_MainPS(
 		const FStrataIntegrationSettings Settings = InitStrataIntegrationSettings(false /*bForceFullyRough*/, false /*bRoughDiffuseEnabled*/, 0, false/*bRoughnessTracking*/);
 		StrataHeader.StrataUpdateTree(NullStrataAddressing, V, Settings);
 
+		// Either propagate the BSDF opacity, or use the one for AlphaComposite (premultiplied alpha)
+		float CoverageOverBackground = StrataHeader.StrataTree.BSDFs[0].Coverage;
+#if STRATA_USE_PREMULTALPHA_OVERRIDE // AlphaComposite - Premultiplied alpha blending
+		CoverageOverBackground = GetMaterialOpacity(PixelMaterialInputs);
+#endif 
+
 		// Since the Convert-To-Decal node forces parameter blending for the entire Strata tree, the BSDF we are interested in will be in 0.
-		Coverage = DecalCommonOutput(Out, StrataHeader, StrataHeader.StrataTree.BSDFs[0], 1.0f /*CoverageWeight*/);
+		Coverage = DecalCommonOutput(Out, StrataHeader, StrataHeader.StrataTree.BSDFs[0], 1.0f /*CoverageWeight*/, CoverageOverBackground);
 	}
 	const float Opacity = Coverage;
 #else // STRATA_ENABLED
diff --git a/Engine/Shaders/Private/MobileBasePassPixelShader.usf b/Engine/Shaders/Private/MobileBasePassPixelShader.usf
--- a/Engine/Shaders/Private/MobileBasePassPixelShader.usf
+++ b/Engine/Shaders/Private/MobileBasePassPixelShader.usf
@@ -691,11 +691,10 @@ void Main(
 	#endif
 
 	#if STRATA_TRANSLUCENT_MATERIAL
-			
-	#if STRATA_LEGACY_PREMULT_ALPHA_OVERRIDE && !STRATA_OPTIMIZED_UNLIT
-		// This data patching is working because DualBlendSurfaceCoverage is always 1 when converting material from legacy to Strata with premultiplied alpha blending.
-		DualBlendSurfaceTransmittancePreCoverage= StrataData.PreMultipliedAlphaOverrideCoverage >= 0.0f ? 0.0f											: DualBlendSurfaceTransmittancePreCoverage;
-		DualBlendSurfaceCoverage				= StrataData.PreMultipliedAlphaOverrideCoverage >= 0.0f ? StrataData.PreMultipliedAlphaOverrideCoverage	: DualBlendSurfaceCoverage;
+
+	#if STRATA_USE_PREMULTALPHA_OVERRIDE // AlphaComposite - Premultiplied alpha blending
+		DualBlendSurfaceTransmittancePreCoverage= 0.0;
+		DualBlendSurfaceCoverage = GetMaterialOpacity(PixelMaterialInputs);
 	#endif
 
 		// Add fog luminance according to surfacecoverage and reduce surface luminance according to fog coverage.
diff --git a/Engine/Shaders/Private/PostProcessMaterialShaders.usf b/Engine/Shaders/Private/PostProcessMaterialShaders.usf
--- a/Engine/Shaders/Private/PostProcessMaterialShaders.usf
+++ b/Engine/Shaders/Private/PostProcessMaterialShaders.usf
@@ -248,30 +248,51 @@ void MainPS(
 	// from the engine will have their alpha channel inverted to match the PNG convention. However, for consistency inside the rendering chain,
 	// we need to stick to 0.0f as the default value.
 	float Alpha = 0.0f;
+	float3 EmissiveColor = 0.0f;
 #if STRATA_ENABLED
 
 	// Unlit forces a single BSDF
 	FStrataBSDF UnlitBSDF = PixelMaterialInputs.FrontMaterial.InlinedBSDF;
 	UnlitBSDF.StrataSanitizeBSDF();
 
-	// Grab emissive colour as output
-  #if MATERIAL_OUTPUT_OPACITY_AS_ALPHA
-	// Transmittance = 1 - opacity. See UMaterialExpressionStrataPostProcess.
-	const float UnlitBSDFTransmittance = UNLIT_TRANSMITTANCE(UnlitBSDF).x;
-	const float UnlitBSDFOpacity = 1.0f - UnlitBSDFTransmittance;
-	Alpha = UnlitBSDFOpacity;
-  #endif
+	half3 PPColor = BSDF_GETEMISSIVE(UnlitBSDF);
+	half  PPAlpha = saturate(1.0 - UNLIT_TRANSMITTANCE(UnlitBSDF).x);	// The current interface only allows legacy opacity transformed to transmittance as 1-Opacity. See UMaterialExpressionStrataPostProcess.
+
+	// Substrate only use premultiplied alpha blending state (except for modulate and holdout) so we need to convert to the here.
+#if MATERIALBLENDING_SOLID || MATERIALBLENDING_MASKED
+	EmissiveColor = PPColor;
+#elif MATERIALBLENDING_ALPHACOMPOSITE
+	EmissiveColor = PPColor;	// Premultipled alpha already baked in by the user
+#elif MATERIALBLENDING_ALPHAHOLDOUT
+	EmissiveColor = 0.0;
+#elif MATERIALBLENDING_TRANSLUCENT
+	EmissiveColor = PPColor * PPAlpha;
+#elif MATERIALBLENDING_ADDITIVE
+	EmissiveColor = PPColor;	// not affected by alpha
+#elif MATERIALBLENDING_MODULATE
+	EmissiveColor = PPColor;	// not affected by alpha
+#else
+	EmissiveColor = PPColor;
+#endif
 
-	const float3 Emissive = BSDF_GETEMISSIVE(UnlitBSDF);
+#if STRATA_USE_PREMULTALPHA_OVERRIDE // AlphaComposite - Premultiplied alpha blending
+	PPAlpha = GetMaterialOpacity(PixelMaterialInputs);
+#endif
+
+#if MATERIAL_OUTPUT_OPACITY_AS_ALPHA
+	Alpha = PPAlpha;
+#endif
+
+#else // STRATA_ENABLED
 
-#else
 	// Grab emissive colour as output
   #if MATERIAL_OUTPUT_OPACITY_AS_ALPHA
 	Alpha = GetMaterialOpacity(PixelMaterialInputs);
   #endif
-	const float3 Emissive = GetMaterialEmissive(PixelMaterialInputs);
-#endif
-	OutColor = float4(Emissive, Alpha );
+	EmissiveColor = GetMaterialEmissive(PixelMaterialInputs);
+#endif // STRATA_ENABLED
+
+	OutColor = float4(EmissiveColor, Alpha );
 
 #if !MATERIALBLENDING_MODULATE && POST_PROCESS_MATERIAL_BEFORE_TONEMAP
 	OutColor.xyz *= View.PreExposure;
@@ -402,33 +423,51 @@ void MainPS(
 	// fill out other related material parameters
 	CalcMaterialParametersPost(Parameters, PixelMaterialInputs, SvPosition, true);
 
+	float Alpha = 0.0f;
+	float3 EmissiveColor = 0.0f;
 #if STRATA_ENABLED
 
 	// Unlit forces a single BSDF
 	FStrataBSDF UnlitBSDF = PixelMaterialInputs.FrontMaterial.InlinedBSDF;
 	UnlitBSDF.StrataSanitizeBSDF();
 
-	// Grab emissive colour as output
-  #if MATERIAL_OUTPUT_OPACITY_AS_ALPHA
-  // Transmittance = 1 - opacity. See UMaterialExpressionStrataPostProcess.
-	const half UnlitBSDFTransmittance = UNLIT_TRANSMITTANCE(UnlitBSDF).x;
-	const half UnlitBSDFOpacity = 1.0f - UnlitBSDFTransmittance;
-	half Alpha = UnlitBSDFOpacity;
-  #else
-	half Alpha = 0.0f;
-  #endif
+	half3 PPColor = BSDF_GETEMISSIVE(UnlitBSDF);
+	half  PPAlpha = saturate(1.0 - UNLIT_TRANSMITTANCE(UnlitBSDF).x);	// The current interface only allows legacy opacity transformed to transmittance as 1-Opacity. See UMaterialExpressionStrataPostProcess.
+
+	// Substrate only use premultiplied alpha blending state (except for modulate and holdout) so we need to convert to the here.
+#if MATERIALBLENDING_SOLID || MATERIALBLENDING_MASKED
+	EmissiveColor = PPColor;
+#elif MATERIALBLENDING_ALPHACOMPOSITE
+	EmissiveColor = PPColor;	// Premultipled alpha already baked in by the user
+#elif MATERIALBLENDING_ALPHAHOLDOUT
+	EmissiveColor = 0.0;
+#elif MATERIALBLENDING_TRANSLUCENT
+	EmissiveColor = PPColor * PPAlpha;
+#elif MATERIALBLENDING_ADDITIVE
+	EmissiveColor = PPColor;	// not affected by alpha
+#elif MATERIALBLENDING_MODULATE
+	EmissiveColor = PPColor;	// not affected by alpha
+#else
+	EmissiveColor = PPColor;
+#endif
+
+#if STRATA_USE_PREMULTALPHA_OVERRIDE // AlphaComposite - Premultiplied alpha blending
+	PPAlpha = GetMaterialOpacity(PixelMaterialInputs);
+#endif
 
-	const half3 EmissiveColor = BSDF_GETEMISSIVE(UnlitBSDF);
+#if MATERIAL_OUTPUT_OPACITY_AS_ALPHA
+	Alpha = PPAlpha;
+#endif
 
 #else // STRATA_ENABLED
 
 	// Grab emissive colour as output
-	half3 EmissiveColor = GetMaterialEmissive(PixelMaterialInputs);
+	EmissiveColor = GetMaterialEmissive(PixelMaterialInputs);
 	
   #if MATERIAL_OUTPUT_OPACITY_AS_ALPHA
-	half Alpha = GetMaterialOpacity(PixelMaterialInputs);
+	Alpha = GetMaterialOpacity(PixelMaterialInputs);
   #else
-	half Alpha = 0.0f;
+	Alpha = 0.0f;
   #endif
 
 #endif // STRATA_ENABLED
diff --git a/Engine/Shaders/Private/SlateElementPixelShader.usf b/Engine/Shaders/Private/SlateElementPixelShader.usf
--- a/Engine/Shaders/Private/SlateElementPixelShader.usf
+++ b/Engine/Shaders/Private/SlateElementPixelShader.usf
@@ -169,17 +169,30 @@ half4 GetMaterialColor( VertexToPixelInterpolants InVertex, float2 MeshUV )
 	UnlitBSDF.StrataSanitizeBSDF();
 
 #if ((STRATA_OPAQUE_DEFERRED || (STRATA_OPAQUE_MATERIAL && STRATA_FORWARD_SHADING)))
+
 	OutColor = half4(BSDF_GETEMISSIVE(UnlitBSDF), 1.0f);
-#else // STRATA_TRANSLUCENT_FORWARD || STRATA_FORWARD_SHADING || STRATA_MATERIAL_EXPORT_FROM_TRANSLUCENT
-	const float Coverage = saturate(1.0f - UNLIT_TRANSMITTANCE(UnlitBSDF).x);	// Unlit node transmittance is only generated from opacity
-#if MATERIALBLENDING_ADDITIVE
-	OutColor = half4(BSDF_GETEMISSIVE(UnlitBSDF) * Coverage, 0.0f);
-#else
-	OutColor = half4(BSDF_GETEMISSIVE(UnlitBSDF), Coverage);
-#endif
+
+#else // STRATA_OPAQUE_DEFERRED
+
+	float3 Color = BSDF_GETEMISSIVE(UnlitBSDF);
+	float Coverage = saturate(1.0f - UNLIT_TRANSMITTANCE(UnlitBSDF).x);	// Unlit node transmittance is only generated from opacity
+
+#if MATERIALBLENDING_ALPHACOMPOSITE
+
+#if STRATA_USE_PREMULTALPHA_OVERRIDE // AlphaComposite - Premultiplied alpha blending
+	Coverage = GetMaterialOpacity(PixelMaterialInputs);
 #endif
+	OutColor = half4(Color, Coverage);
+
+#elif MATERIALBLENDING_MODULATE
+	OutColor = half4(Color, 0.0f);
+#elif MATERIALBLENDING_ALPHAHOLDOUT
+	OutColor = half4(0.0f, 0.0f, 0.0f, Coverage);
+#else // MATERIALBLENDING_ADDITIVE || MATERIALBLENDING_TRANSLUCENT
+	OutColor = half4(Color * Coverage, Coverage);
+#endif // MATERIALBLENDING_xxx
 
-	// STRATA_TODO fix premultiplied alpha properly
+#endif// STRATA_OPAQUE_DEFERRED
 
 #else // STRATA_ENABLED
 
diff --git a/Engine/Shaders/Private/Strata/Strata.ush b/Engine/Shaders/Private/Strata/Strata.ush
--- a/Engine/Shaders/Private/Strata/Strata.ush
+++ b/Engine/Shaders/Private/Strata/Strata.ush
@@ -97,9 +97,6 @@
 
 #define STRATA_COMPLEXPATH (!STRATA_FASTPATH && !STRATA_SINGLEPATH)
 
-// When converting legacy materials to Strata, we need to override transmittance/coverage to be independent from the material and directly linked to opacity
-#define STRATA_LEGACY_PREMULT_ALPHA_OVERRIDE (STRATA_USES_CONVERSION_FROM_LEGACY && MATERIALBLENDING_ALPHACOMPOSITE)
-
 #define STRATA_LEGACY_MATERIAL_APPLIES_FINAL_WEIGHT (!MATERIALBLENDING_MASKED && !MATERIALBLENDING_SOLID && STRATA_LEGACY_PREMULT_ALPHA_OVERRIDE==0 || STRATA_MATERIAL_EXPORT_FROM_TRANSLUCENT)
 
 #define STRATA_INLINE_SINGLELAYERWATER (STRATA_ENABLED && MATERIAL_IS_STRATA && STRATA_INLINE_SHADING && MATERIAL_SHADINGMODEL_SINGLELAYERWATER)
@@ -1378,18 +1375,11 @@ struct FStrataData
 #if USE_DEVELOPMENT_SHADERS
 	float3 PreviewColor;	// Used as preview in the material graph within the editor.
 #endif
-
-#if STRATA_LEGACY_PREMULT_ALPHA_OVERRIDE
-	float PreMultipliedAlphaOverrideCoverage;
-#endif
 };
 
 FStrataData GetInitialisedStrataData()
 {
 	FStrataData StrataData = (FStrataData)0;
-#if STRATA_LEGACY_PREMULT_ALPHA_OVERRIDE
-	StrataData.PreMultipliedAlphaOverrideCoverage = -1.0f;
-#endif
 	return StrataData;
 }
 
diff --git a/Engine/Shaders/Private/Strata/StrataLegacyConversion.ush b/Engine/Shaders/Private/Strata/StrataLegacyConversion.ush
--- a/Engine/Shaders/Private/Strata/StrataLegacyConversion.ush
+++ b/Engine/Shaders/Private/Strata/StrataLegacyConversion.ush
@@ -436,10 +436,6 @@ FStrataData StrataConvertLegacyMaterialDynamic(
 	Slab.InlinedBSDF.Coverage = LegacySlab.Weight * FinalWeight;	//////////////////// Remove FinalWeight? Do that in CreateLegacySlab?
 #endif
 
-	#if STRATA_LEGACY_PREMULT_ALPHA_OVERRIDE
-	Slab.PreMultipliedAlphaOverrideCoverage = Opacity;
-	#endif
-
 	return Slab;
 }
 
@@ -472,6 +468,12 @@ FStrataData StrataConvertLegacyMaterialStatic(
 	const float3 Zeros = float3(0, 0, 0);
 	const float3 Ones = float3(1, 1, 1);
 
+#if STRATA_USE_PREMULTALPHA_OVERRIDE
+	// Opacity is applied using the "Opacity Override" pin on the root node when using alpha composite (pre multiplied alpha).
+	// This is needed because some material might have the pin plugged in and the blending mode overriden on the material instance.
+	Opacity = 1.0f;
+#endif
+
 	if (ShadingModel == SHADINGMODELID_UNLIT)
 	{
 		const float GreyTransmittance = saturate(1.0f - Opacity);
diff --git a/Engine/Source/Editor/MaterialEditor/Private/MaterialEditorDetailCustomization.cpp b/Engine/Source/Editor/MaterialEditor/Private/MaterialEditorDetailCustomization.cpp
--- a/Engine/Source/Editor/MaterialEditor/Private/MaterialEditorDetailCustomization.cpp
+++ b/Engine/Source/Editor/MaterialEditor/Private/MaterialEditorDetailCustomization.cpp
@@ -60,9 +60,6 @@ UEnum* GetBlendModeEnum()
 
 		// BLEND_TranslucentColoredTransmittance is only supported in Strata mode
 		BlendModeEnum->SetMetaData(TEXT("DisplayName"), TEXT("TranslucentColoredTransmittance"), BLEND_TranslucentColoredTransmittance);
-
-		// BLEND_AlphaComposite is identical to BLEND_Translucent in Strata mode
-		BlendModeEnum->SetMetaData(TEXT("Hidden"), TEXT("True"), BLEND_AlphaComposite);
 	}
 	else
 	{
diff --git a/Engine/Source/Editor/UnrealEd/Classes/MaterialGraph/MaterialGraph.h b/Engine/Source/Editor/UnrealEd/Classes/MaterialGraph/MaterialGraph.h
--- a/Engine/Source/Editor/UnrealEd/Classes/MaterialGraph/MaterialGraph.h
+++ b/Engine/Source/Editor/UnrealEd/Classes/MaterialGraph/MaterialGraph.h
@@ -67,7 +67,7 @@ struct FMaterialInputInfo
 		{
 			if (Strata::IsStrataEnabled())
 			{
-				if (Property == MP_FrontMaterial || Property == MP_OpacityMask || Property == MP_WorldPositionOffset)
+				if (Property == MP_FrontMaterial || Property == MP_Opacity || Property == MP_OpacityMask || Property == MP_WorldPositionOffset)
 				{
 					return true;
 				}
diff --git a/Engine/Source/Editor/UnrealEd/Private/MaterialGraph.cpp b/Engine/Source/Editor/UnrealEd/Private/MaterialGraph.cpp
--- a/Engine/Source/Editor/UnrealEd/Private/MaterialGraph.cpp
+++ b/Engine/Source/Editor/UnrealEd/Private/MaterialGraph.cpp
@@ -242,7 +242,7 @@ void UMaterialGraph::RebuildGraphInternal(const TMap<UMaterialExpression*, TArra
 		MaterialInputs.Add(FMaterialInputInfo(FMaterialAttributeDefinitionMap::GetDisplayNameForMaterial(MP_Roughness, Material), MP_Roughness, LOCTEXT("RoughnessToolTip", "Controls how rough the Material is. Roughness of 0 (smooth) is a mirror reflection and 1 (rough) is completely matte or diffuse")));
 		MaterialInputs.Add(FMaterialInputInfo(FMaterialAttributeDefinitionMap::GetDisplayNameForMaterial(MP_Anisotropy, Material), MP_Anisotropy, LOCTEXT("AnisotropyToolTip", "Determines the extent the specular highlight is stretched along the tangent. Anisotropy from 0 to 1 results in a specular highlight that stretches from uniform to maximally stretched along the tangent direction.")));
 		MaterialInputs.Add(FMaterialInputInfo(FMaterialAttributeDefinitionMap::GetDisplayNameForMaterial(MP_EmissiveColor, Material), MP_EmissiveColor, LOCTEXT("EmissiveToolTip", "Controls which parts of your Material will appear to glow")));
-		MaterialInputs.Add(FMaterialInputInfo(FMaterialAttributeDefinitionMap::GetDisplayNameForMaterial(MP_Opacity, Material), MP_Opacity, LOCTEXT("OpacityToolTip", "Controls the translucency of the Material")));
+		MaterialInputs.Add(FMaterialInputInfo(FMaterialAttributeDefinitionMap::GetDisplayNameForMaterial(MP_Opacity, Material), MP_Opacity, LOCTEXT("OpacityToolTip", "Controls the translucency of the Material. When Substrate is enabled, this input override alpha blending over the background when AlphaComposite blend mode is selected.")));
 		MaterialInputs.Add(FMaterialInputInfo(FMaterialAttributeDefinitionMap::GetDisplayNameForMaterial(MP_OpacityMask, Material), MP_OpacityMask, LOCTEXT("OpacityMaskToolTip", "When in Masked mode, a Material is either completely visible or completely invisible")));
 		MaterialInputs.Add(FMaterialInputInfo(FMaterialAttributeDefinitionMap::GetDisplayNameForMaterial(MP_Normal, Material), MP_Normal, LOCTEXT("NormalToolTip", "Takes the input of a normal map")));
 		MaterialInputs.Add(FMaterialInputInfo(FMaterialAttributeDefinitionMap::GetDisplayNameForMaterial(MP_Tangent, Material), MP_Tangent, LOCTEXT("TangentToolTip", "Takes the input of a tangent map. Useful for specifying anisotropy direction.")));
diff --git a/Engine/Source/Runtime/Engine/Private/Materials/HLSLMaterialTranslator.cpp b/Engine/Source/Runtime/Engine/Private/Materials/HLSLMaterialTranslator.cpp
--- a/Engine/Source/Runtime/Engine/Private/Materials/HLSLMaterialTranslator.cpp
+++ b/Engine/Source/Runtime/Engine/Private/Materials/HLSLMaterialTranslator.cpp
@@ -1227,9 +1227,11 @@ bool FHLSLMaterialTranslator::Translate()
 			AddEstimatedTextureSample(1);
 		}
 
+		bOpacityPropertyIsUsed = IsMaterialPropertyUsed(MP_Opacity, Chunk[MP_Opacity], FLinearColor(1, 0, 0, 0), 1);
+
 		bUsesCurvature = FeatureLevel == ERHIFeatureLevel::ES3_1 && 
 						((MaterialShadingModels.HasShadingModel(MSM_SubsurfaceProfile) && IsMaterialPropertyUsed(MP_CustomData0, Chunk[MP_CustomData0], FLinearColor(1, 0, 0, 0), 1))
-						|| (MaterialShadingModels.HasShadingModel(MSM_Eye) && IsMaterialPropertyUsed(MP_Opacity, Chunk[MP_Opacity], FLinearColor(1, 0, 0, 0), 1)));
+						|| (MaterialShadingModels.HasShadingModel(MSM_Eye) && bOpacityPropertyIsUsed));
 
 		if (!bStrataEnabled && IsModulateBlendMode(BlendMode) && MaterialShadingModels.IsLit() && !Material->IsDeferredDecal())
 		{
@@ -2255,7 +2257,9 @@ void FHLSLMaterialTranslator::GetMaterialEnvironment(EShaderPlatform InPlatform,
 	bMaterialRequestsDualSourceBlending |= bMaterialIsStrata && Material->GetBlendMode() == EBlendMode::BLEND_TranslucentColoredTransmittance;
 
 	// if duals source blending (colored transmittance) is not supported on a platform, it will fall back to standard alpha blending (grey scale transmittance)
-	OutEnvironment.SetDefine(TEXT("DUAL_SOURCE_COLOR_BLENDING_ENABLED"), bMaterialRequestsDualSourceBlending && Material->IsDualBlendingEnabled(Platform) ? TEXT("1") : TEXT("0"));
+	OutEnvironment.SetDefine(TEXT("DUAL_SOURCE_COLOR_BLENDING_ENABLED"), bMaterialRequestsDualSourceBlending&& Material->IsDualBlendingEnabled(Platform) ? TEXT("1") : TEXT("0"));
+
+	OutEnvironment.SetDefine(TEXT("STRATA_PREMULTIPLIED_ALPHA_OPACITY_OVERRIDEN"), bMaterialIsStrata && bOpacityPropertyIsUsed ? TEXT("1") : TEXT("0"));
 
 	if (bMaterialIsStrata)
 	{
diff --git a/Engine/Source/Runtime/Engine/Private/Materials/HLSLMaterialTranslator.h b/Engine/Source/Runtime/Engine/Private/Materials/HLSLMaterialTranslator.h
--- a/Engine/Source/Runtime/Engine/Private/Materials/HLSLMaterialTranslator.h
+++ b/Engine/Source/Runtime/Engine/Private/Materials/HLSLMaterialTranslator.h
@@ -408,7 +408,10 @@ protected:
 	/** True if the material is detected as a strata material at compile time.
 	 * This is decoupled from runtime FMaterialResource::IsStrataMaterial but practically fine since this is only temporary until Strata is the main shading system. Only really used at runtime for translucency dual source blending.
 	 */
-	uint32 bMaterialIsStrata : 1; // 
+	uint32 bMaterialIsStrata : 1;
+
+	/** True if the opacity input is plugged in */
+	uint32 bOpacityPropertyIsUsed : 1;
 	
 	uint32 bEnableExecutionFlow : 1;
 
diff --git a/Engine/Source/Runtime/Engine/Private/Materials/Material.cpp b/Engine/Source/Runtime/Engine/Private/Materials/Material.cpp
--- a/Engine/Source/Runtime/Engine/Private/Materials/Material.cpp
+++ b/Engine/Source/Runtime/Engine/Private/Materials/Material.cpp
@@ -3116,7 +3116,7 @@ EBlendMode ConvertLegacyBlendMode(EBlendMode InBlendMode, FMaterialShadingModelF
 	{
 		return BLEND_TranslucentColoredTransmittance;
 	}
-	else if (InBlendMode == BLEND_Translucent || InBlendMode == BLEND_AlphaComposite)
+	else if (InBlendMode == BLEND_Translucent)
 	{
 		return BLEND_TranslucentGreyTransmittance;
 	}
@@ -3297,6 +3297,7 @@ void UMaterial::ConvertMaterialToStrataMaterial()
 		// * Always forward masked opacity because this is required when the blend mode is overriden to masked on a material instance.
 		bUseMaterialAttributes = false;
 		EditorOnly->FrontMaterial.Connect(0, ConvertNode);
+		EditorOnly->Opacity.Connect(6, BreakMatAtt);			// For alpha composite blend mode, Opacity on the root node in case BLEND_AlphaComposite is selected.
 		EditorOnly->OpacityMask.Connect(7, BreakMatAtt);
 		EditorOnly->WorldPositionOffset.Connect(10, BreakMatAtt);
 		EditorOnly->AmbientOcclusion.Connect(14, BreakMatAtt);
@@ -3472,7 +3473,7 @@ void UMaterial::ConvertMaterialToStrataMaterial()
 				MoveConnectionTo(EditorOnly->SubsurfaceColor, ConvertNode, 8);
 				MoveConnectionTo(EditorOnly->ClearCoat, ConvertNode, 9);
 				MoveConnectionTo(EditorOnly->ClearCoatRoughness, ConvertNode, 10);
-				MoveConnectionTo(EditorOnly->Opacity, ConvertNode, 11);
+				CopyConnectionTo(EditorOnly->Opacity, ConvertNode, 11);	// We only copy, to keep Opacity on the root node in case BLEND_AlphaComposite is selected.
 				bRelinkCustomOutputNodes = true;
 			
 				// Shading Model
@@ -3568,26 +3569,9 @@ void UMaterial::ConvertMaterialToStrataMaterial()
 
 			UMaterialExpressionStrataPostProcess* PostProcNode = NewObject<UMaterialExpressionStrataPostProcess>(this);
 			SetPosXAndMoveReferenceToTheRight(PostProcNode);
-			if (GetBlendMode() == BLEND_Translucent)
-			{
-				// If the blending mode was translucent, we must multiply the color by the opacity before the output node
-				UMaterialExpressionMultiply* ColorMultiplyOpacityNode = NewObject<UMaterialExpressionMultiply>(this);
-				ReplaceNodeAndMoveToTheRight(PostProcNode, ColorMultiplyOpacityNode);
-				CopyConnectionTo(EditorOnly->Opacity, ColorMultiplyOpacityNode, 0);
-				MoveConnectionTo(EditorOnly->EmissiveColor, ColorMultiplyOpacityNode, 1);
 
-				PostProcNode->Color.Connect(0, ColorMultiplyOpacityNode);
-				MoveConnectionTo(EditorOnly->Opacity, PostProcNode, 1);
-			}
-			else
-			{
-				// Only Emissive & Opacity are valid input for PostProcess material
-				MoveConnectionTo(EditorOnly->EmissiveColor, PostProcNode, 0);
-				if (GetBlendMode() != BLEND_Additive) // Legacy material additive was ignoring opacity
-				{
-					MoveConnectionTo(EditorOnly->Opacity, PostProcNode, 1);
-				}
-			}
+			MoveConnectionTo(EditorOnly->EmissiveColor, PostProcNode, 0);
+			CopyConnectionTo(EditorOnly->Opacity, PostProcNode, 1);	// We only copy, to keep Opacity on the root node in case BLEND_AlphaComposite is selected.
 
 			EditorOnly->FrontMaterial.Connect(0, PostProcNode);
 			bInvalidateShader = true;
@@ -3612,7 +3596,7 @@ void UMaterial::ConvertMaterialToStrataMaterial()
 			MoveConnectionTo(EditorOnly->SubsurfaceColor, ConvertNode, 8);
 			MoveConnectionTo(EditorOnly->ClearCoat, ConvertNode, 9);
 			MoveConnectionTo(EditorOnly->ClearCoatRoughness, ConvertNode, 10);
-			MoveConnectionTo(EditorOnly->Opacity, ConvertNode, 11);
+			CopyConnectionTo(EditorOnly->Opacity, ConvertNode, 11);	// We only copy, to keep Opacity on the root node in case BLEND_AlphaComposite is selected.
 
 			// Add constant for the Unlit shading model
 			UMaterialExpressionShadingModel* ShadingModelNode = NewObject<UMaterialExpressionShadingModel>(this);
@@ -3641,7 +3625,7 @@ void UMaterial::ConvertMaterialToStrataMaterial()
 			UMaterialExpressionStrataUI* UINode = NewObject<UMaterialExpressionStrataUI>(this);
 			SetPosXAndMoveReferenceToTheRight(UINode);
 			MoveConnectionTo(EditorOnly->EmissiveColor, UINode, 0);
-			MoveConnectionTo(EditorOnly->Opacity, UINode, 1);
+			CopyConnectionTo(EditorOnly->Opacity, UINode, 1);	// We only copy, to keep Opacity on the root node in case BLEND_AlphaComposite is selected.
 
 			EditorOnly->FrontMaterial.Connect(0, UINode);
 			bInvalidateShader = true;
@@ -5061,7 +5045,6 @@ void UMaterial::RebuildShadingModelField()
 				// Decal can have multiple shading model
 				MaterialDomain = EMaterialDomain::MD_DeferredDecal;
 				ShadingModel = MSM_DefaultLit;
-				BlendMode = EBlendMode::BLEND_TranslucentGreyTransmittance;
 			}
 
 			// Also update the ShadingModels for remaining pipeline operation
@@ -6667,6 +6650,7 @@ bool UMaterial::IsPropertySupported(EMaterialProperty InProperty) const
 		switch (InProperty)
 		{
 		case MP_Refraction:
+		case MP_Opacity:	// For Substrate, this is only use for premultiply alpha blending where throughput applied to the background can be overriden.
 		case MP_OpacityMask:
 		case MP_AmbientOcclusion:
 		case MP_WorldPositionOffset:
@@ -6708,12 +6692,13 @@ static bool IsPropertyActive_Internal(EMaterialProperty InProperty,
 	bool bIsSupported)
 {
 	const bool bStrataEnabled = Strata::IsStrataEnabled();
+	const bool bStrataOpacityOverrideAllowed = BlendMode == BLEND_AlphaComposite; // Should we always have it enabled to be able to be plugged in an fed when blend mode is toggled later on a material instance?
 
 	if (Domain == MD_PostProcess)
 	{
 		if (bStrataEnabled)
 		{
-			return InProperty == MP_FrontMaterial;
+			return InProperty == MP_FrontMaterial || (InProperty == MP_Opacity && bStrataOpacityOverrideAllowed);
 		}
 		else
 		{
@@ -6737,7 +6722,8 @@ static bool IsPropertyActive_Internal(EMaterialProperty InProperty,
 		if (bStrataEnabled)
 		{
 			return InProperty == MP_FrontMaterial
-				|| InProperty == MP_AmbientOcclusion;
+				|| InProperty == MP_AmbientOcclusion
+				|| (InProperty == MP_Opacity && bStrataOpacityOverrideAllowed);
 		}
 		else if (InProperty >= MP_CustomizedUVs0 && InProperty <= MP_CustomizedUVs7)
 		{
@@ -6792,9 +6778,9 @@ static bool IsPropertyActive_Internal(EMaterialProperty InProperty,
 	}
 	else if (Domain == MD_Volume)
 	{
-		if (InProperty == MP_FrontMaterial)
+		if (bStrataEnabled)
 		{
-			return bStrataEnabled;
+			return InProperty == MP_FrontMaterial;
 		}
 		return InProperty == MP_EmissiveColor
 			|| InProperty == MP_SubsurfaceColor
@@ -6808,6 +6794,7 @@ static bool IsPropertyActive_Internal(EMaterialProperty InProperty,
 			return InProperty == MP_FrontMaterial
 				|| (InProperty == MP_WorldPositionOffset)
 				|| (InProperty == MP_OpacityMask && IsMaskedBlendMode(BlendMode))
+				|| (InProperty == MP_Opacity && bStrataOpacityOverrideAllowed)
 				|| (InProperty >= MP_CustomizedUVs0 && InProperty <= MP_CustomizedUVs7);
 		}
 		else
@@ -6842,7 +6829,8 @@ static bool IsPropertyActive_Internal(EMaterialProperty InProperty,
 				Active = (bIsTranslucentBlendMode && !IsAlphaHoldoutBlendMode(BlendMode) && !IsModulateBlendMode(BlendMode) && bUsesDistortion) || ShadingModels.HasShadingModel(MSM_SingleLayerWater);
 				break;
 			case MP_Opacity:
-				Active = (bIsTranslucentBlendMode && !IsModulateBlendMode(BlendMode)) || ShadingModels.HasShadingModel(MSM_SingleLayerWater);
+				// Opacity is used as alpha override for alpha composite blending. 
+				Active = bStrataOpacityOverrideAllowed;
 				break;
 			case MP_OpacityMask:
 				Active = IsMaskedBlendMode(BlendMode);
diff --git a/Engine/Source/Runtime/Engine/Private/Materials/MaterialAttributeDefinitionMap.cpp b/Engine/Source/Runtime/Engine/Private/Materials/MaterialAttributeDefinitionMap.cpp
--- a/Engine/Source/Runtime/Engine/Private/Materials/MaterialAttributeDefinitionMap.cpp
+++ b/Engine/Source/Runtime/Engine/Private/Materials/MaterialAttributeDefinitionMap.cpp
@@ -353,12 +353,15 @@ FText FMaterialAttributeDefinitionMap::GetAttributeOverrideForMaterial(const FGu
 	TArray<TKeyValuePair<EMaterialShadingModel, FString>> CustomPinNames;
 	EMaterialProperty Property = GMaterialPropertyAttributesMap.Find(AttributeID)->Property;
 
+	static const IConsoleVariable* SubstrateCVar = IConsoleManager::Get().FindConsoleVariable(TEXT("r.substrate"));
+	const bool bStrataEnabled = (SubstrateCVar && SubstrateCVar->GetInt() != 0);
+
 	switch (Property)
 	{
 	case MP_EmissiveColor:
 		return Material->IsUIMaterial() ? LOCTEXT("UIOutputColor", "Final Color") : LOCTEXT("EmissiveColor", "Emissive Color");
 	case MP_Opacity:
-		return LOCTEXT("Opacity", "Opacity");
+		return bStrataEnabled ? LOCTEXT("OpacityOverride", "Opacity Override") : LOCTEXT("Opacity", "Opacity");
 	case MP_OpacityMask:
 		return LOCTEXT("OpacityMask", "Opacity Mask");
 	case MP_DiffuseColor:
diff --git a/Engine/Source/Runtime/Engine/Private/Materials/MaterialShared.cpp b/Engine/Source/Runtime/Engine/Private/Materials/MaterialShared.cpp
--- a/Engine/Source/Runtime/Engine/Private/Materials/MaterialShared.cpp
+++ b/Engine/Source/Runtime/Engine/Private/Materials/MaterialShared.cpp
@@ -2431,7 +2431,7 @@ void FMaterial::SetupMaterialEnvironment(
 			OutEnvironment.SetDefine(TEXT("STRATA_BLENDING_MASKED"), TEXT("1"));
 			break;
 		}
-		case BLEND_Additive: // STRATA_TODO_BLENDMODE_ADDITIVE
+		case BLEND_Additive:
 		{
 			OutEnvironment.SetDefine(TEXT("MATERIALBLENDING_ADDITIVE"),TEXT("1"));
 			OutEnvironment.SetDefine(TEXT("STRATA_BLENDING_TRANSLUCENT_GREYTRANSMITTANCE"), TEXT("1"));
diff --git a/Engine/Source/Runtime/Renderer/Private/DecalRenderingCommon.cpp b/Engine/Source/Runtime/Renderer/Private/DecalRenderingCommon.cpp
--- a/Engine/Source/Runtime/Renderer/Private/DecalRenderingCommon.cpp
+++ b/Engine/Source/Runtime/Renderer/Private/DecalRenderingCommon.cpp
@@ -50,7 +50,7 @@ namespace DecalRendering
 		}
 
 		// Enforce blend modes output limitations.
-		if (!bIsStrataEnabled && IsAlphaCompositeBlendMode(Desc))
+		if (IsAlphaCompositeBlendMode(Desc))
 		{
 			Desc.bWriteNormal = false;
 		}
diff --git a/Engine/Source/Runtime/Renderer/Private/PostProcess/PostProcessMaterial.cpp b/Engine/Source/Runtime/Renderer/Private/PostProcess/PostProcessMaterial.cpp
--- a/Engine/Source/Runtime/Renderer/Private/PostProcess/PostProcessMaterial.cpp
+++ b/Engine/Source/Runtime/Renderer/Private/PostProcess/PostProcessMaterial.cpp
@@ -144,7 +144,8 @@ static FRHIBlendState* GetMaterialBlendState(const FMaterial* Material)
 		case EBlendMode::BLEND_Opaque:
 		case EBlendMode::BLEND_Masked:
 			return TStaticBlendState<>::GetRHI();
-		case EBlendMode::BLEND_Additive: 						// STRATA_TODO_BLENDMODE_ADDITIVE
+		case EBlendMode::BLEND_Additive:
+			return TStaticBlendState<CW_RGB, BO_Add, BF_One, BF_One>::GetRHI();
 		case EBlendMode::BLEND_AlphaComposite:
 		case EBlendMode::BLEND_TranslucentColoredTransmittance: // A platform may not support dual source blending so we always only use grey scale transmittance
 		case EBlendMode::BLEND_TranslucentGreyTransmittance:
