Fix hash calculation of VT pool settings.
Old code was not hashing array contents.
This could lead to not correctly detecting when the values are changing.

diff --git a/Engine/Source/Runtime/Engine/Private/VT/VirtualTextureScalability.cpp b/Engine/Source/Runtime/Engine/Private/VT/VirtualTextureScalability.cpp
--- a/Engine/Source/Runtime/Engine/Private/VT/VirtualTextureScalability.cpp
+++ b/Engine/Source/Runtime/Engine/Private/VT/VirtualTextureScalability.cpp
@@ -295,8 +295,12 @@ namespace VirtualTextureScalability
 	
 	uint32 GetPhysicalPoolSettingsHash()
 	{
-		uint32 Hash = GetTypeHash(GPoolSizeScales);
-		Hash = HashCombine(Hash, GetSplitPhysicalPoolSize());
+		uint32 Hash = 0;
+		for (float PoolSizeScale : GPoolSizeScales)
+		{
+			Hash = HashCombine(Hash, GetTypeHash(PoolSizeScale));
+		}
+		Hash = HashCombine(Hash, GetTypeHash(GetSplitPhysicalPoolSize()));
 		return Hash;
 	}
 }
