Disable new debug checks since some stuff was missed in testing, sigh.

diff --git a/Engine/Source/Runtime/Engine/Private/Rendering/Texture2DResource.cpp b/Engine/Source/Runtime/Engine/Private/Rendering/Texture2DResource.cpp
--- a/Engine/Source/Runtime/Engine/Private/Rendering/Texture2DResource.cpp
+++ b/Engine/Source/Runtime/Engine/Private/Rendering/Texture2DResource.cpp
@@ -142,7 +142,7 @@ void FTexture2DResource::CreateTexture()
 				uint32 DestPitch = -1;
 				uint64 Size = ~0ULL;
 				void* TheMipData = RHILockTexture2D(TextureRHI, RHIMipIdx, RLM_WriteOnly, DestPitch, false /* bLockWithinMiptail */, true /* bFlushRHIThread */, &Size);
-				check(Size != ~0ULL);
+				//check(Size != ~0ULL);
 
 				GetData( ResourceMipIdx, TheMipData, DestPitch, Size );
 				RHIUnlockTexture2D(TextureRHI, RHIMipIdx, false );
@@ -181,7 +181,7 @@ void FTexture2DResource::CreatePartiallyResidentTexture()
 			uint32 DestPitch = -1;
 			uint64 Size = ~0ULL;
 			void* TheMipData = RHILockTexture2D(TextureRHI, MipIndex, RLM_WriteOnly, DestPitch, false /* bLockWithinMiptail */, true /* bFlushRHIThread */, &Size);
-			check(Size != ~0ULL);
+			//check(Size != ~0ULL);
 
 			GetData( MipIndex, TheMipData, DestPitch, Size);
 			RHIUnlockTexture2D(TextureRHI, MipIndex, false );
@@ -271,7 +271,7 @@ void FTexture2DResource::GetData( uint32 MipIndex, void* Dest, uint32 DestPitch,
 		EffectiveSize,(int)BulkDataSize,
 		SrcPitch,DestPitch, PixelFormat, DestSize);
 
-	check(DestSize >= (uint64)BulkDataSize);
+	//check(DestSize >= (uint64)BulkDataSize);
 
 	// for platforms that returned 0 pitch from Lock, we need to just use the bulk data directly, never do 
 	// runtime block size checking, conversion, or the like
