Fix for cluster union sometimes resetting to origin
- make sure we properly initialize the cluster union proxy external particle with current component position

diff --git a/Engine/Source/Runtime/Engine/Private/PhysicsEngine/ClusterUnionComponent.cpp b/Engine/Source/Runtime/Engine/Private/PhysicsEngine/ClusterUnionComponent.cpp
--- a/Engine/Source/Runtime/Engine/Private/PhysicsEngine/ClusterUnionComponent.cpp
+++ b/Engine/Source/Runtime/Engine/Private/PhysicsEngine/ClusterUnionComponent.cpp
@@ -14,6 +14,11 @@
 
 DEFINE_LOG_CATEGORY(LogClusterUnion);
 
+// todo : remove when proven to be the right fix 
+bool bClusterUnionClientTransformFix = true;
+FAutoConsoleVariableRef CVarClusterUnionClientTransformFix(TEXT("clusterunion.clienttransformfix"), bClusterUnionClientTransformFix, TEXT("fix transform not always being properly set for cluster union"));
+
+
 UClusterUnionComponent::UClusterUnionComponent(const FObjectInitializer& ObjectInitializer)
 	: UPrimitiveComponent(ObjectInitializer)
 {
@@ -297,6 +302,16 @@ void UClusterUnionComponent::OnCreatePhysicsState()
 			AddComponentToCluster(Cast<UPrimitiveComponent>(ComponentReference.GetComponent(ComponentReference.OtherActor.Get())), {});
 		}
 	}
+	else
+	{
+		if (bClusterUnionClientTransformFix)
+		{
+			// we should not rely on the OnTransformUpdated callback to set the initial position of the external particle
+			// if it is not set and the callback does not get called then the cluster union component will be moved to the origin of the world
+			const FTransform Transform = GetComponentTransform();
+			PhysicsProxy->SetXR_External(Transform.GetLocation(), Transform.GetRotation());
+		}
+	}
 }
 
 void UClusterUnionComponent::OnDestroyPhysicsState()
diff --git a/Engine/Source/Runtime/Experimental/Chaos/Private/Chaos/PBDRigidClustering.cpp b/Engine/Source/Runtime/Experimental/Chaos/Private/Chaos/PBDRigidClustering.cpp
--- a/Engine/Source/Runtime/Experimental/Chaos/Private/Chaos/PBDRigidClustering.cpp
+++ b/Engine/Source/Runtime/Experimental/Chaos/Private/Chaos/PBDRigidClustering.cpp
@@ -1143,14 +1143,16 @@ namespace Chaos
 				SCOPE_CYCLE_COUNTER(STAT_UpdateDirtyImpulses);
 				for (const auto& ActiveCluster : TopLevelClusterParentsStrained)
 				{
+					bool bIgnoreDisabledCheck = false;
 					FClusterUnion* ClusterUnion = ClusterUnionManager.FindClusterUnionFromParticle(ActiveCluster);
 					if (ClusterUnion && ClusterUnion->InternalCluster != ActiveCluster)
 					{
 						// Need to pre-emptively remove the particle from the cluster union otherwise we won't be passing the disabled check.
-						ClusterUnionManager.HandleRemoveOperationWithClusterLookup({ ActiveCluster }, EClusterUnionOperationTiming::Defer);
+						//ClusterUnionManager.HandleRemoveOperationWithClusterLookup({ ActiveCluster }, EClusterUnionOperationTiming::Defer);
+						bIgnoreDisabledCheck = true;
 					}
 
-					if (!ActiveCluster->Disabled())
+					if (!ActiveCluster->Disabled() || bIgnoreDisabledCheck)
 					{
 						if (ActiveCluster->ClusterIds().NumChildren > 0) //active index is a cluster
 						{
