Fix missing member initialization.

diff --git a/Engine/Source/Runtime/Engine/Public/WorldPartition/WorldPartitionRuntimeCell.h b/Engine/Source/Runtime/Engine/Public/WorldPartition/WorldPartitionRuntimeCell.h
--- a/Engine/Source/Runtime/Engine/Public/WorldPartition/WorldPartitionRuntimeCell.h
+++ b/Engine/Source/Runtime/Engine/Public/WorldPartition/WorldPartitionRuntimeCell.h
@@ -39,6 +39,7 @@ struct FWorldPartitionRuntimeCellObjectMapping
 		, ContainerTransform(FTransform::Identity)
 		, ContainerPackage(NAME_None)
 		, LoadedPath(NAME_None)
+		, bIsEditorOnly(false)
 #endif
 	{}
 
