Swap usage of TObjectPtr to TWeakObjectPtr in various cluster union components

diff --git a/Engine/Source/Runtime/Engine/Classes/PhysicsEngine/ClusterUnionComponent.h b/Engine/Source/Runtime/Engine/Classes/PhysicsEngine/ClusterUnionComponent.h
--- a/Engine/Source/Runtime/Engine/Classes/PhysicsEngine/ClusterUnionComponent.h
+++ b/Engine/Source/Runtime/Engine/Classes/PhysicsEngine/ClusterUnionComponent.h
@@ -58,7 +58,7 @@ struct FClusteredActorData
 	GENERATED_BODY()
 
 	UPROPERTY()
-	TSet<TObjectPtr<UPrimitiveComponent>> Components;
+	TSet<TWeakObjectPtr<UPrimitiveComponent>> Components;
 
 	UPROPERTY()
 	bool bWasReplicatingMovement = true;
@@ -101,7 +101,7 @@ struct FClusterUnionParticleCandidateData
 	GENERATED_BODY()
 
 	UPROPERTY()
-	TObjectPtr<UPrimitiveComponent> Component;
+	TWeakObjectPtr<UPrimitiveComponent> Component;
 
 	UPROPERTY()
 	int32 BoneId = INDEX_NONE;
diff --git a/Engine/Source/Runtime/Engine/Classes/PhysicsEngine/ClusterUnionReplicatedProxyComponent.h b/Engine/Source/Runtime/Engine/Classes/PhysicsEngine/ClusterUnionReplicatedProxyComponent.h
--- a/Engine/Source/Runtime/Engine/Classes/PhysicsEngine/ClusterUnionReplicatedProxyComponent.h
+++ b/Engine/Source/Runtime/Engine/Classes/PhysicsEngine/ClusterUnionReplicatedProxyComponent.h
@@ -29,7 +29,7 @@ public:
 	UClusterUnionReplicatedProxyComponent(const FObjectInitializer& ObjectInitializer);
 
 	UFUNCTION()
-	UClusterUnionComponent* GetParentClusterUnionComponent() const { return ParentClusterUnion; }
+	UClusterUnionComponent* GetParentClusterUnionComponent() const { return ParentClusterUnion.Get(); }
 
 	UFUNCTION()
 	void SetParentClusterUnion(UClusterUnionComponent* InComponent);
@@ -65,13 +65,13 @@ protected:
 
 private:
 	UPROPERTY(ReplicatedUsing=OnRep_ParentClusterUnion)
-	TObjectPtr<UClusterUnionComponent> ParentClusterUnion;
+	TWeakObjectPtr<UClusterUnionComponent> ParentClusterUnion;
 
 	UPROPERTY()
 	bool bNetUpdateParentClusterUnion;
 
 	UPROPERTY(ReplicatedUsing=OnRep_ChildClusteredComponent)
-	TObjectPtr<UPrimitiveComponent> ChildClusteredComponent;
+	TWeakObjectPtr<UPrimitiveComponent> ChildClusteredComponent;
 
 	UPROPERTY()
 	bool bNetUpdateChildClusteredComponent;
diff --git a/Engine/Source/Runtime/Engine/Private/PhysicsEngine/ClusterUnionComponent.cpp b/Engine/Source/Runtime/Engine/Private/PhysicsEngine/ClusterUnionComponent.cpp
--- a/Engine/Source/Runtime/Engine/Private/PhysicsEngine/ClusterUnionComponent.cpp
+++ b/Engine/Source/Runtime/Engine/Private/PhysicsEngine/ClusterUnionComponent.cpp
@@ -445,7 +445,10 @@ void UClusterUnionComponent::SyncClusterUnionFromProxy()
 	{
 		if (FClusterUnionParticleCandidateData* Data = UniqueIdxToComponent.Find(ChildData.ParticleIdx.Idx))
 		{
-			MappedData.FindOrAdd(Data->Component).Add(Data->BoneId, ChildData.ChildToParent);
+			if (Data->Component.IsValid())
+			{
+				MappedData.FindOrAdd(Data->Component.Get()).Add(Data->BoneId, ChildData.ChildToParent);
+			}
 		}
 	}
 
diff --git a/Engine/Source/Runtime/Engine/Private/PhysicsEngine/ClusterUnionReplicatedProxyComponent.cpp b/Engine/Source/Runtime/Engine/Private/PhysicsEngine/ClusterUnionReplicatedProxyComponent.cpp
--- a/Engine/Source/Runtime/Engine/Private/PhysicsEngine/ClusterUnionReplicatedProxyComponent.cpp
+++ b/Engine/Source/Runtime/Engine/Private/PhysicsEngine/ClusterUnionReplicatedProxyComponent.cpp
@@ -47,9 +47,9 @@ void UClusterUnionReplicatedProxyComponent::EndPlay(const EEndPlayReason::Type E
 		Owner->GetWorldTimerManager().ClearTimer(DeferSetChildToParentHandle);
 	}
 
-	if (!GetOwner()->HasAuthority() && ParentClusterUnion && ChildClusteredComponent)
+	if (!GetOwner()->HasAuthority() && ParentClusterUnion.IsValid() && ChildClusteredComponent.IsValid())
 	{
-		ParentClusterUnion->RemoveComponentFromCluster(ChildClusteredComponent);
+		ParentClusterUnion->RemoveComponentFromCluster(ChildClusteredComponent.Get());
 	}
 }
 
@@ -113,7 +113,7 @@ void UClusterUnionReplicatedProxyComponent::PostRepNotifies()
 {
 	UActorComponent::PostRepNotifies();
 
-	const bool bIsValid = ParentClusterUnion && ChildClusteredComponent && !ParticleBoneIds.IsEmpty();
+	const bool bIsValid = ParentClusterUnion.IsValid() && ChildClusteredComponent.IsValid() && !ParticleBoneIds.IsEmpty();
 	if (!bIsValid)
 	{
 		return;
@@ -146,7 +146,7 @@ void UClusterUnionReplicatedProxyComponent::PostRepNotifies()
 
 void UClusterUnionReplicatedProxyComponent::DeferAddComponentToClusterHandleUntilInitialTransformUpdate()
 {
-	if (!ParentClusterUnion || !ChildClusteredComponent || IsPendingDeletion())
+	if (!ParentClusterUnion.IsValid() || !ChildClusteredComponent.IsValid() || IsPendingDeletion())
 	{
 		return;
 	}
@@ -155,7 +155,7 @@ void UClusterUnionReplicatedProxyComponent::DeferAddComponentToClusterHandleUnti
 
 	if (ParentClusterUnion->HasReceivedTransform())
 	{
-		ParentClusterUnion->AddComponentToCluster(ChildClusteredComponent, ParticleBoneIds);
+		ParentClusterUnion->AddComponentToCluster(ChildClusteredComponent.Get(), ParticleBoneIds);
 	}
 	else if (AActor* Owner = GetOwner())
 	{
@@ -165,16 +165,16 @@ void UClusterUnionReplicatedProxyComponent::DeferAddComponentToClusterHandleUnti
 
 void UClusterUnionReplicatedProxyComponent::DeferSetChildToParentChildUntilClusteredComponentInParentUnion()
 {
-	if (!ParentClusterUnion || !ChildClusteredComponent || IsPendingDeletion())
+	if (!ParentClusterUnion.IsValid() || !ChildClusteredComponent.IsValid() || IsPendingDeletion())
 	{
 		return;
 	}
 
 	DeferSetChildToParentHandle.Invalidate();
 
-	if (ParentClusterUnion->IsComponentAdded(ChildClusteredComponent))
+	if (ParentClusterUnion->IsComponentAdded(ChildClusteredComponent.Get()))
 	{
-		ParentClusterUnion->ForceSetChildToParent(ChildClusteredComponent, ParticleBoneIds, ParticleChildToParents);
+		ParentClusterUnion->ForceSetChildToParent(ChildClusteredComponent.Get(), ParticleBoneIds, ParticleChildToParents);
 	}
 	else if (AActor* Owner = GetOwner())
 	{
