Add sphere check to cluster union overlaps before running overlaps on all cluster sub components

diff --git a/Engine/Source/Runtime/Engine/Private/PhysicsEngine/ClusterUnionComponent.cpp b/Engine/Source/Runtime/Engine/Private/PhysicsEngine/ClusterUnionComponent.cpp
--- a/Engine/Source/Runtime/Engine/Private/PhysicsEngine/ClusterUnionComponent.cpp
+++ b/Engine/Source/Runtime/Engine/Private/PhysicsEngine/ClusterUnionComponent.cpp
@@ -724,12 +724,17 @@ bool UClusterUnionComponent::SweepComponent(FHitResult& OutHit, const FVector St
 bool UClusterUnionComponent::OverlapComponentWithResult(const FVector& Pos, const FQuat& Rot, const FCollisionShape& CollisionShape, TArray<FOverlapResult>& OutOverlap) const
 {
 	bool bHasOverlap = false;
+
+	FVector QueryHalfExtent = CollisionShape.GetExtent();
+	FBoxSphereBounds QueryBounds(FBox(-QueryHalfExtent, QueryHalfExtent));
+	QueryBounds = QueryBounds.TransformBy(FTransform(Rot, Pos));
+
 	for (UPrimitiveComponent* Component : GetAllCurrentChildComponents())
 	{
-		if (Component)
+		if(Component && FBoxSphereBounds::SpheresIntersect(QueryBounds, Component->Bounds))
 		{
 			TArray<FOverlapResult> SubOverlaps;
-			if (Component->OverlapComponentWithResult(Pos, Rot, CollisionShape, SubOverlaps))
+			if(Component->OverlapComponentWithResult(Pos, Rot, CollisionShape, SubOverlaps))
 			{
 				bHasOverlap = true;
 				OutOverlap.Append(SubOverlaps);
@@ -741,6 +746,11 @@ bool UClusterUnionComponent::OverlapComponentWithResult(const FVector& Pos, cons
 
 bool UClusterUnionComponent::ComponentOverlapComponentWithResultImpl(const class UPrimitiveComponent* const PrimComp, const FVector& Pos, const FQuat& Rot, const FCollisionQueryParams& Params, TArray<FOverlapResult>& OutOverlap) const
 {
+	if(!PrimComp)
+	{
+		return false;
+	}
+
 	bool bHasOverlap = false;
 
 	TSet<uint32> IgnoredActors;
@@ -757,25 +767,27 @@ bool UClusterUnionComponent::ComponentOverlapComponentWithResultImpl(const class
 		IgnoredComponents.Add(Id);
 	}
 
+	FBoxSphereBounds QueryBounds = PrimComp->Bounds.TransformBy(FTransform(Rot, Pos));
+
 	for (UPrimitiveComponent* Component : GetAllCurrentChildComponents())
 	{
-		if (Component)
+		if(Component && FBoxSphereBounds::SpheresIntersect(QueryBounds, Component->Bounds))
 		{
-			if (IgnoredComponents.Contains(Component->GetUniqueID()))
+			if(IgnoredComponents.Contains(Component->GetUniqueID()))
 			{
 				continue;
 			}
 
-			if (AActor* Owner = Component->GetOwner())
+			if(AActor* Owner = Component->GetOwner())
 			{
-				if (IgnoredActors.Contains(Owner->GetUniqueID()))
+				if(IgnoredActors.Contains(Owner->GetUniqueID()))
 				{
 					continue;
 				}
 			}
 
 			TArray<FOverlapResult> SubOverlaps;
-			if (Component->ComponentOverlapComponentWithResult(PrimComp, Pos, Rot, Params, SubOverlaps))
+			if(Component->ComponentOverlapComponentWithResult(PrimComp, Pos, Rot, Params, SubOverlaps))
 			{
 				bHasOverlap = true;
 				OutOverlap.Append(SubOverlaps);
