HLODs: Fixed issue where HLODs bounds were off for instanced WP levels
* Bounds optimization (bComputeBoundsOnceForGame) can't really work with instanced levels...

diff --git a/Engine/Source/Runtime/Engine/Private/WorldPartition/HLOD/HLODActor.cpp b/Engine/Source/Runtime/Engine/Private/WorldPartition/HLOD/HLODActor.cpp
--- a/Engine/Source/Runtime/Engine/Private/WorldPartition/HLOD/HLODActor.cpp
+++ b/Engine/Source/Runtime/Engine/Private/WorldPartition/HLOD/HLODActor.cpp
@@ -140,6 +140,13 @@ void AWorldPartitionHLOD::PostLoad()
 		});
 	}
 
+	// If world is instanced, we need to recompute our bounds since they are in the instanced-world space
+	ForEachComponent<USceneComponent>(false, [](USceneComponent* SceneComponent)
+	{
+		// Clear bComputedBoundsOnceForGame so that the bounds are recomputed once
+		SceneComponent->bComputedBoundsOnceForGame = false;
+	});
+
 #if WITH_EDITOR
 	if (GetLinkerCustomVersion(FFortniteMainBranchObjectVersion::GUID) < FFortniteMainBranchObjectVersion::WorldPartitionStreamingCellsNamingShortened)
 	{
