Do not hash padding bytes which can have random values when deduplicating preshader parameters

diff --git a/Engine/Source/Runtime/Engine/Private/HLSLTree/HLSLTreeEmit.cpp b/Engine/Source/Runtime/Engine/Private/HLSLTree/HLSLTreeEmit.cpp
--- a/Engine/Source/Runtime/Engine/Private/HLSLTree/HLSLTreeEmit.cpp
+++ b/Engine/Source/Runtime/Engine/Private/HLSLTree/HLSLTreeEmit.cpp
@@ -926,7 +926,9 @@ FEmitShaderExpression* FEmitContext::EmitPreshaderOrConstant(FEmitScope& Scope,
 	}
 
 	FXxHash64Builder Hasher;
-	Hasher.Update(&Type, sizeof(Type));
+	Hasher.Update(&Type.StructType, sizeof(Type.StructType));
+	Hasher.Update(&Type.ObjectType, sizeof(Type.ObjectType));
+	Hasher.Update(&Type.ValueType, sizeof(Type.ValueType));
 	LocalPreshader.AppendHash(Hasher);
 	const FXxHash64 Hash = Hasher.Finalize();
 	FEmitShaderExpression* const* PrevShaderValue = EmitPreshaderMap.Find(Hash);
