Add subsurface profile support to the new HLSL generator

diff --git a/Engine/Source/Runtime/Engine/Private/Materials/MaterialHLSLEmitter.cpp b/Engine/Source/Runtime/Engine/Private/Materials/MaterialHLSLEmitter.cpp
--- a/Engine/Source/Runtime/Engine/Private/Materials/MaterialHLSLEmitter.cpp
+++ b/Engine/Source/Runtime/Engine/Private/Materials/MaterialHLSLEmitter.cpp
@@ -17,6 +17,7 @@
 #include "Containers/LazyPrintf.h"
 #include "RenderUtils.h"
 #include "DataDrivenShaderPlatformInfo.h"
+#include "Engine/SubsurfaceProfile.h"
 
 static bool SharedPixelProperties[CompiledMP_MAX];
 
@@ -58,6 +59,7 @@ static FString GenerateMaterialTemplateHLSL(EShaderPlatform ShaderPlatform,
 	const TCHAR* VertexShaderCode,
 	const TCHAR* PixelShaderCodePhase0[2],
 	const TCHAR* PixelShaderCodePhase1[2],
+	const TCHAR* SubsurfaceProfileCode,
 	FMaterialCompilationOutput& OutCompilationOutput)
 {
 	using namespace UE::HLSLTree;
@@ -353,8 +355,11 @@ static FString GenerateMaterialTemplateHLSL(EShaderPlatform ShaderPlatform,
 				}
 				else if (PropertyIndex == MP_SubsurfaceColor)
 				{
-					// TODO - properly handle subsurface profile
-					EvaluateMaterialAttributesCode += FString::Printf("    PixelMaterialInputs.Subsurface = float4(MaterialAttributesPhase%d.%s, 0.0f);" LINE_TERMINATOR, EvaluatePhase, *PropertyName);
+					EvaluateMaterialAttributesCode += FString::Printf(
+						"    PixelMaterialInputs.Subsurface = float4(MaterialAttributesPhase%d.%s, %s);" LINE_TERMINATOR,
+						EvaluatePhase,
+						*PropertyName,
+						SubsurfaceProfileCode ? SubsurfaceProfileCode : TEXT("1.0f"));
 				}
 				else
 				{
@@ -823,6 +828,7 @@ bool MaterialEmitHLSL(const FMaterialCompileTargetParameters& InCompilerTarget,
 	// Prepare pixel shader code
 	FStringBuilderBase* PixelCodePhase0[2] = { nullptr };
 	FStringBuilderBase* PixelCodePhase1[2] = { nullptr };
+	const TCHAR* SubsurfaceProfileShaderCode = nullptr;
 
 	for(int32 DerivativeIndex = 0; DerivativeIndex < 2; ++DerivativeIndex)
 	{
@@ -846,6 +852,27 @@ bool MaterialEmitHLSL(const FMaterialCompileTargetParameters& InCompilerTarget,
 		EmitContext.PrepareExpression(CachedTree->GetResultExpression(), *EmitResultScope, RequestedPixelAttributesType);
 		EmitContext.bMarkLiveValues = false;
 
+		if (DerivativeIndex == 0)
+		{
+			const EMaterialDomain MaterialDomain = EmitContext.Material->GetMaterialDomain();
+			if (MaterialDomain == MD_Volume || (MaterialDomain == MD_Surface && IsSubsurfaceShadingModel(EmitMaterialData.ShadingModelsFromCompilation)))
+			{
+				const FMaterialParameterInfo SubsurfaceProfileParameterInfo(GetSubsurfaceProfileParameterName());
+				const FMaterialParameterMetadata SubsurfaceProfileParameterMetadata(1.f);
+				const Material::FExpressionParameter SubsurfaceProfileExpression(SubsurfaceProfileParameterInfo, SubsurfaceProfileParameterMetadata);
+
+				const FPreparedType SubsurfaceProfilePreparedType = EmitContext.PrepareExpression(&SubsurfaceProfileExpression, *EmitResultScope, EValueType::Float1);
+				const FEmitShaderExpression* SubsurfaceProfileEmitExpression = SubsurfaceProfileExpression.GetValueShader(
+					EmitContext,
+					*EmitResultScope,
+					EValueType::Float1,
+					SubsurfaceProfilePreparedType,
+					EValueType::Float1);
+				check(SubsurfaceProfileEmitExpression && SubsurfaceProfileEmitExpression->Reference);
+				SubsurfaceProfileShaderCode = SubsurfaceProfileEmitExpression->Reference;
+			}
+		}
+
 		if (CachedTree->IsAttributeUsed(EmitContext, *EmitResultScope, PixelResultType0, MP_PixelDepthOffset))
 		{
 			bUsesPixelDepthOffset = true;
@@ -980,6 +1007,7 @@ bool MaterialEmitHLSL(const FMaterialCompileTargetParameters& InCompilerTarget,
 			VertexCode.ToString(),
 			InputPixelCodePhase0,
 			InputPixelCodePhase1,
+			SubsurfaceProfileShaderCode,
 			OutCompilationOutput);
 	}
 
