Preshader evaluation support for FExpressionSelect

diff --git a/Engine/Source/Runtime/Engine/Private/HLSLTree/HLSLTreeCommon.cpp b/Engine/Source/Runtime/Engine/Private/HLSLTree/HLSLTreeCommon.cpp
--- a/Engine/Source/Runtime/Engine/Private/HLSLTree/HLSLTreeCommon.cpp
+++ b/Engine/Source/Runtime/Engine/Private/HLSLTree/HLSLTreeCommon.cpp
@@ -382,7 +382,6 @@ bool FExpressionSelect::PrepareValue(FEmitContext& Context, FEmitScope& Scope, c
 	{
 		const bool bCondition = ConditionExpression->GetValueConstant(Context, Scope, ConditionType, Shader::EValueType::Bool1).AsBoolScalar();
 		FPreparedType ResultType = Context.PrepareExpression(bCondition ? TrueExpression : FalseExpression, Scope, RequestedType);
-		ResultType.MergeEvaluation(EExpressionEvaluation::Shader); // TODO - support preshader
 		return OutResult.SetType(Context, RequestedType, ResultType);
 	}
 
@@ -398,9 +397,16 @@ bool FExpressionSelect::PrepareValue(FEmitContext& Context, FEmitScope& Scope, c
 	{
 		return Context.Error(TEXT("Type mismatch"));
 	}
+	if (LhsType.Type != RhsType.Type)
+	{
+		check(LhsType.Type.IsNumericVector() && RhsType.Type.IsNumericVector());
+		if (LhsType.Type.GetNumComponents() != RhsType.Type.GetNumComponents() && !LhsType.IsNumericScalar() && !RhsType.IsNumericScalar())
+		{
+			return Context.Error(TEXT("Cannot mix numeric vectors with different number of components"));
+		}
+	}
 
 	ResultType.MergeEvaluation(ConditionEvaluation);
-	ResultType.MergeEvaluation(EExpressionEvaluation::Shader); // TODO - support preshader
 	return OutResult.SetType(Context, RequestedType, ResultType);
 }
 
@@ -429,10 +435,57 @@ void FExpressionSelect::EmitValueShader(FEmitContext& Context, FEmitScope& Scope
 
 void FExpressionSelect::EmitValuePreshader(FEmitContext& Context, FEmitScope& Scope, const FRequestedType& RequestedType, FEmitValuePreshaderResult& OutResult) const
 {
-	// TODO
-	Context.PreshaderStackPosition++;
-	OutResult.Type = Context.GetType(this);
-	OutResult.Preshader.WriteOpcode(Shader::EPreshaderOpcode::ConstantZero).Write(OutResult.Type);
+	const FPreparedType& ConditionType = Context.GetPreparedType(ConditionExpression);
+	const EExpressionEvaluation ConditionEvaluation = ConditionType.GetEvaluation(Scope, Shader::EValueType::Bool1);
+	if (IsConstantEvaluation(ConditionEvaluation))
+	{
+		const bool bCondition = ConditionExpression->GetValueConstant(Context, Scope, ConditionType, Shader::EValueType::Bool1).AsBoolScalar();
+		const FExpression* InputExpression = bCondition ? TrueExpression : FalseExpression;
+		OutResult.Type = InputExpression->GetValuePreshader(Context, Scope, RequestedType, OutResult.Preshader);
+	}
+	else
+	{
+		const Shader::FType ResultType = Context.GetType(this);
+		auto ConvertToResultType = [&ResultType](const Shader::FType& Type, Shader::FPreshaderData& Preshader)
+		{
+			if (Type == ResultType)
+			{
+				return;
+			}
+
+			check(Type.IsNumericVector() && ResultType.IsNumericVector());
+			const Shader::EValueComponentType ComponentType = Type.GetComponentType(0);
+			const Shader::EValueComponentType ResultComponentType = ResultType.GetComponentType(0);
+			const int32 NumComponents = Type.GetNumComponents();
+			const int32 NumResultComponents = ResultType.GetNumComponents();
+			check(NumComponents == NumResultComponents || NumComponents == 1);
+
+			if (ComponentType != ResultComponentType || NumComponents != NumResultComponents)
+			{
+				Preshader.WriteOpcode(Shader::EPreshaderOpcode::ConstantZero).Write(ResultType);
+				Preshader.WriteOpcode(Shader::EPreshaderOpcode::Add);
+			}
+		};
+
+		check(ConditionEvaluation == EExpressionEvaluation::Preshader);
+		ConditionExpression->GetValuePreshader(Context, Scope, Shader::EValueType::Bool1, OutResult.Preshader);
+
+		check(Context.PreshaderStackPosition > 0);
+		Context.PreshaderStackPosition--;
+		const Shader::FPreshaderLabel Label0 = OutResult.Preshader.WriteJump(Shader::EPreshaderOpcode::JumpIfFalse);
+
+		const Shader::FType TrueType = TrueExpression->GetValuePreshader(Context, Scope, RequestedType, OutResult.Preshader);
+		ConvertToResultType(TrueType, OutResult.Preshader);
+
+		const Shader::FPreshaderLabel Label1 = OutResult.Preshader.WriteJump(Shader::EPreshaderOpcode::Jump);
+		OutResult.Preshader.SetLabel(Label0);
+
+		const Shader::FType FalseType = FalseExpression->GetValuePreshader(Context, Scope, RequestedType, OutResult.Preshader);
+		ConvertToResultType(FalseType, OutResult.Preshader);
+
+		OutResult.Preshader.SetLabel(Label1);
+		OutResult.Type = ResultType;
+	}
 }
 
 void FExpressionDerivative::ComputeAnalyticDerivatives(FTree& Tree, FExpressionDerivatives& OutResult) const
