Do not emit code to set struct field if not requested

diff --git a/Engine/Source/Runtime/Engine/Private/HLSLTree/HLSLTreeCommon.cpp b/Engine/Source/Runtime/Engine/Private/HLSLTree/HLSLTreeCommon.cpp
--- a/Engine/Source/Runtime/Engine/Private/HLSLTree/HLSLTreeCommon.cpp
+++ b/Engine/Source/Runtime/Engine/Private/HLSLTree/HLSLTreeCommon.cpp
@@ -283,21 +283,25 @@ bool FExpressionSetStructField::PrepareValue(FEmitContext& Context, FEmitScope&
 	FRequestedType RequestedStructType(RequestedType);
 	RequestedStructType.ClearFieldRequested(Field);
 
-	FPreparedType StructPreparedType = Context.PrepareExpression(StructExpression, Scope, RequestedStructType);
+	const FPreparedType StructPreparedType = Context.PrepareExpression(StructExpression, Scope, RequestedStructType);
 	if (!StructPreparedType.IsVoid() && StructPreparedType.Type.StructType != StructType)
 	{
 		return Context.Errorf(TEXT("Expected type %s"), StructType->Name);
 	}
 
-	const FRequestedType RequestedFieldType = RequestedType.GetField(Field);
-	FPreparedType FieldPreparedType = Context.PrepareExpression(FieldExpression, Scope, RequestedFieldType);
-
 	FPreparedType ResultType(StructPreparedType);
 	if (ResultType.IsVoid())
 	{
 		ResultType = StructType;
 	}
-	ResultType.SetField(Field, FieldPreparedType);
+	
+	const FRequestedType RequestedFieldType = RequestedType.GetField(Field);
+	if (!RequestedType.IsEmpty())
+	{
+		const FPreparedType FieldPreparedType = Context.PrepareExpression(FieldExpression, Scope, RequestedFieldType);
+		ResultType.SetField(Field, FieldPreparedType);
+	}
+
 	return OutResult.SetType(Context, RequestedType, ResultType);
 }
 
@@ -319,6 +323,14 @@ void FExpressionSetStructField::EmitValueShader(FEmitContext& Context, FEmitScop
 	else
 	{
 		FEmitShaderExpression* StructValue = StructExpression->GetValueShader(Context, Scope, RequestedStructType, StructType);
+
+		// Do not emit code to set the field if not requested
+		if (RequestedFieldType.IsEmpty())
+		{
+			OutResult.Code = StructValue;
+			return;
+		}
+
 		FEmitShaderExpression* FieldValue = FieldExpression->GetValueShader(Context, Scope, RequestedFieldType, Field->Type);
 		OutResult.Code = Context.EmitExpression(Scope, StructType, TEXT("%_Set%(%, %)"),
 			StructType->Name,
