Hash maps for speeding up AST generation

diff --git a/Engine/Source/Runtime/Engine/Private/HLSLTree/HLSLTree.cpp b/Engine/Source/Runtime/Engine/Private/HLSLTree/HLSLTree.cpp
--- a/Engine/Source/Runtime/Engine/Private/HLSLTree/HLSLTree.cpp
+++ b/Engine/Source/Runtime/Engine/Private/HLSLTree/HLSLTree.cpp
@@ -1573,7 +1573,17 @@ FExpressionDerivatives FTree::GetAnalyticDerivatives(const FExpression* InExpres
 	if (InExpression)
 	{
 		//FOwnerScope OwnerScope(*this, InExpression->GetOwner()); // Associate any newly created nodes with the same owner as the input expression
-		InExpression->ComputeAnalyticDerivatives(*this, Derivatives);
+		
+		FExpressionDerivatives* Found = ExpressionDerivativesMap.Find(InExpression);
+		if (Found)
+		{
+			Derivatives = *Found;
+		}
+		else
+		{
+			InExpression->ComputeAnalyticDerivatives(*this, Derivatives);
+			ExpressionDerivativesMap.Add(InExpression, Derivatives);
+		}
 	}
 	return Derivatives;
 }
@@ -1584,7 +1594,23 @@ const FExpression* FTree::GetPreviousFrame(const FExpression* InExpression, cons
 	if (Result && !RequestedType.IsVoid())
 	{
 		//FOwnerScope OwnerScope(*this, InExpression->GetOwner()); // Associate any newly created nodes with the same owner as the input expression
-		const FExpression* PrevFrameExpression = InExpression->ComputePreviousFrame(*this, RequestedType);
+		FHasher Hasher;
+		Hasher.AppendData(&InExpression, sizeof(InExpression));
+		AppendHash(Hasher, RequestedType);
+		const FXxHash64 KeyHash = Hasher.Finalize();
+		const FExpression** Found = PreviousFrameExpressionMap.Find(KeyHash);
+
+		const FExpression* PrevFrameExpression;
+		if (Found)
+		{
+			PrevFrameExpression = *Found;
+		}
+		else
+		{
+			PrevFrameExpression = InExpression->ComputePreviousFrame(*this, RequestedType);
+			PreviousFrameExpressionMap.Add(KeyHash, PrevFrameExpression);
+		}
+
 		if (PrevFrameExpression)
 		{
 			Result = PrevFrameExpression;
diff --git a/Engine/Source/Runtime/Engine/Private/Materials/MaterialHLSLGenerator.cpp b/Engine/Source/Runtime/Engine/Private/Materials/MaterialHLSLGenerator.cpp
--- a/Engine/Source/Runtime/Engine/Private/Materials/MaterialHLSLGenerator.cpp
+++ b/Engine/Source/Runtime/Engine/Private/Materials/MaterialHLSLGenerator.cpp
@@ -424,8 +424,41 @@ const UE::HLSLTree::FExpression* FMaterialHLSLGenerator::AcquireExpression(UE::H
 	UObject* InputOwner = GetTree().GetCurrentOwner();
 	FOwnerScope OwnerScope(GetTree(), MaterialExpression, NeedToPushOwnerExpression());
 
-	const FExpression* Expression = nullptr;
-	if (MaterialExpression->GenerateHLSLExpression(*this, Scope, OutputIndex, Expression))
+	const FExpression* Expression;
+	{
+		FHasher Hasher;
+		const FScope* ScopeAddress = &Scope;
+		Hasher.AppendData(&ScopeAddress, sizeof(ScopeAddress));
+		Hasher.AppendData(&MaterialExpression, sizeof(MaterialExpression));
+		Hasher.AppendData(&OutputIndex, sizeof(OutputIndex));
+
+		// Make sure the callstack is the same because results can vary when the containing function
+		// is called at different places due to the difference in inputs
+		for (int32 Index = FunctionCallStack.Num() - 1; Index >= 0; --Index)
+		{
+			const FFunctionCallEntry* FunctionCall = FunctionCallStack[Index];
+			Hasher.AppendData(&FunctionCall->MaterialFunction, sizeof(FunctionCall->MaterialFunction));
+			Hasher.AppendData(&FunctionCall->ParameterAssociation, sizeof(FunctionCall->ParameterAssociation));
+			Hasher.AppendData(&FunctionCall->ParameterIndex, sizeof(FunctionCall->ParameterIndex));
+
+			for (int32 InputIndex = 0; InputIndex < FunctionCall->ConnectedInputs.Num(); ++InputIndex)
+			{
+				const FExpression* InputExpression = FunctionCall->ConnectedInputs[InputIndex];
+				Hasher.AppendData(&InputExpression, sizeof(InputExpression));
+			}
+		}
+
+		const FXxHash64 KeyHash = Hasher.Finalize();
+		const FExpression** Found = GeneratedExpressionMap.Find(KeyHash);
+		Expression = Found ? *Found : nullptr;
+
+		if (!Expression && MaterialExpression->GenerateHLSLExpression(*this, Scope, OutputIndex, Expression))
+		{
+			GeneratedExpressionMap.Add(KeyHash, Expression);
+		}
+	}
+
+	if (Expression)
 	{
 		Expression = GetTree().NewSwizzle(Swizzle, Expression);
 		if (MaterialExpression == PreviewExpression &&
diff --git a/Engine/Source/Runtime/Engine/Public/HLSLTree/HLSLTree.h b/Engine/Source/Runtime/Engine/Public/HLSLTree/HLSLTree.h
--- a/Engine/Source/Runtime/Engine/Public/HLSLTree/HLSLTree.h
+++ b/Engine/Source/Runtime/Engine/Public/HLSLTree/HLSLTree.h
@@ -689,6 +689,9 @@ private:
 	TMap<FXxHash64, FExpression*> ExpressionMap;
 	TArray<FExpressionLocalPHI*> PHIExpressions;
 
+	TMap<const FExpression*, FExpressionDerivatives> ExpressionDerivativesMap;
+	TMap<FXxHash64, const FExpression*> PreviousFrameExpressionMap;
+
 	friend class FExpressionLocalPHI;
 };
 
diff --git a/Engine/Source/Runtime/Engine/Public/MaterialHLSLGenerator.h b/Engine/Source/Runtime/Engine/Public/MaterialHLSLGenerator.h
--- a/Engine/Source/Runtime/Engine/Public/MaterialHLSLGenerator.h
+++ b/Engine/Source/Runtime/Engine/Public/MaterialHLSLGenerator.h
@@ -265,6 +265,8 @@ private:
 	TMap<FExpressionDataKey, void*> ExpressionDataMap;
 	const UE::HLSLTree::FExpression* PreviewExpressionResult = nullptr;
 	bool bGeneratedResult;
+
+	TMap<FXxHash64, const UE::HLSLTree::FExpression*> GeneratedExpressionMap;
 };
 
 #endif // WITH_EDITOR
