Fix custom expression output struct name clash

diff --git a/Engine/Source/Runtime/Engine/Private/Materials/MaterialExpressionHLSL.cpp b/Engine/Source/Runtime/Engine/Private/Materials/MaterialExpressionHLSL.cpp
--- a/Engine/Source/Runtime/Engine/Private/Materials/MaterialExpressionHLSL.cpp
+++ b/Engine/Source/Runtime/Engine/Private/Materials/MaterialExpressionHLSL.cpp
@@ -2426,7 +2426,9 @@ bool UMaterialExpressionCustom::GenerateHLSLExpression(FMaterialHLSLGenerator& G
 		OutputFieldInitializers.Emplace(OutputNames.Last(), GetCustomOutputType(Generator, Output.OutputType));
 	}
 
-	FString OutputStructName = TEXT("FCustomOutput") + GetName();
+	const int32 OutputStructId = Generator.FindOrAddCustomExpressionOutputStructId(OutputFieldInitializers);
+
+	FString OutputStructName = FString::Printf(TEXT("FCustomExpressionOutput%d"), OutputStructId);
 	FStructTypeInitializer OutputStructInitializer;
 	OutputStructInitializer.Name = OutputStructName;
 	OutputStructInitializer.Fields = OutputFieldInitializers;
diff --git a/Engine/Source/Runtime/Engine/Private/Materials/MaterialHLSLGenerator.cpp b/Engine/Source/Runtime/Engine/Private/Materials/MaterialHLSLGenerator.cpp
--- a/Engine/Source/Runtime/Engine/Private/Materials/MaterialHLSLGenerator.cpp
+++ b/Engine/Source/Runtime/Engine/Private/Materials/MaterialHLSLGenerator.cpp
@@ -834,4 +834,34 @@ void* FMaterialHLSLGenerator::InternalFindExpressionData(const FName& Type, cons
 	return Result ? *Result : nullptr;
 }
 
+int32 FMaterialHLSLGenerator::FindOrAddCustomExpressionOutputStructId(TArrayView<UE::Shader::FStructFieldInitializer> StructFields)
+{
+	FXxHash64Builder Hasher;
+	for (const UE::Shader::FStructFieldInitializer& StructField : StructFields)
+	{
+		Hasher.Update(StructField.Name.GetData(), StructField.Name.Len() * sizeof(TCHAR));
+		const UE::Shader::FType& FieldType = StructField.Type;
+		if (FieldType.IsStruct())
+		{
+			Hasher.Update(&FieldType.StructType->Hash, sizeof(FieldType.StructType->Hash));
+		}
+		else
+		{
+			Hasher.Update(&FieldType.ValueType, sizeof(FieldType.ValueType));
+		}
+	}
+	const FXxHash64 KeyHash = Hasher.Finalize();
+	int32* Found = CustomExpressionOutputStructIdMap.Find(KeyHash);
+	if (Found)
+	{
+		return *Found;
+	}
+	else
+	{
+		const int32 NewId = CustomExpressionOutputStructIdMap.Num();
+		CustomExpressionOutputStructIdMap.Add(KeyHash, NewId);
+		return NewId;
+	}
+}
+
 #endif // WITH_EDITOR
diff --git a/Engine/Source/Runtime/Engine/Public/MaterialHLSLGenerator.h b/Engine/Source/Runtime/Engine/Public/MaterialHLSLGenerator.h
--- a/Engine/Source/Runtime/Engine/Public/MaterialHLSLGenerator.h
+++ b/Engine/Source/Runtime/Engine/Public/MaterialHLSLGenerator.h
@@ -186,6 +186,8 @@ public:
 	bool GetParameterOverrideValueForCurrentFunction(EMaterialParameterType ParameterType, FName ParameterName, FMaterialParameterMetadata& OutResult) const;
 	FMaterialParameterInfo GetParameterInfo(const FName& ParameterName) const;
 
+	int32 FindOrAddCustomExpressionOutputStructId(TArrayView<UE::Shader::FStructFieldInitializer> StructFields);
+
 private:
 	static constexpr int32 MaxNumPreviousScopes = UE::HLSLTree::MaxNumPreviousScopes;
 	
@@ -267,6 +269,7 @@ private:
 	bool bGeneratedResult;
 
 	TMap<FXxHash64, const UE::HLSLTree::FExpression*> GeneratedExpressionMap;
+	TMap<FXxHash64, int32> CustomExpressionOutputStructIdMap;
 };
 
 #endif // WITH_EDITOR
