Add parameter collection support to the new HLSL generator

diff --git a/Engine/Source/Runtime/Engine/Classes/Materials/MaterialExpressionCollectionParameter.h b/Engine/Source/Runtime/Engine/Classes/Materials/MaterialExpressionCollectionParameter.h
--- a/Engine/Source/Runtime/Engine/Classes/Materials/MaterialExpressionCollectionParameter.h
+++ b/Engine/Source/Runtime/Engine/Classes/Materials/MaterialExpressionCollectionParameter.h
@@ -49,6 +49,8 @@ class UMaterialExpressionCollectionParameter : public UMaterialExpression
 	virtual void ValidateParameterName(const bool bAllowDuplicateName) override {};
 
 	virtual bool MatchesSearchQuery( const TCHAR* SearchQuery ) override;
+
+	virtual bool GenerateHLSLExpression(FMaterialHLSLGenerator& Generator, UE::HLSLTree::FScope& Scope, int32 OutputIndex, UE::HLSLTree::FExpression const*& OutExpression) const override;
 #endif
 	//~ End UMaterialExpression Interface
 };
diff --git a/Engine/Source/Runtime/Engine/Private/Materials/MaterialExpressionHLSL.cpp b/Engine/Source/Runtime/Engine/Private/Materials/MaterialExpressionHLSL.cpp
--- a/Engine/Source/Runtime/Engine/Private/Materials/MaterialExpressionHLSL.cpp
+++ b/Engine/Source/Runtime/Engine/Private/Materials/MaterialExpressionHLSL.cpp
@@ -18,6 +18,8 @@
 #include "Materials/MaterialExpressionConstantBiasScale.h"
 #include "Materials/MaterialExpressionShadingModel.h"
 #include "Materials/MaterialExpressionParameter.h"
+#include "Materials/MaterialExpressionCollectionParameter.h"
+#include "Materials/MaterialParameterCollection.h"
 #include "Materials/MaterialExpressionCurveAtlasRowParameter.h"
 #include "Materials/MaterialExpressionStaticSwitchParameter.h"
 #include "Materials/MaterialExpressionStaticComponentMaskParameter.h"
@@ -344,6 +346,42 @@ bool UMaterialExpressionParameter::GenerateHLSLExpression(FMaterialHLSLGenerator
 	return true;
 }
 
+bool UMaterialExpressionCollectionParameter::GenerateHLSLExpression(FMaterialHLSLGenerator& Generator, UE::HLSLTree::FScope& Scope, int32 OutputIndex, UE::HLSLTree::FExpression const*& OutExpression) const
+{
+	int32 ParameterIndex = INDEX_NONE;
+	int32 ComponentIndex = INDEX_NONE;
+
+	if (Collection)
+	{
+		Collection->GetParameterIndex(ParameterId, ParameterIndex, ComponentIndex);
+	}
+
+	if (ParameterIndex != INDEX_NONE)
+	{
+		const int32 CollectionIndex = Generator.FindOrAddParameterCollection(Collection);
+		OutExpression = Generator.GetTree().NewExpression<UE::HLSLTree::Material::FExpressionCollectionParameter>(CollectionIndex, ParameterIndex);
+		OutExpression = Generator.GetTree().NewExpression<UE::HLSLTree::FExpressionSwizzle>(
+			UE::HLSLTree::MakeSwizzleMask(
+				ComponentIndex == INDEX_NONE ? true : ComponentIndex % 4 == 0,
+				ComponentIndex == INDEX_NONE ? true : ComponentIndex % 4 == 1,
+				ComponentIndex == INDEX_NONE ? true : ComponentIndex % 4 == 2,
+				ComponentIndex == INDEX_NONE ? true : ComponentIndex % 4 == 3),
+			OutExpression);
+		return true;
+	}
+	else
+	{
+		if (!Collection)
+		{
+			return Generator.Error(TEXT("CollectionParameter has invalid Collection!"));
+		}
+		else
+		{
+			return Generator.Errorf(TEXT("CollectionParameter has invalid parameter %s"), *ParameterName.ToString());
+		}
+	}
+}
+
 bool UMaterialExpressionCurveAtlasRowParameter::GenerateHLSLExpression(FMaterialHLSLGenerator& Generator, UE::HLSLTree::FScope& Scope, int32 OutputIndex, UE::HLSLTree::FExpression const*& OutExpression) const
 {
 	using namespace UE::HLSLTree;
diff --git a/Engine/Source/Runtime/Engine/Private/Materials/MaterialHLSLEmitter.cpp b/Engine/Source/Runtime/Engine/Private/Materials/MaterialHLSLEmitter.cpp
--- a/Engine/Source/Runtime/Engine/Private/Materials/MaterialHLSLEmitter.cpp
+++ b/Engine/Source/Runtime/Engine/Private/Materials/MaterialHLSLEmitter.cpp
@@ -8,6 +8,7 @@
 #include "Materials/MaterialExpressionFunctionInput.h"
 #include "Materials/MaterialExpressionFunctionOutput.h"
 #include "Materials/MaterialExpressionVolumetricAdvancedMaterialOutput.h"
+#include "Materials/MaterialParameterCollection.h"
 #include "Materials/Material.h"
 #include "MaterialDomain.h"
 #include "MaterialHLSLTree.h"
@@ -534,7 +535,8 @@ static void GetMaterialEnvironment(EShaderPlatform InPlatform,
 		OutEnvironment.SetDefine(PageTableName.ToString(), PageTableValue.ToString());
 	}
 
-	/*for (int32 CollectionIndex = 0; CollectionIndex < ParameterCollections.Num(); CollectionIndex++)
+	const TArray<UMaterialParameterCollection*>& ParameterCollections = InMaterial.GetCachedHLSLTree()->GetParameterCollections();
+	for (int32 CollectionIndex = 0; CollectionIndex < ParameterCollections.Num(); CollectionIndex++)
 	{
 		// Add uniform buffer declarations for any parameter collections referenced
 		const FString CollectionName = FString::Printf(TEXT("MaterialCollection%u"), CollectionIndex);
@@ -542,7 +544,7 @@ static void GetMaterialEnvironment(EShaderPlatform InPlatform,
 		// OutEnvironment.ResourceTableMap has a map by name, and the N ParameterCollection Uniform Buffers ALL are names "MaterialCollection"
 		// (and the hlsl cbuffers are named MaterialCollection0, etc, so the names don't match the layout)
 		FShaderUniformBufferParameter::ModifyCompilationEnvironment(*CollectionName, ParameterCollections[CollectionIndex]->GetUniformBufferStruct(), InPlatform, OutEnvironment);
-	}*/
+	}
 	OutEnvironment.SetDefine(TEXT("IS_MATERIAL_SHADER"), true);
 
 	// Set all the shading models for this material here 
@@ -1011,6 +1013,8 @@ bool MaterialEmitHLSL(const FMaterialCompileTargetParameters& InCompilerTarget,
 			OutCompilationOutput);
 	}
 
+	OutCompilationOutput.UniformExpressionSet.SetParameterCollections(CachedTree->GetParameterCollections());
+
 	OutMaterialEnvironment = new FSharedShaderCompilerEnvironment();
 	OutMaterialEnvironment->TargetPlatform = InCompilerTarget.TargetPlatform;
 	GetMaterialEnvironment(InCompilerTarget.ShaderPlatform, InOutMaterial, EmitContext, OutCompilationOutput, *OutMaterialEnvironment);
diff --git a/Engine/Source/Runtime/Engine/Private/Materials/MaterialHLSLGenerator.cpp b/Engine/Source/Runtime/Engine/Private/Materials/MaterialHLSLGenerator.cpp
--- a/Engine/Source/Runtime/Engine/Private/Materials/MaterialHLSLGenerator.cpp
+++ b/Engine/Source/Runtime/Engine/Private/Materials/MaterialHLSLGenerator.cpp
@@ -9,11 +9,13 @@
 #include "Materials/MaterialExpressionFunctionOutput.h"
 #include "Materials/MaterialExpressionCustomOutput.h"
 #include "Materials/MaterialExpressionExecBegin.h"
+#include "Materials/MaterialParameterCollection.h"
 #include "Materials/MaterialFunction.h"
 #include "MaterialHLSLTree.h"
 #include "Materials/Material.h"
 #include "MaterialCachedHLSLTree.h"
 #include "Misc/MemStackUtility.h"
+#include "ParameterCollection.h"
 
 FMaterialHLSLGenerator::FMaterialHLSLGenerator(UMaterial* Material,
 	const FMaterialLayersFunctions* InLayerOverrides,
@@ -864,4 +866,22 @@ int32 FMaterialHLSLGenerator::FindOrAddCustomExpressionOutputStructId(TArrayView
 	}
 }
 
+int32 FMaterialHLSLGenerator::FindOrAddParameterCollection(UMaterialParameterCollection* ParameterCollection)
+{
+	int32 CollectionIndex = CachedTree.ParameterCollections.Find(ParameterCollection);
+
+	if (CollectionIndex == INDEX_NONE)
+	{
+		if (CachedTree.ParameterCollections.Num() >= MaxNumParameterCollectionsPerMaterial)
+		{
+			return Error(TEXT("Material references too many MaterialParameterCollections!  A material may only reference 2 different collections."));
+		}
+
+		CachedTree.ParameterCollections.Add(ParameterCollection);
+		CollectionIndex = CachedTree.ParameterCollections.Num() - 1;
+	}
+
+	return CollectionIndex;
+}
+
 #endif // WITH_EDITOR
diff --git a/Engine/Source/Runtime/Engine/Private/Materials/MaterialHLSLTree.cpp b/Engine/Source/Runtime/Engine/Private/Materials/MaterialHLSLTree.cpp
--- a/Engine/Source/Runtime/Engine/Private/Materials/MaterialHLSLTree.cpp
+++ b/Engine/Source/Runtime/Engine/Private/Materials/MaterialHLSLTree.cpp
@@ -606,6 +606,22 @@ bool FExpressionParameter::EmitCustomHLSLParameter(FEmitContext& Context, FEmitS
 	return false;
 }
 
+void FExpressionCollectionParameter::ComputeAnalyticDerivatives(FTree& Tree, FExpressionDerivatives& OutResult) const
+{
+	OutResult.ExpressionDdx = Tree.NewConstant(0.f);
+	OutResult.ExpressionDdy = OutResult.ExpressionDdx;
+}
+
+bool FExpressionCollectionParameter::PrepareValue(FEmitContext& Context, FEmitScope& Scope, const FRequestedType& RequestedType, FPrepareValueResult& OutResult) const
+{
+	return OutResult.SetType(Context, RequestedType, EExpressionEvaluation::Shader, Shader::EValueType::Float4);
+}
+
+void FExpressionCollectionParameter::EmitValueShader(FEmitContext& Context, FEmitScope& Scope, const FRequestedType& RequestedType, FEmitValueShaderResult& OutResult) const
+{
+	OutResult.Code = Context.EmitInlineExpression(Scope, Shader::EValueType::Float4, TEXT("MaterialCollection%.Vectors[%]"), CollectionIndex, ParameterIndex);
+}
+
 namespace Private
 {
 
diff --git a/Engine/Source/Runtime/Engine/Public/MaterialCachedHLSLTree.h b/Engine/Source/Runtime/Engine/Public/MaterialCachedHLSLTree.h
--- a/Engine/Source/Runtime/Engine/Public/MaterialCachedHLSLTree.h
+++ b/Engine/Source/Runtime/Engine/Public/MaterialCachedHLSLTree.h
@@ -62,6 +62,8 @@ public:
 	const UE::Shader::FStructType* GetVTPageTableResultType() const { return VTPageTableResultType; }
 	const UE::Shader::FValue& GetMaterialAttributesDefaultValue() const { return MaterialAttributesDefaultValue; }
 
+	const TArray<UMaterialParameterCollection*>& GetParameterCollections() const { return ParameterCollections; }
+
 	const TMap<FMaterialConnectionKey, const UE::HLSLTree::FExpression*>& GetConnections() const { return ConnectionMap; }
 
 	ENGINE_API void SetRequestedFields(EShaderFrequency ShaderFrequency, UE::HLSLTree::FRequestedType& OutRequestedType) const;
@@ -80,6 +82,7 @@ private:
 	UE::HLSLTree::FScope* ResultScope = nullptr;
 
 	TArray<UMaterialExpressionCustomOutput*> MaterialCustomOutputs;
+	TArray<UMaterialParameterCollection*> ParameterCollections;
 	TMap<FMaterialConnectionKey, const UE::HLSLTree::FExpression*> ConnectionMap;
 	const UE::Shader::FStructType* MaterialAttributesType = nullptr;
 	const UE::Shader::FStructType* VTPageTableResultType = nullptr;
diff --git a/Engine/Source/Runtime/Engine/Public/MaterialHLSLGenerator.h b/Engine/Source/Runtime/Engine/Public/MaterialHLSLGenerator.h
--- a/Engine/Source/Runtime/Engine/Public/MaterialHLSLGenerator.h
+++ b/Engine/Source/Runtime/Engine/Public/MaterialHLSLGenerator.h
@@ -187,6 +187,8 @@ public:
 	FMaterialParameterInfo GetParameterInfo(const FName& ParameterName) const;
 
 	int32 FindOrAddCustomExpressionOutputStructId(TArrayView<UE::Shader::FStructFieldInitializer> StructFields);
+	
+	int32 FindOrAddParameterCollection(UMaterialParameterCollection* ParameterCollection);
 
 private:
 	static constexpr int32 MaxNumPreviousScopes = UE::HLSLTree::MaxNumPreviousScopes;
diff --git a/Engine/Source/Runtime/Engine/Public/MaterialHLSLTree.h b/Engine/Source/Runtime/Engine/Public/MaterialHLSLTree.h
--- a/Engine/Source/Runtime/Engine/Public/MaterialHLSLTree.h
+++ b/Engine/Source/Runtime/Engine/Public/MaterialHLSLTree.h
@@ -271,6 +271,22 @@ public:
 	virtual bool EmitCustomHLSLParameter(FEmitContext& Context, FEmitScope& Scope, const FName& ObjectTypeName, const TCHAR* ParameterName, FEmitCustomHLSLParameterResult& OutResult) const override;
 };
 
+class FExpressionCollectionParameter : public FExpression
+{
+public:
+	int32 CollectionIndex;
+	int32 ParameterIndex;
+
+	FExpressionCollectionParameter(int32 InCollectionIndex, int32 InParameterIndex)
+		: CollectionIndex(InCollectionIndex)
+		, ParameterIndex(InParameterIndex)
+	{}
+
+	virtual void ComputeAnalyticDerivatives(FTree& Tree, FExpressionDerivatives& OutResult) const override;
+	virtual bool PrepareValue(FEmitContext& Context, FEmitScope& Scope, const FRequestedType& RequestedType, FPrepareValueResult& OutResult) const override;
+	virtual void EmitValueShader(FEmitContext& Context, FEmitScope& Scope, const FRequestedType& RequestedType, FEmitValueShaderResult& OutResult) const override;
+};
+
 class FExpressionTextureSample : public FExpression
 {
 public:
