Fix a CIS error

diff --git a/Engine/Source/Runtime/Engine/Private/Materials/MaterialHLSLGenerator.cpp b/Engine/Source/Runtime/Engine/Private/Materials/MaterialHLSLGenerator.cpp
--- a/Engine/Source/Runtime/Engine/Private/Materials/MaterialHLSLGenerator.cpp
+++ b/Engine/Source/Runtime/Engine/Private/Materials/MaterialHLSLGenerator.cpp
@@ -443,9 +443,9 @@ const UE::HLSLTree::FExpression* FMaterialHLSLGenerator::AcquireExpression(UE::H
 			Hasher.AppendData(&FunctionCall->ParameterAssociation, sizeof(FunctionCall->ParameterAssociation));
 			Hasher.AppendData(&FunctionCall->ParameterIndex, sizeof(FunctionCall->ParameterIndex));
 
-			for (int32 InputIndex = 0; InputIndex < FunctionCall->ConnectedInputs.Num(); ++InputIndex)
+			for (int32 FunctionInputIndex = 0; FunctionInputIndex < FunctionCall->ConnectedInputs.Num(); ++FunctionInputIndex)
 			{
-				const FExpression* InputExpression = FunctionCall->ConnectedInputs[InputIndex];
+				const FExpression* InputExpression = FunctionCall->ConnectedInputs[FunctionInputIndex];
 				Hasher.AppendData(&InputExpression, sizeof(InputExpression));
 			}
 		}
