- Implemented GenerateHLSLExpression interface for various UMaterialExpressions: UMaterialExpressionDynamicParameter, UMaterialExpressionInverseLinearInterpolate, UMaterialExpressionQualitySwitch, UMaterialExpressionShaderStageSwitch, UMaterialExpressionSkyAtmosphereLightDirection, UMaterialExpressionSkyLightEnvMapSample, UMaterialExpressionSmoothStep, UMaterialExpressionStep, and UMaterialExpressionReflectionVectorWS.

- Implemented ComputeAnalyticDerivatives interface for some of the expressions above plus FExpressionFunctionCall and various unary/binary operations: abs, saturate, log2, exp2, length, normalize, and PowPositiveClamped.

- Miscellaneous fixes: FExpressionAppend::PrepareValue now correctly handles requested type EValueType::Any; MP_ShadingModel is now blended with a custom function instead of a regular lerp.

diff --git a/Engine/Source/Runtime/Engine/Classes/Materials/MaterialExpressionDynamicParameter.h b/Engine/Source/Runtime/Engine/Classes/Materials/MaterialExpressionDynamicParameter.h
--- a/Engine/Source/Runtime/Engine/Classes/Materials/MaterialExpressionDynamicParameter.h
+++ b/Engine/Source/Runtime/Engine/Classes/Materials/MaterialExpressionDynamicParameter.h
@@ -53,6 +53,8 @@ class UMaterialExpressionDynamicParameter : public UMaterialExpression
 	virtual int32 GetWidth() const override;
 	virtual int32 GetLabelPadding() override { return 8; }
 	virtual FText GetKeywords() const override { return FText::FromString(TEXT("dynamic")); }
+
+	virtual bool GenerateHLSLExpression(FMaterialHLSLGenerator& Generator, UE::HLSLTree::FScope& Scope, int32 OutputIndex, UE::HLSLTree::FExpression const*& OutExpression) const override;
 #endif
 	//~ End UMaterialExpression Interface
 
diff --git a/Engine/Source/Runtime/Engine/Classes/Materials/MaterialExpressionInverseLinearInterpolate.h b/Engine/Source/Runtime/Engine/Classes/Materials/MaterialExpressionInverseLinearInterpolate.h
--- a/Engine/Source/Runtime/Engine/Classes/Materials/MaterialExpressionInverseLinearInterpolate.h
+++ b/Engine/Source/Runtime/Engine/Classes/Materials/MaterialExpressionInverseLinearInterpolate.h
@@ -45,6 +45,8 @@ class UMaterialExpressionInverseLinearInterpolate : public UMaterialExpression
 	virtual void GetCaption(TArray<FString>& OutCaptions) const override;
 	virtual FText GetKeywords() const override { return FText::FromString(TEXT("InvLerp")); }
 	virtual FText GetCreationName() const override { return FText::FromString(TEXT("InverseLinearInterpolate")); }
+
+	virtual bool GenerateHLSLExpression(FMaterialHLSLGenerator& Generator, UE::HLSLTree::FScope& Scope, int32 OutputIndex, UE::HLSLTree::FExpression const*& OutExpression) const override;
 #endif // WITH_EDITOR
 	//~ End UMaterialExpression Interface
 };
diff --git a/Engine/Source/Runtime/Engine/Classes/Materials/MaterialExpressionQualitySwitch.h b/Engine/Source/Runtime/Engine/Classes/Materials/MaterialExpressionQualitySwitch.h
--- a/Engine/Source/Runtime/Engine/Classes/Materials/MaterialExpressionQualitySwitch.h
+++ b/Engine/Source/Runtime/Engine/Classes/Materials/MaterialExpressionQualitySwitch.h
@@ -38,6 +38,8 @@ class UMaterialExpressionQualitySwitch : public UMaterialExpression
 	virtual bool IsResultStrataMaterial(int32 OutputIndex) override;
 	virtual void GatherStrataMaterialInfo(FStrataMaterialInfo& StrataMaterialInfo, int32 OutputIndex) override;
 	virtual FStrataOperator* StrataGenerateMaterialTopologyTree(class FMaterialCompiler* Compiler, class UMaterialExpression* Parent, int32 OutputIndex) override;
+
+	virtual bool GenerateHLSLExpression(FMaterialHLSLGenerator& Generator, UE::HLSLTree::FScope& Scope, int32 OutputIndex, UE::HLSLTree::FExpression const*& OutExpression) const override;
 #endif // WITH_EDITOR
 	//~ End UMaterialExpression Interface
 
diff --git a/Engine/Source/Runtime/Engine/Classes/Materials/MaterialExpressionShaderStageSwitch.h b/Engine/Source/Runtime/Engine/Classes/Materials/MaterialExpressionShaderStageSwitch.h
--- a/Engine/Source/Runtime/Engine/Classes/Materials/MaterialExpressionShaderStageSwitch.h
+++ b/Engine/Source/Runtime/Engine/Classes/Materials/MaterialExpressionShaderStageSwitch.h
@@ -37,6 +37,8 @@ class UMaterialExpressionShaderStageSwitch : public UMaterialExpression
 	virtual bool IsResultStrataMaterial(int32 OutputIndex) override;
 	virtual void GatherStrataMaterialInfo(FStrataMaterialInfo& StrataMaterialInfo, int32 OutputIndex) override;
 	virtual FStrataOperator* StrataGenerateMaterialTopologyTree(class FMaterialCompiler* Compiler, class UMaterialExpression* Parent, int32 OutputIndex) override;
+
+	virtual bool GenerateHLSLExpression(FMaterialHLSLGenerator& Generator, UE::HLSLTree::FScope& Scope, int32 OutputIndex, UE::HLSLTree::FExpression const*& OutExpression) const override;
 #endif
 	//~ End UMaterialExpression Interface
 };
diff --git a/Engine/Source/Runtime/Engine/Classes/Materials/MaterialExpressionSkyAtmosphereLightDirection.h b/Engine/Source/Runtime/Engine/Classes/Materials/MaterialExpressionSkyAtmosphereLightDirection.h
--- a/Engine/Source/Runtime/Engine/Classes/Materials/MaterialExpressionSkyAtmosphereLightDirection.h
+++ b/Engine/Source/Runtime/Engine/Classes/Materials/MaterialExpressionSkyAtmosphereLightDirection.h
@@ -21,6 +21,8 @@ class UMaterialExpressionSkyAtmosphereLightDirection : public UMaterialExpressio
 #if WITH_EDITOR
 	virtual int32 Compile(class FMaterialCompiler* Compiler, int32 OutputIndex) override;
 	virtual void GetCaption(TArray<FString>& OutCaptions) const override;
+
+	virtual bool GenerateHLSLExpression(FMaterialHLSLGenerator& Generator, UE::HLSLTree::FScope& Scope, int32 OutputIndex, UE::HLSLTree::FExpression const*& OutExpression) const override;
 #endif
 	//~ End UMaterialExpression Interface
 };
diff --git a/Engine/Source/Runtime/Engine/Classes/Materials/MaterialExpressionSkyLightEnvMapSample.h b/Engine/Source/Runtime/Engine/Classes/Materials/MaterialExpressionSkyLightEnvMapSample.h
--- a/Engine/Source/Runtime/Engine/Classes/Materials/MaterialExpressionSkyLightEnvMapSample.h
+++ b/Engine/Source/Runtime/Engine/Classes/Materials/MaterialExpressionSkyLightEnvMapSample.h
@@ -23,6 +23,8 @@ class UMaterialExpressionSkyLightEnvMapSample : public UMaterialExpression
 #if WITH_EDITOR
 	virtual int32 Compile(class FMaterialCompiler* Compiler, int32 OutputIndex) override;
 	virtual void GetCaption(TArray<FString>& OutCaptions) const override;
+
+	virtual bool GenerateHLSLExpression(FMaterialHLSLGenerator& Generator, UE::HLSLTree::FScope& Scope, int32 OutputIndex, UE::HLSLTree::FExpression const*& OutExpression) const override;
 #endif
 	//~ End UMaterialExpression Interface
 };
diff --git a/Engine/Source/Runtime/Engine/Classes/Materials/MaterialExpressionSmoothStep.h b/Engine/Source/Runtime/Engine/Classes/Materials/MaterialExpressionSmoothStep.h
--- a/Engine/Source/Runtime/Engine/Classes/Materials/MaterialExpressionSmoothStep.h
+++ b/Engine/Source/Runtime/Engine/Classes/Materials/MaterialExpressionSmoothStep.h
@@ -42,6 +42,8 @@ class UMaterialExpressionSmoothStep : public UMaterialExpression
 	virtual void GetCaption(TArray<FString>& OutCaptions) const override;
 	virtual FText GetKeywords() const override { return FText::FromString(TEXT("SmoothStep")); }
 	virtual FText GetCreationName() const override { return FText::FromString(TEXT("SmoothStep")); }
+
+	virtual bool GenerateHLSLExpression(FMaterialHLSLGenerator& Generator, UE::HLSLTree::FScope& Scope, int32 OutputIndex, UE::HLSLTree::FExpression const*& OutExpression) const override;
 #endif // WITH_EDITOR
 	//~ End UMaterialExpression Interface
 };
diff --git a/Engine/Source/Runtime/Engine/Classes/Materials/MaterialExpressionStep.h b/Engine/Source/Runtime/Engine/Classes/Materials/MaterialExpressionStep.h
--- a/Engine/Source/Runtime/Engine/Classes/Materials/MaterialExpressionStep.h
+++ b/Engine/Source/Runtime/Engine/Classes/Materials/MaterialExpressionStep.h
@@ -36,6 +36,8 @@ class UMaterialExpressionStep : public UMaterialExpression
 	virtual void GetCaption(TArray<FString>& OutCaptions) const override;
 	virtual FText GetKeywords() const override { return FText::FromString(TEXT("Step")); }
 	virtual FText GetCreationName() const override { return FText::FromString(TEXT("Step")); }
+
+	virtual bool GenerateHLSLExpression(FMaterialHLSLGenerator& Generator, UE::HLSLTree::FScope& Scope, int32 OutputIndex, UE::HLSLTree::FExpression const*& OutExpression) const override;
 #endif // WITH_EDITOR
 	//~ End UMaterialExpression Interface
 };
diff --git a/Engine/Source/Runtime/Engine/Private/HLSLTree/HLSLTree.cpp b/Engine/Source/Runtime/Engine/Private/HLSLTree/HLSLTree.cpp
--- a/Engine/Source/Runtime/Engine/Private/HLSLTree/HLSLTree.cpp
+++ b/Engine/Source/Runtime/Engine/Private/HLSLTree/HLSLTree.cpp
@@ -67,6 +67,7 @@ class FExpressionFunctionCall final : public FExpression
 public:
 	FExpressionFunctionCall(FFunction* InFunction, int32 InOutputIndex) : Function(InFunction), OutputIndex(InOutputIndex) {}
 
+	virtual void ComputeAnalyticDerivatives(FTree& Tree, FExpressionDerivatives& OutResult) const override;
 	virtual bool PrepareValue(FEmitContext& Context, FEmitScope& Scope, const FRequestedType& RequestedType, FPrepareValueResult& OutResult) const override;
 	virtual void EmitValueShader(FEmitContext& Context, FEmitScope& Scope, const FRequestedType& RequestedType, FEmitValueShaderResult& OutResult) const override;
 	virtual void EmitValuePreshader(FEmitContext& Context, FEmitScope& Scope, const FRequestedType& RequestedType, FEmitValuePreshaderResult& OutResult) const override;
@@ -488,6 +489,13 @@ void FExpressionLocalPHI::EmitValuePreshader(FEmitContext& Context, FEmitScope&
 	}
 }
 
+void FExpressionFunctionCall::ComputeAnalyticDerivatives(FTree& Tree, FExpressionDerivatives& OutResult) const
+{
+	check(Function->OutputExpressions.IsValidIndex(OutputIndex));
+
+	Function->OutputExpressions[OutputIndex]->ComputeAnalyticDerivatives(Tree, OutResult);
+}
+
 bool FExpressionFunctionCall::PrepareValue(FEmitContext& Context, FEmitScope& Scope, const FRequestedType& RequestedType, FPrepareValueResult& OutResult) const
 {
 	FEmitScope* EmitFunctionScope = Context.PrepareScopeWithParent(Function->RootScope, Function->CalledScope);
diff --git a/Engine/Source/Runtime/Engine/Private/HLSLTree/HLSLTreeCommon.cpp b/Engine/Source/Runtime/Engine/Private/HLSLTree/HLSLTreeCommon.cpp
--- a/Engine/Source/Runtime/Engine/Private/HLSLTree/HLSLTreeCommon.cpp
+++ b/Engine/Source/Runtime/Engine/Private/HLSLTree/HLSLTreeCommon.cpp
@@ -871,12 +871,19 @@ bool FExpressionAppend::PrepareValue(FEmitContext& Context, FEmitScope& Scope, c
 	const int32 NumRhsComponents = FMath::Max(NumRequestedComponents - NumLhsComponents, 1);
 
 	FRequestedType RhsRequestedType;
-	RhsRequestedType.Type = Shader::MakeValueType(RequestedType.GetValueComponentType(), NumRhsComponents);
-	for (int32 Index = NumLhsComponents; Index < NumRequestedComponents; ++Index)
+	if (RequestedType.Type.IsAny())
 	{
-		if (RequestedType.IsComponentRequested(Index))
+		RhsRequestedType = Shader::EValueType::Any;
+	}
+	else
+	{
+		RhsRequestedType.Type = Shader::MakeValueType(RequestedType.GetValueComponentType(), NumRhsComponents);
+		for (int32 Index = NumLhsComponents; Index < NumRequestedComponents; ++Index)
 		{
-			RhsRequestedType.SetComponentRequest(Index - NumLhsComponents);
+			if (RequestedType.IsComponentRequested(Index))
+			{
+				RhsRequestedType.SetComponentRequest(Index - NumLhsComponents);
+			}
 		}
 	}
 
@@ -1046,6 +1053,31 @@ bool FExpressionShadingPathSwitch::IsInputActive(const FEmitContext& Context, in
 	return Index == (int32)ShadingPathToCompile;
 }
 
+FExpressionQualitySwitch::FExpressionQualitySwitch(TConstArrayView<const FExpression*> InInputs)
+	: FExpressionSwitchBase(InInputs)
+{
+	static_assert(MaxInputs >= (int32)EMaterialQualityLevel::Num + 1, "FExpressionSwitchBase is too small for FExpressionQualitySwitch");
+	check(InInputs.Num() == (int32)EMaterialQualityLevel::Num + 1);
+}
+
+bool FExpressionQualitySwitch::IsInputActive(const FEmitContext& Context, int32 Index) const
+{
+	return Context.TargetParameters.IsGenericTarget() || Index == (int32)Context.Material->GetQualityLevel();
+}
+
+FExpressionShaderStageSwitch::FExpressionShaderStageSwitch(TConstArrayView<const FExpression*> InInputs)
+	: FExpressionSwitchBase(InInputs)
+{
+	static_assert(MaxInputs >= 2, "FExpressionSwitchBase is too small for FExpressionShaderStageSwitch");
+	check(InInputs.Num() == 2);
+}
+
+bool FExpressionShaderStageSwitch::IsInputActive(const FEmitContext& Context, int32 Index) const
+{
+	check(Context.ShaderFrequency == SF_Pixel || Context.ShaderFrequency == SF_Vertex);
+	return (Context.ShaderFrequency == SF_Pixel && Index == 0) || (Context.ShaderFrequency == SF_Vertex && Index == 1);
+}
+
 bool FExpressionInlineCustomHLSL::PrepareValue(FEmitContext& Context, FEmitScope& Scope, const FRequestedType& RequestedType, FPrepareValueResult& OutResult) const
 {
 	return OutResult.SetType(Context, RequestedType, EExpressionEvaluation::Shader, Shader::FType(ResultType));
diff --git a/Engine/Source/Runtime/Engine/Private/HLSLTree/HLSLTreeOperation.cpp b/Engine/Source/Runtime/Engine/Private/HLSLTree/HLSLTreeOperation.cpp
--- a/Engine/Source/Runtime/Engine/Private/HLSLTree/HLSLTreeOperation.cpp
+++ b/Engine/Source/Runtime/Engine/Private/HLSLTree/HLSLTreeOperation.cpp
@@ -13,7 +13,7 @@ class FExpressionOperation : public FExpression
 public:
 	FExpressionOperation(EOperation InOp, TConstArrayView<const FExpression*> InInputs);
 
-	static constexpr int8 MaxInputs = 2;
+	static constexpr int8 MaxInputs = 3;
 
 	EOperation Op;
 	const FExpression* Inputs[MaxInputs];
@@ -74,6 +74,7 @@ FOperationDescription GetOperationDescription(EOperation Op)
 	case EOperation::Mul: return FOperationDescription(TEXT("Multiply"), TEXT("*"), 2, Shader::EPreshaderOpcode::Mul); break;
 	case EOperation::Div: return FOperationDescription(TEXT("Divide"), TEXT("/"), 2, Shader::EPreshaderOpcode::Div); break;
 	case EOperation::Fmod: return FOperationDescription(TEXT("Fmod"), TEXT("%"), 2, Shader::EPreshaderOpcode::Fmod); break;
+	case EOperation::Step: return FOperationDescription(TEXT("Step"), TEXT("step"), 2, Shader::EPreshaderOpcode::Nop); break;
 	case EOperation::PowPositiveClamped: return FOperationDescription(TEXT("PowPositiveClamped"), TEXT("PowPositiveClamped"), 2, Shader::EPreshaderOpcode::Nop); break;
 	case EOperation::Atan2: return FOperationDescription(TEXT("Atan2"), TEXT("atan2"), 2, Shader::EPreshaderOpcode::Atan2); break;
 	case EOperation::Atan2Fast: return FOperationDescription(TEXT("Atan2Fast"), TEXT("atan2Fast"), 2, Shader::EPreshaderOpcode::Atan2); break;
@@ -87,6 +88,10 @@ FOperationDescription GetOperationDescription(EOperation Op)
 	case EOperation::VecMulMatrix4: return FOperationDescription(TEXT("VecMulMatrix4"), TEXT("mul"), 2, Shader::EPreshaderOpcode::Nop); break;
 	case EOperation::Matrix3MulVec: return FOperationDescription(TEXT("Matrix3MulVec"), TEXT("mul"), 2, Shader::EPreshaderOpcode::Nop); break;
 	case EOperation::Matrix4MulVec: return FOperationDescription(TEXT("Matrix4MulVec"), TEXT("mul"), 2, Shader::EPreshaderOpcode::Nop); break;
+
+	// Ternary
+	case EOperation::SmoothStep: return FOperationDescription(TEXT("SmoothStep"), TEXT("smoothstep"), 3, Shader::EPreshaderOpcode::Nop); break;
+
 	default: checkNoEntry(); return FOperationDescription();
 	}
 }
@@ -103,6 +108,12 @@ const FExpression* FTree::NewBinaryOp(EOperation Op, const FExpression* Lhs, con
 	return NewExpression<FExpressionOperation>(Op, Inputs);
 }
 
+const FExpression* FTree::NewTernaryOp(EOperation Op, const FExpression* Input0, const FExpression* Input1, const FExpression* Input2)
+{
+	const FExpression* Inputs[3] = { Input0, Input1, Input2 };
+	return NewExpression<FExpressionOperation>(Op, Inputs);
+}
+
 const FExpression* FTree::NewLog(const FExpression* Input)
 {
 	const float Log2ToLog = 1.0f / FMath::Log2(UE_EULERS_NUMBER);
@@ -302,6 +313,8 @@ FOperationTypes GetOperationTypes(EOperation Op, TConstArrayView<FPreparedType>
 			break;
 		case EOperation::Saturate:
 		case EOperation::Frac:
+		case EOperation::Step:
+		case EOperation::SmoothStep:
 			for (int32 Index = 0; Index < Types.ResultType.PreparedComponents.Num(); ++Index)
 			{
 				Types.ResultType.SetComponentBounds(Index, Shader::FComponentBounds(Shader::EComponentBound::Zero, Shader::EComponentBound::One));
@@ -374,6 +387,7 @@ void FExpressionOperation::ComputeAnalyticDerivatives(FTree& Tree, FExpressionDe
 	case EOperation::Round:
 	case EOperation::Trunc:
 	case EOperation::Sign:
+	case EOperation::Step:
 		OutResult.ExpressionDdx = OutResult.ExpressionDdy = Tree.NewConstant(0.0f);
 		break;
 	default:
@@ -486,15 +500,65 @@ void FExpressionOperation::ComputeAnalyticDerivatives(FTree& Tree, FExpressionDe
 		OutResult.ExpressionDdy = Tree.NewAdd(Tree.NewMul(dFdA, InputDerivatives[0].ExpressionDdy), Tree.NewMul(dFdB, InputDerivatives[1].ExpressionDdy));
 		break;
 	}
-	case EOperation::Length:
-	case EOperation::Normalize:
 	case EOperation::Abs:
+	{
+		const FExpression* ConstantZero = Tree.NewConstant(0.f);
+		const FExpression* ConstantOne = Tree.NewConstant(1.f);
+		const FExpression* Cond = Tree.NewGreaterEqual(Inputs[0], ConstantZero);
+		const FExpression* Sign = Tree.NewExpression<FExpressionSelect>(Cond, ConstantOne, Tree.NewNeg(ConstantOne));
+		OutResult.ExpressionDdx = Tree.NewMul(Sign, InputDerivatives[0].ExpressionDdx);
+		OutResult.ExpressionDdy = Tree.NewMul(Sign, InputDerivatives[0].ExpressionDdy);
+		break;
+	}
 	case EOperation::Saturate:
-	case EOperation::PowPositiveClamped:
+	{
+		const FExpression* ConstantZero = Tree.NewConstant(0.f);
+		const FExpression* ConstantOne = Tree.NewConstant(1.f);
+		const FExpression* Cond0 = Tree.NewGreater(Inputs[0], ConstantZero);
+		const FExpression* Cond1 = Tree.NewLess(Inputs[0], ConstantOne);
+		OutResult.ExpressionDdx = Tree.NewExpression<FExpressionSelect>(Cond0, Tree.NewExpression<FExpressionSelect>(Cond1, InputDerivatives[0].ExpressionDdx, ConstantOne), ConstantZero);
+		OutResult.ExpressionDdy = Tree.NewExpression<FExpressionSelect>(Cond0, Tree.NewExpression<FExpressionSelect>(Cond1, InputDerivatives[0].ExpressionDdy, ConstantOne), ConstantZero);
+		break;
+	}
 	case EOperation::Log2:
+	{
+		const FExpression* dFdA = Tree.NewMul(Tree.NewRcp(Inputs[0]), Tree.NewConstant(1.442695f));
+		OutResult.ExpressionDdx = Tree.NewMul(dFdA, InputDerivatives[0].ExpressionDdx);
+		OutResult.ExpressionDdy = Tree.NewMul(dFdA, InputDerivatives[0].ExpressionDdy);
+		break;
+	}
 	case EOperation::Exp2:
-		// TODO
+	{
+		const FExpression* dFdA = Tree.NewMul(Tree.NewExp2(Inputs[0]), Tree.NewConstant(0.693147f));
+		OutResult.ExpressionDdx = Tree.NewMul(dFdA, InputDerivatives[0].ExpressionDdx);
+		OutResult.ExpressionDdy = Tree.NewMul(dFdA, InputDerivatives[0].ExpressionDdy);
+		break;
+	}
+	case EOperation::Length:
+	{
+		const FExpression* LengthExpression = Tree.NewSqrt(Tree.NewSum(Tree.NewMul(Inputs[0], Inputs[0])));
+		LengthExpression->ComputeAnalyticDerivatives(Tree, OutResult);
+		break;
+	}
+	case EOperation::Normalize:
+	{
+		const FExpression* NormalizeExpression = Tree.NewMul(Inputs[0], Tree.NewRsqrt(Tree.NewSum(Tree.NewMul(Inputs[0], Inputs[0]))));
+		NormalizeExpression->ComputeAnalyticDerivatives(Tree, OutResult);
+		break;
+	}
+	case EOperation::PowPositiveClamped:
+	{
+		const FExpression* ConstantZero = Tree.NewConstant(0.f);
+		const FExpression* Cond = Tree.NewLess(ConstantZero, Inputs[1]); // should we check for A as well?
+		const FExpression* F = Tree.NewPowClamped(Inputs[0], Inputs[1]);
+		const FExpression* LnA = Tree.NewLog(Inputs[0]);
+		const FExpression* BByA = Tree.NewDiv(Inputs[1], Inputs[0]);
+		const FExpression* InRangeDdx = Tree.NewMul(F, Tree.NewAdd(Tree.NewMul(InputDerivatives[1].ExpressionDdx, LnA), Tree.NewMul(BByA, InputDerivatives[0].ExpressionDdx)));
+		const FExpression* InRangeDdy = Tree.NewMul(F, Tree.NewAdd(Tree.NewMul(InputDerivatives[1].ExpressionDdy, LnA), Tree.NewMul(BByA, InputDerivatives[0].ExpressionDdy)));
+		OutResult.ExpressionDdx = Tree.NewExpression<FExpressionSelect>(Cond, InRangeDdx, ConstantZero);
+		OutResult.ExpressionDdy = Tree.NewExpression<FExpressionSelect>(Cond, InRangeDdy, ConstantZero);
 		break;
+	}
 	case EOperation::Add:
 		OutResult.ExpressionDdx = Tree.NewAdd(InputDerivatives[0].ExpressionDdx, InputDerivatives[1].ExpressionDdx);
 		OutResult.ExpressionDdy = Tree.NewAdd(InputDerivatives[0].ExpressionDdy, InputDerivatives[1].ExpressionDdy);
@@ -542,6 +606,9 @@ void FExpressionOperation::ComputeAnalyticDerivatives(FTree& Tree, FExpressionDe
 		// TODO
 		OutResult.ExpressionDdx = OutResult.ExpressionDdy = Tree.NewConstant(FVector3f(0.0f, 0.0f, 0.0f));
 		break;
+	case EOperation::SmoothStep:
+		// TODO
+		break;
 	default:
 		checkNoEntry();
 		break;
@@ -714,6 +781,7 @@ void FExpressionOperation::EmitValueShader(FEmitContext& Context, FEmitScope& Sc
 	case EOperation::Mul: OutResult.Code = Context.EmitExpression(Scope, ResultType, Types.bIsLWC ? TEXT("LWCMultiply(%, %)") : TEXT("(% * %)"), InputValue[0], InputValue[1]); break;
 	case EOperation::Div: OutResult.Code = Context.EmitExpression(Scope, ResultType, Types.bIsLWC ? TEXT("LWCDivide(%, %)") : TEXT("(% / %)"), InputValue[0], InputValue[1]); break;
 	case EOperation::Fmod: OutResult.Code = Context.EmitExpression(Scope, ResultType, Types.bIsLWC ? TEXT("LWCFmod(%, %)") : TEXT("fmod(%, %)"), InputValue[0], InputValue[1]); break;
+	case EOperation::Step: OutResult.Code = Context.EmitExpression(Scope, ResultType, Types.bIsLWC ? TEXT("LWCStep(%, %)") : TEXT("step(%, %)"), InputValue[0], InputValue[1]); break;
 	case EOperation::PowPositiveClamped: OutResult.Code = Context.EmitExpression(Scope, ResultType, TEXT("PositiveClampedPow(%, %)"), InputValue[0], InputValue[1]); break;
 	case EOperation::Atan2: OutResult.Code = Context.EmitExpression(Scope, ResultType, TEXT("atan2(%, %)"), InputValue[0], InputValue[1]); break;
 	case EOperation::Atan2Fast: OutResult.Code = Context.EmitExpression(Scope, ResultType, TEXT("atan2Fast(%, %)"), InputValue[0], InputValue[1]); break;
@@ -749,6 +817,10 @@ void FExpressionOperation::EmitValueShader(FEmitContext& Context, FEmitScope& Sc
 	case EOperation::Matrix4MulVec:
 		OutResult.Code = Context.EmitExpression(Scope, ResultType, TEXT("mul(%, %)"), InputValue[0], InputValue[1]);
 		break;
+
+	// Ternary Ops
+	case EOperation::SmoothStep: OutResult.Code = Context.EmitExpression(Scope, ResultType, Types.bIsLWC ? TEXT("LWCSmoothStep(%, %, %)") : TEXT("smoothstep(%, %, %)"), InputValue[0], InputValue[1], InputValue[2]); break;
+
 	default:
 		checkNoEntry();
 		break;
diff --git a/Engine/Source/Runtime/Engine/Private/Materials/MaterialCachedData.cpp b/Engine/Source/Runtime/Engine/Private/Materials/MaterialCachedData.cpp
--- a/Engine/Source/Runtime/Engine/Private/Materials/MaterialCachedData.cpp
+++ b/Engine/Source/Runtime/Engine/Private/Materials/MaterialCachedData.cpp
@@ -952,7 +952,7 @@ void PrepareHLSLTree(UE::HLSLTree::FEmitContext& EmitContext,
 			if (FMaterialAttributeDefinitionMap::GetShaderFrequency(AttributeID) == ShaderFrequency)
 			{
 				const EMaterialProperty Property = FMaterialAttributeDefinitionMap::GetProperty(AttributeID);
-				if (CachedTree.IsAttributeUsed(EmitContext, *EmitResultScope, ResultType, Property))
+				if (Property != MP_CustomOutput && CachedTree.IsAttributeUsed(EmitContext, *EmitResultScope, ResultType, Property))
 				{
 					CachedData.SetPropertyConnected(Property);
 				}
diff --git a/Engine/Source/Runtime/Engine/Private/Materials/MaterialExpressionHLSL.cpp b/Engine/Source/Runtime/Engine/Private/Materials/MaterialExpressionHLSL.cpp
--- a/Engine/Source/Runtime/Engine/Private/Materials/MaterialExpressionHLSL.cpp
+++ b/Engine/Source/Runtime/Engine/Private/Materials/MaterialExpressionHLSL.cpp
@@ -20,6 +20,7 @@
 #include "Materials/MaterialExpressionParameter.h"
 #include "Materials/MaterialExpressionCollectionParameter.h"
 #include "Materials/MaterialParameterCollection.h"
+#include "Materials/MaterialExpressionDynamicParameter.h"
 #include "Materials/MaterialExpressionCurveAtlasRowParameter.h"
 #include "Materials/MaterialExpressionStaticSwitchParameter.h"
 #include "Materials/MaterialExpressionStaticComponentMaskParameter.h"
@@ -51,6 +52,7 @@
 #include "Materials/MaterialExpressionParticleColor.h"
 #include "Materials/MaterialExpressionParticlePositionWS.h"
 #include "Materials/MaterialExpressionParticleRadius.h"
+#include "Materials/MaterialExpressionSkyLightEnvMapSample.h"
 #include "Materials/MaterialExpressionTextureSample.h"
 #include "Materials/MaterialExpressionTextureSampleParameter.h"
 #include "Materials/MaterialExpressionFontSample.h"
@@ -68,6 +70,8 @@
 #include "Materials/MaterialExpressionStaticSwitch.h"
 #include "Materials/MaterialExpressionFeatureLevelSwitch.h"
 #include "Materials/MaterialExpressionShadingPathSwitch.h"
+#include "Materials/MaterialExpressionQualitySwitch.h"
+#include "Materials/MaterialExpressionShaderStageSwitch.h"
 #include "Materials/MaterialExpressionGetLocal.h"
 #include "Materials/MaterialExpressionOneMinus.h"
 #include "Materials/MaterialExpressionAbs.h"
@@ -107,6 +111,9 @@
 #include "Materials/MaterialExpressionMax.h"
 #include "Materials/MaterialExpressionClamp.h"
 #include "Materials/MaterialExpressionLinearInterpolate.h"
+#include "Materials/MaterialExpressionInverseLinearInterpolate.h"
+#include "Materials/MaterialExpressionStep.h"
+#include "Materials/MaterialExpressionSmoothStep.h"
 #include "Materials/MaterialExpressionDistance.h"
 #include "Materials/MaterialExpressionNormalize.h"
 #include "Materials/MaterialExpressionBinaryOp.h"
@@ -126,6 +133,7 @@
 #include "Materials/MaterialExpressionBumpOffset.h"
 #include "Materials/MaterialExpressionPanner.h"
 #include "Materials/MaterialExpressionRotator.h"
+#include "Materials/MaterialExpressionSkyAtmosphereLightDirection.h"
 #include "Materials/MaterialExpressionVolumetricAdvancedMaterialInput.h"
 #include "Materials/MaterialExpressionVolumetricAdvancedMaterialOutput.h"
 #include "Materials/MaterialExpressionHairAttributes.h"
@@ -144,6 +152,7 @@
 #include "Materials/MaterialFunctionInterface.h"
 #include "MaterialHLSLTree.h"
 #include "MaterialShared.h"
+#include "MaterialDomain.h"
 #include "Engine/Engine.h"
 #include "Engine/Font.h"
 #include "Curves/CurveLinearColorAtlas.h"
@@ -324,6 +333,59 @@ bool UMaterialExpressionShadingPathSwitch::GenerateHLSLExpression(FMaterialHLSLG
 	return true;
 }
 
+bool UMaterialExpressionQualitySwitch::GenerateHLSLExpression(FMaterialHLSLGenerator& Generator, UE::HLSLTree::FScope& Scope, int32 OutputIndex, UE::HLSLTree::FExpression const*& OutExpression) const
+{
+	using namespace UE::HLSLTree;
+	if (!Default.GetTracedInput().Expression)
+	{
+		return Generator.Error(TEXT("Missing default input"));
+	}
+
+	const FExpression* ExpressionDefault = Default.AcquireHLSLExpression(Generator, Scope);
+	const FExpression* ExpressionInputs[EMaterialQualityLevel::Num + 1] = { nullptr };
+
+	for (int32 Index = 0; Index < EMaterialQualityLevel::Num; ++Index)
+	{
+		const FExpression* Expression = nullptr;
+		const FExpressionInput& QualityInput = Inputs[Index];
+		if (QualityInput.GetTracedInput().Expression)
+		{
+			Expression = QualityInput.AcquireHLSLExpression(Generator, Scope);
+		}
+		else
+		{
+			Expression = ExpressionDefault;
+		}
+		ExpressionInputs[Index] = Expression;
+	}
+	ExpressionInputs[EMaterialQualityLevel::Num] = ExpressionDefault;
+
+	OutExpression = Generator.GetTree().NewExpression<FExpressionQualitySwitch>(ExpressionInputs);
+	return true;
+}
+
+bool UMaterialExpressionShaderStageSwitch::GenerateHLSLExpression(FMaterialHLSLGenerator& Generator, UE::HLSLTree::FScope& Scope, int32 OutputIndex, UE::HLSLTree::FExpression const*& OutExpression) const
+{
+	if (!PixelShader.GetTracedInput().Expression)
+	{
+		return Generator.Error(TEXT("Missing input PixelShader"));
+	}
+	else if (!VertexShader.GetTracedInput().Expression)
+	{
+		return Generator.Error(TEXT("Missing input VertexShader"));
+	}
+	else
+	{
+		using namespace UE::HLSLTree;
+		const FExpression* ExpressionInputs[2];
+		ExpressionInputs[0] = PixelShader.AcquireHLSLExpression(Generator, Scope);
+		ExpressionInputs[1] = VertexShader.AcquireHLSLExpression(Generator, Scope);
+
+		OutExpression = Generator.GetTree().NewExpression<FExpressionShaderStageSwitch>(ExpressionInputs);
+		return true;
+	}
+}
+
 bool UMaterialExpressionGetLocal::GenerateHLSLExpression(FMaterialHLSLGenerator& Generator, UE::HLSLTree::FScope& Scope, int32 OutputIndex, UE::HLSLTree::FExpression const*& OutExpression) const
 {
 	OutExpression = Generator.GetTree().AcquireLocal(Scope, LocalName);
@@ -382,6 +444,15 @@ bool UMaterialExpressionCollectionParameter::GenerateHLSLExpression(FMaterialHLS
 	}
 }
 
+bool UMaterialExpressionDynamicParameter::GenerateHLSLExpression(FMaterialHLSLGenerator& Generator, UE::HLSLTree::FScope& Scope, int32 OutputIndex, UE::HLSLTree::FExpression const*& OutExpression) const
+{
+	using namespace UE::HLSLTree;
+	FTree& Tree = Generator.GetTree();
+	const FExpression* DefaultValueExpression = Tree.NewConstant(DefaultValue);
+	OutExpression = Tree.NewExpression<Material::FExpressionDynamicParameter>(DefaultValueExpression, ParameterIndex);
+	return true;
+}
+
 bool UMaterialExpressionCurveAtlasRowParameter::GenerateHLSLExpression(FMaterialHLSLGenerator& Generator, UE::HLSLTree::FScope& Scope, int32 OutputIndex, UE::HLSLTree::FExpression const*& OutExpression) const
 {
 	using namespace UE::HLSLTree;
@@ -615,10 +686,25 @@ bool UMaterialExpressionTwoSidedSign::GenerateHLSLExpression(FMaterialHLSLGenera
 
 bool UMaterialExpressionReflectionVectorWS::GenerateHLSLExpression(FMaterialHLSLGenerator& Generator, UE::HLSLTree::FScope& Scope, int32 OutputIndex, UE::HLSLTree::FExpression const*& OutExpression) const
 {
-	using namespace UE::HLSLTree::Material;
-	check(!CustomWorldNormal.GetTracedInput().Expression); // TODO
-
-	OutExpression = Generator.GetTree().NewExpression<FExpressionExternalInput>(EExternalInput::WorldReflection);
+	FExpressionInput CustomWorldNormalInput = CustomWorldNormal.GetTracedInput();
+	if (CustomWorldNormalInput.IsConnected())
+	{
+		using namespace UE::HLSLTree;
+		FTree& Tree = Generator.GetTree();
+		const FExpression* CustomWorldNormalExpression = CustomWorldNormalInput.TryAcquireHLSLExpression(Generator, Scope);
+		if (bNormalizeCustomWorldNormal)
+		{
+			CustomWorldNormalExpression = Tree.NewNormalize(CustomWorldNormalExpression);
+		}
+		const FExpression* CameraVectorExpression = Tree.NewExpression<Material::FExpressionExternalInput>(Material::EExternalInput::CameraVector);
+		const FExpression* LengthExpression = Tree.NewMul(Tree.NewDot(CustomWorldNormalExpression, CameraVectorExpression), Tree.NewConstant(2.f));
+		OutExpression = Tree.NewAdd(Tree.NewNeg(CameraVectorExpression), Tree.NewMul(CustomWorldNormalExpression, LengthExpression));
+	}
+	else
+	{
+		using namespace UE::HLSLTree::Material;
+		OutExpression = Generator.GetTree().NewExpression<FExpressionExternalInput>(EExternalInput::WorldReflection);
+	}
 	return true;
 }
 
@@ -797,6 +883,27 @@ bool UMaterialExpressionTextureObjectParameter::GenerateHLSLExpression(FMaterial
 	return true;
 }
 
+bool UMaterialExpressionSkyLightEnvMapSample::GenerateHLSLExpression(FMaterialHLSLGenerator& Generator, UE::HLSLTree::FScope& Scope, int32 OutputIndex, UE::HLSLTree::FExpression const*& OutExpression) const
+{
+	using namespace UE::HLSLTree;
+
+	const FExpression* DirectionExpression = Direction.AcquireHLSLExpressionOrConstant(Generator, Scope, FVector3f(0.f, 0.f, 1.f));
+	const FExpression* RoughnessExpression = Roughness.AcquireHLSLExpressionOrConstant(Generator, Scope, 0.f);
+
+	if (!DirectionExpression || !RoughnessExpression)
+	{
+		return false;
+	}
+
+	if (Generator.GetTargetMaterial()->MaterialDomain != MD_Surface)
+	{
+		return Generator.Error(TEXT("The SkyLightEnvMapSample node can only be used when material Domain is set to Surface."));
+	}
+
+	OutExpression = Generator.GetTree().NewExpression<Material::FExpressionSkyLightEnvMapSample>(DirectionExpression, RoughnessExpression);
+	return true;
+}
+
 bool UMaterialExpressionTextureSample::GenerateHLSLExpressionBase(FMaterialHLSLGenerator& Generator, UE::HLSLTree::FScope& Scope, const UE::HLSLTree::FExpression* TextureExpression, UE::HLSLTree::FExpression const*& OutExpression) const
 {
 	using namespace UE::HLSLTree;
@@ -816,10 +923,13 @@ bool UMaterialExpressionTextureSample::GenerateHLSLExpressionBase(FMaterialHLSLG
 	case TMVM_MipBias:
 	{
 		MipLevelExpression = MipValue.AcquireHLSLExpressionOrConstant(Generator, Scope, ConstMipValue);
-		const FExpression* DerivativeScale = Generator.GetTree().NewExp2(MipLevelExpression);
 		TexCoordDerivatives = Generator.GetTree().GetAnalyticDerivatives(TexCoordExpression);
-		TexCoordDerivatives.ExpressionDdx = Generator.GetTree().NewMul(TexCoordDerivatives.ExpressionDdx, DerivativeScale);
-		TexCoordDerivatives.ExpressionDdy = Generator.GetTree().NewMul(TexCoordDerivatives.ExpressionDdy, DerivativeScale);
+		if (TexCoordDerivatives.IsValid())
+		{
+			const FExpression* DerivativeScale = Generator.GetTree().NewExp2(MipLevelExpression);
+			TexCoordDerivatives.ExpressionDdx = Generator.GetTree().NewMul(TexCoordDerivatives.ExpressionDdx, DerivativeScale);
+			TexCoordDerivatives.ExpressionDdy = Generator.GetTree().NewMul(TexCoordDerivatives.ExpressionDdy, DerivativeScale);
+		}
 		break;
 	}
 	case TMVM_MipLevel:
@@ -1419,6 +1529,55 @@ bool UMaterialExpressionLinearInterpolate::GenerateHLSLExpression(FMaterialHLSLG
 	return true;
 }
 
+bool UMaterialExpressionInverseLinearInterpolate::GenerateHLSLExpression(FMaterialHLSLGenerator& Generator, UE::HLSLTree::FScope& Scope, int32 OutputIndex, UE::HLSLTree::FExpression const*& OutExpression) const
+{
+	const UE::HLSLTree::FExpression* ExpressionA = A.AcquireHLSLExpressionOrConstant(Generator, Scope, ConstA);
+	const UE::HLSLTree::FExpression* ExpressionB = B.AcquireHLSLExpressionOrConstant(Generator, Scope, ConstB);
+	const UE::HLSLTree::FExpression* ExpressionValue = Value.AcquireHLSLExpressionOrConstant(Generator, Scope, ConstValue);
+	if (!ExpressionA || !ExpressionB || !ExpressionValue || ExpressionA == ExpressionB)
+	{
+		return false;
+	}
+	UE::HLSLTree::FTree& Tree = Generator.GetTree();
+	OutExpression = Tree.NewDiv(Tree.NewSub(ExpressionValue, ExpressionA), Tree.NewSub(ExpressionB, ExpressionA));
+	if (bClampResult)
+	{
+		OutExpression = Tree.NewSaturate(OutExpression);
+	}
+	return true;
+}
+
+bool UMaterialExpressionStep::GenerateHLSLExpression(FMaterialHLSLGenerator& Generator, UE::HLSLTree::FScope& Scope, int32 OutputIndex, UE::HLSLTree::FExpression const*& OutExpression) const
+{
+	const UE::HLSLTree::FExpression* ExpressionY = Y.AcquireHLSLExpressionOrConstant(Generator, Scope, ConstY);
+	const UE::HLSLTree::FExpression* ExpressionX = X.AcquireHLSLExpressionOrConstant(Generator, Scope, ConstX);
+	if (!ExpressionY || !ExpressionX)
+	{
+		return false;
+	}
+	OutExpression = Generator.GetTree().NewStep(ExpressionY, ExpressionX);
+	return true;
+}
+
+bool UMaterialExpressionSmoothStep::GenerateHLSLExpression(FMaterialHLSLGenerator& Generator, UE::HLSLTree::FScope& Scope, int32 OutputIndex, UE::HLSLTree::FExpression const*& OutExpression) const
+{
+	const UE::HLSLTree::FExpression* ExpressionMin = Min.AcquireHLSLExpressionOrConstant(Generator, Scope, ConstMin);
+	const UE::HLSLTree::FExpression* ExpressionMax = Max.AcquireHLSLExpressionOrConstant(Generator, Scope, ConstMax);
+	const UE::HLSLTree::FExpression* ExpressionValue = Value.AcquireHLSLExpressionOrConstant(Generator, Scope, ConstValue);
+	if (!ExpressionMin || !ExpressionMax || !ExpressionValue)
+	{
+		return false;
+	}
+	UE::HLSLTree::FTree& Tree = Generator.GetTree();
+	if (ExpressionMin == ExpressionMax)
+	{
+		OutExpression = Tree.NewStep(ExpressionMin, ExpressionValue);
+		return true;
+	}
+	OutExpression = Tree.NewSmoothStep(ExpressionMin, ExpressionMax, ExpressionValue);
+	return true;
+}
+
 bool UMaterialExpressionDistance::GenerateHLSLExpression(FMaterialHLSLGenerator& Generator, UE::HLSLTree::FScope& Scope, int32 OutputIndex, UE::HLSLTree::FExpression const*& OutExpression) const
 {
 	const UE::HLSLTree::FExpression* ExpressionA = A.AcquireHLSLExpression(Generator, Scope);
@@ -1805,7 +1964,17 @@ bool UMaterialExpressionBlendMaterialAttributes::GenerateHLSLExpression(FMateria
 		{
 			const FExpression* ExpressionFieldA = Generator.GetTree().NewExpression<FExpressionGetStructField>(MaterialAttributesType, Field, ExpressionA);
 			const FExpression* ExpressionFieldB = Generator.GetTree().NewExpression<FExpressionGetStructField>(MaterialAttributesType, Field, ExpressionB);
-			ExpressionFieldResult = Generator.GetTree().NewLerp(ExpressionFieldA, ExpressionFieldB, ExpressionAlpha);
+			const EMaterialProperty AttributeEnum = FMaterialAttributeDefinitionMap::GetProperty(AttributeID);
+			if (AttributeEnum == MP_ShadingModel)
+			{
+				const FExpression* ConditionExpression = Generator.GetTree().NewLess(ExpressionAlpha, Generator.GetTree().NewConstant(0.5f));
+				ExpressionFieldResult = Generator.GenerateBranch(Scope, ConditionExpression, ExpressionFieldA, ExpressionFieldB);
+			}
+			else
+			{
+				// TODO: MP_FrontMaterial also seems to require special handling
+				ExpressionFieldResult = Generator.GetTree().NewLerp(ExpressionFieldA, ExpressionFieldB, ExpressionAlpha);
+			}
 		}
 		ExpressionResult = Generator.GetTree().NewExpression<FExpressionSetStructField>(MaterialAttributesType, Field, ExpressionResult, ExpressionFieldResult);
 	}
@@ -2263,6 +2432,12 @@ bool UMaterialExpressionRotator::GenerateHLSLExpression(FMaterialHLSLGenerator&
 	return true;
 }
 
+bool UMaterialExpressionSkyAtmosphereLightDirection::GenerateHLSLExpression(FMaterialHLSLGenerator& Generator, UE::HLSLTree::FScope& Scope, int32 OutputIndex, UE::HLSLTree::FExpression const*& OutExpression) const
+{
+	OutExpression = Generator.GetTree().NewExpression<UE::HLSLTree::Material::FExpressionSkyAtmosphereLightDirection>(LightIndex);
+	return true;
+}
+
 static UE::Shader::FType GetCustomOutputType(const FMaterialHLSLGenerator& Generator, ECustomMaterialOutputType Type)
 {
 	using namespace UE::Shader;
diff --git a/Engine/Source/Runtime/Engine/Private/Materials/MaterialHLSLEmitter.cpp b/Engine/Source/Runtime/Engine/Private/Materials/MaterialHLSLEmitter.cpp
--- a/Engine/Source/Runtime/Engine/Private/Materials/MaterialHLSLEmitter.cpp
+++ b/Engine/Source/Runtime/Engine/Private/Materials/MaterialHLSLEmitter.cpp
@@ -406,11 +406,10 @@ static void GetMaterialEnvironment(EShaderPlatform InPlatform,
 		OutEnvironment.SetDefine(TEXT("NEEDS_PARTICLE_VELOCITY"), 1);
 	}
 
-	uint32 DynamicParticleParameterMask = 0u;
-	if (DynamicParticleParameterMask)
+	if (EmitContext.DynamicParticleParameterMask)
 	{
 		OutEnvironment.SetDefine(TEXT("USE_DYNAMIC_PARAMETERS"), 1);
-		OutEnvironment.SetDefine(TEXT("DYNAMIC_PARAMETERS_MASK"), DynamicParticleParameterMask);
+		OutEnvironment.SetDefine(TEXT("DYNAMIC_PARAMETERS_MASK"), EmitContext.DynamicParticleParameterMask);
 	}
 
 	if (false)//bNeedsParticleTime)
@@ -490,7 +489,7 @@ static void GetMaterialEnvironment(EShaderPlatform InPlatform,
 
 	// @todo MetalMRT: Remove this hack and implement proper atmospheric-fog solution for Metal MRT...
 	OutEnvironment.SetDefine(TEXT("MATERIAL_ATMOSPHERIC_FOG"), false);// !IsMetalMRTPlatform(InPlatform) ? bUsesAtmosphericFog : 0);
-	OutEnvironment.SetDefine(TEXT("MATERIAL_SKY_ATMOSPHERE"), false);// bUsesSkyAtmosphere);
+	OutEnvironment.SetDefine(TEXT("MATERIAL_SKY_ATMOSPHERE"), EmitContext.bUsesSkyAtmosphere);
 	OutEnvironment.SetDefine(TEXT("INTERPOLATE_VERTEX_COLOR"), bUsesVertexColor);
 	OutEnvironment.SetDefine(TEXT("NEEDS_PARTICLE_COLOR"), bUsesParticleColor);
 	OutEnvironment.SetDefine(TEXT("NEEDS_PARTICLE_LOCAL_TO_WORLD"), false);// bUsesParticleLocalToWorld);
diff --git a/Engine/Source/Runtime/Engine/Private/Materials/MaterialHLSLGenerator.cpp b/Engine/Source/Runtime/Engine/Private/Materials/MaterialHLSLGenerator.cpp
--- a/Engine/Source/Runtime/Engine/Private/Materials/MaterialHLSLGenerator.cpp
+++ b/Engine/Source/Runtime/Engine/Private/Materials/MaterialHLSLGenerator.cpp
@@ -30,6 +30,11 @@ FMaterialHLSLGenerator::FMaterialHLSLGenerator(UMaterial* Material,
 	FunctionCallStack.Add(&RootFunctionCallEntry);
 }
 
+const UMaterial* FMaterialHLSLGenerator::GetTargetMaterial() const
+{
+	return TargetMaterial;
+}
+
 UE::HLSLTree::FTree& FMaterialHLSLGenerator::GetTree() const
 {
 	return CachedTree.GetTree();
@@ -754,7 +759,7 @@ const UE::HLSLTree::FExpression* FMaterialHLSLGenerator::GenerateBranch(UE::HLSL
 		FXxHash64Builder Hasher;
 		Hasher.Update(&ConditionExpression, sizeof(FExpression*));
 		Hasher.Update(&TrueExpression, sizeof(FExpression*));
-		Hasher.Update(&FalseExpression, sizeof(FExpression));
+		Hasher.Update(&FalseExpression, sizeof(FExpression*));
 		Hash = Hasher.Finalize();
 	}
 
diff --git a/Engine/Source/Runtime/Engine/Private/Materials/MaterialHLSLTree.cpp b/Engine/Source/Runtime/Engine/Private/Materials/MaterialHLSLTree.cpp
--- a/Engine/Source/Runtime/Engine/Private/Materials/MaterialHLSLTree.cpp
+++ b/Engine/Source/Runtime/Engine/Private/Materials/MaterialHLSLTree.cpp
@@ -622,6 +622,42 @@ void FExpressionCollectionParameter::EmitValueShader(FEmitContext& Context, FEmi
 	OutResult.Code = Context.EmitInlineExpression(Scope, Shader::EValueType::Float4, TEXT("MaterialCollection%.Vectors[%]"), CollectionIndex, ParameterIndex);
 }
 
+void FExpressionDynamicParameter::ComputeAnalyticDerivatives(FTree& Tree, FExpressionDerivatives& OutResult) const
+{
+	OutResult.ExpressionDdx = Tree.NewConstant(0.f);
+	OutResult.ExpressionDdy = OutResult.ExpressionDdx;
+}
+
+bool FExpressionDynamicParameter::PrepareValue(FEmitContext& Context, FEmitScope& Scope, const FRequestedType& RequestedType, FPrepareValueResult& OutResult) const
+{
+	if (Context.ShaderFrequency != SF_Vertex && Context.ShaderFrequency != SF_Pixel && Context.ShaderFrequency != SF_Compute)
+	{
+		return Context.Error(TEXT("Invalid node used in hull/domain shader input!"));
+	}
+
+	return OutResult.SetType(Context, RequestedType, EExpressionEvaluation::Shader, Shader::EValueType::Float4);
+}
+
+void FExpressionDynamicParameter::EmitValueShader(FEmitContext& Context, FEmitScope& Scope, const FRequestedType& RequestedType, FEmitValueShaderResult& OutResult) const
+{
+	FEmitShaderExpression* EmitDefaultValueExpression = DefaultValueExpression->GetValueShader(Context, Scope, Shader::EValueType::Float4);
+	OutResult.Code = Context.EmitExpression(Scope, Shader::EValueType::Float4, TEXT("GetDynamicParameter(Parameters.Particle, %, %)"), EmitDefaultValueExpression, ParameterIndex);
+}
+
+bool FExpressionSkyLightEnvMapSample::PrepareValue(FEmitContext& Context, FEmitScope& Scope, const FRequestedType& RequestedType, FPrepareValueResult& OutResult) const
+{
+	// This expression evaluates to constant 1.0f when IS_BASE_PASS is not defined but we cannot do special handling
+	// here since the define is tied to specific shader types.
+	return OutResult.SetType(Context, RequestedType, EExpressionEvaluation::Shader, Shader::EValueType::Float3);
+}
+
+void FExpressionSkyLightEnvMapSample::EmitValueShader(FEmitContext& Context, FEmitScope& Scope, const FRequestedType& RequestedType, FEmitValueShaderResult& OutResult) const
+{
+	FEmitShaderExpression* EmitDirectionExpression = DirectionExpression->GetValueShader(Context, Scope, Shader::EValueType::Float3);
+	FEmitShaderExpression* EmitRoughnessExpression = RoughnessExpression->GetValueShader(Context, Scope, Shader::EValueType::Float1);
+	OutResult.Code = Context.EmitExpression(Scope, Shader::EValueType::Float3, TEXT("MaterialExpressionSkyLightEnvMapSample(%, %)"), EmitDirectionExpression, EmitRoughnessExpression);
+}
+
 namespace Private
 {
 
@@ -1327,6 +1363,17 @@ void FExpressionVertexInterpolator::EmitValuePreshader(FEmitContext& Context, FE
 	OutResult.Type = VertexExpression->GetValuePreshader(Context, Scope, RequestedType, OutResult.Preshader);
 }
 
+bool FExpressionSkyAtmosphereLightDirection::PrepareValue(FEmitContext& Context, FEmitScope& Scope, const FRequestedType& RequestedType, FPrepareValueResult& OutResult) const
+{
+	return OutResult.SetType(Context, RequestedType, EExpressionEvaluation::Shader, Shader::EValueType::Float3);
+}
+
+void FExpressionSkyAtmosphereLightDirection::EmitValueShader(FEmitContext& Context, FEmitScope& Scope, const FRequestedType& RequestedType, FEmitValueShaderResult& OutResult) const
+{
+	Context.bUsesSkyAtmosphere = true;
+	OutResult.Code = Context.EmitInlineExpression(Scope, Shader::EValueType::Float3, TEXT("MaterialExpressionSkyAtmosphereLightDirection(Parameters, %)"), LightIndex);
+}
+
 int32 FEmitData::FindInterpolatorIndex(const FExpression* Expression) const
 {
 	for (int32 Index = 0; Index < VertexInterpolators.Num(); ++Index)
diff --git a/Engine/Source/Runtime/Engine/Public/HLSLTree/HLSLTree.h b/Engine/Source/Runtime/Engine/Public/HLSLTree/HLSLTree.h
--- a/Engine/Source/Runtime/Engine/Public/HLSLTree/HLSLTree.h
+++ b/Engine/Source/Runtime/Engine/Public/HLSLTree/HLSLTree.h
@@ -502,6 +502,8 @@ private:
 	friend class FTree;
 	friend class FEmitContext;
 	friend class FExpressionForward;
+	friend class FExpressionFunctionCall;
+	friend class FExpressionOperation;
 };
 
 /**
@@ -629,6 +631,7 @@ public:
 	const FExpression* NewSwizzle(const FSwizzleParameters& Params, const FExpression* Input);
 	const FExpression* NewUnaryOp(EOperation Op, const FExpression* Input);
 	const FExpression* NewBinaryOp(EOperation Op, const FExpression* Lhs, const FExpression* Rhs);
+	const FExpression* NewTernaryOp(EOperation Op, const FExpression* Input0, const FExpression* Input1, const FExpression* Input2);
 
 	const FExpression* NewAbs(const FExpression* Input) { return NewUnaryOp(EOperation::Abs, Input); }
 	const FExpression* NewNeg(const FExpression* Input) { return NewUnaryOp(EOperation::Neg, Input); }
@@ -650,6 +653,7 @@ public:
 	const FExpression* NewMul(const FExpression* Lhs, const FExpression* Rhs) { return NewBinaryOp(EOperation::Mul, Lhs, Rhs); }
 	const FExpression* NewDiv(const FExpression* Lhs, const FExpression* Rhs) { return NewBinaryOp(EOperation::Div, Lhs, Rhs); }
 	const FExpression* NewFmod(const FExpression* Lhs, const FExpression* Rhs) { return NewBinaryOp(EOperation::Fmod, Lhs, Rhs); }
+	const FExpression* NewStep(const FExpression* Lhs, const FExpression* Rhs) { return NewBinaryOp(EOperation::Step, Lhs, Rhs); }
 	const FExpression* NewPowClamped(const FExpression* Lhs, const FExpression* Rhs) { return NewBinaryOp(EOperation::PowPositiveClamped, Lhs, Rhs); }
 	const FExpression* NewMin(const FExpression* Lhs, const FExpression* Rhs) { return NewBinaryOp(EOperation::Min, Lhs, Rhs); }
 	const FExpression* NewMax(const FExpression* Lhs, const FExpression* Rhs) { return NewBinaryOp(EOperation::Max, Lhs, Rhs); }
@@ -658,6 +662,8 @@ public:
 	const FExpression* NewLessEqual(const FExpression* Lhs, const FExpression* Rhs) { return NewBinaryOp(EOperation::LessEqual, Lhs, Rhs); }
 	const FExpression* NewGreaterEqual(const FExpression* Lhs, const FExpression* Rhs) { return NewBinaryOp(EOperation::GreaterEqual, Lhs, Rhs); }
 
+	const FExpression* NewSmoothStep(const FExpression* Input0, const FExpression* Input1, const FExpression* Input2) { return NewTernaryOp(EOperation::SmoothStep, Input0,	Input1, Input2); }
+
 	const FExpression* NewLog(const FExpression* Input);
 	const FExpression* NewPow2(const FExpression* Input) { return NewMul(Input, Input); }
 	const FExpression* NewCross(const FExpression* Lhs, const FExpression* Rhs);
diff --git a/Engine/Source/Runtime/Engine/Public/HLSLTree/HLSLTreeCommon.h b/Engine/Source/Runtime/Engine/Public/HLSLTree/HLSLTreeCommon.h
--- a/Engine/Source/Runtime/Engine/Public/HLSLTree/HLSLTreeCommon.h
+++ b/Engine/Source/Runtime/Engine/Public/HLSLTree/HLSLTreeCommon.h
@@ -261,6 +261,24 @@ public:
 	virtual bool IsInputActive(const FEmitContext& Context, int32 Index) const override;
 };
 
+class FExpressionQualitySwitch : public FExpressionSwitchBase
+{
+public:
+	FExpressionQualitySwitch(TConstArrayView<const FExpression*> InInputs);
+
+	virtual const FExpression* NewSwitch(FTree& Tree, TConstArrayView<const FExpression*> InInputs) const override { return Tree.NewExpression<FExpressionQualitySwitch>(InInputs); }
+	virtual bool IsInputActive(const FEmitContext& Context, int32 Index) const override;
+};
+
+class FExpressionShaderStageSwitch : public FExpressionSwitchBase
+{
+public:
+	FExpressionShaderStageSwitch(TConstArrayView<const FExpression*> InInputs);
+
+	virtual const FExpression* NewSwitch(FTree& Tree, TConstArrayView<const FExpression*> InInputs) const override { return Tree.NewExpression<FExpressionShaderStageSwitch>(InInputs); }
+	virtual bool IsInputActive(const FEmitContext& Context, int32 Index) const override;
+};
+
 /** Can be used to emit HLSL chunks with no inputs, where it's not worth the trouble of defining a new expression type */
 class FExpressionInlineCustomHLSL : public FExpression
 {
diff --git a/Engine/Source/Runtime/Engine/Public/HLSLTree/HLSLTreeEmit.h b/Engine/Source/Runtime/Engine/Public/HLSLTree/HLSLTreeEmit.h
--- a/Engine/Source/Runtime/Engine/Public/HLSLTree/HLSLTreeEmit.h
+++ b/Engine/Source/Runtime/Engine/Public/HLSLTree/HLSLTreeEmit.h
@@ -577,6 +577,10 @@ public:
 	bool bMarkLiveValues = false;
 	bool bUseAnalyticDerivatives = false;
 
+	bool bUsesSkyAtmosphere = false;
+
+	uint32 DynamicParticleParameterMask = 0u;
+
 	TArray<const FOwnedNode*, TInlineAllocator<32>> OwnerStack;
 	TArray<FEmitShaderNode*> EmitNodes;
 	TMap<const FScope*, FEmitScope*> EmitScopeMap;
diff --git a/Engine/Source/Runtime/Engine/Public/HLSLTree/HLSLTreeTypes.h b/Engine/Source/Runtime/Engine/Public/HLSLTree/HLSLTreeTypes.h
--- a/Engine/Source/Runtime/Engine/Public/HLSLTree/HLSLTreeTypes.h
+++ b/Engine/Source/Runtime/Engine/Public/HLSLTree/HLSLTreeTypes.h
@@ -110,6 +110,7 @@ enum class EOperation : uint8
 	Mul,
 	Div,
 	Fmod,
+	Step,
 	PowPositiveClamped,
 	Atan2,
 	Atan2Fast,
@@ -124,6 +125,9 @@ enum class EOperation : uint8
 	VecMulMatrix4,
 	Matrix3MulVec,
 	Matrix4MulVec,
+
+	// Ternary Ops
+	SmoothStep,
 };
 
 struct FOperationDescription
diff --git a/Engine/Source/Runtime/Engine/Public/MaterialHLSLGenerator.h b/Engine/Source/Runtime/Engine/Public/MaterialHLSLGenerator.h
--- a/Engine/Source/Runtime/Engine/Public/MaterialHLSLGenerator.h
+++ b/Engine/Source/Runtime/Engine/Public/MaterialHLSLGenerator.h
@@ -71,6 +71,7 @@ public:
 
 	const FMaterialLayersFunctions* GetLayerOverrides() const { return LayerOverrides; }
 
+	const UMaterial* GetTargetMaterial() const;
 	UE::HLSLTree::FTree& GetTree() const;
 	UE::Shader::FStructTypeRegistry& GetTypeRegistry() const;
 	const UE::Shader::FStructType* GetMaterialAttributesType() const;
diff --git a/Engine/Source/Runtime/Engine/Public/MaterialHLSLTree.h b/Engine/Source/Runtime/Engine/Public/MaterialHLSLTree.h
--- a/Engine/Source/Runtime/Engine/Public/MaterialHLSLTree.h
+++ b/Engine/Source/Runtime/Engine/Public/MaterialHLSLTree.h
@@ -287,6 +287,37 @@ public:
 	virtual void EmitValueShader(FEmitContext& Context, FEmitScope& Scope, const FRequestedType& RequestedType, FEmitValueShaderResult& OutResult) const override;
 };
 
+class FExpressionDynamicParameter : public FExpression
+{
+public:
+	const FExpression* DefaultValueExpression;
+	int32 ParameterIndex;
+
+	FExpressionDynamicParameter(const FExpression* InDefaultValueExpression, int32 InParameterIndex)
+		: DefaultValueExpression(InDefaultValueExpression)
+		, ParameterIndex(InParameterIndex)
+	{}
+
+	virtual void ComputeAnalyticDerivatives(FTree& Tree, FExpressionDerivatives& OutResult) const override;
+	virtual bool PrepareValue(FEmitContext& Context, FEmitScope& Scope, const FRequestedType& RequestedType, FPrepareValueResult& OutResult) const override;
+	virtual void EmitValueShader(FEmitContext& Context, FEmitScope& Scope, const FRequestedType& RequestedType, FEmitValueShaderResult& OutResult) const override;
+};
+
+class FExpressionSkyLightEnvMapSample : public FExpression
+{
+public:
+	const FExpression* DirectionExpression;
+	const FExpression* RoughnessExpression;
+
+	FExpressionSkyLightEnvMapSample(const FExpression* InDirectionExpression, const FExpression* InRoughnessExpression)
+		: DirectionExpression(InDirectionExpression)
+		, RoughnessExpression(InRoughnessExpression)
+	{}
+
+	virtual bool PrepareValue(FEmitContext& Context, FEmitScope& Scope, const FRequestedType& RequestedType, FPrepareValueResult& OutResult) const override;
+	virtual void EmitValueShader(FEmitContext& Context, FEmitScope& Scope, const FRequestedType& RequestedType, FEmitValueShaderResult& OutResult) const override;
+};
+
 class FExpressionTextureSample : public FExpression
 {
 public:
@@ -425,6 +456,19 @@ public:
 	virtual void EmitValuePreshader(FEmitContext& Context, FEmitScope& Scope, const FRequestedType& RequestedType, FEmitValuePreshaderResult& OutResult) const override;
 };
 
+class FExpressionSkyAtmosphereLightDirection : public FExpression
+{
+public:
+	int32 LightIndex;
+
+	FExpressionSkyAtmosphereLightDirection(int32 InLightIndex)
+		: LightIndex(InLightIndex)
+	{}
+
+	virtual bool PrepareValue(FEmitContext& Context, FEmitScope& Scope, const FRequestedType& RequestedType, FPrepareValueResult& OutResult) const override;
+	virtual void EmitValueShader(FEmitContext& Context, FEmitScope& Scope, const FRequestedType& RequestedType, FEmitValueShaderResult& OutResult) const override;
+};
+
 struct FVertexInterpolator
 {
 	FVertexInterpolator() = default;
