Make it more explicit that DataLayerInstanceWithAsset are expected to have a valid asset in cook builds

diff --git a/Engine/Source/Runtime/Engine/Private/WorldPartition/DataLayer/WorldDataLayers.cpp b/Engine/Source/Runtime/Engine/Private/WorldPartition/DataLayer/WorldDataLayers.cpp
--- a/Engine/Source/Runtime/Engine/Private/WorldPartition/DataLayer/WorldDataLayers.cpp
+++ b/Engine/Source/Runtime/Engine/Private/WorldPartition/DataLayer/WorldDataLayers.cpp
@@ -775,7 +775,7 @@ void AWorldDataLayers::PostLoad()
 
 	ConvertDataLayerToInstancces();
 
-	// Remove all Editor Data Layers when cooking or when in a game world
+	// Prune non-runtime Data Layers when cooking or when in a game world
 	if (IsRunningCookCommandlet() || GetWorld()->IsGameWorld())
 	{
 		ForEachDataLayer([](UDataLayerInstance* DataLayer)
@@ -784,16 +784,16 @@ void AWorldDataLayers::PostLoad()
 			return true;
 		});
 
-		TArray<UDataLayerInstance*> EditorDataLayers;
-		ForEachDataLayer([&EditorDataLayers](UDataLayerInstance* DataLayer)
+		TArray<UDataLayerInstance*> NonRuntimeDataLayers;
+		ForEachDataLayer([&NonRuntimeDataLayers](UDataLayerInstance* DataLayer)
 		{
 			if (DataLayer && !DataLayer->IsRuntime())
 			{
-				EditorDataLayers.Add(DataLayer);
+				NonRuntimeDataLayers.Add(DataLayer);
 			}
 			return true;
 		});
-		RemoveDataLayers(EditorDataLayers);
+		RemoveDataLayers(NonRuntimeDataLayers);
 	}
 
 	// Empty Data Layers for non-runtime relevant AWorldDataLayers
@@ -836,7 +836,10 @@ void AWorldDataLayers::PostLoad()
 		static_assert(DATALAYER_TO_INSTANCE_RUNTIME_CONVERSION_ENABLED, "Remove unnecessary cast. All DataLayerInstance now have assets");
 		if (const UDataLayerInstanceWithAsset* DataLayerInstanceWithAsset = Cast<UDataLayerInstanceWithAsset>(DataLayerInstance))
 		{
-			AssetNameToInstance.Add(DataLayerInstanceWithAsset->GetAsset()->GetFullName(), DataLayerInstance);
+			// DataLayerInstanceWithAsset with invalid asset are removed during cook when pruning non-runtime DataLayers.
+			// All DataLayerInstanceWithAsset are expected to have an asset in cook builds.
+			check(DataLayerInstanceWithAsset->GetAsset());
+			AssetNameToInstance.Add(DataLayerInstanceWithAsset->GetAsset()->GetFullName(), DataLayerInstanceWithAsset);
 		}
 	}
 #endif
