Fix some non unity compile errors

diff --git a/Engine/Source/Runtime/Engine/Private/HLSLTree/HLSLTreeCommon.cpp b/Engine/Source/Runtime/Engine/Private/HLSLTree/HLSLTreeCommon.cpp
--- a/Engine/Source/Runtime/Engine/Private/HLSLTree/HLSLTreeCommon.cpp
+++ b/Engine/Source/Runtime/Engine/Private/HLSLTree/HLSLTreeCommon.cpp
@@ -5,6 +5,8 @@
 #include "HLSLTree/HLSLTreeEmit.h"
 #include "RenderUtils.h"
 #include "Shader/Preshader.h"
+#include "SceneTypes.h"
+#include "MaterialShared.h"
 
 namespace UE::HLSLTree
 {
diff --git a/Engine/Source/Runtime/Engine/Private/Materials/MaterialExpressionHLSL.cpp b/Engine/Source/Runtime/Engine/Private/Materials/MaterialExpressionHLSL.cpp
--- a/Engine/Source/Runtime/Engine/Private/Materials/MaterialExpressionHLSL.cpp
+++ b/Engine/Source/Runtime/Engine/Private/Materials/MaterialExpressionHLSL.cpp
@@ -150,6 +150,7 @@
 #include "Materials/MaterialExpressionForLoop.h"
 #include "Materials/MaterialExpressionWhileLoop.h"
 #include "Materials/MaterialFunctionInterface.h"
+#include "Materials/Material.h"
 #include "MaterialHLSLTree.h"
 #include "MaterialShared.h"
 #include "MaterialDomain.h"
diff --git a/Engine/Source/Runtime/Engine/Public/MaterialCachedHLSLTree.h b/Engine/Source/Runtime/Engine/Public/MaterialCachedHLSLTree.h
--- a/Engine/Source/Runtime/Engine/Public/MaterialCachedHLSLTree.h
+++ b/Engine/Source/Runtime/Engine/Public/MaterialCachedHLSLTree.h
@@ -12,6 +12,7 @@
 class UMaterial;
 class UMaterialExpression;
 class UMaterialExpressionCustomOutput;
+class UMaterialParameterCollection;
 struct FMaterialLayersFunctions;
 
 struct FMaterialConnectionKey
diff --git a/Engine/Source/Runtime/Engine/Public/MaterialHLSLGenerator.h b/Engine/Source/Runtime/Engine/Public/MaterialHLSLGenerator.h
--- a/Engine/Source/Runtime/Engine/Public/MaterialHLSLGenerator.h
+++ b/Engine/Source/Runtime/Engine/Public/MaterialHLSLGenerator.h
@@ -27,6 +27,7 @@ class UMaterialExpressionFunctionInput;
 class UMaterialExpressionFunctionOutput;
 class UMaterialExpressionCustomOutput;
 class UMaterialExpressionVertexInterpolator;
+class UMaterialParameterCollection;
 struct FFunctionExpressionInput;
 struct FExpressionInput;
 class ITargetPlatform;
