Fixed mismatched copy constructor/assignment operator in FShaderRecompileData.

diff --git a/Engine/Source/Runtime/Engine/Public/ShaderCompiler.h b/Engine/Source/Runtime/Engine/Public/ShaderCompiler.h
--- a/Engine/Source/Runtime/Engine/Public/ShaderCompiler.h
+++ b/Engine/Source/Runtime/Engine/Public/ShaderCompiler.h
@@ -1147,24 +1147,6 @@ struct FShaderRecompileData
 	/** For recompiling just global shaders. */
 	ENGINE_API FShaderRecompileData(const FString& InPlatformName, EShaderPlatform InShaderPlatform, ODSCRecompileCommand InCommandType, TArray<FString>* OutModifiedFiles, TArray<uint8>* OutMeshMaterialMaps, TArray<uint8>* OutGlobalShaderMap);
 
-	FShaderRecompileData& operator=(const FShaderRecompileData& Other)
-	{
-		PlatformName = Other.PlatformName;
-		ShaderPlatform = Other.ShaderPlatform;
-		FeatureLevel = Other.FeatureLevel;
-		QualityLevel = Other.QualityLevel;
-		ModifiedFiles = Other.ModifiedFiles;
-		MeshMaterialMaps = Other.MeshMaterialMaps;
-		MaterialsToLoad = Other.MaterialsToLoad;
-		ShaderTypesToLoad = Other.ShaderTypesToLoad;
-		CommandType = Other.CommandType;
-		GlobalShaderMap = Other.GlobalShaderMap;
-
-		ShadersToRecompile = Other.ShadersToRecompile;
-
-		return *this;
-	}
-
 	ENGINE_API friend FArchive& operator<<(FArchive& Ar, FShaderRecompileData& Elem);
 };
 
