Add functions to convert from classes and soft class references to TopLevelAssetPath used by asset utility functions
Fix comments and add BlueprintThreadsafe to some other converter functions

diff --git a/Engine/Source/Runtime/Engine/Classes/Kismet/KismetSystemLibrary.h b/Engine/Source/Runtime/Engine/Classes/Kismet/KismetSystemLibrary.h
--- a/Engine/Source/Runtime/Engine/Classes/Kismet/KismetSystemLibrary.h
+++ b/Engine/Source/Runtime/Engine/Classes/Kismet/KismetSystemLibrary.h
@@ -133,6 +133,10 @@ class ENGINE_API UKismetSystemLibrary : public UBlueprintFunctionLibrary
 	UFUNCTION(BlueprintPure, Category = "Utilities")
 	static FSoftClassPath GetSoftClassPath(const UClass* Class);
 
+	// Returns the full path to the specified class as a Top Level Asset Path used by asset utilities
+	UFUNCTION(BlueprintPure, Category = "Utilities")
+	static FTopLevelAssetPath GetClassTopLevelAssetPath(const UClass* Class);
+
 	// Returns the outer object of an object.
 	UFUNCTION(BlueprintPure, Category = "Utilities")
 	static UObject* GetOuterObject(const UObject* Object);
@@ -270,16 +274,16 @@ class ENGINE_API UKismetSystemLibrary : public UBlueprintFunctionLibrary
 	static void BreakSoftObjectPath(FSoftObjectPath InSoftObjectPath, FString& PathString);
 
 	/** Converts a Soft Object Path into a base Soft Object Reference, this is not guaranteed to be resolvable */
-	UFUNCTION(BlueprintPure, meta = (DisplayName = "To Soft Object Reference"), Category = "Utilities")
+	UFUNCTION(BlueprintPure, meta = (DisplayName = "To Soft Object Reference", BlueprintThreadSafe), Category = "Utilities")
 	static TSoftObjectPtr<UObject> Conv_SoftObjPathToSoftObjRef(const FSoftObjectPath& SoftObjectPath);
 
 	/** Converts a Soft Object Reference into a Soft Object Path */
-	UFUNCTION(BlueprintPure, meta = (DisplayName = "To Soft Object Path", CompactNodeTitle = "->", BlueprintAutocast), Category = "Utilities")
+	UFUNCTION(BlueprintPure, meta = (DisplayName = "To Soft Object Path", CompactNodeTitle = "->", BlueprintThreadSafe, BlueprintAutocast), Category = "Utilities")
 	static FSoftObjectPath Conv_SoftObjRefToSoftObjPath(TSoftObjectPtr<UObject> SoftObjectReference);
 
 	/** Builds a TopLevelAssetPath struct from single Path string or from PackageName and AssetName string. */
 	UFUNCTION(BlueprintPure, Category = "Utilities", meta = (Keywords = "construct build", NativeMakeFunc, BlueprintThreadSafe, BlueprintAutocast))
-	static FORCENOINLINE FTopLevelAssetPath MakeTopLevelAssetPath(UPARAM(DisplayName="FullPathOrPackageName") const FString& PackageName, const FString& AssetName);
+	static FTopLevelAssetPath MakeTopLevelAssetPath(UPARAM(DisplayName="FullPathOrPackageName") const FString& PackageName, const FString& AssetName);
 
 	/** Gets the path string out of a TopLevelAssetPath */
 	UFUNCTION(BlueprintPure, Category = "Utilities", meta = (NativeBreakFunc, BlueprintThreadSafe, BlueprintAutocast))
@@ -297,11 +301,11 @@ class ENGINE_API UKismetSystemLibrary : public UBlueprintFunctionLibrary
 	static void BreakSoftClassPath(FSoftClassPath InSoftClassPath, FString& PathString);
 
 	/** Converts a Soft Class Path into a base Soft Class Reference, this is not guaranteed to be resolvable */
-	UFUNCTION(BlueprintPure, meta = (DisplayName = "To Soft Class Reference"), Category = "Utilities")
+	UFUNCTION(BlueprintPure, meta = (DisplayName = "To Soft Class Reference", BlueprintThreadSafe), Category = "Utilities")
 	static TSoftClassPtr<UObject> Conv_SoftClassPathToSoftClassRef(const FSoftClassPath& SoftClassPath);
 
-	/** Converts a Soft Object Reference into a Soft Class Path (which can be used like a Soft Object Path) */
-	UFUNCTION(BlueprintPure, meta = (DisplayName = "To Soft Class Path", CompactNodeTitle = "->", BlueprintAutocast), Category = "Utilities")
+	/** Converts a Soft Class Reference into a Soft Class Path (which can be used like a Soft Object Path) */
+	UFUNCTION(BlueprintPure, meta = (DisplayName = "To Soft Class Path", CompactNodeTitle = "->", BlueprintThreadSafe, BlueprintAutocast), Category = "Utilities")
 	static FSoftClassPath Conv_SoftObjRefToSoftClassPath(TSoftClassPtr<UObject> SoftClassReference);
 
 	/** Returns true if the Soft Object Reference is not null */
@@ -328,6 +332,10 @@ class ENGINE_API UKismetSystemLibrary : public UBlueprintFunctionLibrary
 	UFUNCTION(BlueprintPure, Category = "Utilities", meta = (BlueprintThreadSafe))
 	static bool IsValidSoftClassReference(const TSoftClassPtr<UObject>& SoftClassReference);
 
+	/** Converts a Soft Class Reference to a Top Level Asset Path used by asset utilities */
+	UFUNCTION(BlueprintPure, meta = (BlueprintThreadSafe, BlueprintAutocast), Category = "Utilities")
+	static FTopLevelAssetPath GetSoftClassTopLevelAssetPath(TSoftClassPtr<UObject> SoftClassReference);
+
 	/** Converts a Soft Class Reference to a path string */
 	UFUNCTION(BlueprintPure, meta = (DisplayName = "To String (SoftClassReference)", CompactNodeTitle = "->", BlueprintThreadSafe, BlueprintAutocast), Category = "Utilities")
 	static FString Conv_SoftClassReferenceToString(const TSoftClassPtr<UObject>& SoftClassReference);
diff --git a/Engine/Source/Runtime/Engine/Private/KismetSystemLibrary.cpp b/Engine/Source/Runtime/Engine/Private/KismetSystemLibrary.cpp
--- a/Engine/Source/Runtime/Engine/Private/KismetSystemLibrary.cpp
+++ b/Engine/Source/Runtime/Engine/Private/KismetSystemLibrary.cpp
@@ -163,6 +163,12 @@ FSoftClassPath UKismetSystemLibrary::GetSoftClassPath(const UClass* Class)
 	return FSoftClassPath(Class);
 }
 
+FTopLevelAssetPath UKismetSystemLibrary::GetClassTopLevelAssetPath(const UClass* Class)
+{
+	// This will succeed for all valid classes as they are never subobjects
+	return FTopLevelAssetPath(Class);
+}
+
 UObject* UKismetSystemLibrary::GetOuterObject(const UObject* Object)
 {
 	return Object ? Object->GetOuter() : nullptr;
@@ -1360,6 +1366,15 @@ bool UKismetSystemLibrary::IsValidSoftClassReference(const TSoftClassPtr<UObject
 	return !SoftClassReference.IsNull();
 }
 
+FTopLevelAssetPath UKismetSystemLibrary::GetSoftClassTopLevelAssetPath(TSoftClassPtr<UObject> SoftClassReference)
+{
+	const FSoftObjectPath& ObjectPath = SoftClassReference.ToSoftObjectPath();
+
+	// Class paths should never have a subpath
+	ensure(ObjectPath.GetSubPathString().IsEmpty());
+	return ObjectPath.GetAssetPath();
+}
+
 FString UKismetSystemLibrary::Conv_SoftClassReferenceToString(const TSoftClassPtr<UObject>& SoftClassReference)
 {
 	return SoftClassReference.ToString();
