AssetManager: Deprecate the function RebuildObjectReferenceList. Calculate the ObjectReferenceList in PreGarbageCollect instead. Replace the AssetPathMap update the function was doing with incremental updates when the AssetMaps are modified.
#rnx

diff --git a/Engine/Source/Runtime/Engine/Classes/Engine/AssetManager.h b/Engine/Source/Runtime/Engine/Classes/Engine/AssetManager.h
--- a/Engine/Source/Runtime/Engine/Classes/Engine/AssetManager.h
+++ b/Engine/Source/Runtime/Engine/Classes/Engine/AssetManager.h
@@ -491,6 +491,8 @@ public:
 	class FTimerManager* GetTimerManager() const;
 
 	// Overrides
+	virtual void BeginDestroy() override;
+
 	virtual void PostInitProperties() override;
 
 	/** Get the encryption key guid attached to this primary asset. Can be invalid if the asset is not encrypted */
@@ -649,8 +651,14 @@ protected:
 	const FPrimaryAssetData* GetNameData(const FPrimaryAssetId& PrimaryAssetId, bool bCheckRedirector = true) const;
 
 	/** Rebuilds the ObjectReferenceList, needed after global object state has changed */
+	UE_DEPRECATED(5.3, "Function was split up into different callsites. Use OnObjectReferenceListInvalidated() if you need a hook in the same place.")
 	virtual void RebuildObjectReferenceList();
 
+	virtual void OnObjectReferenceListInvalidated();
+
+	void CallPreGarbageCollect();
+	virtual void PreGarbageCollect();
+
 	/** Called when an internal load handle finishes, handles setting to pending state */
 	virtual void OnAssetStateChangeCompleted(FPrimaryAssetId PrimaryAssetId, TSharedPtr<FStreamableHandle> BoundHandle, FStreamableDelegate WrappedDelegate);
 
@@ -803,6 +811,8 @@ protected:
 	/** Suppresses bOnlyCookProductionAssets based on the AllowsDevelopmentObjects() property of the TargetPlatforms being cooked. */
 	bool bTargetPlatformsAllowDevelopmentObjects;
 
+	bool bObjectReferenceListDirty = true;
+
 	/** >0 if we are currently in bulk scanning mode */
 	UPROPERTY()
 	int32 NumBulkScanRequests;
diff --git a/Engine/Source/Runtime/Engine/Private/AssetManager.cpp b/Engine/Source/Runtime/Engine/Private/AssetManager.cpp
--- a/Engine/Source/Runtime/Engine/Private/AssetManager.cpp
+++ b/Engine/Source/Runtime/Engine/Private/AssetManager.cpp
@@ -81,25 +81,33 @@ struct FPrimaryAssetLoadState
 	}
 };
 
+struct FPrimaryAssetTypeData;
+
 /** Structure representing data about a specific asset */
 struct FPrimaryAssetData
 {
-	/** Path used to look up cached asset data in the asset registry. This will be missing the _C for blueprint classes */
-	FSoftObjectPath AssetDataPath;
+public:
+	FPrimaryAssetData() {}
 
 	/** Path to this asset on disk */
+	const FSoftObjectPtr& GetAssetPtr() const { return AssetPtr; }
+	/** Path used to look up cached asset data in the asset registry. This will be missing the _C for blueprint classes */
+	const FSoftObjectPath& GetARLookupPath() const { return ARLookupPath; }
+	/** Asset is considered loaded at all if there is an active handle for it */
+	bool IsLoaded() const { return CurrentState.IsValid(); }
+
+private:
+	// These are used as the keys in a reverse map and cannot be modified except when modifying the reversemap as well
 	FSoftObjectPtr AssetPtr;
+	FSoftObjectPath ARLookupPath;
 
+public:
 	/** Current state of this asset */
 	FPrimaryAssetLoadState CurrentState;
-
 	/** Pending state of this asset, will be copied to CurrentState when load finishes */
 	FPrimaryAssetLoadState PendingState;
 
-	FPrimaryAssetData() {}
-
-	/** Asset is considered loaded at all if there is an active handle for it */
-	bool IsLoaded() const { return CurrentState.IsValid(); }
+	friend FPrimaryAssetTypeData;
 };
 
 /** Structure representing all items of a specific asset type */
@@ -108,8 +116,18 @@ struct FPrimaryAssetTypeData
 	/** The public info struct */
 	FPrimaryAssetTypeInfo Info;
 
-	/** Map of scanned assets */
-	TMap<FName, FPrimaryAssetData> AssetMap;
+	// Each PrimaryAssetType has a list of PrimaryAssetNames that are PrimaryAssets of its type, and a map to store a
+	// FPrimaryAssetData for each one, including the SoftObjectPtr to the asset represented by that Name.
+	// The AssetManager also has a reverse map from SoftObjectPath to the PrimaryAssetId
+	// (PrimaryAssetTypeName + PrimaryAssetName) for every PrimaryAsset.
+	// We need to modify the list of assets and the reverse map together, so direct access to this->AssetMap is
+	// private, and modifying it is done through functions that take the ReverseMap as an additional In/Out argument.
+	const TMap<FName, FPrimaryAssetData>& GetAssets() const;
+	FPrimaryAssetData& FindOrAddAsset(FName AssetName, const FSoftObjectPtr& AssetPtr, const FSoftObjectPath& ARLookupPath,
+		TMap<FSoftObjectPath, FPrimaryAssetId>& InOutReverseMap);
+	void RemoveAsset(FName AssetName, TMap<FSoftObjectPath, FPrimaryAssetId>& InOutReverseMap);
+	void ResetAssets(TMap<FSoftObjectPath, FPrimaryAssetId>& InOutReverseMap);
+	void ShrinkAssets();
 
 	/** In the editor, paths that we need to scan once asset registry is done loading */
 	TArray<FString> DeferredAssetScanPaths;
@@ -118,12 +136,88 @@ struct FPrimaryAssetTypeData
 	TArray<FString> RealAssetScanPaths;
 
 	FPrimaryAssetTypeData() {}
+	~FPrimaryAssetTypeData();
 
 	FPrimaryAssetTypeData(FName InPrimaryAssetType, UClass* InAssetBaseClass, bool bInHasBlueprintClasses, bool bInIsEditorOnly)
 		: Info(InPrimaryAssetType, InAssetBaseClass, bInHasBlueprintClasses, bInIsEditorOnly)
 		{}
+
+private:
+	/** Map of scanned assets */
+	TMap<FName, FPrimaryAssetData> AssetMap;
 };
 
+FPrimaryAssetTypeData::~FPrimaryAssetTypeData()
+{
+	// Assets must be removed via a call to e.g. ResetAssets before destruction, so that the AssetManager's AssetPathMap
+	// can be updated. For an example, see UAssetManager::RemovePrimaryAssetType
+	checkf(AssetMap.IsEmpty(), TEXT("FPrimaryAssetTypeData is being destructed while still containing assets"));
+}
+
+const TMap<FName, FPrimaryAssetData>& FPrimaryAssetTypeData::GetAssets() const
+{
+	return AssetMap;
+}
+
+FPrimaryAssetData& FPrimaryAssetTypeData::FindOrAddAsset(FName AssetName, const FSoftObjectPtr& AssetPtr,
+	const FSoftObjectPath& ARLookupPath, TMap<FSoftObjectPath, FPrimaryAssetId>& InOutReverseMap)
+{
+	FPrimaryAssetData& ExistingAssetData = AssetMap.FindOrAdd(AssetName);
+	const FSoftObjectPath& OldAssetRef = ExistingAssetData.AssetPtr.ToSoftObjectPath();
+	// If the old reference exists and we are replacing it, remove it from the reverse map before adding the new one
+	if (!OldAssetRef.IsNull())
+	{
+		InOutReverseMap.Remove(OldAssetRef);
+	}
+
+	ExistingAssetData.AssetPtr = AssetPtr;
+	ExistingAssetData.ARLookupPath = ARLookupPath;
+
+	const FSoftObjectPath& NewAssetRef = ExistingAssetData.AssetPtr.ToSoftObjectPath();
+	// Dynamic types can have null asset refs NewAssetRef might be null
+	if (!NewAssetRef.IsNull())
+	{
+		InOutReverseMap.Add(NewAssetRef, FPrimaryAssetId(Info.PrimaryAssetType, AssetName));
+	}
+	Info.NumberOfAssets = AssetMap.Num();
+
+	return ExistingAssetData;
+}
+
+void FPrimaryAssetTypeData::RemoveAsset(FName AssetName, TMap<FSoftObjectPath, FPrimaryAssetId>& InOutReverseMap)
+{
+	FPrimaryAssetData ExistingAssetData;
+	if (AssetMap.RemoveAndCopyValue(AssetName, ExistingAssetData))
+	{
+		const FSoftObjectPath& OldAssetRef = ExistingAssetData.AssetPtr.ToSoftObjectPath();
+		if (!OldAssetRef.IsNull()) // Dynamic types can have null asset refs
+		{
+			InOutReverseMap.Remove(OldAssetRef);
+		}
+	}
+	Info.NumberOfAssets = AssetMap.Num();
+}
+
+void FPrimaryAssetTypeData::ResetAssets(TMap<FSoftObjectPath, FPrimaryAssetId>& InOutReverseMap)
+{
+	for (const TPair<FName, FPrimaryAssetData>& Pair : AssetMap)
+	{
+		const FPrimaryAssetData& AssetData = Pair.Value;
+		const FSoftObjectPath& AssetRef = AssetData.AssetPtr.ToSoftObjectPath();
+		if (!AssetRef.IsNull()) // Dynamic types can have null asset refs
+		{
+			InOutReverseMap.Remove(AssetRef);
+		}
+	}
+	AssetMap.Reset();
+	Info.NumberOfAssets = 0;
+}
+
+void FPrimaryAssetTypeData::ShrinkAssets()
+{
+	AssetMap.Shrink();
+}
+
 /** Version of rules with cached data */
 struct FCompiledAssetManagerSearchRules : FAssetManagerSearchRules
 {
@@ -230,11 +324,8 @@ FSimpleMulticastDelegate UAssetManager::OnAssetManagerCreatedDelegate;
 // Allow StartInitialLoading() bulk scan to continue past FinishInitialLoading() to cover:
 // 1. Loading in subclass FinishInitialLoading() override after calling base FinishInitialLoading() 
 // 2. Plugin loading in OnPostEngine callback, e.g. UDataRegistrySubsystem::LoadAllRegistries()
-// 
-// Bulk scanning should stop before the early GC when loading the first map, to ensure
-// RebuildObjectReferenceList() stops INI-referenced UClasses from being GCed.
-//
-// Ideally some future asset manager refactoring can get rid of this complexity / hack.
+// Ideally a future AssetManager refactor will reduce update costs and remove the need for batching
+// things up with bulk scanning.
 static class FInitialBulkScanHelper
 {
 	UAssetManager* Scanner = nullptr;
@@ -285,6 +376,18 @@ UAssetManager::UAssetManager()
 	NumberOfSpawnedNotifications = 0;
 }
 
+void UAssetManager::BeginDestroy()
+{
+	Super::BeginDestroy();
+
+	// FPrimaryAssetTypeDatas need to have their AssetMaps cleared before their destructor,
+	// because they have a contract that modifying AssetMap also modifies AssetPathMap.
+	for (TPair<FName, TSharedRef<FPrimaryAssetTypeData>>& Pair : AssetTypeMap)
+	{
+		Pair.Value->ResetAssets(AssetPathMap);
+	}
+}
+
 void UAssetManager::PostInitProperties()
 {
 	Super::PostInitProperties();
@@ -339,6 +442,8 @@ void UAssetManager::PostInitProperties()
 
 		// Add /Game to initial list, games can add additional ones if desired
 		AllAssetSearchRoots.Add(TEXT("/Game"));
+
+		FCoreUObjectDelegates::GetPreGarbageCollectDelegate().AddUObject(this, &UAssetManager::CallPreGarbageCollect);
 	}
 }
 
@@ -898,10 +1003,7 @@ int32 UAssetManager::ScanPathsForPrimaryAssets(FPrimaryAssetType PrimaryAssetTyp
 		TryUpdateCachedAssetData(PrimaryAssetId, Data, false);
 	}
 
-	if (!IsBulkScanning())
-	{
-		RebuildObjectReferenceList();
-	}
+	OnObjectReferenceListInvalidated();
 
 	return NumAdded;
 }
@@ -969,7 +1071,11 @@ void UAssetManager::RemoveScanPathsForPrimaryAssets(FPrimaryAssetType PrimaryAss
 
 void UAssetManager::RemovePrimaryAssetType(FPrimaryAssetType PrimaryAssetType)
 {
-	AssetTypeMap.Remove(PrimaryAssetType);
+	TSharedRef<FPrimaryAssetTypeData> TypeData;
+	if (AssetTypeMap.RemoveAndCopyValue(PrimaryAssetType, TypeData))
+	{
+		TypeData->ResetAssets(AssetPathMap);
+	}
 }
 
 void UAssetManager::StartBulkScanning()
@@ -987,16 +1093,16 @@ void UAssetManager::StopBulkScanning()
 	check(!IsBulkScanning());
 
 	GetAssetRegistry().SetTemporaryCachingMode(bOldTemporaryCachingMode);
-	
-	RebuildObjectReferenceList();
+
+	OnObjectReferenceListInvalidated();
 
 	// Shrink some big containers because we've finished modifying them.
 	AssetPathMap.Shrink();
 	AssetTypeMap.Shrink();
-	for (const auto& Pair : AssetTypeMap)
+	for (const TPair<FName, TSharedRef<FPrimaryAssetTypeData>>& Pair : AssetTypeMap)
 	{
 		FPrimaryAssetTypeData& AssetTypeData = *Pair.Value;
-		AssetTypeData.AssetMap.Shrink();
+		AssetTypeData.ShrinkAssets();
 		AssetTypeData.RealAssetScanPaths.Shrink();
 	}
 }
@@ -1019,10 +1125,7 @@ bool UAssetManager::RegisterSpecificPrimaryAsset(const FPrimaryAssetId& PrimaryA
 		return false;
 	}
 
-	if (!IsBulkScanning())
-	{
-		RebuildObjectReferenceList();
-	}
+	OnObjectReferenceListInvalidated();
 
 	return true;
 }
@@ -1047,7 +1150,7 @@ bool UAssetManager::TryUpdateCachedAssetData(const FPrimaryAssetId& PrimaryAsset
 	{
 		FPrimaryAssetTypeData& TypeData = FoundType->Get();
 
-		FPrimaryAssetData* OldData = TypeData.AssetMap.Find(PrimaryAssetId.PrimaryAssetName);
+		const FPrimaryAssetData* OldData = TypeData.GetAssets().Find(PrimaryAssetId.PrimaryAssetName);
 
 		FSoftObjectPath NewAssetPath;
 		if (!TryToSoftObjectPath<EAssetDataCanBeSubobject::No>(NewAssetData, NewAssetPath))
@@ -1059,11 +1162,11 @@ bool UAssetManager::TryUpdateCachedAssetData(const FPrimaryAssetId& PrimaryAsset
 
 		ensure(NewAssetPath.IsAsset());
 
-		if (OldData && OldData->AssetPtr.ToSoftObjectPath() != NewAssetPath)
+		if (OldData && OldData->GetAssetPtr().ToSoftObjectPath() != NewAssetPath)
 		{
 			UE_LOG(LogAssetManager, Warning, 
 				TEXT("Found duplicate PrimaryAssetID %s, path %s conflicts with existing path %s. Two different primary assets can not have the same type and name."), 
-				*PrimaryAssetId.ToString(), *OldData->AssetPtr.ToString(), *NewAssetPath.ToString());
+				*PrimaryAssetId.ToString(), *OldData->GetAssetPtr().ToString(), *NewAssetPath.ToString());
 
 #if WITH_EDITOR
 			if (GIsEditor)
@@ -1072,7 +1175,9 @@ bool UAssetManager::TryUpdateCachedAssetData(const FPrimaryAssetId& PrimaryAsset
 				if (NumberOfSpawnedNotifications++ < MaxNotificationsPerFrame)
 				{
 					FNotificationInfo Info(FText::Format(LOCTEXT("DuplicateAssetId", "Duplicate Asset ID {0} used by both {1} and {2}, rename one to avoid a conflict."),
-						FText::FromString(PrimaryAssetId.ToString()), FText::FromString(OldData->AssetPtr.ToSoftObjectPath().GetLongPackageName()), FText::FromString(NewAssetPath.GetLongPackageName())));
+						FText::FromString(PrimaryAssetId.ToString()),
+						FText::FromString(OldData->GetAssetPtr().ToSoftObjectPath().GetLongPackageName()),
+						FText::FromString(NewAssetPath.GetLongPackageName())));
 					Info.ExpireDuration = 30.0f;
 
 					TSharedPtr<SNotificationItem> Notification = FSlateNotificationManager::Get().AddNotification(Info);
@@ -1087,29 +1192,23 @@ bool UAssetManager::TryUpdateCachedAssetData(const FPrimaryAssetId& PrimaryAsset
 			// Don't ensure for editor only types as they will not cause an actual game problem
 			if (!bAllowDuplicates && !TypeData.Info.bIsEditorOnly)
 			{
-				ensureMsgf(!OldData, TEXT("Found Duplicate PrimaryAssetID %s! Path %s conflicts with existing path %s"), *PrimaryAssetId.ToString(), *OldData->AssetPtr.ToString(), *NewAssetPath.ToString());
+				ensureMsgf(!OldData, TEXT("Found Duplicate PrimaryAssetID %s! Path %s conflicts with existing path %s"),
+					*PrimaryAssetId.ToString(), *OldData->GetAssetPtr().ToString(), *NewAssetPath.ToString());
 			}
 		}
 
-		FPrimaryAssetData& NameData = TypeData.AssetMap.FindOrAdd(PrimaryAssetId.PrimaryAssetName);
-
 		// Update data and path, don't touch state or references
-		NameData.AssetDataPath = NewAssetData.GetSoftObjectPath(); // This will not have _C
-		NameData.AssetPtr = FSoftObjectPtr(NewAssetPath); // This will have _C
+		FSoftObjectPath NewARLookupPath = NewAssetData.GetSoftObjectPath(); // This will not have _C
+		FSoftObjectPtr NewAssetPtr = FSoftObjectPtr(NewAssetPath); // This will have _C
+		TypeData.FindOrAddAsset(PrimaryAssetId.PrimaryAssetName, NewAssetPtr, NewARLookupPath, AssetPathMap);
 
 		// If the types don't match, update the registry
 		IAssetRegistry& LocalAssetRegistry = GetAssetRegistry();
 		FPrimaryAssetId SavedId = NewAssetData.GetPrimaryAssetId();
-		FPrimaryAssetId ObjectPathId = LocalAssetRegistry.GetAssetByObjectPath(NameData.AssetDataPath, true).GetPrimaryAssetId();
+		FPrimaryAssetId ObjectPathId = LocalAssetRegistry.GetAssetByObjectPath(NewARLookupPath, true).GetPrimaryAssetId();
 		if (SavedId != PrimaryAssetId || (ObjectPathId.IsValid() && SavedId != ObjectPathId))
 		{
-			LocalAssetRegistry.SetPrimaryAssetIdForObjectPath(NameData.AssetDataPath, PrimaryAssetId);
-		}
-
-		if (IsBulkScanning())
-		{
-			// Do a partial update, add to the path->asset map
-			AssetPathMap.Add(NewAssetPath, PrimaryAssetId);
+			LocalAssetRegistry.SetPrimaryAssetIdForObjectPath(NewARLookupPath, PrimaryAssetId);
 		}
 
 		if (NewAssetData.TaggedAssetBundles)
@@ -1131,7 +1230,9 @@ bool UAssetManager::TryUpdateCachedAssetData(const FPrimaryAssetId& PrimaryAsset
 				static FName AssetBundleDataName("AssetBundleData");
 
 				// Mark these as editor only if our type is editor only
-				FSoftObjectPathSerializationScope SerializationScope(NewAssetData.PackageName, AssetBundleDataName, TypeData.Info.bIsEditorOnly ? ESoftObjectPathCollectType::EditorOnlyCollect : ESoftObjectPathCollectType::AlwaysCollect, ESoftObjectPathSerializeType::AlwaysSerialize);
+				FSoftObjectPathSerializationScope SerializationScope(NewAssetData.PackageName, AssetBundleDataName,
+					TypeData.Info.bIsEditorOnly ? ESoftObjectPathCollectType::EditorOnlyCollect : ESoftObjectPathCollectType::AlwaysCollect,
+					ESoftObjectPathSerializeType::AlwaysSerialize);
 			
 				for (const FAssetBundleEntry& Entry : NewAssetData.TaggedAssetBundles->Bundles)
 				{
@@ -1192,22 +1293,14 @@ bool UAssetManager::AddDynamicAsset(const FPrimaryAssetId& PrimaryAssetId, const
 		return false;
 	}
 
-	FPrimaryAssetData* OldData = TypeData.AssetMap.Find(PrimaryAssetId.PrimaryAssetName);
-	FPrimaryAssetData& NameData = TypeData.AssetMap.FindOrAdd(PrimaryAssetId.PrimaryAssetName);
-
-	if (OldData && OldData->AssetPtr.ToSoftObjectPath() != AssetPath)
-	{
-		UE_LOG(LogAssetManager, Warning, TEXT("AddDynamicAsset on %s called with conflicting path. Path %s is replacing path %s"), *PrimaryAssetId.ToString(), *OldData->AssetPtr.ToString(), *AssetPath.ToString());
-	}
-
-	NameData.AssetPtr = FSoftObjectPtr(AssetPath);
-
-	if (IsBulkScanning() && AssetPath.IsValid())
+	const FPrimaryAssetData* OldData = TypeData.GetAssets().Find(PrimaryAssetId.PrimaryAssetName);
+	if (OldData && OldData->GetAssetPtr().ToSoftObjectPath() != AssetPath)
 	{
-		// Do a partial update, add to the path->asset map
-		AssetPathMap.Add(AssetPath, PrimaryAssetId);
+		UE_LOG(LogAssetManager, Warning, TEXT("AddDynamicAsset on %s called with conflicting path. Path %s is replacing path %s"),
+			*PrimaryAssetId.ToString(), *OldData->GetAssetPtr().ToString(), *AssetPath.ToString());
 	}
 
+	TypeData.FindOrAddAsset(PrimaryAssetId.PrimaryAssetName, FSoftObjectPtr(AssetPath), AssetPath, AssetPathMap);
 
 	if (BundleData.Bundles.Num() > 0)
 	{
@@ -1342,7 +1435,8 @@ bool UAssetManager::GetPrimaryAssetData(const FPrimaryAssetId& PrimaryAssetId, F
 
 	if (NameData)
 	{
-		FAssetData CachedAssetData = GetAssetRegistry().GetAssetByObjectPath(NameData->AssetDataPath, false /* bIncludeOnlyOnDiskAssets */);
+		FAssetData CachedAssetData = GetAssetRegistry().GetAssetByObjectPath(NameData->GetARLookupPath(),
+			false /* bIncludeOnlyOnDiskAssets */);
 
 		if (CachedAssetData.IsValid())
 		{
@@ -1363,9 +1457,10 @@ bool UAssetManager::GetPrimaryAssetDataList(FPrimaryAssetType PrimaryAssetType,
 	{
 		const FPrimaryAssetTypeData& TypeData = FoundType->Get();
 
-		for (const TPair<FName, FPrimaryAssetData>& Pair : TypeData.AssetMap)
+		for (const TPair<FName, FPrimaryAssetData>& Pair : TypeData.GetAssets())
 		{
-			FAssetData CachedAssetData = Registry.GetAssetByObjectPath(Pair.Value.AssetDataPath, false /* bIncludeOnlyOnDiskAssets */);
+			FAssetData CachedAssetData = Registry.GetAssetByObjectPath(Pair.Value.GetARLookupPath(),
+				false /* bIncludeOnlyOnDiskAssets */);
 
 			if (CachedAssetData.IsValid())
 			{
@@ -1384,7 +1479,7 @@ UObject* UAssetManager::GetPrimaryAssetObject(const FPrimaryAssetId& PrimaryAsse
 
 	if (NameData)
 	{
-		return NameData->AssetPtr.Get();
+		return NameData->GetAssetPtr().Get();
 	}
 
 	return nullptr;
@@ -1399,9 +1494,9 @@ bool UAssetManager::GetPrimaryAssetObjectList(FPrimaryAssetType PrimaryAssetType
 	{
 		const FPrimaryAssetTypeData& TypeData = FoundType->Get();
 
-		for (const TPair<FName, FPrimaryAssetData>& Pair : TypeData.AssetMap)
+		for (const TPair<FName, FPrimaryAssetData>& Pair : TypeData.GetAssets())
 		{
-			UObject* FoundObject = Pair.Value.AssetPtr.Get();
+			UObject* FoundObject = Pair.Value.GetAssetPtr().Get();
 
 			if (FoundObject)
 			{
@@ -1420,7 +1515,7 @@ FSoftObjectPath UAssetManager::GetPrimaryAssetPath(const FPrimaryAssetId& Primar
 
 	if (NameData)
 	{
-		return NameData->AssetPtr.ToSoftObjectPath();
+		return NameData->GetAssetPtr().ToSoftObjectPath();
 	}
 	return FSoftObjectPath();
 }
@@ -1433,11 +1528,12 @@ bool UAssetManager::GetPrimaryAssetPathList(FPrimaryAssetType PrimaryAssetType,
 	{
 		const FPrimaryAssetTypeData& TypeData = FoundType->Get();
 
-		for (const TPair<FName, FPrimaryAssetData>& Pair : TypeData.AssetMap)
+		for (const TPair<FName, FPrimaryAssetData>& Pair : TypeData.GetAssets())
 		{
-			if (!Pair.Value.AssetPtr.IsNull())
+			const FSoftObjectPtr& AssetPtr = Pair.Value.GetAssetPtr();
+			if (!AssetPtr.IsNull())
 			{
-				AssetPathList.AddUnique(Pair.Value.AssetPtr.ToSoftObjectPath());
+				AssetPathList.AddUnique(AssetPtr.ToSoftObjectPath());
 			}
 		}
 	}
@@ -1549,7 +1645,7 @@ bool UAssetManager::GetPrimaryAssetIdList(FPrimaryAssetType PrimaryAssetType, TA
 	{
 		const FPrimaryAssetTypeData& TypeData = FoundType->Get();
 
-		for (const TPair<FName, FPrimaryAssetData>& Pair : TypeData.AssetMap)
+		for (const TPair<FName, FPrimaryAssetData>& Pair : TypeData.GetAssets())
 		{
 			if ((!(Filter & EAssetManagerFilter::UnloadedOnly)) || ((Pair.Value.CurrentState.BundleNames.Num() == 0) && (Pair.Value.PendingState.BundleNames.Num() == 0)))
 			{
@@ -1654,7 +1750,7 @@ TSharedPtr<FStreamableHandle> UAssetManager::ChangeBundleStateForPrimaryAssets(c
 			TSet<FSoftObjectPath> PathsToLoad;
 
 			// Gather asset refs
-			const FSoftObjectPath& AssetPath = NameData->AssetPtr.ToSoftObjectPath();
+			const FSoftObjectPath& AssetPath = NameData->GetAssetPtr().ToSoftObjectPath();
 
 			if (!AssetPath.IsNull())
 			{
@@ -1804,7 +1900,7 @@ bool UAssetManager::GetPrimaryAssetLoadSet(TSet<FSoftObjectPath>& OutAssetLoadSe
 	if (NameData)
 	{
 		// Gather asset refs
-		const FSoftObjectPath& AssetPath = NameData->AssetPtr.ToSoftObjectPath();
+		const FSoftObjectPath& AssetPath = NameData->GetAssetPtr().ToSoftObjectPath();
 		if (!AssetPath.IsNull())
 		{
 			// Dynamic types can have no base asset path
@@ -1957,7 +2053,7 @@ bool UAssetManager::GetPrimaryAssetsWithBundleState(TArray<FPrimaryAssetId>& Pri
 
 		const FPrimaryAssetTypeData& TypeData = TypePair.Value.Get();
 
-		for (const TPair<FName, FPrimaryAssetData>& NamePair : TypeData.AssetMap)
+		for (const TPair<FName, FPrimaryAssetData>& NamePair : TypeData.GetAssets())
 		{
 			const FPrimaryAssetData& NameData = NamePair.Value;
 
@@ -2009,7 +2105,7 @@ void UAssetManager::GetPrimaryAssetBundleStateMap(TMap<FPrimaryAssetId, TArray<F
 	{
 		const FPrimaryAssetTypeData& TypeData = TypePair.Value.Get();
 
-		for (const TPair<FName, FPrimaryAssetData>& NamePair : TypeData.AssetMap)
+		for (const TPair<FName, FPrimaryAssetData>& NamePair : TypeData.GetAssets())
 		{
 			const FPrimaryAssetData& NameData = NamePair.Value;
 
@@ -2285,7 +2381,7 @@ void UAssetManager::AcquireResourcesForPrimaryAssetList(const TArray<FPrimaryAss
 		if (NameData)
 		{
 			// Gather asset refs
-			const FSoftObjectPath& AssetPath = NameData->AssetPtr.ToSoftObjectPath();
+			const FSoftObjectPath& AssetPath = NameData->GetAssetPtr().ToSoftObjectPath();
 
 			if (!AssetPath.IsNull())
 			{
@@ -2494,7 +2590,7 @@ bool UAssetManager::OnAssetRegistryAvailableAfterInitialization(FName InName, FA
 
 			if (bRebuildReferenceList)
 			{
-				RebuildObjectReferenceList();
+				OnObjectReferenceListInvalidated();
 			}
 		}
 	}
@@ -2515,7 +2611,7 @@ const FPrimaryAssetData* UAssetManager::GetNameData(const FPrimaryAssetId& Prima
 	// Try redirected name
 	if (FoundType)
 	{
-		const FPrimaryAssetData* FoundName = (*FoundType)->AssetMap.Find(PrimaryAssetId.PrimaryAssetName);
+		const FPrimaryAssetData* FoundName = (*FoundType)->GetAssets().Find(PrimaryAssetId.PrimaryAssetName);
 		
 		if (FoundName)
 		{
@@ -2539,37 +2635,38 @@ const FPrimaryAssetData* UAssetManager::GetNameData(const FPrimaryAssetId& Prima
 
 void UAssetManager::RebuildObjectReferenceList()
 {
-	AssetPathMap.Reset();
-	ObjectReferenceList.Reset();
+}
 
-	// Iterate primary asset map
-	for (TPair<FName, TSharedRef<FPrimaryAssetTypeData>>& TypePair : AssetTypeMap)
-	{
-		FPrimaryAssetTypeData& TypeData = TypePair.Value.Get();
+void UAssetManager::OnObjectReferenceListInvalidated()
+{
+	bIsManagementDatabaseCurrent = false;
+	bObjectReferenceListDirty = true;
+}
 
-		// Add base class in case it's a blueprint
-		if (!TypeData.Info.bIsDynamicAsset)
-		{
-			ObjectReferenceList.AddUnique(TypeData.Info.AssetBaseClassLoaded);
-		}
+void UAssetManager::CallPreGarbageCollect()
+{
+	PreGarbageCollect();
+}
 
-		TypeData.Info.NumberOfAssets = TypeData.AssetMap.Num();
+void UAssetManager::PreGarbageCollect()
+{
+	if (bObjectReferenceListDirty)
+	{
+		bObjectReferenceListDirty = false;
+		ObjectReferenceList.Reset();
 
-		for (TPair<FName, FPrimaryAssetData>& NamePair : TypeData.AssetMap)
+		// Iterate primary asset map
+		for (TPair<FName, TSharedRef<FPrimaryAssetTypeData>>& TypePair : AssetTypeMap)
 		{
-			FPrimaryAssetData& NameData = NamePair.Value;
-			
-			const FSoftObjectPath& AssetRef = NameData.AssetPtr.ToSoftObjectPath();
+			FPrimaryAssetTypeData& TypeData = TypePair.Value.Get();
 
-			// Dynamic types can have null asset refs
-			if (!AssetRef.IsNull())
+			// Add base class in case it's a blueprint
+			if (!TypeData.Info.bIsDynamicAsset)
 			{
-				AssetPathMap.Add(AssetRef, FPrimaryAssetId(TypePair.Key, NamePair.Key));
+				ObjectReferenceList.AddUnique(TypeData.Info.AssetBaseClassLoaded);
 			}
 		}
 	}
-
-	bIsManagementDatabaseCurrent = false;
 }
 
 void UAssetManager::LoadRedirectorMaps()
@@ -2955,7 +3052,7 @@ void UAssetManager::DumpLoadedAssetState()
 
 		FPrimaryAssetTypeData& TypeData = Manager.AssetTypeMap.Find(TypeInfo.PrimaryAssetType)->Get();
 
-		for (const TPair<FName, FPrimaryAssetData>& NamePair : TypeData.AssetMap)
+		for (const TPair<FName, FPrimaryAssetData>& NamePair : TypeData.GetAssets())
 		{
 			const FPrimaryAssetData& NameData = NamePair.Value;
 
@@ -3627,7 +3724,7 @@ void UAssetManager::RefreshPrimaryAssetDirectory(bool bForceRefresh)
 			if (TypeData.Info.AssetScanPaths.Num())
 			{
 				// Clear old data if this type has actual scan paths
-				TypeData.AssetMap.Reset();
+				TypeData.ResetAssets(AssetPathMap);
 
 				// Rescan all assets. We don't force synchronous here as in the editor it was already loaded async
 				ScanPathsForPrimaryAssets(TypePair.Key, TypeData.Info.AssetScanPaths, TypeData.Info.AssetBaseClassLoaded, TypeData.Info.bHasBlueprintClasses, TypeData.Info.bIsEditorOnly, false);
@@ -3786,7 +3883,7 @@ void UAssetManager::UpdateManagementDatabase(bool bForceRefresh)
 	{
 		const FPrimaryAssetTypeData& TypeData = TypePair.Value.Get();
 
-		for (const TPair<FName, FPrimaryAssetData>& NamePair : TypeData.AssetMap)
+		for (const TPair<FName, FPrimaryAssetData>& NamePair : TypeData.GetAssets())
 		{
 			const FPrimaryAssetData& NameData = NamePair.Value;
 			FPrimaryAssetId PrimaryAssetId(TypePair.Key, NamePair.Key);
@@ -3796,7 +3893,7 @@ void UAssetManager::UpdateManagementDatabase(bool bForceRefresh)
 			// Get the list of directly referenced assets, the registry wants it as FNames
 			AssetPackagesReferenced.Reset();
 
-			const FSoftObjectPath& AssetRef = NameData.AssetPtr.ToSoftObjectPath();
+			const FSoftObjectPath& AssetRef = NameData.GetAssetPtr().ToSoftObjectPath();
 
 			if (AssetRef.IsValid())
 			{
@@ -3897,11 +3994,11 @@ void UAssetManager::UpdateManagementDatabase(bool bForceRefresh)
 	{
 		const FPrimaryAssetTypeData& TypeData = TypePair.Value.Get();
 
-		for (const TPair<FName, FPrimaryAssetData>& NamePair : TypeData.AssetMap)
+		for (const TPair<FName, FPrimaryAssetData>& NamePair : TypeData.GetAssets())
 		{
 			const FPrimaryAssetData& NameData = NamePair.Value;
 			FPrimaryAssetId PrimaryAssetId(TypePair.Key, NamePair.Key);
-			const FSoftObjectPath& AssetRef = NameData.AssetPtr.ToSoftObjectPath();
+			const FSoftObjectPath& AssetRef = NameData.GetAssetPtr().ToSoftObjectPath();
 
 			TSet<FPrimaryAssetId> Managers;
 
@@ -4458,7 +4555,7 @@ void UAssetManager::EndPIE(bool bStartSimulate)
 	{
 		const FPrimaryAssetTypeData& TypeData = TypePair.Value.Get();
 
-		for (const TPair<FName, FPrimaryAssetData>& NamePair : TypeData.AssetMap)
+		for (const TPair<FName, FPrimaryAssetData>& NamePair : TypeData.GetAssets())
 		{
 			const FPrimaryAssetData& NameData = NamePair.Value;
 			const FPrimaryAssetLoadState& LoadState = (!NameData.PendingState.IsValid()) ? NameData.CurrentState : NameData.PendingState;
@@ -4490,7 +4587,11 @@ void UAssetManager::EndPIE(bool bStartSimulate)
 void UAssetManager::ReinitializeFromConfig()
 {
 	// We specifically do not reset AssetRuleOverrides as those can be set by something other than inis
-	AssetPathMap.Reset();
+	for (TPair<FName, TSharedRef<FPrimaryAssetTypeData>>& Pair : AssetTypeMap)
+	{
+		Pair.Value->ResetAssets(AssetPathMap);
+	}
+	check(AssetPathMap.IsEmpty()); // Should have been emptied by the ResetAssets calls
 	ManagementParentMap.Reset();
 	CachedAssetBundles.Reset();
 	AlreadyScannedDirectories.Reset();
@@ -4557,7 +4658,7 @@ void UAssetManager::OnInMemoryAssetCreated(UObject *Object)
 					// Add or update asset data
 					if (TryUpdateCachedAssetData(PrimaryAssetId, NewAssetData, true))
 					{
-						RebuildObjectReferenceList();
+						OnObjectReferenceListInvalidated();
 					}
 				}
 			}
@@ -4652,9 +4753,9 @@ void UAssetManager::RemovePrimaryAssetId(const FPrimaryAssetId& PrimaryAssetId)
 		check(FoundType);
 		FPrimaryAssetTypeData& TypeData = FoundType->Get();
 
-		TypeData.AssetMap.Remove(PrimaryAssetId.PrimaryAssetName);
+		TypeData.RemoveAsset(PrimaryAssetId.PrimaryAssetName, AssetPathMap);
 
-		RebuildObjectReferenceList();
+		OnObjectReferenceListInvalidated();
 	}
 }
 
