[Backout] - CL25148391
Original CL Desc
-----------------------------------------------------------------
Make it more explicit that DataLayerInstanceWithAsset are expected to have a valid asset in cook builds

diff --git a/Engine/Source/Runtime/Engine/Private/WorldPartition/DataLayer/WorldDataLayers.cpp b/Engine/Source/Runtime/Engine/Private/WorldPartition/DataLayer/WorldDataLayers.cpp
--- a/Engine/Source/Runtime/Engine/Private/WorldPartition/DataLayer/WorldDataLayers.cpp
+++ b/Engine/Source/Runtime/Engine/Private/WorldPartition/DataLayer/WorldDataLayers.cpp
@@ -775,7 +775,7 @@ void AWorldDataLayers::PostLoad()
 
 	ConvertDataLayerToInstancces();
 
-	// Prune non-runtime Data Layers when cooking or when in a game world
+	// Remove all Editor Data Layers when cooking or when in a game world
 	if (IsRunningCookCommandlet() || GetWorld()->IsGameWorld())
 	{
 		ForEachDataLayer([](UDataLayerInstance* DataLayer)
@@ -784,16 +784,16 @@ void AWorldDataLayers::PostLoad()
 			return true;
 		});
 
-		TArray<UDataLayerInstance*> NonRuntimeDataLayers;
-		ForEachDataLayer([&NonRuntimeDataLayers](UDataLayerInstance* DataLayer)
+		TArray<UDataLayerInstance*> EditorDataLayers;
+		ForEachDataLayer([&EditorDataLayers](UDataLayerInstance* DataLayer)
 		{
 			if (DataLayer && !DataLayer->IsRuntime())
 			{
-				NonRuntimeDataLayers.Add(DataLayer);
+				EditorDataLayers.Add(DataLayer);
 			}
 			return true;
 		});
-		RemoveDataLayers(NonRuntimeDataLayers);
+		RemoveDataLayers(EditorDataLayers);
 	}
 
 	// Empty Data Layers for non-runtime relevant AWorldDataLayers
@@ -836,10 +836,7 @@ void AWorldDataLayers::PostLoad()
 		static_assert(DATALAYER_TO_INSTANCE_RUNTIME_CONVERSION_ENABLED, "Remove unnecessary cast. All DataLayerInstance now have assets");
 		if (const UDataLayerInstanceWithAsset* DataLayerInstanceWithAsset = Cast<UDataLayerInstanceWithAsset>(DataLayerInstance))
 		{
-			// DataLayerInstanceWithAsset with invalid asset are removed during cook when pruning non-runtime DataLayers.
-			// All DataLayerInstanceWithAsset are expected to have an asset in cook builds.
-			check(DataLayerInstanceWithAsset->GetAsset());
-			AssetNameToInstance.Add(DataLayerInstanceWithAsset->GetAsset()->GetFullName(), DataLayerInstanceWithAsset);
+			AssetNameToInstance.Add(DataLayerInstanceWithAsset->GetAsset()->GetFullName(), DataLayerInstance);
 		}
 	}
 #endif
