Fixed mismatched copy constructor/assignment operator in FMeshReductionSettings.

diff --git a/Engine/Source/Runtime/Engine/Classes/Engine/MeshMerging.h b/Engine/Source/Runtime/Engine/Classes/Engine/MeshMerging.h
--- a/Engine/Source/Runtime/Engine/Classes/Engine/MeshMerging.h
+++ b/Engine/Source/Runtime/Engine/Classes/Engine/MeshMerging.h
@@ -140,30 +140,6 @@ struct FMeshReductionSettings
 	{
 	}
 
-	FMeshReductionSettings(const FMeshReductionSettings& Other)
-		: PercentTriangles(Other.PercentTriangles)
-		, MaxNumOfTriangles(Other.MaxNumOfTriangles)
-		, PercentVertices(Other.PercentVertices)
-		, MaxNumOfVerts(Other.MaxNumOfVerts)
-		, MaxDeviation(Other.MaxDeviation)
-		, PixelError(Other.PixelError)
-		, WeldingThreshold(Other.WeldingThreshold)
-		, HardAngleThreshold(Other.HardAngleThreshold)
-		, BaseLODModel(Other.BaseLODModel)
-		, SilhouetteImportance(Other.SilhouetteImportance)
-		, TextureImportance(Other.TextureImportance)
-		, ShadingImportance(Other.ShadingImportance)
-		, bRecalculateNormals(Other.bRecalculateNormals)
-		, bGenerateUniqueLightmapUVs(Other.bGenerateUniqueLightmapUVs)
-		, bKeepSymmetry(Other.bKeepSymmetry)
-		, bVisibilityAided(Other.bVisibilityAided)
-		, bCullOccluded(Other.bCullOccluded)
-		, TerminationCriterion(Other.TerminationCriterion)
-		, VisibilityAggressiveness(Other.VisibilityAggressiveness)
-		, VertexColorImportance(Other.VertexColorImportance)
-	{
-	}
-
 	/** Equality operator. */
 	bool operator==(const FMeshReductionSettings& Other) const
 	{
