(Github PR) Make the IsSeekable method virtual in USoundWave
#9901

diff --git a/Engine/Source/Runtime/Engine/Classes/Sound/SoundWave.h b/Engine/Source/Runtime/Engine/Classes/Sound/SoundWave.h
--- a/Engine/Source/Runtime/Engine/Classes/Sound/SoundWave.h
+++ b/Engine/Source/Runtime/Engine/Classes/Sound/SoundWave.h
@@ -1132,7 +1132,7 @@ public:
 	bool IsStreaming(const FPlatformAudioCookOverrides& Overrides) const;
 
 	/** Returns whether the sound is seekable. */
-	bool IsSeekable() const;
+	virtual bool IsSeekable() const;
 
 	/**
 	 * Checks whether we should use the load on demand cache.
