Fixed standard blending not working as expected with sync groups. Additionally, added support for syncing when using inertialization for a state machine transition.

diff --git a/Engine/Source/Runtime/Engine/Classes/Animation/AnimationAsset.h b/Engine/Source/Runtime/Engine/Classes/Animation/AnimationAsset.h
--- a/Engine/Source/Runtime/Engine/Classes/Animation/AnimationAsset.h
+++ b/Engine/Source/Runtime/Engine/Classes/Animation/AnimationAsset.h
@@ -384,6 +384,8 @@ struct FAnimTickRecord
 
 	bool bLooping = false;
 	bool bIsEvaluator = false;
+	bool bRequestedInertialization = false;
+
 	const UMirrorDataTable* MirrorDataTable = nullptr;
 
 	TSharedPtr<TArray<TUniquePtr<const UE::Anim::IAnimNotifyEventContextDataInterface>>> ContextData;
diff --git a/Engine/Source/Runtime/Engine/Private/Animation/AnimInertializationSyncScope.cpp b/Engine/Source/Runtime/Engine/Private/Animation/AnimInertializationSyncScope.cpp
new file mode 100644
--- /dev/null
+++ b/Engine/Source/Runtime/Engine/Private/Animation/AnimInertializationSyncScope.cpp
@@ -0,0 +1,5 @@
+﻿// Copyright Epic Games, Inc. All Rights Reserved.
+
+#include "Animation/AnimInertializationSyncScope.h"
+
+IMPLEMENT_ANIMGRAPH_MESSAGE(UE::Anim::FAnimInertializationSyncScope);
diff --git a/Engine/Source/Runtime/Engine/Private/Animation/AnimNode_AssetPlayerBase.cpp b/Engine/Source/Runtime/Engine/Private/Animation/AnimNode_AssetPlayerBase.cpp
--- a/Engine/Source/Runtime/Engine/Private/Animation/AnimNode_AssetPlayerBase.cpp
+++ b/Engine/Source/Runtime/Engine/Private/Animation/AnimNode_AssetPlayerBase.cpp
@@ -3,6 +3,7 @@
 #include "Animation/AnimNode_AssetPlayerBase.h"
 #include "Animation/AnimSyncScope.h"
 #include "Animation/AnimTrace.h"
+#include "Animation/AnimInertializationSyncScope.h"
 
 #include UE_INLINE_GENERATED_CPP_BY_NAME(AnimNode_AssetPlayerBase)
 
@@ -32,21 +33,41 @@ void FAnimNode_AssetPlayerBase::CreateTickRecordForNode(const FAnimationUpdateCo
 
 	const EAnimGroupRole::Type SyncGroupRole = GetGroupRole();
 	const FName SyncGroupName = GetGroupName();
-
-	const FName GroupNameToUse = ((SyncGroupRole < EAnimGroupRole::TransitionLeader) || bHasBeenFullWeight) ? SyncGroupName : NAME_None;
+	FName GroupNameToUse = SyncGroupName;
 	EAnimSyncMethod MethodToUse = GetGroupMethod();
-	if(GroupNameToUse == NAME_None && MethodToUse == EAnimSyncMethod::SyncGroup)
+	const bool bRequestedInertialization = Context.GetMessage<UE::Anim::FAnimInertializationSyncScope>() != nullptr;
+	
+	// Skip sync based on roles.
 	{
-		MethodToUse = EAnimSyncMethod::DoNotSync;
+		// Only allow transition leader/follower part of a sync group once after inertialization request. (Inertilization)
+		if (bRequestedInertialization)
+		{
+			if (SyncGroupRole == EAnimGroupRole::TransitionLeader || SyncGroupRole == EAnimGroupRole::TransitionFollower)
+			{
+				GroupNameToUse = NAME_None;
+			}
+		}
+		// Only allow transition leader/follower part of a sync group once it has full weight (Standard blend).
+		else if ((SyncGroupRole == EAnimGroupRole::TransitionLeader || SyncGroupRole == EAnimGroupRole::TransitionFollower) && !bHasBeenFullWeight)
+		{
+			GroupNameToUse = NAME_None;
+		}
+
+		// Do not use sync groups.
+		if (GroupNameToUse == NAME_None && MethodToUse == EAnimSyncMethod::SyncGroup)
+		{
+			MethodToUse = EAnimSyncMethod::DoNotSync;
+		}
 	}
-
+	
 	const UE::Anim::FAnimSyncParams SyncParams(GroupNameToUse, SyncGroupRole, MethodToUse);
 	FAnimTickRecord TickRecord(Sequence, bLooping, PlayRate, bIsEvaluator, FinalBlendWeight, /*inout*/ InternalTimeAccumulator, MarkerTickRecord);
 	TickRecord.GatherContextData(Context);
 
 	TickRecord.RootMotionWeightModifier = Context.GetRootMotionWeightModifier();
 	TickRecord.DeltaTimeRecord = &DeltaTimeRecord;
-
+	TickRecord.bRequestedInertialization = bRequestedInertialization;
+	
 	SyncScope.AddTickRecord(TickRecord, SyncParams, UE::Anim::FAnimSyncDebugInfo(Context));
 
 	TRACE_ANIM_TICK_RECORD(Context, TickRecord);
diff --git a/Engine/Source/Runtime/Engine/Private/Animation/AnimNode_StateMachine.cpp b/Engine/Source/Runtime/Engine/Private/Animation/AnimNode_StateMachine.cpp
--- a/Engine/Source/Runtime/Engine/Private/Animation/AnimNode_StateMachine.cpp
+++ b/Engine/Source/Runtime/Engine/Private/Animation/AnimNode_StateMachine.cpp
@@ -13,6 +13,7 @@
 #include "Animation/AnimBlueprintGeneratedClass.h"
 #include "Animation/ExposedValueHandler.h"
 #include "Logging/TokenizedMessage.h"
+#include "Animation/AnimInertializationSyncScope.h"
 
 #include UE_INLINE_GENERATED_CPP_BY_NAME(AnimNode_StateMachine)
 
@@ -473,6 +474,8 @@ void FAnimNode_StateMachine::Update_AnyThread(const FAnimationUpdateContext& Con
 
 	StatesUpdated.Reset();
 
+	bool bLastActiveTransitionRequestedInertialization = false;
+	
 	// Tick the individual state/states that are active
 	if (ActiveTransitionArray.Num() > 0)
 	{
@@ -492,6 +495,12 @@ void FAnimNode_StateMachine::Update_AnyThread(const FAnimationUpdateContext& Con
 					Context.AnimInstanceProxy->AddAnimNotifyFromGeneratedClass(ActiveTransitionArray[Index].EndNotify);
 					Context.AnimInstanceProxy->AddAnimNotifyFromGeneratedClass(GetStateInfo().FullyBlendedNotify);
 				}
+
+				// we were the last active transition and used inertialization
+				if (ActiveTransitionArray[Index].LogicType == ETransitionLogicType::TLT_Inertialization && (ActiveTransitionArray.Num() - 1 ==  Index))
+				{
+					bLastActiveTransitionRequestedInertialization = true;
+				}
 			}
 			else
 			{
@@ -516,7 +525,9 @@ void FAnimNode_StateMachine::Update_AnyThread(const FAnimationUpdateContext& Con
 	// Update the only active state if there are no transitions still in flight
 	if (ActiveTransitionArray.Num() == 0 && !IsAConduitState(CurrentState) && !StatesUpdated.Contains(CurrentState))
 	{
+		UE::Anim::TOptionalScopedGraphMessage<UE::Anim::FAnimInertializationSyncScope> InertializationSync(bLastActiveTransitionRequestedInertialization, Context);
 		UE::Anim::TOptionalScopedGraphMessage<UE::Anim::FActiveStateMachineScope> Message(bCreateNotifyMetaData, Context, Context, this, CurrentState);
+		
 		StatePoseLinks[CurrentState].Update(Context);
 	}
 
@@ -766,6 +777,8 @@ void FAnimNode_StateMachine::UpdateTransitionStates(const FAnimationUpdateContex
 
 		case ETransitionLogicType::TLT_Inertialization:
 			{
+				UE::Anim::TScopedGraphMessage<UE::Anim::FAnimInertializationSyncScope> InertializationSync(Context);
+
 				// update target state
 				UpdateState(Transition.NextState, Context);
 			}
diff --git a/Engine/Source/Runtime/Engine/Private/Animation/AnimSync.cpp b/Engine/Source/Runtime/Engine/Private/Animation/AnimSync.cpp
--- a/Engine/Source/Runtime/Engine/Private/Animation/AnimSync.cpp
+++ b/Engine/Source/Runtime/Engine/Private/Animation/AnimSync.cpp
@@ -148,19 +148,26 @@ void FAnimSync::TickAssetPlayerInstances(FAnimInstanceProxy& InProxy, float InDe
 				// if it has leader score
 				SCOPE_CYCLE_COUNTER(STAT_TickAssetPlayerInstance);
 				FScopeCycleCounterUObject Scope(GroupLeader.SourceAsset);
-
-				auto TickRecordMatchesGroupLeader = [&GroupLeader](const FAnimTickRecord& TickRecord) -> bool
+				
+				// Inertialization was requested therefore we resync to previous leader, if needed.
+				if (GroupLeader.bRequestedInertialization)
 				{
-					return TickRecord.SourceAsset == GroupLeader.SourceAsset && TickRecord.TimeAccumulator == GroupLeader.TimeAccumulator;
-				};
-
-				// if the group leader was previously inactive then set it to resynchronize to the sync group's time
-				// (maintains sync during inertialization or zero-length blends where the previous sync leader is no longer active)
-				const bool bGroupLeaderWasActive = (PreviousGroup && PreviousGroup->ActivePlayers.ContainsByPredicate(TickRecordMatchesGroupLeader)) ||
-					(PreviousUngroupedActivePlayers.ContainsByPredicate(TickRecordMatchesGroupLeader));
-				const bool bGroupWasActive = PreviousGroup && !PreviousGroup->ActivePlayers.IsEmpty();
-				TickContext.SetResyncToSyncGroup(!bGroupLeaderWasActive && bGroupWasActive);
-
+					// Ensure we have a previous group leader to sync to otherwise we play normally. 
+					if (PreviousGroup && !PreviousGroup->ActivePlayers.IsEmpty())
+					{
+						// Only need to resync when using length based syncing since we initialized the tick context with the previous group's sync end position
+						// and that will take care of marker based syncing.
+						if (!TickContext.CanUseMarkerPosition() || !(PreviousGroup->ValidMarkers.Num() > 0))
+						{
+							const FAnimTickRecord & PreviousGroupLeader = PreviousGroup->ActivePlayers[PreviousGroup->GroupLeaderIndex];
+							const bool bShouldResyncToSyncGroup = PreviousGroupLeader.LeaderScore >= GroupLeader.LeaderScore;
+
+							// Sync to previous group leader if we have a lower score than them.
+							TickContext.SetResyncToSyncGroup(bShouldResyncToSyncGroup);
+						}
+					}
+				}
+				
 				TickContext.MarkerTickContext.MarkersPassedThisTick.Reset();
 				TickContext.RootMotionMovementParams.Clear();
 				GroupLeader.SourceAsset->TickAssetPlayer(GroupLeader, InProxy.NotifyQueue, TickContext);
@@ -254,7 +261,6 @@ void FAnimSync::TickAssetPlayerInstances(FAnimInstanceProxy& InProxy, float InDe
 				// if we don't have a good leader, no reason to convert to follower
 				// tick as leader
 				TickContext.ConvertToFollower();
-				TickContext.SetResyncToSyncGroup(false);
 
 				for (int32 TickIndex = GroupLeaderIndex + 1; TickIndex < SyncGroup.ActivePlayers.Num(); ++TickIndex)
 				{
@@ -262,7 +268,10 @@ void FAnimSync::TickAssetPlayerInstances(FAnimInstanceProxy& InProxy, float InDe
 					{
 						SCOPE_CYCLE_COUNTER(STAT_TickAssetPlayerInstance);
 						FScopeCycleCounterUObject Scope(AssetPlayer.SourceAsset);
+						
+						TickContext.SetResyncToSyncGroup(AssetPlayer.bRequestedInertialization);
 						TickContext.RootMotionMovementParams.Clear();
+						
 						AssetPlayer.SourceAsset->TickAssetPlayer(AssetPlayer, InProxy.NotifyQueue, TickContext);
 					}
 					if (RootMotionMode == ERootMotionMode::RootMotionFromEverything && TickContext.RootMotionMovementParams.bHasRootMotion)
diff --git a/Engine/Source/Runtime/Engine/Public/Animation/AnimInertializationSyncScope.h b/Engine/Source/Runtime/Engine/Public/Animation/AnimInertializationSyncScope.h
new file mode 100644
--- /dev/null
+++ b/Engine/Source/Runtime/Engine/Public/Animation/AnimInertializationSyncScope.h
@@ -0,0 +1,16 @@
+﻿// Copyright Epic Games, Inc. All Rights Reserved.
+
+#pragma once
+
+#include "CoreMinimal.h"
+#include "Animation/AnimNodeMessages.h"
+
+namespace UE { namespace Anim {
+
+// Scoped graph message used to synchronize tick records using inertialization 
+class ENGINE_API FAnimInertializationSyncScope : public IGraphMessage
+{
+	DECLARE_ANIMGRAPH_MESSAGE(FAnimInertializationSyncScope)
+};
+
+}} // namespace UE::Anim
\ No newline at end of file
