Texture : fix duplication changes owner to Default texture because fields are reset from archetype after constructor

#rnx

diff --git a/Engine/Source/Runtime/Engine/Private/Texture.cpp b/Engine/Source/Runtime/Engine/Private/Texture.cpp
--- a/Engine/Source/Runtime/Engine/Private/Texture.cpp
+++ b/Engine/Source/Runtime/Engine/Private/Texture.cpp
@@ -1048,6 +1048,12 @@ void UTexture::AppendToClassSchema(FAppendToClassSchemaContext& Context)
 void UTexture::PostInitProperties()
 {
 #if WITH_EDITORONLY_DATA
+	// this was set in the constructor :
+	// but it can be stomped from the archetype in duplication
+	// re-set it now :
+	//check( Source.Owner == this );
+	Source.SetOwner(this);
+
 	if (!HasAnyFlags(RF_ClassDefaultObject | RF_NeedLoad))
 	{
 		AssetImportData = NewObject<UAssetImportData>(this, TEXT("AssetImportData"));
@@ -2100,6 +2106,7 @@ FTextureSource FTextureSource::CopyTornOff() const
 	Result.Owner = nullptr; // TornOffs don't count as belonging to the same owner
 	// Result can't talk to Owner any more, so save info we need :
 	check( Owner != nullptr );
+	check( &(Owner->Source) == this );
 	Result.TornOffGammaSpace = Owner->GetGammaSpace();
 	Result.TornOffTextureClass = Owner->GetTextureClass();
 	return Result;
@@ -2622,6 +2629,8 @@ EGammaSpace FTextureSource::GetGammaSpace(int LayerIndex) const
 	// TextureSource does not know its own gamma, but its owning Texture does :
 	if ( Owner != nullptr )
 	{
+		check( &(Owner->Source) == this );
+
 		// same as Owner->GetGammaSpace , but uses LayerFormatSettings for SRGB flag
 		FTextureFormatSettings FormatSettings;
 		Owner->GetLayerFormatSettings(LayerIndex, FormatSettings);
@@ -3056,6 +3065,10 @@ void UTexture::GetDefaultFormatSettings(FTextureFormatSettings& OutSettings) con
 
 void UTexture::GetLayerFormatSettings(int32 LayerIndex, FTextureFormatSettings& OutSettings) const
 {
+#if WITH_EDITORONLY_DATA
+	check( Source.Owner == this );
+#endif
+
 	check(LayerIndex >= 0);
 	if (LayerIndex < LayerFormatSettings.Num())
 	{
