World Partition: don't allow editing IsSpatiallyLoaded flag on level instance actors if the main world has streaming disabled.

diff --git a/Engine/Source/Runtime/Engine/Private/ActorEditor.cpp b/Engine/Source/Runtime/Engine/Private/ActorEditor.cpp
--- a/Engine/Source/Runtime/Engine/Private/ActorEditor.cpp
+++ b/Engine/Source/Runtime/Engine/Private/ActorEditor.cpp
@@ -87,12 +87,23 @@ bool AActor::CanEditChange(const FProperty* PropertyThatWillChange) const
 	{
 		if (!IsTemplate())
 		{
-			UWorld* World = GetTypedOuter<UWorld>();
-			UWorldPartition* WorldPartition = World ? World->GetWorldPartition() : nullptr;
-			if (!WorldPartition || (!WorldPartition->IsStreamingEnabled() && bIsRuntimeGridProperty))
+			UWorld* OwningWorld = GetWorld();
+			UWorld* OuterWorld = GetTypedOuter<UWorld>();
+			UWorldPartition* OwningWorldPartition = OwningWorld ? OwningWorld->GetWorldPartition() : nullptr;
+			UWorldPartition* OuterWorldPartition = OuterWorld ? OuterWorld->GetWorldPartition() : nullptr;			
+
+			if (!OwningWorldPartition || !OuterWorldPartition)
 			{
 				return false;
 			}
+
+			if (!OwningWorldPartition->IsStreamingEnabled())
+			{
+				if (bIsSpatiallyLoadedProperty || bIsRuntimeGridProperty)
+				{
+					return false;
+				}
+			}
 		}
 	}
 
