WorldDataLayers: Add instrumentation and don't crash if DataLayerInstanceWithAsset is missing its Asset; log a warning instead.
#rnx

diff --git a/Engine/Source/Runtime/Engine/Private/WorldPartition/DataLayer/WorldDataLayers.cpp b/Engine/Source/Runtime/Engine/Private/WorldPartition/DataLayer/WorldDataLayers.cpp
--- a/Engine/Source/Runtime/Engine/Private/WorldPartition/DataLayer/WorldDataLayers.cpp
+++ b/Engine/Source/Runtime/Engine/Private/WorldPartition/DataLayer/WorldDataLayers.cpp
@@ -836,6 +836,11 @@ void AWorldDataLayers::PostLoad()
 		static_assert(DATALAYER_TO_INSTANCE_RUNTIME_CONVERSION_ENABLED, "Remove unnecessary cast. All DataLayerInstance now have assets");
 		if (const UDataLayerInstanceWithAsset* DataLayerInstanceWithAsset = Cast<UDataLayerInstanceWithAsset>(DataLayerInstance))
 		{
+			if (!DataLayerInstanceWithAsset->GetAsset())
+			{
+				UE_LOG(LogWorldPartition, Warning, TEXT("DataLayerWithAsset %s has null asset."), *DataLayerInstanceWithAsset->GetPathName());
+				continue;
+			}
 			AssetNameToInstance.Add(DataLayerInstanceWithAsset->GetAsset()->GetFullName(), DataLayerInstance);
 		}
 	}
