Fixed mismatched copy constructor/assignment operator in FHLODProxyMesh.

diff --git a/Engine/Source/Runtime/Engine/Private/HLOD/HLODProxyMesh.cpp b/Engine/Source/Runtime/Engine/Private/HLOD/HLODProxyMesh.cpp
--- a/Engine/Source/Runtime/Engine/Private/HLOD/HLODProxyMesh.cpp
+++ b/Engine/Source/Runtime/Engine/Private/HLOD/HLODProxyMesh.cpp
@@ -11,13 +11,6 @@ FHLODProxyMesh::FHLODProxyMesh()
 }
 
 #if WITH_EDITOR
-FHLODProxyMesh::FHLODProxyMesh(const FHLODProxyMesh& Other)
-	: LODActor(Other.LODActor)
-	, StaticMesh(Other.StaticMesh)
-	, Key(Other.Key)
-{
-}
-
 FHLODProxyMesh::FHLODProxyMesh(ALODActor* InLODActor, UStaticMesh* InStaticMesh, const FName& InKey)
 	: LODActor(InLODActor)
 	, StaticMesh(InStaticMesh)
diff --git a/Engine/Source/Runtime/Engine/Public/HLOD/HLODProxyMesh.h b/Engine/Source/Runtime/Engine/Public/HLOD/HLODProxyMesh.h
--- a/Engine/Source/Runtime/Engine/Public/HLOD/HLODProxyMesh.h
+++ b/Engine/Source/Runtime/Engine/Public/HLOD/HLODProxyMesh.h
@@ -17,7 +17,6 @@ struct FHLODProxyMesh
 	FHLODProxyMesh();
 
 #if WITH_EDITOR
-	FHLODProxyMesh(const FHLODProxyMesh& Other);
 	FHLODProxyMesh(UStaticMesh* InStaticMesh, const FName& InKey);
 
 	FHLODProxyMesh(ALODActor* InLODActor, UStaticMesh* InStaticMesh, const FName& InKey);
