Fixed mismatched copy constructor/assignment operator in FContentComparisonAssetInfo.

diff --git a/Engine/Source/Runtime/Engine/Public/EngineUtils.h b/Engine/Source/Runtime/Engine/Public/EngineUtils.h
--- a/Engine/Source/Runtime/Engine/Public/EngineUtils.h
+++ b/Engine/Source/Runtime/Engine/Public/EngineUtils.h
@@ -752,14 +752,6 @@ struct FContentComparisonAssetInfo
 			(ResourceSize == Other.ResourceSize)
 			);
 	}
-
-	/** operator = */
-	FContentComparisonAssetInfo& operator=(const FContentComparisonAssetInfo& Other)
-	{
-		AssetName = Other.AssetName;
-		ResourceSize = Other.ResourceSize;
-		return *this;
-	}
 };
 
 /** Helper class for performing the content comparison console command */
