Fixed mismatched and incomplete copy constructor/assignment operator in FVisualLogEntry.

diff --git a/Engine/Source/Runtime/Engine/Private/VisualLogger/VisualLoggerTypes.cpp b/Engine/Source/Runtime/Engine/Private/VisualLogger/VisualLoggerTypes.cpp
--- a/Engine/Source/Runtime/Engine/Private/VisualLogger/VisualLoggerTypes.cpp
+++ b/Engine/Source/Runtime/Engine/Private/VisualLogger/VisualLoggerTypes.cpp
@@ -62,20 +62,6 @@ bool FVisualLogStatusCategory::GetDesc(int32 Index, FString& Key, FString& Value
 	return false;
 }
 
-FVisualLogEntry::FVisualLogEntry(const FVisualLogEntry& Entry)
-{
-	TimeStamp = Entry.TimeStamp;
-	Location = Entry.Location;
-	bIsLocationValid = Entry.bIsLocationValid;
-
-	Events = Entry.Events;
-	LogLines = Entry.LogLines;
-	Status = Entry.Status;
-	ElementsToDraw = Entry.ElementsToDraw;
-	HistogramSamples = Entry.HistogramSamples;
-	DataBlocks = Entry.DataBlocks;
-}
-
 FVisualLogEntry::FVisualLogEntry(const AActor* InActor, TArray<TWeakObjectPtr<UObject> >* Children)
 {
 	if (!IsValid(InActor))
diff --git a/Engine/Source/Runtime/Engine/Public/VisualLogger/VisualLoggerTypes.h b/Engine/Source/Runtime/Engine/Public/VisualLogger/VisualLoggerTypes.h
--- a/Engine/Source/Runtime/Engine/Public/VisualLogger/VisualLoggerTypes.h
+++ b/Engine/Source/Runtime/Engine/Public/VisualLogger/VisualLoggerTypes.h
@@ -182,7 +182,6 @@ struct ENGINE_API FVisualLogEntry
 	TArray<FVisualLogDataBlock>	DataBlocks;
 
 	FVisualLogEntry() { Reset(); }
-	FVisualLogEntry(const FVisualLogEntry& Entry);
 	FVisualLogEntry(const AActor* InActor, TArray<TWeakObjectPtr<UObject> >* Children);
 	FVisualLogEntry(double InTimeStamp, FVector InLocation, const UObject* Object, TArray<TWeakObjectPtr<UObject> >* Children);
 
