Fixed mismatched copy constructor/assignment operator in UE::HLSLTree::FRequestedType.

diff --git a/Engine/Source/Runtime/Engine/Public/HLSLTree/HLSLTree.h b/Engine/Source/Runtime/Engine/Public/HLSLTree/HLSLTree.h
--- a/Engine/Source/Runtime/Engine/Public/HLSLTree/HLSLTree.h
+++ b/Engine/Source/Runtime/Engine/Public/HLSLTree/HLSLTree.h
@@ -140,8 +140,13 @@ class ENGINE_API FRequestedType
 {
 public:
 	FRequestedType() = default;
+	FRequestedType(FRequestedType&&) = default;
+	FRequestedType(const FRequestedType&) = default;
+	FRequestedType& operator=(FRequestedType&&) = default;
+	FRequestedType& operator=(const FRequestedType&) = default;
+
 	FRequestedType(const Shader::FType& InType, bool bDefaultRequest = true);
-	FRequestedType(const FRequestedType& InType, bool bDefaultRequest = true);
+	FRequestedType(const FRequestedType& InType, bool bDefaultRequest);
 	FRequestedType(Shader::EValueType InType, bool bDefaultRequest = true) : FRequestedType(Shader::FType(InType), bDefaultRequest) {}
 	FRequestedType(const Shader::FStructType* InType, bool bDefaultRequest = true) : FRequestedType(Shader::FType(InType), bDefaultRequest) {}
 	FRequestedType(const FName& InType, bool bDefaultRequest = true) : FRequestedType(Shader::FType(InType), bDefaultRequest) {}
