Fixed mismatched copy constructor/assignment operator in FAnimationCurveIdentifier.

diff --git a/Engine/Source/Runtime/Engine/Classes/Animation/AnimData/CurveIdentifier.h b/Engine/Source/Runtime/Engine/Classes/Animation/AnimData/CurveIdentifier.h
--- a/Engine/Source/Runtime/Engine/Classes/Animation/AnimData/CurveIdentifier.h
+++ b/Engine/Source/Runtime/Engine/Classes/Animation/AnimData/CurveIdentifier.h
@@ -75,16 +75,6 @@ struct ENGINE_API FAnimationCurveIdentifier
 		return CurveName != NAME_None && CurveType != ERawCurveTrackTypes::RCT_MAX;
 	}
 
-	FAnimationCurveIdentifier& operator=(const FAnimationCurveIdentifier& Other)
-	{
-		CurveName = Other.CurveName;
-		CurveType = Other.CurveType;
-		Channel = Other.Channel;
-		Axis = Other.Axis;
-
-		return *this;
-	}
-	
 	friend uint32 GetTypeHash(const FAnimationCurveIdentifier& CurveId)
 	{
 		return HashCombine(HashCombine(HashCombine((uint32)CurveId.Channel, (uint32)CurveId.Axis), (uint32)CurveId.CurveType), GetTypeHash(CurveId.CurveName));
