Fixed mismatched copy constructor/assignment operator in SkeletalMeshImportData::FTriangle.

diff --git a/Engine/Source/Runtime/Engine/Public/Rendering/SkeletalMeshLODImporterData.h b/Engine/Source/Runtime/Engine/Public/Rendering/SkeletalMeshLODImporterData.h
--- a/Engine/Source/Runtime/Engine/Public/Rendering/SkeletalMeshLODImporterData.h
+++ b/Engine/Source/Runtime/Engine/Public/Rendering/SkeletalMeshLODImporterData.h
@@ -117,29 +117,6 @@ namespace SkeletalMeshImportData
 		FVector3f	TangentZ[3];
 
 
-		FTriangle& operator=(const FTriangle& Other)
-		{
-			this->AuxMatIndex = Other.AuxMatIndex;
-			this->MatIndex = Other.MatIndex;
-			this->SmoothingGroups = Other.SmoothingGroups;
-			this->WedgeIndex[0] = Other.WedgeIndex[0];
-			this->WedgeIndex[1] = Other.WedgeIndex[1];
-			this->WedgeIndex[2] = Other.WedgeIndex[2];
-			this->TangentX[0] = Other.TangentX[0];
-			this->TangentX[1] = Other.TangentX[1];
-			this->TangentX[2] = Other.TangentX[2];
-
-			this->TangentY[0] = Other.TangentY[0];
-			this->TangentY[1] = Other.TangentY[1];
-			this->TangentY[2] = Other.TangentY[2];
-
-			this->TangentZ[0] = Other.TangentZ[0];
-			this->TangentZ[1] = Other.TangentZ[1];
-			this->TangentZ[2] = Other.TangentZ[2];
-
-			return *this;
-		}
-
 		friend FArchive &operator<<(FArchive& Ar, FTriangle& F)
 		{
 			Ar.UsingCustomVersion(FEditorObjectVersion::GUID);
