Removed RenderCore's FTracingProfiler which is deprecated since UE 4.26. Please use Trace system and Unreal Insights instead.

diff --git a/Engine/Source/Runtime/Engine/Private/UnrealEngine.cpp b/Engine/Source/Runtime/Engine/Private/UnrealEngine.cpp
--- a/Engine/Source/Runtime/Engine/Private/UnrealEngine.cpp
+++ b/Engine/Source/Runtime/Engine/Private/UnrealEngine.cpp
@@ -208,7 +208,6 @@ UnrealEngine.cpp: Implements the UEngine class and helpers.
 #include "GeneralProjectSettings.h"
 #include "ProfilingDebugging/LoadTimeTracker.h"
 #include "AssetRegistry/AssetRegistryModule.h"
-#include "ProfilingDebugging/TracingProfiler.h"
 #include "Engine/CoreSettings.h"
 #include "IEyeTrackerModule.h"
 #include "Interfaces/IPluginManager.h"
@@ -12340,21 +12339,6 @@ void DrawStatsHUD( UWorld* World, FViewport* Viewport, FCanvas* Canvas, UCanvas*
 		}
 #endif // FRAMEPRO_ENABLED
 
-#if TRACING_PROFILER
-PRAGMA_DISABLE_DEPRECATION_WARNINGS
-		if (FTracingProfiler::Get()->IsCapturing())
-		{
-			SmallTextItem.SetColor(FLinearColor(0.0f, 1.0f, 0.0f, 1.0f));
-			SmallTextItem.Text = FText::Format(LOCTEXT("TracingProfilerFrameFmt", "TracingProfiler frame: {0}"), FTracingProfiler::Get()->GetCaptureFrameNumber());
-
-			MessageY += 250.0f;
-			Canvas->DrawItem(SmallTextItem, FVector2D(MessageX, MessageY));
-			MessageY += FontSizeY;
-
-		}
-PRAGMA_ENABLE_DEPRECATION_WARNINGS
-#endif
-
 #if !(UE_BUILD_TEST)
 		if (GEngine->bEnableOnScreenDebugMessagesDisplay && GEngine->bEnableOnScreenDebugMessages)
 		{
diff --git a/Engine/Source/Runtime/Launch/Private/LaunchEngineLoop.cpp b/Engine/Source/Runtime/Launch/Private/LaunchEngineLoop.cpp
--- a/Engine/Source/Runtime/Launch/Private/LaunchEngineLoop.cpp
+++ b/Engine/Source/Runtime/Launch/Private/LaunchEngineLoop.cpp
@@ -259,7 +259,6 @@ class FFeedbackContext;
 	#endif
 	#include "FramePro/FrameProProfiler.h"
 	#include "ProfilingDebugging/CsvProfiler.h"
-	#include "ProfilingDebugging/TracingProfiler.h"
 #endif
 
 #if defined(WITH_LAUNCHERCHECK) && WITH_LAUNCHERCHECK
@@ -2798,11 +2797,6 @@ int32 FEngineLoop::PreInitPreStartupScreen(const TCHAR* CmdLine)
 		GLog->EnableBacklog(false);
 	}
 
-#if WITH_ENGINE && TRACING_PROFILER
-PRAGMA_DISABLE_DEPRECATION_WARNINGS
-	FTracingProfiler::Get()->Init();
-PRAGMA_ENABLE_DEPRECATION_WARNINGS
-#endif
 #if WITH_ENGINE && FRAMEPRO_ENABLED
 	FFrameProProfiler::Initialize();
 #endif // FRAMEPRO_ENABLED
diff --git a/Engine/Source/Runtime/RenderCore/Private/ProfilingDebugging/RealtimeGPUProfiler.cpp b/Engine/Source/Runtime/RenderCore/Private/ProfilingDebugging/RealtimeGPUProfiler.cpp
--- a/Engine/Source/Runtime/RenderCore/Private/ProfilingDebugging/RealtimeGPUProfiler.cpp
+++ b/Engine/Source/Runtime/RenderCore/Private/ProfilingDebugging/RealtimeGPUProfiler.cpp
@@ -2,7 +2,6 @@
 
 #include "ProfilingDebugging/RealtimeGPUProfiler.h"
 #include "ProfilingDebugging/CsvProfiler.h"
-#include "ProfilingDebugging/TracingProfiler.h"
 #include "RHI.h"
 #include "RenderCore.h"
 #include "RenderingThread.h"
@@ -607,24 +606,6 @@ public:
 				}
 #endif
 			}
-
-#if TRACING_PROFILER
-PRAGMA_DISABLE_DEPRECATION_WARNINGS
-			const bool bTracingStatsEnabled = !!CVarGPUTracingStatsEnabled.GetValueOnRenderThread();
-			if (bTracingStatsEnabled)
-			{
-				for (uint32 GPUIndex : Event.GetGPUMask())
-				{
-					FTracingProfiler::Get()->AddGPUEvent(
-						Event.GetName(),
-						Event.GetStartResultMicroseconds(GPUIndex),
-						Event.GetEndResultMicroseconds(GPUIndex),
-						GPUIndex,
-						Event.GetFrameNumber());
-				}
-			}
-PRAGMA_ENABLE_DEPRECATION_WARNINGS
-#endif //TRACING_PROFILER
 		}
 
 #if STATS
diff --git a/Engine/Source/Runtime/RenderCore/Private/ProfilingDebugging/TracingProfiler.cpp b/Engine/Source/Runtime/RenderCore/Private/ProfilingDebugging/TracingProfiler.cpp
deleted file mode 100644
diff --git a/Engine/Source/Runtime/RenderCore/Public/ProfilingDebugging/TracingProfiler.h b/Engine/Source/Runtime/RenderCore/Public/ProfilingDebugging/TracingProfiler.h
deleted file mode 100644
