Added helper function to find all external saveable packages

Useful for when a user clicks 'save' on an asset and expects its 'contents' to also be saved

diff --git a/Engine/Source/Runtime/Engine/Private/ExternalPackageHelper.cpp b/Engine/Source/Runtime/Engine/Private/ExternalPackageHelper.cpp
--- a/Engine/Source/Runtime/Engine/Private/ExternalPackageHelper.cpp
+++ b/Engine/Source/Runtime/Engine/Private/ExternalPackageHelper.cpp
@@ -125,4 +125,23 @@ FString FExternalPackageHelper::GetExternalObjectPackageInstanceName(const FStri
 	return FLinkerInstancingContext::GetInstancedPackageName(OuterPackageName, ObjectPackageName);
 }
 
+void FExternalPackageHelper::GetExternalSaveableObjects(UObject* InOuter, TArray<UObject*>& OutObjects)
+{
+	// Get external packages
+	TSet<UPackage*> ExternalObjectPackages;
+	ExternalObjectPackages.Append(InOuter->GetPackage()->GetExternalPackages());
+
+	// Find assets for external packages
+	for (UPackage* ExternalPackage : ExternalObjectPackages)
+	{
+		if(FPackageName::IsValidLongPackageName(ExternalPackage->GetName()) && ExternalPackage->IsDirty())
+		{
+			if(UObject* Asset = ExternalPackage->FindAssetInPackage())
+			{
+				OutObjects.Add(Asset);
+			}
+		}
+	}
+}
+
 #endif
diff --git a/Engine/Source/Runtime/Engine/Public/ExternalPackageHelper.h b/Engine/Source/Runtime/Engine/Public/ExternalPackageHelper.h
--- a/Engine/Source/Runtime/Engine/Public/ExternalPackageHelper.h
+++ b/Engine/Source/Runtime/Engine/Public/ExternalPackageHelper.h
@@ -72,6 +72,13 @@ public:
 	template<typename T>
 	static void LoadObjectsFromExternalPackages(UObject* InOuter, TFunctionRef<void(T*)> Operation);
 
+	/**
+	 * Get the saveable external objects that should be saved alongside this outer's package
+	 * @param InOuter		The external object's outer
+	 * @param OutObjects	The objects that should be saved
+	 */
+	static void GetExternalSaveableObjects(UObject* InOuter, TArray<UObject*>& OutObjects);
+	
 private:
 	/** Get the external object package instance name. */
 	static FString GetExternalObjectPackageInstanceName(const FString& OuterPackageName, const FString& ObjectPackageName);
