Path Tracer: Fix rough diffuse not working when Substrate is enabled

Try to simplify logic around STRATA_DIFFUSE_CHAN being set from several places at once and remove logic which only looked at the rough diffuse setting on consoles

diff --git a/Engine/Shaders/Private/PathTracing/Material/PathTracingMaterialCommon.ush b/Engine/Shaders/Private/PathTracing/Material/PathTracingMaterialCommon.ush
--- a/Engine/Shaders/Private/PathTracing/Material/PathTracingMaterialCommon.ush
+++ b/Engine/Shaders/Private/PathTracing/Material/PathTracingMaterialCommon.ush
@@ -212,7 +212,8 @@ float ShadowTerminatorTerm(float3 L, float3 N, float3 Ns)
 
 float3 GetPathTracingDiffuseModel(float3 DiffuseColor, float Roughness, float NoV, float NoL, float VoH, float NoH)
 {
-#if MATERIAL_ROUGHDIFFUSE
+	// TODO: Always use MATERIAL_ROUGHDIFFUSE instead of have two different defines for Substrate vs. Legacy?
+#if (!STRATA_ENABLED && MATERIAL_ROUGHDIFFUSE) || (STRATA_ENABLED && STRATA_DIFFUSE_CHAN)
 	// Clamp to [0,1] as the model is ever so slightly non-energy conserving at grazing angles which can cause unwanted noise
 	return saturate(PI * Diffuse_Chan(DiffuseColor, Pow4(Roughness), NoV, NoL, VoH, NoH, 1.0f));
 #else
diff --git a/Engine/Shaders/Private/Strata/StrataEvaluation.ush b/Engine/Shaders/Private/Strata/StrataEvaluation.ush
--- a/Engine/Shaders/Private/Strata/StrataEvaluation.ush
+++ b/Engine/Shaders/Private/Strata/StrataEvaluation.ush
@@ -28,13 +28,6 @@
 #define STRATA_TRANSLUCENT_ENABLED 0
 #endif
 
-#ifndef STRATA_DIFFUSE_CHAN
-	#if PROJECT_STRATA_BACKCOMPATIBILITY > 0
-		#define STRATA_DIFFUSE_CHAN 0
-	#else
-		#define STRATA_DIFFUSE_CHAN 1
-	#endif
-#endif
 
 #ifndef STRATA_SSS_TRANSMISSION
 #define STRATA_SSS_TRANSMISSION 0
diff --git a/Engine/Source/Runtime/Engine/Private/ShaderCompiler/ShaderCompiler.cpp b/Engine/Source/Runtime/Engine/Private/ShaderCompiler/ShaderCompiler.cpp
--- a/Engine/Source/Runtime/Engine/Private/ShaderCompiler/ShaderCompiler.cpp
+++ b/Engine/Source/Runtime/Engine/Private/ShaderCompiler/ShaderCompiler.cpp
@@ -6571,16 +6571,11 @@ void GlobalBeginCompileShader(
 		const bool bStrata = Strata::IsStrataEnabled();
 		Input.Environment.SetDefine(TEXT("STRATA_ENABLED"), bStrata ? 1 : 0);
 
-		// Force rough diffuse to be disable on platform which explicitly disable it from their settings
-		if (bStrata && IsConsolePlatform(Target.GetPlatform()))
+		if (bStrata)
 		{
-			const bool bStrataRoughDiffuse = Strata::IsRoughDiffuseEnabled();
+			const bool bStrataRoughDiffuse = Strata::IsRoughDiffuseEnabled() && !Strata::IsBackCompatibilityEnabled();
 			Input.Environment.SetDefine(TEXT("STRATA_DIFFUSE_CHAN"), bStrataRoughDiffuse ? 1 : 0);
-		}
 
-		// Force rough diffuse to be disable on platform which explicitly disable it from their settings
-		if (bStrata)
-		{
 			const uint32 StrataShadingQuality = Strata::GetShadingQuality(Target.GetPlatform());
 			Input.Environment.SetDefine(TEXT("STRATA_SHADING_QUALITY"), StrataShadingQuality);
 
