Deprecate Modulator Parameter (non-set) on Modulation Destination Settings

diff --git a/Engine/Source/Runtime/Engine/Classes/Sound/SoundModulationDestination.h b/Engine/Source/Runtime/Engine/Classes/Sound/SoundModulationDestination.h
--- a/Engine/Source/Runtime/Engine/Classes/Sound/SoundModulationDestination.h
+++ b/Engine/Source/Runtime/Engine/Classes/Sound/SoundModulationDestination.h
@@ -55,11 +55,11 @@ struct ENGINE_API FSoundModulationDestinationSettings
 	/** Whether or not to enable modulation */
 	UPROPERTY(EditAnywhere, Category = Modulation, meta = (DisplayName = "Modulate"))
 	bool bEnableModulation = false;
-#endif // WITH_EDITORONLY_DATA
 
-	/** Deprecated in favor of Modulators set. */
-	UPROPERTY(EditAnywhere, BlueprintReadWrite, Category = Modulation)
+	UE_DEPRECATED(5.3, "Use the Modulators set property instead.")
+	UPROPERTY(meta = (DeprecatedProperty))
 	TObjectPtr<USoundModulatorBase> Modulator = nullptr;
+#endif // WITH_EDITORONLY_DATA
 
 	/** Set of modulation sources, which provides values to mix with base value. */
 	UPROPERTY(EditAnywhere, BlueprintReadWrite, Category = Modulation)
@@ -69,6 +69,17 @@ struct ENGINE_API FSoundModulationDestinationSettings
 	/** Versioning utility function to upgrade single modulator field to set of modulators */
 	void VersionModulators();
 #endif // WITH_EDITORONLY_DATA
+
+	// We need to explicitly disable warnings on these constructors/operators for clang to be happy with deprecated variables
+	PRAGMA_DISABLE_DEPRECATION_WARNINGS
+	FSoundModulationDestinationSettings() = default;
+	~FSoundModulationDestinationSettings() = default;
+	FSoundModulationDestinationSettings(const FSoundModulationDestinationSettings&) = default;
+	FSoundModulationDestinationSettings(FSoundModulationDestinationSettings&&) = default;
+	FSoundModulationDestinationSettings& operator=(const FSoundModulationDestinationSettings&) = default;
+	FSoundModulationDestinationSettings& operator=(FSoundModulationDestinationSettings&&) = default;
+	PRAGMA_ENABLE_DEPRECATION_WARNINGS
+
 };
 
 /** Default parameter destination settings for source audio object. */
