Removing many vestigates of old audio engine code.

diff --git a/Engine/Config/Android/AndroidEngine.ini b/Engine/Config/Android/AndroidEngine.ini
--- a/Engine/Config/Android/AndroidEngine.ini
+++ b/Engine/Config/Android/AndroidEngine.ini
@@ -13,7 +13,6 @@ PakSigningRequired=False
 gc.MaxObjectsInGame=131072
 
 [Audio]
-AudioDeviceModuleName=AndroidAudio
 PlatformFormat=OGG
 PlatformStreamingFormat=OGG
 
@@ -22,7 +21,6 @@ AudioMixerModuleName=AudioMixerAndroid
 
 ; Defines a platform-specific volume headroom (in dB) for audio to provide better platform consistency with respect to volume levels.
 PlatformHeadroomDB=-6
-UseAudioMixer=true
 
 [PlatformMemoryBuckets]
 LargestMemoryBucket_MinGB=8
diff --git a/Engine/Config/HoloLens/HoloLensEngine.ini b/Engine/Config/HoloLens/HoloLensEngine.ini
--- a/Engine/Config/HoloLens/HoloLensEngine.ini
+++ b/Engine/Config/HoloLens/HoloLensEngine.ini
@@ -2,7 +2,7 @@
 AudioDeviceModuleName=XAudio2
 ; Defining below allows switching to audio mixer using -audiomixer commandline
 AudioMixerModuleName=AudioMixerXAudio2
-UseAudioMixer=True
+
 ; Defines a platform-specific volume headroom (in dB) for audio to provide better platform consistency with respect to volume levels.
 PlatformHeadroomDB=0
 PlatformFormat=OGG
diff --git a/Engine/Config/IOS/IOSEngine.ini b/Engine/Config/IOS/IOSEngine.ini
--- a/Engine/Config/IOS/IOSEngine.ini
+++ b/Engine/Config/IOS/IOSEngine.ini
@@ -13,15 +13,10 @@ PlatformRequiresDataCrypto=True
 PakSigningRequired=False
 
 [Audio]
-AudioDeviceModuleName=IOSAudio
-;Uncomment below and comment out above line to enable new audio mixer
-;AudioDeviceModuleName=AudioMixerAudioUnit
-
 ; Defining below allows switching to audio mixer using -audiomixer commandline
 AudioMixerModuleName=AudioMixerAudioUnit
 
 ; Defines a platform-specific volume headroom (in dB) for audio to provide better platform consistency with respect to volume levels.
-UseAudioMixer=true
 PlatformHeadroomDB=-6
 
 PlatformFormat=ADPCM
diff --git a/Engine/Config/Mac/MacEngine.ini b/Engine/Config/Mac/MacEngine.ini
--- a/Engine/Config/Mac/MacEngine.ini
+++ b/Engine/Config/Mac/MacEngine.ini
@@ -1,6 +1,4 @@
 [Audio]
-AudioDeviceModuleName=CoreAudio
-
 ; Defining below allows switching to audio mixer using -audiomixer commandline
 AudioMixerModuleName=AudioMixerCoreAudio
 PlatformFormat=OGG
@@ -8,7 +6,6 @@ PlatformStreamingFormat=OGG
 
 ; Defines a platform-specific volume headroom (in dB) for audio to provide better platform consistency with respect to volume levels.
 PlatformHeadroomDB=-6
-UseAudioMixer=true
 
 [SystemSettings]
 r.setres=1280x720
diff --git a/Engine/Config/Unix/UnixEngine.ini b/Engine/Config/Unix/UnixEngine.ini
--- a/Engine/Config/Unix/UnixEngine.ini
+++ b/Engine/Config/Unix/UnixEngine.ini
@@ -1,5 +1,4 @@
 [Audio]
-AudioDeviceModuleName=AudioMixerSDL
 AudioMixerModuleName=AudioMixerSDL
 
 [PlatformCrypto]
diff --git a/Engine/Config/Windows/WindowsEngine.ini b/Engine/Config/Windows/WindowsEngine.ini
--- a/Engine/Config/Windows/WindowsEngine.ini
+++ b/Engine/Config/Windows/WindowsEngine.ini
@@ -1,5 +1,4 @@
 [Audio]
-AudioDeviceModuleName=XAudio2
 ; Defining below allows switching to audio mixer using -audiomixer commandline
 AudioMixerModuleName=AudioMixerXAudio2
 PlatformFormat=OGG
@@ -7,7 +6,6 @@ PlatformStreamingFormat=OGG
 
 ; Defines a platform-specific volume headroom (in dB) for audio to provide better platform consistency with respect to volume levels.
 PlatformHeadroomDB=-3
-UseAudioMixer=true
 
 [TextureStreaming]
 ; PoolSizeVRAMPercentage is how much percentage of GPU Dedicated VRAM should be used as a TexturePool cache for streaming textures (0 - unlimited streaming)
diff --git a/Engine/Plugins/AudioGameplayVolume/Source/AudioGameplayVolume/Private/SubmixOverrideVolumeComponent.cpp b/Engine/Plugins/AudioGameplayVolume/Source/AudioGameplayVolume/Private/SubmixOverrideVolumeComponent.cpp
--- a/Engine/Plugins/AudioGameplayVolume/Source/AudioGameplayVolume/Private/SubmixOverrideVolumeComponent.cpp
+++ b/Engine/Plugins/AudioGameplayVolume/Source/AudioGameplayVolume/Private/SubmixOverrideVolumeComponent.cpp
@@ -20,7 +20,7 @@ void FProxyMutator_SubmixOverride::Apply(FAudioGameplayVolumeListener& Listener)
 	check(IsInAudioThread());
 
 	FAudioDeviceHandle AudioDeviceHandle = FAudioDeviceManager::Get()->GetAudioDevice(Listener.GetOwningDeviceId());
-	if (AudioDeviceHandle.IsValid() && AudioDeviceHandle->IsAudioMixerEnabled())
+	if (AudioDeviceHandle.IsValid())
 	{
 		FSoundEffectSubmixInitData InitData;
 		InitData.DeviceID = AudioDeviceHandle.GetDeviceID();
@@ -57,7 +57,7 @@ void FProxyMutator_SubmixOverride::Remove(FAudioGameplayVolumeListener& Listener
 	check(IsInAudioThread());
 
 	FAudioDeviceHandle AudioDeviceHandle = FAudioDeviceManager::Get()->GetAudioDevice(Listener.GetOwningDeviceId());
-	if (AudioDeviceHandle.IsValid() && AudioDeviceHandle->IsAudioMixerEnabled())
+	if (AudioDeviceHandle.IsValid())
 	{
 		// Clear out any previous submix effect chain overrides
 		for (const FAudioVolumeSubmixOverrideSettings& OverrideSettings : SubmixOverrideSettings)
diff --git a/Engine/Plugins/Media/BinkMedia/Source/BinkMediaPlayer/Private/BinkMediaPlayerPrivate.h b/Engine/Plugins/Media/BinkMedia/Source/BinkMediaPlayer/Private/BinkMediaPlayerPrivate.h
--- a/Engine/Plugins/Media/BinkMedia/Source/BinkMediaPlayer/Private/BinkMediaPlayerPrivate.h
+++ b/Engine/Plugins/Media/BinkMedia/Source/BinkMediaPlayer/Private/BinkMediaPlayerPrivate.h
@@ -33,7 +33,7 @@ static int GetNumSpeakers()
 		return 2;
 	}
 	FAudioDevice *dev = FAudioDevice::GetMainAudioDevice().GetAudioDevice();
-	if (dev && dev->IsAudioMixerEnabled()) {
+	if (dev) {
 		Audio::FMixerDevice *mix = static_cast<Audio::FMixerDevice*>(dev);
 		if (mix) {
 			return mix->GetNumDeviceChannels();
diff --git a/Engine/Plugins/Media/MediaIOFramework/Source/MediaIOCore/Private/MediaIOCoreAudioOutput.cpp b/Engine/Plugins/Media/MediaIOFramework/Source/MediaIOCore/Private/MediaIOCoreAudioOutput.cpp
--- a/Engine/Plugins/Media/MediaIOFramework/Source/MediaIOCore/Private/MediaIOCoreAudioOutput.cpp
+++ b/Engine/Plugins/Media/MediaIOFramework/Source/MediaIOCore/Private/MediaIOCoreAudioOutput.cpp
@@ -127,7 +127,7 @@ TSharedPtr<FMediaIOAudioOutput> FMediaIOAudioCapture::CreateAudioOutput(int32 In
 
 void FMediaIOAudioCapture::RegisterBufferListener(FAudioDevice* AudioDevice)
 {
-	if (AudioDevice && AudioDevice->IsAudioMixerEnabled())
+	if (AudioDevice)
 	{
 		const Audio::FMixerDevice* MixerDevice = static_cast<Audio::FMixerDevice*>(AudioDevice);
 		NumChannels = MixerDevice->GetDeviceOutputChannels();
diff --git a/Engine/Plugins/MovieScene/MovieRenderPipeline/Source/MovieRenderPipelineCore/Private/MoviePipelineAudioRendering.cpp b/Engine/Plugins/MovieScene/MovieRenderPipeline/Source/MovieRenderPipelineCore/Private/MoviePipelineAudioRendering.cpp
--- a/Engine/Plugins/MovieScene/MovieRenderPipeline/Source/MovieRenderPipelineCore/Private/MoviePipelineAudioRendering.cpp
+++ b/Engine/Plugins/MovieScene/MovieRenderPipeline/Source/MovieRenderPipelineCore/Private/MoviePipelineAudioRendering.cpp
@@ -32,14 +32,7 @@ static Audio::FMixerDevice* GetAudioMixerDeviceFromWorldContext(const UObject* W
 {
 	if (FAudioDevice* AudioDevice = GetAudioDeviceFromWorldContext(WorldContextObject))
 	{
-		if (!AudioDevice->IsAudioMixerEnabled())
-		{
-			return nullptr;
-		}
-		else
-		{
-			return static_cast<Audio::FMixerDevice*>(AudioDevice);
-		}
+		return static_cast<Audio::FMixerDevice*>(AudioDevice);
 	}
 	return nullptr;
 }
diff --git a/Engine/Plugins/MovieScene/MovieRenderPipeline/Source/MovieRenderPipelineRenderPasses/Private/MoviePipelineWaveOutput.cpp b/Engine/Plugins/MovieScene/MovieRenderPipeline/Source/MovieRenderPipelineRenderPasses/Private/MoviePipelineWaveOutput.cpp
--- a/Engine/Plugins/MovieScene/MovieRenderPipeline/Source/MovieRenderPipelineRenderPasses/Private/MoviePipelineWaveOutput.cpp
+++ b/Engine/Plugins/MovieScene/MovieRenderPipeline/Source/MovieRenderPipelineRenderPasses/Private/MoviePipelineWaveOutput.cpp
@@ -40,14 +40,7 @@ static Audio::FMixerDevice* GetAudioMixerDeviceFromWorldContext(const UObject* W
 {
 	if (FAudioDevice* AudioDevice = GetAudioDeviceFromWorldContext(WorldContextObject))
 	{
-		if (!AudioDevice->IsAudioMixerEnabled())
-		{
-			return nullptr;
-		}
-		else
-		{
-			return static_cast<Audio::FMixerDevice*>(AudioDevice);
-		}
+		return static_cast<Audio::FMixerDevice*>(AudioDevice);
 	}
 	return nullptr;
 }
diff --git a/Engine/Plugins/Online/OnlineSubsystemUtils/Source/OnlineSubsystemUtils/Private/OnlineSubsystemUtils.cpp b/Engine/Plugins/Online/OnlineSubsystemUtils/Source/OnlineSubsystemUtils/Private/OnlineSubsystemUtils.cpp
--- a/Engine/Plugins/Online/OnlineSubsystemUtils/Source/OnlineSubsystemUtils/Private/OnlineSubsystemUtils.cpp
+++ b/Engine/Plugins/Online/OnlineSubsystemUtils/Source/OnlineSubsystemUtils/Private/OnlineSubsystemUtils.cpp
@@ -120,17 +120,7 @@ UAudioComponent* CreateVoiceAudioComponent(uint32 SampleRate, int32 NumChannels)
 			SoundStreaming->SoundGroup = SOUNDGROUP_Voice;
 			SoundStreaming->bLooping = false;
 
-			// Turn off async generation in old audio engine on mac.
-			#if PLATFORM_MAC
-			if (!AudioDevice->IsAudioMixerEnabled())
-			{
-				SoundStreaming->bCanProcessAsync = false;
-			}
-			else
-			#endif // #if PLATFORM_MAC
-			{
-				SoundStreaming->bCanProcessAsync = true;
-			}
+			SoundStreaming->bCanProcessAsync = true;
 
 			AudioComponent = AudioDevice->CreateComponent(SoundStreaming);
 			if (AudioComponent)
diff --git a/Engine/Plugins/Runtime/Metasound/Source/MetasoundEngineTest/Private/EngineTestAutomatedNodeTests.cpp b/Engine/Plugins/Runtime/Metasound/Source/MetasoundEngineTest/Private/EngineTestAutomatedNodeTests.cpp
--- a/Engine/Plugins/Runtime/Metasound/Source/MetasoundEngineTest/Private/EngineTestAutomatedNodeTests.cpp
+++ b/Engine/Plugins/Runtime/Metasound/Source/MetasoundEngineTest/Private/EngineTestAutomatedNodeTests.cpp
@@ -32,10 +32,7 @@ namespace Metasound::MetasoundEngineTestPrivate {
 		{
 			if (FAudioDevice* AudioDevice = DeviceManager->GetMainAudioDeviceRaw())
 			{
-				if (AudioDevice->IsAudioMixerEnabled())
-				{
-					return static_cast<Audio::FMixerDevice*>(AudioDevice);
-				}
+				return static_cast<Audio::FMixerDevice*>(AudioDevice);
 			}
 		}
 		return nullptr;
diff --git a/Engine/Plugins/Runtime/Oculus/OculusAudio/Source/OculusAudio/Private/IOculusAudioPlugin.cpp b/Engine/Plugins/Runtime/Oculus/OculusAudio/Source/OculusAudio/Private/IOculusAudioPlugin.cpp
--- a/Engine/Plugins/Runtime/Oculus/OculusAudio/Source/OculusAudio/Private/IOculusAudioPlugin.cpp
+++ b/Engine/Plugins/Runtime/Oculus/OculusAudio/Source/OculusAudio/Private/IOculusAudioPlugin.cpp
@@ -19,12 +19,6 @@ void FOculusAudioPlugin::ShutdownModule()
 
 void FOculusAudioPlugin::RegisterAudioDevice(FAudioDevice* AudioDeviceHandle)
 {
-	// Inject the Context into the spatailizer (and reverb) if they're enabled
-	if (!AudioDeviceHandle->IsAudioMixerEnabled())
-	{
-		return; // Not supported in old audio engine
-	}
-
 	if (!RegisteredAudioDevices.Contains(AudioDeviceHandle))
 	{
 		TAudioPluginListenerPtr ContextManager = TAudioPluginListenerPtr(new FOculusAudioContextManager());
diff --git a/Engine/Source/Editor/AudioEditor/Private/AudioEditorModule.cpp b/Engine/Source/Editor/AudioEditor/Private/AudioEditorModule.cpp
--- a/Engine/Source/Editor/AudioEditor/Private/AudioEditorModule.cpp
+++ b/Engine/Source/Editor/AudioEditor/Private/AudioEditorModule.cpp
@@ -184,52 +184,44 @@ public:
 
 	virtual void RegisterAudioMixerAssetActions() override
 	{
-		// Only register asset actions for when audio mixer data is enabled
-		if (GetDefault<UAudioSettings>()->IsAudioMixerEnabled())
-		{
-			IAssetTools& AssetTools = FModuleManager::LoadModuleChecked<FAssetToolsModule>("AssetTools").Get();
-			AssetTools.RegisterAssetTypeActions(MakeShared<FAssetTypeActions_SoundSubmix>());
-			AssetTools.RegisterAssetTypeActions(MakeShared<FAssetTypeActions_SoundfieldSubmix>());
-			AssetTools.RegisterAssetTypeActions(MakeShared<FAssetTypeActions_EndpointSubmix>());
-			AssetTools.RegisterAssetTypeActions(MakeShared<FAssetTypeActions_SoundfieldEndpointSubmix>());
-			AssetTools.RegisterAssetTypeActions(MakeShared<FAssetTypeActions_SoundfieldEncodingSettings>());
-			AssetTools.RegisterAssetTypeActions(MakeShared<FAssetTypeActions_SoundfieldEffectSettings>());
-			AssetTools.RegisterAssetTypeActions(MakeShared<FAssetTypeActions_SoundfieldEffect>());
-			AssetTools.RegisterAssetTypeActions(MakeShared<FAssetTypeActions_AudioEndpointSettings>());
-			AssetTools.RegisterAssetTypeActions(MakeShared<FAssetTypeActions_SoundfieldEndpointSettings>());
-			AssetTools.RegisterAssetTypeActions(MakeShared<FAssetTypeActions_SoundEffectSubmixPreset>());
-			AssetTools.RegisterAssetTypeActions(MakeShared<FAssetTypeActions_SoundEffectSourcePreset>());
-			AssetTools.RegisterAssetTypeActions(MakeShared<FAssetTypeActions_SoundEffectSourcePresetChain>());
-		}
+		IAssetTools& AssetTools = FModuleManager::LoadModuleChecked<FAssetToolsModule>("AssetTools").Get();
+		AssetTools.RegisterAssetTypeActions(MakeShared<FAssetTypeActions_SoundSubmix>());
+		AssetTools.RegisterAssetTypeActions(MakeShared<FAssetTypeActions_SoundfieldSubmix>());
+		AssetTools.RegisterAssetTypeActions(MakeShared<FAssetTypeActions_EndpointSubmix>());
+		AssetTools.RegisterAssetTypeActions(MakeShared<FAssetTypeActions_SoundfieldEndpointSubmix>());
+		AssetTools.RegisterAssetTypeActions(MakeShared<FAssetTypeActions_SoundfieldEncodingSettings>());
+		AssetTools.RegisterAssetTypeActions(MakeShared<FAssetTypeActions_SoundfieldEffectSettings>());
+		AssetTools.RegisterAssetTypeActions(MakeShared<FAssetTypeActions_SoundfieldEffect>());
+		AssetTools.RegisterAssetTypeActions(MakeShared<FAssetTypeActions_AudioEndpointSettings>());
+		AssetTools.RegisterAssetTypeActions(MakeShared<FAssetTypeActions_SoundfieldEndpointSettings>());
+		AssetTools.RegisterAssetTypeActions(MakeShared<FAssetTypeActions_SoundEffectSubmixPreset>());
+		AssetTools.RegisterAssetTypeActions(MakeShared<FAssetTypeActions_SoundEffectSourcePreset>());
+		AssetTools.RegisterAssetTypeActions(MakeShared<FAssetTypeActions_SoundEffectSourcePresetChain>());
 	}
 
 	virtual void RegisterEffectPresetAssetActions() override
 	{
-		// Only register asset actions for the case where audio mixer data is enabled
-		if (GetDefault<UAudioSettings>()->IsAudioMixerEnabled())
-		{
-			// Register the audio editor asset type actions
-			IAssetTools& AssetTools = FModuleManager::LoadModuleChecked<FAssetToolsModule>("AssetTools").Get();
+		// Register the audio editor asset type actions
+		IAssetTools& AssetTools = FModuleManager::LoadModuleChecked<FAssetToolsModule>("AssetTools").Get();
 
-			// Look for any sound effect presets to register
-			for (TObjectIterator<UClass> It; It; ++It)
+		// Look for any sound effect presets to register
+		for (TObjectIterator<UClass> It; It; ++It)
+		{
+			UClass* ChildClass = *It;
+			if (ChildClass->HasAnyClassFlags(CLASS_Abstract))
 			{
-				UClass* ChildClass = *It;
-				if (ChildClass->HasAnyClassFlags(CLASS_Abstract))
-				{
-					continue;
-				}
+				continue;
+			}
 
-				// Look for submix or source preset classes
-				UClass* ParentClass = ChildClass->GetSuperClass();
-				if (ParentClass && (ParentClass->IsChildOf(USoundEffectSourcePreset::StaticClass()) || ParentClass->IsChildOf(USoundEffectSubmixPreset::StaticClass())))
+			// Look for submix or source preset classes
+			UClass* ParentClass = ChildClass->GetSuperClass();
+			if (ParentClass && (ParentClass->IsChildOf(USoundEffectSourcePreset::StaticClass()) || ParentClass->IsChildOf(USoundEffectSubmixPreset::StaticClass())))
+			{
+				USoundEffectPreset* EffectPreset = ChildClass->GetDefaultObject<USoundEffectPreset>();
+				if (!RegisteredActions.Contains(EffectPreset) && EffectPreset->HasAssetActions())
 				{
-					USoundEffectPreset* EffectPreset = ChildClass->GetDefaultObject<USoundEffectPreset>();
-					if (!RegisteredActions.Contains(EffectPreset) && EffectPreset->HasAssetActions())
-					{
-						RegisteredActions.Add(EffectPreset);
-						AssetTools.RegisterAssetTypeActions(MakeShared<FAssetTypeActions_SoundEffectPreset>(EffectPreset));
-					}
+					RegisteredActions.Add(EffectPreset);
+					AssetTools.RegisterAssetTypeActions(MakeShared<FAssetTypeActions_SoundEffectPreset>(EffectPreset));
 				}
 			}
 		}
diff --git a/Engine/Source/Editor/DetailCustomizations/Private/AttenuationSettingsCustomizations.cpp b/Engine/Source/Editor/DetailCustomizations/Private/AttenuationSettingsCustomizations.cpp
--- a/Engine/Source/Editor/DetailCustomizations/Private/AttenuationSettingsCustomizations.cpp
+++ b/Engine/Source/Editor/DetailCustomizations/Private/AttenuationSettingsCustomizations.cpp
@@ -458,8 +458,6 @@ void FSoundAttenuationSettingsCustomization::CustomizeChildren(TSharedRef<IPrope
 	LayoutBuilder.AddPropertyToCategory(bIsSpatializedHandle)
 		.Visibility(IsAttenuationOverriddenVisibleAttribute());
 
-	bool bIsAudioMixerEnabled = GetDefault<UAudioSettings>()->IsAudioMixerEnabled();
-
 	// Check to see if a spatialization plugin is enabled
 	LayoutBuilder.AddPropertyToCategory(PropertyHandles.FindChecked(GET_MEMBER_NAME_CHECKED(FSoundAttenuationSettings, SpatializationAlgorithm)))
 		.Visibility(IsAttenuationOverriddenVisibleAttribute())
@@ -513,16 +511,13 @@ void FSoundAttenuationSettingsCustomization::CustomizeChildren(TSharedRef<IPrope
 		.Visibility(IsAttenuationOverriddenVisibleAttribute())
 		.EditCondition(GetIsAirAbsorptionEnabledAttribute(), nullptr);
 
-	if (bIsAudioMixerEnabled)
-	{
-		LayoutBuilder.AddPropertyToCategory(PropertyHandles.FindChecked(GET_MEMBER_NAME_CHECKED(FSoundAttenuationSettings, HPFFrequencyAtMin)))
-			.Visibility(IsAttenuationOverriddenVisibleAttribute())
-			.EditCondition(GetIsAirAbsorptionEnabledAttribute(), nullptr);
+	LayoutBuilder.AddPropertyToCategory(PropertyHandles.FindChecked(GET_MEMBER_NAME_CHECKED(FSoundAttenuationSettings, HPFFrequencyAtMin)))
+		.Visibility(IsAttenuationOverriddenVisibleAttribute())
+		.EditCondition(GetIsAirAbsorptionEnabledAttribute(), nullptr);
 
-		LayoutBuilder.AddPropertyToCategory(PropertyHandles.FindChecked(GET_MEMBER_NAME_CHECKED(FSoundAttenuationSettings, HPFFrequencyAtMax)))
-			.Visibility(IsAttenuationOverriddenVisibleAttribute())
-			.EditCondition(GetIsAirAbsorptionEnabledAttribute(), nullptr);
-	}
+	LayoutBuilder.AddPropertyToCategory(PropertyHandles.FindChecked(GET_MEMBER_NAME_CHECKED(FSoundAttenuationSettings, HPFFrequencyAtMax)))
+		.Visibility(IsAttenuationOverriddenVisibleAttribute())
+		.EditCondition(GetIsAirAbsorptionEnabledAttribute(), nullptr);
 
 	LayoutBuilder.AddPropertyToCategory(PropertyHandles.FindChecked(GET_MEMBER_NAME_CHECKED(FSoundAttenuationSettings, bEnableLogFrequencyScaling)))
 		.Visibility(IsAttenuationOverriddenVisibleAttribute())
@@ -536,45 +531,42 @@ void FSoundAttenuationSettingsCustomization::CustomizeChildren(TSharedRef<IPrope
 		.Visibility(TAttribute<EVisibility>(this, &FSoundAttenuationSettingsCustomization::IsCustomAirAbsorptionCurveSelected))
 		.EditCondition(GetIsAirAbsorptionEnabledAttribute(), nullptr);
 
-	if (bIsAudioMixerEnabled)
-	{
-		LayoutBuilder.AddPropertyToCategory(PropertyHandles.FindChecked(GET_MEMBER_NAME_CHECKED(FSoundAttenuationSettings, CustomHighpassAirAbsorptionCurve)))
-			.Visibility(TAttribute<EVisibility>(this, &FSoundAttenuationSettingsCustomization::IsCustomAirAbsorptionCurveSelected))
-			.EditCondition(GetIsAirAbsorptionEnabledAttribute(), nullptr);
+	LayoutBuilder.AddPropertyToCategory(PropertyHandles.FindChecked(GET_MEMBER_NAME_CHECKED(FSoundAttenuationSettings, CustomHighpassAirAbsorptionCurve)))
+		.Visibility(TAttribute<EVisibility>(this, &FSoundAttenuationSettingsCustomization::IsCustomAirAbsorptionCurveSelected))
+		.EditCondition(GetIsAirAbsorptionEnabledAttribute(), nullptr);
 
-		// Add the reverb send enabled handle
-		LayoutBuilder.AddPropertyToCategory(bIsReverbSendEnabledHandle)
-			.Visibility(IsAttenuationOverriddenVisibleAttribute())
-			.EditCondition(IsAttenuationOverriddenAttribute(), nullptr);
+	// Add the reverb send enabled handle
+	LayoutBuilder.AddPropertyToCategory(bIsReverbSendEnabledHandle)
+		.Visibility(IsAttenuationOverriddenVisibleAttribute())
+		.EditCondition(IsAttenuationOverriddenAttribute(), nullptr);
 
-		LayoutBuilder.AddPropertyToCategory(ReverbSendMethodHandle)
-			.Visibility(IsAttenuationOverriddenVisibleAttribute())
-			.EditCondition(GetIsReverbSendEnabledAttribute(), nullptr);
+	LayoutBuilder.AddPropertyToCategory(ReverbSendMethodHandle)
+		.Visibility(IsAttenuationOverriddenVisibleAttribute())
+		.EditCondition(GetIsReverbSendEnabledAttribute(), nullptr);
 
-		LayoutBuilder.AddPropertyToCategory(PropertyHandles.FindChecked(GET_MEMBER_NAME_CHECKED(FSoundAttenuationSettings, ReverbWetLevelMin)))
-			.Visibility(TAttribute<EVisibility>(this, &FSoundAttenuationSettingsCustomization::IsLinearMethodSelected))
-			.EditCondition(GetIsReverbSendEnabledAttribute(), nullptr);
+	LayoutBuilder.AddPropertyToCategory(PropertyHandles.FindChecked(GET_MEMBER_NAME_CHECKED(FSoundAttenuationSettings, ReverbWetLevelMin)))
+		.Visibility(TAttribute<EVisibility>(this, &FSoundAttenuationSettingsCustomization::IsLinearMethodSelected))
+		.EditCondition(GetIsReverbSendEnabledAttribute(), nullptr);
 
-		LayoutBuilder.AddPropertyToCategory(PropertyHandles.FindChecked(GET_MEMBER_NAME_CHECKED(FSoundAttenuationSettings, ReverbWetLevelMax)))
-			.Visibility(TAttribute<EVisibility>(this, &FSoundAttenuationSettingsCustomization::IsLinearMethodSelected))
-			.EditCondition(GetIsReverbSendEnabledAttribute(), nullptr);
+	LayoutBuilder.AddPropertyToCategory(PropertyHandles.FindChecked(GET_MEMBER_NAME_CHECKED(FSoundAttenuationSettings, ReverbWetLevelMax)))
+		.Visibility(TAttribute<EVisibility>(this, &FSoundAttenuationSettingsCustomization::IsLinearMethodSelected))
+		.EditCondition(GetIsReverbSendEnabledAttribute(), nullptr);
 
-		LayoutBuilder.AddPropertyToCategory(PropertyHandles.FindChecked(GET_MEMBER_NAME_CHECKED(FSoundAttenuationSettings, CustomReverbSendCurve)))
-			.Visibility(TAttribute<EVisibility>(this, &FSoundAttenuationSettingsCustomization::IsCustomReverbSendCurveSelected))
-			.EditCondition(GetIsReverbSendEnabledAttribute(), nullptr);
+	LayoutBuilder.AddPropertyToCategory(PropertyHandles.FindChecked(GET_MEMBER_NAME_CHECKED(FSoundAttenuationSettings, CustomReverbSendCurve)))
+		.Visibility(TAttribute<EVisibility>(this, &FSoundAttenuationSettingsCustomization::IsCustomReverbSendCurveSelected))
+		.EditCondition(GetIsReverbSendEnabledAttribute(), nullptr);
 
-		LayoutBuilder.AddPropertyToCategory(PropertyHandles.FindChecked(GET_MEMBER_NAME_CHECKED(FSoundAttenuationSettings, ReverbDistanceMin)))
-			.Visibility(TAttribute<EVisibility>(this, &FSoundAttenuationSettingsCustomization::IsLinearOrCustomReverbMethodSelected))
-			.EditCondition(GetIsReverbSendEnabledAttribute(), nullptr);
+	LayoutBuilder.AddPropertyToCategory(PropertyHandles.FindChecked(GET_MEMBER_NAME_CHECKED(FSoundAttenuationSettings, ReverbDistanceMin)))
+		.Visibility(TAttribute<EVisibility>(this, &FSoundAttenuationSettingsCustomization::IsLinearOrCustomReverbMethodSelected))
+		.EditCondition(GetIsReverbSendEnabledAttribute(), nullptr);
 
-		LayoutBuilder.AddPropertyToCategory(PropertyHandles.FindChecked(GET_MEMBER_NAME_CHECKED(FSoundAttenuationSettings, ReverbDistanceMax)))
-			.Visibility(TAttribute<EVisibility>(this, &FSoundAttenuationSettingsCustomization::IsLinearOrCustomReverbMethodSelected))
-			.EditCondition(GetIsReverbSendEnabledAttribute(), nullptr);
+	LayoutBuilder.AddPropertyToCategory(PropertyHandles.FindChecked(GET_MEMBER_NAME_CHECKED(FSoundAttenuationSettings, ReverbDistanceMax)))
+		.Visibility(TAttribute<EVisibility>(this, &FSoundAttenuationSettingsCustomization::IsLinearOrCustomReverbMethodSelected))
+		.EditCondition(GetIsReverbSendEnabledAttribute(), nullptr);
 
-		LayoutBuilder.AddPropertyToCategory(PropertyHandles.FindChecked(GET_MEMBER_NAME_CHECKED(FSoundAttenuationSettings, ManualReverbSendLevel)))
-			.Visibility(TAttribute<EVisibility>(this, &FSoundAttenuationSettingsCustomization::IsManualReverbSendSelected))
-			.EditCondition(GetIsReverbSendEnabledAttribute(), nullptr);
-	}
+	LayoutBuilder.AddPropertyToCategory(PropertyHandles.FindChecked(GET_MEMBER_NAME_CHECKED(FSoundAttenuationSettings, ManualReverbSendLevel)))
+		.Visibility(TAttribute<EVisibility>(this, &FSoundAttenuationSettingsCustomization::IsManualReverbSendSelected))
+		.EditCondition(GetIsReverbSendEnabledAttribute(), nullptr);
 
 	LayoutBuilder.AddPropertyToCategory(PropertyHandles.FindChecked(GET_MEMBER_NAME_CHECKED(FSoundAttenuationSettings, FocusAzimuth)))
 		.Visibility(IsAttenuationOverriddenVisibleAttribute())
@@ -698,17 +690,14 @@ void FSoundAttenuationSettingsCustomization::CustomizeChildren(TSharedRef<IPrope
 		.Visibility(IsAttenuationOverriddenVisibleAttribute())
 		.EditCondition(IsAttenuationOverriddenAttribute(), nullptr);
 
-	if (bIsAudioMixerEnabled)
+	TSharedPtr<IPropertyHandle> PluginProperty = PropertyHandles.FindChecked(GET_MEMBER_NAME_CHECKED(FSoundAttenuationSettings, PluginSettings));
+	uint32 NumPluginChildren = 0;
+	PluginProperty->GetNumChildren(NumPluginChildren);
+	for (uint32 i = 0; i < NumPluginChildren; ++i)
 	{
-		TSharedPtr<IPropertyHandle> PluginProperty = PropertyHandles.FindChecked(GET_MEMBER_NAME_CHECKED(FSoundAttenuationSettings, PluginSettings));
-		uint32 NumPluginChildren = 0;
-		PluginProperty->GetNumChildren(NumPluginChildren);
-		for (uint32 i = 0; i < NumPluginChildren; ++i)
-		{
-			LayoutBuilder.AddPropertyToCategory(PluginProperty->GetChildHandle(i))
-				.Visibility(IsAttenuationOverriddenVisibleAttribute())
-				.EditCondition(IsAttenuationOverriddenAttribute(), nullptr);
-		}
+		LayoutBuilder.AddPropertyToCategory(PluginProperty->GetChildHandle(i))
+			.Visibility(IsAttenuationOverriddenVisibleAttribute())
+			.EditCondition(IsAttenuationOverriddenAttribute(), nullptr);
 	}
 
 	// Set category display names
diff --git a/Engine/Source/Editor/DetailCustomizations/Private/SoundBaseDetails.cpp b/Engine/Source/Editor/DetailCustomizations/Private/SoundBaseDetails.cpp
--- a/Engine/Source/Editor/DetailCustomizations/Private/SoundBaseDetails.cpp
+++ b/Engine/Source/Editor/DetailCustomizations/Private/SoundBaseDetails.cpp
@@ -16,21 +16,5 @@ TSharedRef<IDetailCustomization> FSoundBaseDetails::MakeInstance()
 
 void FSoundBaseDetails::CustomizeDetails(IDetailLayoutBuilder& DetailBuilder)
 {
-	if (!GetDefault<UAudioSettings>()->IsAudioMixerEnabled())
-	{
-		TSharedRef<IPropertyHandle> Property = DetailBuilder.GetProperty("SoundSubmixObject", USoundBase::StaticClass());
-		Property->MarkHiddenByCustomization();
-
-		Property = DetailBuilder.GetProperty("SourceEffectChain", USoundBase::StaticClass());
-		Property->MarkHiddenByCustomization();
-
-		Property = DetailBuilder.GetProperty("OutputToBusOnly", USoundBase::StaticClass());
-		Property->MarkHiddenByCustomization();
-
-		Property = DetailBuilder.GetProperty("BusSends", USoundBase::StaticClass());
-		Property->MarkHiddenByCustomization();
-
-		Property = DetailBuilder.GetProperty("PreEffectBusSends", USoundBase::StaticClass());
-		Property->MarkHiddenByCustomization();
-	}
+	// Does nothing currently
 }
diff --git a/Engine/Source/Runtime/AudioMixer/Private/Components/SynthComponent.cpp b/Engine/Source/Runtime/AudioMixer/Private/Components/SynthComponent.cpp
--- a/Engine/Source/Runtime/AudioMixer/Private/Components/SynthComponent.cpp
+++ b/Engine/Source/Runtime/AudioMixer/Private/Components/SynthComponent.cpp
@@ -23,18 +23,7 @@ void USynthSound::Init(USynthComponent* InSynthComponent, const int32 InNumChann
 	VirtualizationMode = EVirtualizationMode::PlayWhenSilent;
 	NumChannels = InNumChannels;
 	NumSamplesToGeneratePerCallback = InCallbackSize;
-	// Turn off async generation in old audio engine on mac.
-#if PLATFORM_MAC
-	const FAudioDevice* AudioDevice = InSynthComponent->GetAudioDevice();
-	if (AudioDevice && !AudioDevice->IsAudioMixerEnabled())
-	{
-		bCanProcessAsync = false;
-	}
-	else
-#endif // #if PLATFORM_MAC
-	{
-		bCanProcessAsync = true;
-	}
+	bCanProcessAsync = true;
 
 	Duration = INDEFINITELY_LOOPING_DURATION;
 	bLooping = true;
@@ -43,8 +32,6 @@ void USynthSound::Init(USynthComponent* InSynthComponent, const int32 InNumChann
 
 void USynthSound::StartOnAudioDevice(FAudioDevice* InAudioDevice)
 {
-	check(InAudioDevice != nullptr);
-	bAudioMixer = InAudioDevice->IsAudioMixerEnabled();
 }
 
 void USynthSound::OnBeginGenerate()
@@ -61,46 +48,18 @@ int32 USynthSound::OnGeneratePCMAudio(TArray<uint8>& OutAudio, int32 NumSamples)
 
 	OutAudio.Reset();
 
-	if (bAudioMixer)
-	{
-		// If running with audio mixer, the output audio buffer will be in floats already
-		OutAudio.AddZeroed(NumSamples * sizeof(float));
+	// If running with audio mixer, the output audio buffer will be in floats already
+	OutAudio.AddZeroed(NumSamples * sizeof(float));
 
-		// Mark pending kill can null this out on the game thread in rare cases.
-		if (!OwningSynthComponent.IsValid())
-		{
-			return 0;
-		}
-
-		return OwningSynthComponent->OnGeneratePCMAudio((float*)OutAudio.GetData(), NumSamples);
-	}
-	else
+	// Mark pending kill can null this out on the game thread in rare cases.
+	if (!OwningSynthComponent.IsValid())
 	{
-		// Use the float scratch buffer instead of the out buffer directly
-		FloatBuffer.Reset();
-		FloatBuffer.AddZeroed(NumSamples * sizeof(float));
-
-		// Mark pending kill can null this out on the game thread in rare cases.
-		if (!OwningSynthComponent.IsValid())
-		{
-			return 0;
-		}
-
-		float* FloatBufferDataPtr = FloatBuffer.GetData();
-		int32 NumSamplesGenerated = OwningSynthComponent->OnGeneratePCMAudio(FloatBufferDataPtr, NumSamples);
-
-		// Convert the float buffer to int16 data
-		OutAudio.AddZeroed(NumSamples * sizeof(int16));
-		int16* OutAudioBuffer = (int16*)OutAudio.GetData();
-		for (int32 i = 0; i < NumSamples; ++i)
-		{
-			OutAudioBuffer[i] = (int16)(32767.0f * FMath::Clamp(FloatBufferDataPtr[i], -1.0f, 1.0f));
-		}
-		return NumSamplesGenerated;
+		return 0;
 	}
 
-	return NumSamples;
-}
+	return OwningSynthComponent->OnGeneratePCMAudio((float*)OutAudio.GetData(), NumSamples);
+
+}	
 
 void USynthSound::OnEndGenerate()
 {
@@ -123,7 +82,7 @@ ISoundGeneratorPtr USynthSound::CreateSoundGenerator(const FSoundGeneratorInitPa
 Audio::EAudioMixerStreamDataFormat::Type USynthSound::GetGeneratedPCMDataFormat() const
 {
 	// Only audio mixer supports return float buffers
-	return bAudioMixer ? Audio::EAudioMixerStreamDataFormat::Float : Audio::EAudioMixerStreamDataFormat::Int16;
+	return Audio::EAudioMixerStreamDataFormat::Float;
 }
 
 USynthComponent::USynthComponent(const FObjectInitializer& ObjectInitializer)
diff --git a/Engine/Source/Runtime/AudioMixer/Private/Quartz/QuartzSubsystem.cpp b/Engine/Source/Runtime/AudioMixer/Private/Quartz/QuartzSubsystem.cpp
--- a/Engine/Source/Runtime/AudioMixer/Private/Quartz/QuartzSubsystem.cpp
+++ b/Engine/Source/Runtime/AudioMixer/Private/Quartz/QuartzSubsystem.cpp
@@ -39,7 +39,7 @@ static Audio::FMixerDevice* GetAudioMixerDevice(const UWorld* InWorld)
 	}
 
 	FAudioDevice* AudioDevicePtr = InWorld->GetAudioDeviceRaw();
-	if(AudioDevicePtr && AudioDevicePtr->IsAudioMixerEnabled())
+	if(AudioDevicePtr)
 	{
 		return static_cast<Audio::FMixerDevice*>(AudioDevicePtr);
 	}
@@ -54,7 +54,7 @@ static TUniquePtr<FScopeLock> GetPersistentStateScopeLock(const UWorld* InWorld)
 		return MakeUnique<FScopeLock>(&MixerDevicePtr->QuartzPersistentStateCritSec);
 	}
 
-	return {}; // Empty lock for AudioMixer-less runtimes (server)
+	return {};
 }
 
 
diff --git a/Engine/Source/Runtime/AudioMixer/Public/Components/SynthComponent.h b/Engine/Source/Runtime/AudioMixer/Public/Components/SynthComponent.h
--- a/Engine/Source/Runtime/AudioMixer/Public/Components/SynthComponent.h
+++ b/Engine/Source/Runtime/AudioMixer/Public/Components/SynthComponent.h
@@ -62,7 +62,6 @@ protected:
 	TWeakObjectPtr<USynthComponent> OwningSynthComponent = nullptr;
 
 	TArray<float> FloatBuffer;
-	bool bAudioMixer;
 
 public:
 	USynthComponent* GetOwningSynthComponent()
diff --git a/Engine/Source/Runtime/Engine/Classes/Sound/AudioSettings.h b/Engine/Source/Runtime/Engine/Classes/Sound/AudioSettings.h
--- a/Engine/Source/Runtime/Engine/Classes/Sound/AudioSettings.h
+++ b/Engine/Source/Runtime/Engine/Classes/Sound/AudioSettings.h
@@ -285,12 +285,6 @@ public:
 	// Get the total number of quality level settings
 	int32 GetQualityLevelSettingsNum() const;
 
-	// Sets whether audio mixer is enabled. Set once an audio mixer platform module is loaded.
-	void SetAudioMixerEnabled(const bool bInAudioMixerEnabled);
-
-	// Returns if the audio mixer is currently enabled
-	const bool IsAudioMixerEnabled() const;
-
 	/** Returns the highest value for MaxChannels among all quality levels */
 	int32 GetHighestMaxChannels() const;
 
@@ -310,7 +304,4 @@ private:
 	void AddDefaultSettings();
 
 	bool bParameterInterfacesRegistered;
-
-	// Whether or not the audio mixer is loaded/enabled. Used to toggle visibility of editor features.
-	bool bIsAudioMixerEnabled;
 };
diff --git a/Engine/Source/Runtime/Engine/Private/ActiveSound.cpp b/Engine/Source/Runtime/Engine/Private/ActiveSound.cpp
--- a/Engine/Source/Runtime/Engine/Private/ActiveSound.cpp
+++ b/Engine/Source/Runtime/Engine/Private/ActiveSound.cpp
@@ -996,7 +996,7 @@ void FActiveSound::MarkPendingDestroy(bool bDestroyNow)
 		if (Source)
 		{
 			bool bStopped = false;
-			if (AudioDevice->IsAudioMixerEnabled() && AudioDevice->IsStoppingVoicesEnabled())
+			if (AudioDevice->IsStoppingVoicesEnabled())
 			{
 				if (bDestroyNow || !AudioDevice->GetNumFreeSources())
 				{
diff --git a/Engine/Source/Runtime/Engine/Private/Audio.cpp b/Engine/Source/Runtime/Engine/Private/Audio.cpp
--- a/Engine/Source/Runtime/Engine/Private/Audio.cpp
+++ b/Engine/Source/Runtime/Engine/Private/Audio.cpp
@@ -374,7 +374,7 @@ void FSoundSource::SetFilterFrequency()
 		{
 			// compensate for filter coefficient calculation error for occlusion
 			float OcclusionFilterScale = 1.0f;
-			if (AudioDevice->IsAudioMixerEnabled() && OcclusionFilterScaleEnabledCVar == 1 && !FMath::IsNearlyEqual(WaveInstance->OcclusionFilterFrequency, MAX_FILTER_FREQUENCY))
+			if (OcclusionFilterScaleEnabledCVar == 1 && !FMath::IsNearlyEqual(WaveInstance->OcclusionFilterFrequency, MAX_FILTER_FREQUENCY))
 			{
 				OcclusionFilterScale = 0.25f;
 			}
diff --git a/Engine/Source/Runtime/Engine/Private/AudioBus.cpp b/Engine/Source/Runtime/Engine/Private/AudioBus.cpp
--- a/Engine/Source/Runtime/Engine/Private/AudioBus.cpp
+++ b/Engine/Source/Runtime/Engine/Private/AudioBus.cpp
@@ -41,12 +41,9 @@ void UAudioBus::BeginDestroy()
 		TArray<FAudioDevice*> AudioDevices = AudioDeviceManager->GetAudioDevices();
 		for (FAudioDevice* AudioDevice : AudioDevices)
 		{
-			if (AudioDevice->IsAudioMixerEnabled())
-			{
-				UAudioBusSubsystem* AudioBusSubsystem = AudioDevice->GetSubsystem<UAudioBusSubsystem>();
-				check(AudioBusSubsystem);
-				AudioBusSubsystem->StopAudioBus(Audio::FAudioBusKey(AudioBusId));
-			}
+			UAudioBusSubsystem* AudioBusSubsystem = AudioDevice->GetSubsystem<UAudioBusSubsystem>();
+			check(AudioBusSubsystem);
+			AudioBusSubsystem->StopAudioBus(Audio::FAudioBusKey(AudioBusId));
 		}
 	}
 }
diff --git a/Engine/Source/Runtime/Engine/Private/AudioDevice.cpp b/Engine/Source/Runtime/Engine/Private/AudioDevice.cpp
--- a/Engine/Source/Runtime/Engine/Private/AudioDevice.cpp
+++ b/Engine/Source/Runtime/Engine/Private/AudioDevice.cpp
@@ -542,7 +542,7 @@ bool FAudioDevice::Init(Audio::FDeviceId InDeviceID, int32 InMaxSources, int32 I
 	InitSoundEffectPresets();
 
 	// Audio mixer needs to create effects manager before initializing the plugins.
-	if (IsAudioMixerEnabled() && IsStoppingVoicesEnabled())
+	if (IsStoppingVoicesEnabled())
 	{
 		// create a platform specific effects manager
 		Effects = CreateEffectsManager();
@@ -580,17 +580,6 @@ bool FAudioDevice::Init(Audio::FDeviceId InDeviceID, int32 InMaxSources, int32 I
 			SpatializationInterfaces.Add(SpatPluginInfo);
 		}
 
-		if (!IsAudioMixerEnabled())
-		{
-
-			//Set up initialization parameters for system level effect plugins:
-			PluginInitializationParams.SampleRate = SampleRate;
-			PluginInitializationParams.NumSources = GetMaxSources();
-			PluginInitializationParams.BufferLength = PlatformSettings.CallbackBufferFrameSize;
-			PluginInitializationParams.AudioDevicePtr = this;
-			SetCurrentSpatializationPlugin(AudioPluginUtilities::GetDesiredSpatializationPluginName());
-		}
-
 		bSpatializationInterfaceEnabled = true;
 
 		UE_LOG(LogAudio, Display, TEXT("Audio Spatialization Plugin: %s is external send: %d"), *(SpatializationPluginFactory->GetDisplayName()), SpatializationPluginFactory->IsExternalSend());
@@ -605,7 +594,6 @@ bool FAudioDevice::Init(Audio::FDeviceId InDeviceID, int32 InMaxSources, int32 I
 	{
 		SourceDataOverridePluginInterface = SourceDataOverridePluginFactory->CreateNewSourceDataOverridePlugin(this);
 		bSourceDataOverrideInterfaceEnabled = true;
-		check(IsAudioMixerEnabled());
 		UE_LOG(LogAudio, Display, TEXT("Audio Source Data Override Plugin: %s"), *(SourceDataOverridePluginFactory->GetDisplayName()));
 	}
 
@@ -664,13 +652,6 @@ bool FAudioDevice::Init(Audio::FDeviceId InDeviceID, int32 InMaxSources, int32 I
 		return false;
 	}
 
-	// create a platform specific effects manager
-	// if this is the audio mixer, we initialized the effects manager before the hardware
-	if (!IsAudioMixerEnabled())
-	{
-		Effects = CreateEffectsManager();
-	}
-
 	InitSoundSources();
 
 	// Make sure the Listeners array has at least one entry, so we don't have to check for Listeners.Num() == 0 all the time
@@ -979,14 +960,11 @@ void FAudioDevice::Teardown()
 	// Note: we don't free audio buffers at this stage since they are managed in the audio device manager
 
 	// Must be after FreeBufferResource as that potentially stops sources
-	if (IsAudioMixerEnabled())
+	for (int32 Index = 0; Index < Sources.Num(); Index++)
 	{
-		for (int32 Index = 0; Index < Sources.Num(); Index++)
-		{
-			// Make the sound stop instantly
-			Sources[Index]->StopNow();
-			delete Sources[Index];
-		}
+		// Make the sound stop instantly
+		Sources[Index]->StopNow();
+		delete Sources[Index];
 	}
 
 	Sources.Reset();
@@ -4246,83 +4224,80 @@ void FAudioDevice::UpdateActiveSoundPlaybackTime(bool bIsGameTicking)
 
 void FAudioDevice::StopOldestStoppingSource()
 {
-	if (IsAudioMixerEnabled())
-	{
-		check(!FreeSources.Num());
+	check(!FreeSources.Num());
 
-		FSoundSource* LowestPriStoppingSource = nullptr;
-		FSoundSource* LowestPriSource = nullptr;
-		FSoundSource* LowestPriNonLoopingSource = nullptr;
+	FSoundSource* LowestPriStoppingSource = nullptr;
+	FSoundSource* LowestPriSource = nullptr;
+	FSoundSource* LowestPriNonLoopingSource = nullptr;
 
-		for (FSoundSource* Source : Sources)
+	for (FSoundSource* Source : Sources)
+	{
+		// Find oldest stopping voice first
+		if (Source->IsStopping())
 		{
-			// Find oldest stopping voice first
-			if (Source->IsStopping())
+			if (!LowestPriStoppingSource)
 			{
-				if (!LowestPriStoppingSource)
+				LowestPriStoppingSource = Source;
+			}
+			else
+			{
+				if (Source->WaveInstance->GetVolumeWeightedPriority() < LowestPriStoppingSource->WaveInstance->GetVolumeWeightedPriority())
 				{
 					LowestPriStoppingSource = Source;
 				}
-				else
-				{
-					if (Source->WaveInstance->GetVolumeWeightedPriority() < LowestPriStoppingSource->WaveInstance->GetVolumeWeightedPriority())
-					{
-						LowestPriStoppingSource = Source;
-					}
-				}
 			}
-			else if (Source->WaveInstance)
+		}
+		else if (Source->WaveInstance)
+		{
+			// Find lowest volume/priority non-looping source as fallback
+			if (Source->WaveInstance->LoopingMode != ELoopingMode::LOOP_Forever && !Source->WaveInstance->bIsUISound)
 			{
-				// Find lowest volume/priority non-looping source as fallback
-				if (Source->WaveInstance->LoopingMode != ELoopingMode::LOOP_Forever && !Source->WaveInstance->bIsUISound)
+				if (!LowestPriNonLoopingSource)
 				{
-					if (!LowestPriNonLoopingSource)
+					LowestPriNonLoopingSource = Source;
+				}
+				else
+				{
+					if (Source->WaveInstance->GetVolumeWeightedPriority() < LowestPriNonLoopingSource->WaveInstance->GetVolumeWeightedPriority())
 					{
 						LowestPriNonLoopingSource = Source;
 					}
-					else
-					{
-						if (Source->WaveInstance->GetVolumeWeightedPriority() < LowestPriNonLoopingSource->WaveInstance->GetVolumeWeightedPriority())
-						{
-							LowestPriNonLoopingSource = Source;
-						}
-					}
 				}
+			}
 
 
-				// Find lowest volume/priority source as final fallback
-				if (!LowestPriSource)
+			// Find lowest volume/priority source as final fallback
+			if (!LowestPriSource)
+			{
+				LowestPriSource = Source;
+			}
+			else
+			{
+				if (Source->WaveInstance->GetVolumeWeightedPriority() < LowestPriSource->WaveInstance->GetVolumeWeightedPriority())
 				{
 					LowestPriSource = Source;
 				}
-				else
-				{
-					if (Source->WaveInstance->GetVolumeWeightedPriority() < LowestPriSource->WaveInstance->GetVolumeWeightedPriority())
-					{
-						LowestPriSource = Source;
-					}
-				}
 			}
 		}
+	}
 
-		// Stop oldest stopping source
-		if (LowestPriStoppingSource)
-		{
-			LowestPriStoppingSource->StopNow();
-		}
-		// If no oldest stopping source, stop oldest one-shot
-		else if (LowestPriNonLoopingSource)
-		{
-			LowestPriNonLoopingSource->StopNow();
-		}
-		// Otherwise stop oldest source.
-		else
-		{
-			check(LowestPriSource);
-			LowestPriSource->StopNow();
-		}
-		check(FreeSources.Num() > 0);
+	// Stop oldest stopping source
+	if (LowestPriStoppingSource)
+	{
+		LowestPriStoppingSource->StopNow();
+	}
+	// If no oldest stopping source, stop oldest one-shot
+	else if (LowestPriNonLoopingSource)
+	{
+		LowestPriNonLoopingSource->StopNow();
+	}
+	// Otherwise stop oldest source.
+	else
+	{
+		check(LowestPriSource);
+		LowestPriSource->StopNow();
 	}
+	check(FreeSources.Num() > 0);
 }
 
 void FAudioDevice::StopSources(TArray<FWaveInstance*>& WaveInstances, int32 FirstActiveIndex)
@@ -4448,7 +4423,7 @@ void FAudioDevice::StartSources(TArray<FWaveInstance*>& WaveInstances, int32 Fir
 				IStreamingManager::Get().GetAudioStreamingManager().CanCreateSoundSource(WaveInstance)))
 			{
 				// Check for full sources and stop the oldest stopping source
-				if (IsAudioMixerEnabled() && !FreeSources.Num())
+				if (!FreeSources.Num())
 				{
 					StopOldestStoppingSource();
 				}
@@ -4877,7 +4852,7 @@ void FAudioDevice::UpdateAudioVolumeEffects()
 		Effects->Update();
 
 		// If we any submix override settings apply those overrides to the indicated submixes
-		if (IsAudioMixerEnabled() && bUpdateSubmixEffects)
+		if (bUpdateSubmixEffects)
 		{
 			// Clear out any previous submix effect chain overrides if the audio volume changed
 			if (PreviousPlayerAudioVolumeSettings.SubmixOverrideSettings.Num() > 0)
@@ -5115,17 +5090,6 @@ void FAudioDevice::AddNewActiveSoundInternal(const FActiveSound& InNewActiveSoun
 		return;
 	}
 
-	// Don't allow buses to try to play if we're not using the audio mixer.
-	if (!IsAudioMixerEnabled())
-	{
-		USoundSourceBus* Bus = Cast<USoundSourceBus>(InNewActiveSound.Sound);
-		if (Bus)
-		{
-			ReportSoundFailedToStart(InNewActiveSound.AudioComponentID, InVirtualLoopToRetrigger);
-			return;
-		}
-	}
-
 	if (Sound->GetDuration() <= FMath::Max(0.0f, SoundDistanceOptimizationLengthCVar))
 	{
 		// TODO: Determine if this check has already been completed at AudioComponent level and skip if so. Also,
@@ -6581,13 +6545,10 @@ void FAudioDevice::Flush(UWorld* WorldToFlush, bool bClearActivatedReverb)
 	}
 
 	// Make sure we update any hardware changes that need to happen after flushing
-	if (IsAudioMixerEnabled())
-	{
-		UpdateHardware();
+	UpdateHardware();
 
-		// Make sure any in-flight audio rendering commands get executed.
-		FlushAudioRenderingCommands();
-	}
+	// Make sure any in-flight audio rendering commands get executed.
+	FlushAudioRenderingCommands();
 
 	FlushExtended(WorldToFlush, bClearActivatedReverb);
 }
@@ -6674,17 +6635,8 @@ void FAudioDevice::Precache(USoundWave* SoundWave, bool bSynchronous, bool bTrac
 	}
 	else if (SoundWave->bIsSourceBus)
 	{
-		// Audio data which will be generated by instanced objects, not from the sound wave asset
-		if (IsAudioMixerEnabled())
-		{
-			// Buses will initialize as procedural, but not actually become a procedural sound wave
-			SoundWave->DecompressionType = DTYPE_Procedural;
-		}
-		else
-		{
-			// Buses are only supported with audio mixer
-			SoundWave->DecompressionType = DTYPE_Invalid;
-		}
+		// Buses will initialize as procedural, but not actually become a procedural sound wave
+		SoundWave->DecompressionType = DTYPE_Procedural;
 	}
 	else if (HasCompressedAudioInfoClass(SoundWave))
 	{
diff --git a/Engine/Source/Runtime/Engine/Private/AudioDeviceManager.cpp b/Engine/Source/Runtime/Engine/Private/AudioDeviceManager.cpp
--- a/Engine/Source/Runtime/Engine/Private/AudioDeviceManager.cpp
+++ b/Engine/Source/Runtime/Engine/Private/AudioDeviceManager.cpp
@@ -39,22 +39,6 @@ TAutoConsoleVariable<int32> CVarEnableAudioThreadWait(
 	TEXT("0: Not Enabled, 1: Enabled"),
 	ECVF_Default);
 
-static int32 GCvarIsUsingAudioMixer = 0;
-FAutoConsoleVariableRef CVarIsUsingAudioMixer(
-	TEXT("au.IsUsingAudioMixer"),
-	GCvarIsUsingAudioMixer,
-	TEXT("Whether or not we're currently using the audio mixer. Change to dynamically toggle on/off. This will only take effect if an audio device is currently not in use, unless au.AllowUnsafeAudioMixerToggling is set to 1. Note: sounds will stop. Looping sounds won't automatically resume. \n")
-	TEXT("0: Not Using Audio Mixer, 1: Using Audio Mixer"),
-	ECVF_Default);
-
-static int32 GCvarAllowUnsafeAudioMixerToggling = 0;
-FAutoConsoleVariableRef CVarAllowUnsafeAudioMixerToggling(
-	TEXT("au.AllowUnsafeAudioMixerToggling"),
-	GCvarAllowUnsafeAudioMixerToggling,
-	TEXT("If set to 1, will allow au.IsUsingAudioMixer to swap out the audio engine, even if there are systems in the world currently using the audio engine. \n")
-	TEXT("0: disable usage of au.IsUsingAudioMixer when the audio device is actively in use, 1: enable usage of au.IsUsingAudioMixer."),
-	ECVF_Default);
-
 static int32 CVarIsVisualizeEnabled = 0;
 FAutoConsoleVariableRef CVarAudioVisualizeEnabled(
 	TEXT("au.3dVisualize.Enabled"),
@@ -191,10 +175,7 @@ FAudioDeviceManager::FAudioDeviceManager()
 	, NextResourceID(1)
 	, SoloDeviceHandle(INDEX_NONE)
 	, ActiveAudioDeviceID(INDEX_NONE)
-	, bUsingAudioMixer(false)
 	, bPlayAllDeviceAudio(false)
-	, bOnlyToggleAudioMixerOnce(false)
-	, bToggledAudioMixer(false)
 {
 
 #if ENABLE_AUDIO_DEBUG
@@ -263,191 +244,11 @@ Audio::FMixerDevice* FAudioDeviceManager::GetAudioMixerDeviceFromWorldContext(co
 {
 	if (FAudioDevice* AudioDevice = GetAudioDeviceFromWorldContext(WorldContextObject))
 	{
-		if (!AudioDevice->IsAudioMixerEnabled())
-		{
-			return nullptr;
-		}
-		else
-		{
-			return static_cast<Audio::FMixerDevice*>(AudioDevice);
-		}
+		return static_cast<Audio::FMixerDevice*>(AudioDevice);
 	}
 	return nullptr;
 }
 
-
-bool FAudioDeviceManager::ToggleAudioMixer()
-{
-	// Only need to toggle if we have 2 device module names loaded at init
-	if (AudioDeviceModule && AudioDeviceModuleName.Len() > 0 && AudioMixerModuleName.Len() > 0)
-	{
-		FScopeLock ScopeLock(&DeviceMapCriticalSection);
-		{
-			// Check to see if there are still any active handles using this audio device, and either warn or exit early.
-			for (auto& DeviceContainer : Devices)
-			{
-				if (DeviceContainer.Value.NumberOfHandlesToThisDevice > 0)
-				{
-
-					UE_LOG(LogAudio, Warning, TEXT("Attempted to toggle the audio mixer while that audio device was in use."));
-
-	#if INSTRUMENT_AUDIODEVICE_HANDLES
-					FString ActiveDeviceHandles;
-					for (auto& StackWalkString : DeviceContainer.Value.HandleCreationStackWalks)
-					{
-						ActiveDeviceHandles += StackWalkString.Value;
-						ActiveDeviceHandles += TEXT("\n\n");
-					}
-
-					UE_LOG(LogAudio, Warning, TEXT("List Of Active Handles to this device: \n%s"), *ActiveDeviceHandles);
-	#else
-					UE_LOG(LogAudio, Warning, TEXT("For more information compile with INSTRUMENT_AUDIODEVICE_HANDLES."));
-	#endif
-
-					if (!GCvarAllowUnsafeAudioMixerToggling)
-					{
-						return false;
-					}
-				}
-			}
-		}
-
-		// Suspend the audio thread
-		FAudioThreadSuspendContext AudioThreadSuspendScope;
-
-		// If using audio mixer, we need to toggle back to non-audio mixer
-		FString ModuleToUnload;
-
-		// If currently using the audio mixer, we need to toggle to the old audio engine module
-		if (bUsingAudioMixer)
-		{
-			// Unload the previous module
-			ModuleToUnload = AudioMixerModuleName;
-
-			AudioDeviceModule = FModuleManager::LoadModulePtr<IAudioDeviceModule>(*AudioDeviceModuleName);
-
-			bUsingAudioMixer = false;
-		}
-		// If we're currently using old audio engine module, we toggle to the audio mixer module
-		else
-		{
-			// Unload the previous module
-			ModuleToUnload = AudioDeviceModuleName;
-
-			// Load the audio mixer engine module
-			AudioDeviceModule = FModuleManager::LoadModulePtr<IAudioDeviceModule>(*AudioMixerModuleName);
-
-			bUsingAudioMixer = true;
-		}
-
-		// If we succeeded in loading a new module, create a new main audio device.
-		if (AudioDeviceModule)
-		{
-			// Shutdown and create new audio devices
-			const UAudioSettings* AudioSettings = GetDefault<UAudioSettings>();
-			const int32 QualityLevel = GEngine->GetGameUserSettings()->GetAudioQualityLevel(); // -V595
-			const int32 QualityLevelMaxChannels = AudioSettings->GetQualityLevelSettings(QualityLevel).MaxChannels; //-V595
-
-			// We could have multiple audio devices, so loop through them and patch them up best we can to
-			// get parity. E.g. we need to pass the handle from the odl to the new, set whether or not its active
-			// and try and get the mix-states to be the same.
-			for (auto& DeviceContainer : Devices)
-			{
-				FAudioDevice*& AudioDevice = DeviceContainer.Value.Device;
-
-				check(AudioDevice);
-
-				// Get the audio device handle and whether it is active
-				uint32 DeviceID = AudioDevice->DeviceID;
-				check(DeviceContainer.Key == DeviceID);
-				bool bIsActive = (DeviceID == ActiveAudioDeviceID);
-
-				// To transfer mix states, we need to re-base the absolute clocks on the mix states
-				// so the target audio device timing won't result in the mixes suddenly stopping.
-				TMap<USoundMix*, FSoundMixState> MixModifiers = AudioDevice->GetSoundMixModifiers();
-				TArray<USoundMix*> PrevPassiveSoundMixModifiers = AudioDevice->GetPrevPassiveSoundMixModifiers();
-				USoundMix* BaseSoundMix = AudioDevice->GetDefaultBaseSoundMixModifier();
-				double AudioClock = AudioDevice->GetAudioClock();
-
-				for (TPair<USoundMix*, FSoundMixState>& SoundMixPair : MixModifiers)
-				{
-					// Rebase so that a new clock starting from 0.0 won't cause mixes to stop.
-					SoundMixPair.Value.StartTime -= AudioClock;
-					SoundMixPair.Value.FadeInStartTime -= AudioClock;
-					SoundMixPair.Value.FadeInEndTime -= AudioClock;
-
-					if (SoundMixPair.Value.EndTime > 0.0f)
-					{
-						SoundMixPair.Value.EndTime -= AudioClock;
-					}
-
-					if (SoundMixPair.Value.FadeOutStartTime > 0.0f)
-					{
-						SoundMixPair.Value.FadeOutStartTime -= AudioClock;
-					}
-				}
-
-				// Tear it down and delete the old audio device. This does a bunch of cleanup.
-				AudioDevice->Teardown();
-				delete AudioDevice;
-
-				// Make a new audio device using the new audio device module
-				AudioDevice = AudioDeviceModule->CreateAudioDevice();
-
-				// Some AudioDeviceModules override CreteAudioMixerPlatformInterface, which means we can create a Audio::FMixerDevice.
-				if (AudioDevice == nullptr)
-				{
-					checkf(AudioDeviceModule->IsAudioMixerModule(), TEXT("Please override AudioDeviceModule->CreateAudioDevice()"))
-						AudioDevice = new Audio::FMixerDevice(AudioDeviceModule->CreateAudioMixerPlatformInterface());
-				}
-
-				check(AudioDevice);
-
-				// Re-init the new audio device using appropriate settings so it behaves the same
-				if (AudioDevice->Init(DeviceID, AudioSettings->GetHighestMaxChannels()))
-				{
-					AudioDevice->SetMaxChannels(QualityLevelMaxChannels);
-				}
-
-				// Transfer the sound mix modifiers to the new audio engine
-				AudioDevice->SetSoundMixModifiers(MixModifiers, PrevPassiveSoundMixModifiers, BaseSoundMix);
-				// Setup the mute state of the audio device to be the same that it was
-				if (bIsActive)
-				{
-					AudioDevice->SetDeviceMuted(false);
-				}
-				else
-				{
-					AudioDevice->SetDeviceMuted(true);
-				}
-
-				// Fade in the new audio device (used only in audio mixer to prevent pops on startup/shutdown)
-				AudioDevice->FadeIn();
-			}
-
-			// We now must free any resources that have been cached with the old audio engine
-			// This will result in re-caching of sound waves, but we're forced to do this because FSoundBuffer pointers
-			// are cached and each AudioDevice backend has a derived implementation of this so once we
-			// switch to a new audio engine the FSoundBuffer pointers are totally invalid.
-			for (TObjectIterator<USoundWave> SoundWaveIt; SoundWaveIt; ++SoundWaveIt)
-			{
-				USoundWave* SoundWave = *SoundWaveIt;
-				FreeResource(SoundWave);
-			}
-
-			// Unload the previous audio device module
-			FModuleManager::Get().UnloadModule(*ModuleToUnload);
-		}
-	}
-
-	return false;
-}
-
-bool FAudioDeviceManager::IsUsingAudioMixer() const
-{
-	return bUsingAudioMixer;
-}
-
 IAudioDeviceModule* FAudioDeviceManager::GetAudioDeviceModule()
 {
 	return AudioDeviceModule;
@@ -584,24 +385,14 @@ bool FAudioDeviceManager::InitializeManager()
 		AudioSettings->LoadDefaultObjects();
 		AudioSettings->RegisterParameterInterfaces();
 
-		const bool bIsAudioMixerEnabled = AudioDeviceModule->IsAudioMixerModule();
-		if (bIsAudioMixerEnabled)
-		{
-			FModuleManager::Get().LoadModuleChecked(TEXT("AudioMixer"));
-		}
-		AudioSettings->SetAudioMixerEnabled(bIsAudioMixerEnabled);
-
 		// Register all formats
 		AudioFormatSettings = MakePimpl<Audio::FAudioFormatSettings>(GConfig, GEngineIni, FPlatformProperties::IniPlatformName());
 		RegisterAudioInfoFactories();
 
 #if WITH_EDITOR
-		if (bIsAudioMixerEnabled)
-		{
-			IAudioEditorModule* AudioEditorModule = &FModuleManager::LoadModuleChecked<IAudioEditorModule>("AudioEditor");
-			AudioEditorModule->RegisterAudioMixerAssetActions();
-			AudioEditorModule->RegisterEffectPresetAssetActions();
-		}
+		IAudioEditorModule* AudioEditorModule = &FModuleManager::LoadModuleChecked<IAudioEditorModule>("AudioEditor");
+		AudioEditorModule->RegisterAudioMixerAssetActions();
+		AudioEditorModule->RegisterEffectPresetAssetActions();
 #endif
 
 		FCoreDelegates::ApplicationWillEnterBackgroundDelegate.AddRaw(this, &FAudioDeviceManager::AppWillEnterBackground);
@@ -645,85 +436,22 @@ bool FAudioDeviceManager::LoadDefaultAudioDeviceModule()
 {
 	check(!AudioDeviceModule);
 
-	// Check if we're going to try to force loading the audio mixer from the command line
-	bool bForceAudioMixer = FParse::Param(FCommandLine::Get(), TEXT("AudioMixer"));
-
-	bool bForceNoAudioMixer = FParse::Param(FCommandLine::Get(), TEXT("NoAudioMixer"));
-
 	bool bForceNonRealtimeRenderer = FParse::Param(FCommandLine::Get(), TEXT("DeterministicAudio"));
 
 	// If not using command line switch to use audio mixer, check the game platform engine ini file (e.g. WindowsEngine.ini) which enables it for player
-	bUsingAudioMixer = bForceAudioMixer;
-	if (!bForceAudioMixer && !bForceNoAudioMixer)
-	{
-		GConfig->GetBool(TEXT("Audio"), TEXT("UseAudioMixer"), bUsingAudioMixer, GEngineIni);
-		// Get the audio mixer and non-audio mixer device module names
-		GConfig->GetString(TEXT("Audio"), TEXT("AudioDeviceModuleName"), AudioDeviceModuleName, GEngineIni);
-		GConfig->GetString(TEXT("Audio"), TEXT("AudioMixerModuleName"), AudioMixerModuleName, GEngineIni);
-
-		// If AudioLink is Enabled, use the AudioLink Platform.
-		static const FString AudioLinkMixerModule = TEXT("AudioMixerPlatformAudioLink");
-		if (IModularFeatures::Get().IsModularFeatureAvailable(TEXT("AudioLink Factory")))
-		{
-			AudioMixerModuleName = AudioLinkMixerModule;
-		}
-	}
-	else if (bForceNoAudioMixer)
-	{
-		GConfig->GetString(TEXT("Audio"), TEXT("AudioDeviceModuleName"), AudioDeviceModuleName, GEngineIni);
-
-		// Allow no audio mixer override from command line
-		bUsingAudioMixer = false;
-	}
-	else if(bForceAudioMixer)
-	{ 
-		GConfig->GetString(TEXT("Audio"), TEXT("AudioMixerModuleName"), AudioMixerModuleName, GEngineIni);
-	}
-
-	// Check for config bool that restricts audio mixer toggle to only once. This will allow us to patch audio mixer on or off after initial login.
-	GConfig->GetBool(TEXT("Audio"), TEXT("OnlyToggleAudioMixerOnce"), bOnlyToggleAudioMixerOnce, GEngineIni);
+	GConfig->GetString(TEXT("Audio"), TEXT("AudioMixerModuleName"), AudioMixerModuleName, GEngineIni);
 
 	if (bForceNonRealtimeRenderer)
 	{
 		AudioDeviceModule = FModuleManager::LoadModulePtr<IAudioDeviceModule>(TEXT("NonRealtimeAudioRenderer"));
-
-		static IConsoleVariable* IsUsingAudioMixerCvar = IConsoleManager::Get().FindConsoleVariable(TEXT("au.IsUsingAudioMixer"));
-		check(IsUsingAudioMixerCvar);
-		IsUsingAudioMixerCvar->Set(2, ECVF_SetByConstructor);
-
-		bUsingAudioMixer = true;
-
 		return AudioDeviceModule != nullptr;
 	}
 
-	if (bUsingAudioMixer && AudioMixerModuleName.Len() > 0)
+	if (AudioMixerModuleName.Len() > 0)
 	{
 		AudioDeviceModule = FModuleManager::LoadModulePtr<IAudioDeviceModule>(*AudioMixerModuleName);
-		if (AudioDeviceModule)
-		{
-			static IConsoleVariable* IsUsingAudioMixerCvar = IConsoleManager::Get().FindConsoleVariable(TEXT("au.IsUsingAudioMixer"));
-			check(IsUsingAudioMixerCvar);
-			IsUsingAudioMixerCvar->Set(1, ECVF_SetByConstructor);
-		}
-		else
-		{
-			bUsingAudioMixer = false;
-		}
 	}
 
-	if (!AudioDeviceModule && AudioDeviceModuleName.Len() > 0)
-	{
-		AudioDeviceModule = FModuleManager::LoadModulePtr<IAudioDeviceModule>(*AudioDeviceModuleName);
-
-		static IConsoleVariable* IsUsingAudioMixerCvar = IConsoleManager::Get().FindConsoleVariable(TEXT("au.IsUsingAudioMixer"));
-		check(IsUsingAudioMixerCvar);
-		IsUsingAudioMixerCvar->Set(0, ECVF_SetByConstructor);
-	}
-
-	if (bForceNoAudioMixer)
-	{
-		Audio::Analytics::RecordEvent_Usage(TEXT("ForceNoAudioMixer"));
-	}
 	return AudioDeviceModule != nullptr;
 }
 
@@ -982,38 +710,6 @@ void FAudioDeviceManager::UpdateActiveAudioDevices(bool bGameTicking)
 		SyncFence.Wait();
 	}
 
-	if (!bOnlyToggleAudioMixerOnce || (bOnlyToggleAudioMixerOnce && !bToggledAudioMixer))
-	{
-		if (bUsingAudioMixer && !GCvarIsUsingAudioMixer)
-		{
-			const bool bSuccessfullyToggledAudioMixer = ToggleAudioMixer();
-
-			if (bSuccessfullyToggledAudioMixer)
-			{
-				bToggledAudioMixer = true;
-				bUsingAudioMixer = false;
-			}
-			else
-			{
-				GCvarIsUsingAudioMixer = true;
-			}
-		}
-		else if (!bUsingAudioMixer && GCvarIsUsingAudioMixer)
-		{
-			const bool bSuccessfullyToggledAudioMixer = ToggleAudioMixer();
-
-			if (bSuccessfullyToggledAudioMixer)
-			{
-				bToggledAudioMixer = true;
-				bUsingAudioMixer = true;
-			}
-			else
-			{
-				GCvarIsUsingAudioMixer = false;
-			}
-		}
-	}
-
 	IterateOverAllDevices(
 		[&bGameTicking](Audio::FDeviceId, FAudioDevice* InDevice)
 		{
diff --git a/Engine/Source/Runtime/Engine/Private/AudioSettings.cpp b/Engine/Source/Runtime/Engine/Private/AudioSettings.cpp
--- a/Engine/Source/Runtime/Engine/Private/AudioSettings.cpp
+++ b/Engine/Source/Runtime/Engine/Private/AudioSettings.cpp
@@ -24,7 +24,6 @@ UAudioSettings::UAudioSettings(const FObjectInitializer& ObjectInitializer)
 
 	bAllowPlayWhenSilent = true;
 	bParameterInterfacesRegistered = false;
-	bIsAudioMixerEnabled = false;
 
 	GlobalMinPitchScale = 0.4F;
 	GlobalMaxPitchScale = 2.0F;
@@ -307,16 +306,6 @@ USoundConcurrency* UAudioSettings::GetDefaultSoundConcurrency() const
 	return DefaultSoundConcurrency;
 }
 
-void UAudioSettings::SetAudioMixerEnabled(const bool bInAudioMixerEnabled)
-{
-	bIsAudioMixerEnabled = bInAudioMixerEnabled;
-}
-
-const bool UAudioSettings::IsAudioMixerEnabled() const
-{
-	return bIsAudioMixerEnabled;
-}
-
 int32 UAudioSettings::GetHighestMaxChannels() const
 {
 	check(QualityLevels.Num() > 0);
diff --git a/Engine/Source/Runtime/Engine/Private/SoundWaveProcedural.cpp b/Engine/Source/Runtime/Engine/Private/SoundWaveProcedural.cpp
--- a/Engine/Source/Runtime/Engine/Private/SoundWaveProcedural.cpp
+++ b/Engine/Source/Runtime/Engine/Private/SoundWaveProcedural.cpp
@@ -20,21 +20,6 @@ USoundWaveProcedural::USoundWaveProcedural(const FObjectInitializer& ObjectIniti
 	// We need to do this for procedural sound waves that we do not process asynchronously,
 	// to ensure that we do not underrun.
 	
-	if (GEngine)
-	{
-		FAudioDevice* MainAudioDevice = GEngine->GetMainAudioDeviceRaw();
-		if (MainAudioDevice && !MainAudioDevice->IsAudioMixerEnabled())
-		{
-#if PLATFORM_MAC
-			// We special case the mac callback on the old audio engine, Since Buffer Length is smaller than the device callback size.
-			NumSamplesToGeneratePerCallback = 2048;
-#else
-			NumSamplesToGeneratePerCallback = MainAudioDevice->GetBufferLength();
-#endif
-			NumBufferUnderrunSamples = NumSamplesToGeneratePerCallback / 2;
-		}
-	}
-
 	SampleByteSize = 2;
 
 	checkf(NumSamplesToGeneratePerCallback >= NumBufferUnderrunSamples, TEXT("Should generate more samples than this per callback."));
diff --git a/Engine/Source/Runtime/Engine/Public/AudioDevice.h b/Engine/Source/Runtime/Engine/Public/AudioDevice.h
--- a/Engine/Source/Runtime/Engine/Public/AudioDevice.h
+++ b/Engine/Source/Runtime/Engine/Public/AudioDevice.h
@@ -1389,12 +1389,6 @@ public:
 
 	ENGINE_API static bool IsSourceDataOverridePluginLoaded();
 
-	/** Returns if this is the multi-platform audio mixer. */
-	bool IsAudioMixerEnabled() const
-	{
-		return bAudioMixerModuleLoaded;
-	}
-
 	/** Returns if stopping voices is enabled. */
 	bool IsStoppingVoicesEnabled() const
 	{
diff --git a/Engine/Source/Runtime/Engine/Public/AudioDeviceManager.h b/Engine/Source/Runtime/Engine/Public/AudioDeviceManager.h
--- a/Engine/Source/Runtime/Engine/Public/AudioDeviceManager.h
+++ b/Engine/Source/Runtime/Engine/Public/AudioDeviceManager.h
@@ -105,9 +105,6 @@ public:
 	static FAudioDevice* GetAudioDeviceFromWorldContext(const UObject* WorldContextObject);
 	static Audio::FMixerDevice* GetAudioMixerDeviceFromWorldContext(const UObject* WorldContextObject);
 
-	/** Returns true if we're currently using the audio mixer. */
-	bool IsUsingAudioMixer() const;
-
 	/**
 	 * returns the currently used audio device module for this platform.
 	 * returns nullptr if Initialize() has not been called yet.
@@ -309,9 +306,6 @@ private:
 	/** Creates a handle given the index and a generation value. */
 	uint32 GetNewDeviceID();
 
-	/** Toggles between audio mixer and non-audio mixer audio engine. */
-	bool ToggleAudioMixer();
-
 	/** Load audio device module. */
 	bool LoadDefaultAudioDeviceModule();
 
@@ -328,12 +322,9 @@ private:
 	
 	void RegisterAudioInfoFactories();
 
-	/** Audio device module which creates (old backend) audio devices. */
+	/** Audio device module which creates audio devices. */
 	IAudioDeviceModule* AudioDeviceModule;
 
-	/** Audio device module name. This is the "old" audio engine module name to use. E.g. XAudio2 */
-	FString AudioDeviceModuleName;
-
 	/** The audio mixer module name. This is the audio mixer module name to use. E.g. AudioMixerXAudio2 */
 	FString AudioMixerModuleName;
 
@@ -413,18 +404,9 @@ private:
 	/** Dynamic volume map */
 	TMap<TTuple<ESoundType, FName>, float> DynamicSoundVolumes;
 
-	/** Whether we're currently using the audio mixer or not. Used to toggle between old/new audio engines. */
-	bool bUsingAudioMixer;
-
 	/** Whether or not to play all audio in all active audio devices. */
 	bool bPlayAllDeviceAudio;
 
-	/** Whether or not we check to toggle audio mixer once. */
-	bool bOnlyToggleAudioMixerOnce;
-
-	/** Whether or not we've toggled the audio mixer. */
-	bool bToggledAudioMixer;
-
 	/** Audio Fence to ensure that we don't allow the audio thread to drift never endingly behind. */
 	FAudioCommandFence SyncFence;
 
diff --git a/Engine/Source/Runtime/Launch/Private/Android/AndroidEventManager.cpp b/Engine/Source/Runtime/Launch/Private/Android/AndroidEventManager.cpp
--- a/Engine/Source/Runtime/Launch/Private/Android/AndroidEventManager.cpp
+++ b/Engine/Source/Runtime/Launch/Private/Android/AndroidEventManager.cpp
@@ -330,23 +330,11 @@ void FAppEventManager::PauseAudio()
 	FAudioDeviceHandle AudioDevice = GEngine->GetMainAudioDevice();
 	if (AudioDevice)
 	{
-		if (AudioDevice->IsAudioMixerEnabled())
-		{
-			FAudioCommandFence Fence;
-			Fence.BeginFence();
-			Fence.Wait();
+		FAudioCommandFence Fence;
+		Fence.BeginFence();
+		Fence.Wait();
 
-			AudioDevice->SuspendContext();
-		}
-		else
-		{
-			GEngine->GetMainAudioDevice()->Suspend(false);
-
-			// make sure the audio thread runs the pause request
-			FAudioCommandFence Fence;
-			Fence.BeginFence();
-			Fence.Wait();
-		}
+		AudioDevice->SuspendContext();
 	}
 }
 
@@ -365,14 +353,7 @@ void FAppEventManager::ResumeAudio()
 	FAudioDeviceHandle AudioDevice = GEngine->GetMainAudioDevice();
 	if (AudioDevice)
 	{
-		if (AudioDevice->IsAudioMixerEnabled())
-		{
-			AudioDevice->ResumeContext();
-		}
-		else
-		{
-			GEngine->GetMainAudioDevice()->Suspend(true);
-		}
+		AudioDevice->ResumeContext();
 	}
 }
 
diff --git a/Engine/Source/Runtime/SoundFieldRendering/Private/SoundFieldRendering.cpp b/Engine/Source/Runtime/SoundFieldRendering/Private/SoundFieldRendering.cpp
--- a/Engine/Source/Runtime/SoundFieldRendering/Private/SoundFieldRendering.cpp
+++ b/Engine/Source/Runtime/SoundFieldRendering/Private/SoundFieldRendering.cpp
@@ -641,14 +641,10 @@ TArray<Audio::FChannelPositionInfo>* FSoundFieldDecoder::GetDefaultChannelPositi
 	if (GEngine)
 	{
 		FAudioDevice* AudioDevice = GEngine->GetMainAudioDeviceRaw();
-		if (AudioDevice->IsAudioMixerEnabled())
+		Audio::FMixerDevice* MixerDevice = static_cast<Audio::FMixerDevice*>(AudioDevice);
+		if (MixerDevice)
 		{
-			// downcast.
-			Audio::FMixerDevice* MixerDevice = static_cast<Audio::FMixerDevice*>(AudioDevice);
-			if (MixerDevice)
-			{
-				return MixerDevice->GetDefaultPositionMap(InNumChannels);
-			}
+			return MixerDevice->GetDefaultPositionMap(InNumChannels);
 		}
 	}
 
