(Github PR) Override Runtime Format in FSoundWaveData
#9885

diff --git a/Engine/Source/Runtime/Engine/Classes/Sound/SoundWave.h b/Engine/Source/Runtime/Engine/Classes/Sound/SoundWave.h
--- a/Engine/Source/Runtime/Engine/Classes/Sound/SoundWave.h
+++ b/Engine/Source/Runtime/Engine/Classes/Sound/SoundWave.h
@@ -1283,6 +1283,8 @@ public:
 	// non-const arg for internal call to FAudioDevice::GetRuntimeFormat
 	void InitializeDataFromSoundWave(USoundWave& InWave);
 
+	void OverrideRuntimeFormat(const FName& InRuntimeFormat);
+
 	const FName& GetFName() const { return NameCached; }
 	const FName& GetPackageName() const { return PackageNameCached; }
 	const FName& GetRuntimeFormat() const { return RuntimeFormat; }
diff --git a/Engine/Source/Runtime/Engine/Private/SoundWave.cpp b/Engine/Source/Runtime/Engine/Private/SoundWave.cpp
--- a/Engine/Source/Runtime/Engine/Private/SoundWave.cpp
+++ b/Engine/Source/Runtime/Engine/Private/SoundWave.cpp
@@ -172,6 +172,11 @@ void FSoundWaveData::InitializeDataFromSoundWave(USoundWave& InWave)
 #endif //WITH_EDITOR
 }
 
+void FSoundWaveData::OverrideRuntimeFormat(const FName& InRuntimeFormat)
+{
+	RuntimeFormat = InRuntimeFormat;
+}
+
 FName FSoundWaveData::FindRuntimeFormat(const USoundWave& InWave) const
 {		
 #if WITH_EDITOR	
