Reenable level instance partial loading with the following rules:
- If the main world has streaming disabled, don't partially load.
- If the level instance (or any parent level instance) is in edit mode, don't partially load.
- If the level instance actor (or any parent actor) is not spatially loaded, don't partially load.
- If the level instance actor (or any parent actor) is dirty (has pending save changes), don't partially load.

diff --git a/Engine/Source/Runtime/Engine/Private/LevelInstance/LevelInstanceInterface.cpp b/Engine/Source/Runtime/Engine/Private/LevelInstance/LevelInstanceInterface.cpp
--- a/Engine/Source/Runtime/Engine/Private/LevelInstance/LevelInstanceInterface.cpp
+++ b/Engine/Source/Runtime/Engine/Private/LevelInstance/LevelInstanceInterface.cpp
@@ -30,12 +30,14 @@ bool ILevelInstanceInterface::SupportsPartialEditorLoading() const
 
 			LevelInstanceSubsystem->ForEachLevelInstanceAncestorsAndSelf(Actor, [LevelInstanceSubsystem, &bResult](const ILevelInstanceInterface* LevelInstance)
 			{
+				// If the level instance (or any parent level instance) is in edit mode, don't partially load.
 				if (LevelInstanceSubsystem->IsEditingLevelInstance(LevelInstance))
 				{
 					bResult = false;
 					return false;
 				}
 
+				// If the level instance actor (or any parent actor) is not spatially loaded, don't partially load.
 				const AActor* LevelInstanceActor = CastChecked<AActor>(LevelInstance);
 				if (!LevelInstanceActor->GetIsSpatiallyLoaded())
 				{
@@ -43,6 +45,14 @@ bool ILevelInstanceInterface::SupportsPartialEditorLoading() const
 					return false;
 				}
 
+				// If the level instance actor (or any parent actor) is dirty (has pending save changes), don't partially load.
+				if (LevelInstanceActor->GetPackage()->IsDirty())
+				{
+					bResult = false;
+					return false;
+				}
+
+				// We can partially load
 				return true;
 			});
 
diff --git a/Engine/Source/Runtime/Engine/Private/LevelInstance/LevelInstanceLevelStreaming.cpp b/Engine/Source/Runtime/Engine/Private/LevelInstance/LevelInstanceLevelStreaming.cpp
--- a/Engine/Source/Runtime/Engine/Private/LevelInstance/LevelInstanceLevelStreaming.cpp
+++ b/Engine/Source/Runtime/Engine/Private/LevelInstance/LevelInstanceLevelStreaming.cpp
@@ -27,7 +27,7 @@
 #include "Misc/ScopeExit.h"
 #include "UObject/LinkerLoad.h"
 
-static bool GDisableLevelInstanceEditorPartialLoading = true;
+static bool GDisableLevelInstanceEditorPartialLoading = false;
 FAutoConsoleVariableRef CVarDisableLevelInstanceEditorPartialLoading(
 	TEXT("wp.Editor.DisableLevelInstanceEditorPartialLoading"),
 	GDisableLevelInstanceEditorPartialLoading,
@@ -365,21 +365,13 @@ void ULevelStreamingLevelInstance::OnLevelLoadedChanged(ULevel* InLevel)
 			LevelInstanceSubsystem->RegisterLoadedLevelStreamingLevelInstance(this);
 
 #if WITH_EDITOR
-			if (!GDisableLevelInstanceEditorPartialLoading && GetWorld()->IsPartitionedWorld())
+			if (UWorldPartition* OwningWorldPartition = GetWorld()->GetWorldPartition(); OwningWorldPartition && OwningWorldPartition->IsStreamingEnabled() && !GDisableLevelInstanceEditorPartialLoading)
 			{
-				if (UWorldPartition* WorldPartition = NewLoadedLevel->GetWorldPartition())
+				if (UWorldPartition* OuterWorldPartition = NewLoadedLevel->GetWorldPartition())
 				{
-					check(!WorldPartition->IsInitialized());
-					if (ILevelInstanceInterface* LevelInstance = LevelInstanceSubsystem->GetLevelInstance(LevelInstanceID))
-					{
-						if (AActor* Actor = CastChecked<AActor>(LevelInstance))
-						{
-							if (!Actor->GetPackage()->IsDirty())
-							{
-								WorldPartition->bForceEnableStreamingInEditor = LevelInstance->SupportsPartialEditorLoading();
-							}
-						}
-					}
+					check(!OuterWorldPartition->IsInitialized());
+					ILevelInstanceInterface* LevelInstance = LevelInstanceSubsystem->GetLevelInstance(LevelInstanceID);
+					OuterWorldPartition->bForceEnableStreamingInEditor = LevelInstance ? LevelInstance->SupportsPartialEditorLoading() : false;
 				}
 			}
 #endif
