Another set of additional data added for softlock issue, also added BeginCache call to theoretically remove the soft-locking as data will be fetched according to expected hash

diff --git a/Engine/Source/Runtime/Engine/Classes/Animation/AnimSequence.h b/Engine/Source/Runtime/Engine/Classes/Animation/AnimSequence.h
--- a/Engine/Source/Runtime/Engine/Classes/Animation/AnimSequence.h
+++ b/Engine/Source/Runtime/Engine/Classes/Animation/AnimSequence.h
@@ -896,6 +896,7 @@ protected:
 		FDateTime LastWillNeverCook;
 	};
 	TMap<FIoHash, FCompressionRequestTimings> TimingsByKeyHash;
+	TMap<FIoHash, FString> DDCStringByKeyHash;
 protected:
 	bool TryCancelAsyncTasks();
 	void FinishAsyncTasks();
diff --git a/Engine/Source/Runtime/Engine/Private/Animation/AnimSequence.cpp b/Engine/Source/Runtime/Engine/Private/Animation/AnimSequence.cpp
--- a/Engine/Source/Runtime/Engine/Private/Animation/AnimSequence.cpp
+++ b/Engine/Source/Runtime/Engine/Private/Animation/AnimSequence.cpp
@@ -1106,6 +1106,12 @@ bool UAnimSequence::IsCachedCookedPlatformDataLoaded(const ITargetPlatform* Targ
 		UE_LOG(LogAnimation, Warning, TEXT("Expecting either valid compressed Animation Data, or an in-flight compression task for asset %s\nPlatform: %s\nHash: %s"), *GetPathName(), *TargetPlatform->DisplayName().ToString(), *LexToString(KeyHash));
 
 		LogCompressionRequestTimings();
+
+		if (IsInGameThread())
+		{
+			// Temporary fix for soft-lock, to be removed when root-cause is addressed
+			BeginCacheForCookedPlatformData(TargetPlatform);			
+		}
 	};
 
 	if (KeyHash == DataKeyHash)
@@ -5042,7 +5048,9 @@ FIoHash UAnimSequence::CreateDerivedDataKeyHash(const ITargetPlatform* TargetPla
 
 	Writer << Ret;
 
-	return Writer.Finalize();
+	const FIoHash FinalHash = Writer.Finalize();
+	DDCStringByKeyHash.FindOrAdd(FinalHash) = Ret;
+	return FinalHash;
 }
 
 FIoHash UAnimSequence::BeginCacheDerivedData(const ITargetPlatform* TargetPlatform)
@@ -5185,6 +5193,13 @@ void UAnimSequence::LogCompressionRequestTimings() const
 			*TimingsPair.Value.LastClearCache.ToString(),
 			*TimingsPair.Value.LastWillNeverCook.ToString());
 	}
+	
+	for (const TPair<FIoHash, FString>& StringPair : DDCStringByKeyHash)
+	{
+		UE_LOG(LogAnimation, Warning, TEXT("Hash: %s\nDDC key: %s"),
+		*LexToString(StringPair.Key),
+		*StringPair.Value);
+	}
 }
 
 FCompressedAnimSequence& UAnimSequence::CacheDerivedData(const ITargetPlatform* TargetPlatform)
