StaticMesh: Fixed race condition when undoing UStaticMesh by canceling distance field and mesh card build tasks.

diff --git a/Engine/Source/Runtime/Engine/Private/StaticMesh.cpp b/Engine/Source/Runtime/Engine/Private/StaticMesh.cpp
--- a/Engine/Source/Runtime/Engine/Private/StaticMesh.cpp
+++ b/Engine/Source/Runtime/Engine/Private/StaticMesh.cpp
@@ -3913,6 +3913,22 @@ void UStaticMesh::PreEditChange(FProperty* PropertyAboutToChange)
 		FStaticMeshCompilingManager::Get().FinishCompilation({ this });
 	}
 
+	// We need to cancel these builds manually since they rely on PostGC reachability analysis to delete invalid tasks themselves.
+	// If the mesh is invalidated and the async task executes before GC we may attempt to build with invalid data.
+	// PostEditChange will enqueue builds for each of these again so it's okay to cancel them completely here.
+	// #todo: this should be modified to use the FAssetCompilationManager::Get().FinishAllCompilationForObjects function in 5.2
+	{
+		if (GDistanceFieldAsyncQueue)
+		{
+			GDistanceFieldAsyncQueue->CancelBuild(this);
+		}
+		if (GCardRepresentationAsyncQueue)
+		{
+			GCardRepresentationAsyncQueue->CancelBuild(this);
+		}
+	}
+
+
 	Super::PreEditChange(PropertyAboutToChange);
 
 	// Release the static mesh's resources.
