WorldPartiton: Better handling of Streaming out sub WorldPartitions
- UWorldPartition::CanStream() now returns false if WP is in a LevelStreaming with ShouldBeVisible() == false
- UWorldPartitionSubsystem calls UpdateStreamingState on streaming out sub WorldPartitions (so that it can unload all its cells asap)

I have been unable to reproduce the crash but it was possible for WP Cells to not be in the proper state which could lead to this kind of issue.

#rnx

diff --git a/Engine/Source/Runtime/Engine/Private/WorldPartition/WorldPartition.cpp b/Engine/Source/Runtime/Engine/Private/WorldPartition/WorldPartition.cpp
--- a/Engine/Source/Runtime/Engine/Private/WorldPartition/WorldPartition.cpp
+++ b/Engine/Source/Runtime/Engine/Private/WorldPartition/WorldPartition.cpp
@@ -17,6 +17,7 @@
 #include "ProfilingDebugging/ScopedTimers.h"
 #include "UObject/UE5MainStreamObjectVersion.h"
 #include "LandscapeProxy.h"
+#include "Engine/LevelStreaming.h"
 
 #include UE_INLINE_GENERATED_CPP_BY_NAME(WorldPartition)
 
@@ -766,9 +767,26 @@ bool UWorldPartition::IsStreamingEnabled() const
 
 bool UWorldPartition::CanStream() const
 {
-	// WorldPartition can't stream if not initialized or it it's part of a partitioned 
-	// sub-level that was removed from its owning world.
-	return IsInitialized() && (GetTypedOuter<UWorld>()->PersistentLevel->GetWorld() != nullptr);
+	if (!IsInitialized())
+	{
+		return false;
+	}
+
+	const ULevel* PersistentLevel = GetTypedOuter<UWorld>()->PersistentLevel;
+	// Is it a level streamed World Partition that was removed from its owning world
+	// or is the World requesting unloading of all streaming levels.
+	if (!PersistentLevel->GetWorld() || PersistentLevel->GetWorld()->GetShouldForceUnloadStreamingLevels())
+	{
+		return false;
+	}
+		
+	// Is it part of a Sub-level that should be visible.
+	if (ULevelStreaming* LevelStreaming = ULevelStreaming::FindStreamingLevel(PersistentLevel))
+	{
+		return LevelStreaming->ShouldBeVisible();
+	}
+
+	return true;
 }
 
 bool UWorldPartition::IsMainWorldPartition() const
diff --git a/Engine/Source/Runtime/Engine/Private/WorldPartition/WorldPartitionSubsystem.cpp b/Engine/Source/Runtime/Engine/Private/WorldPartition/WorldPartitionSubsystem.cpp
--- a/Engine/Source/Runtime/Engine/Private/WorldPartition/WorldPartitionSubsystem.cpp
+++ b/Engine/Source/Runtime/Engine/Private/WorldPartition/WorldPartitionSubsystem.cpp
@@ -16,6 +16,8 @@
 #include "WorldPartition/ContentBundle/ContentBundleWorldSubsystem.h"
 #include "Engine/Canvas.h"
 #include "Engine/CoreSettings.h"
+#include "Engine/LevelStreaming.h"
+#include "Streaming/LevelStreamingDelegates.h"
 #include "Engine/LevelBounds.h"
 #include "Debug/DebugDrawService.h"
 #include "GameFramework/PlayerController.h"
@@ -405,6 +407,7 @@ void UWorldPartitionSubsystem::Initialize(FSubsystemCollectionBase& Collection)
 
 	GetWorld()->OnWorldPartitionInitialized().AddUObject(this, &UWorldPartitionSubsystem::OnWorldPartitionInitialized);
 	GetWorld()->OnWorldPartitionUninitialized().AddUObject(this, &UWorldPartitionSubsystem::OnWorldPartitionUninitialized);
+	FLevelStreamingDelegates::OnLevelStreamingTargetStateChanged.AddUObject(this, &UWorldPartitionSubsystem::OnLevelStreamingTargetStateChanged);
 }
 
 void UWorldPartitionSubsystem::Deinitialize()
@@ -419,6 +422,7 @@ void UWorldPartitionSubsystem::Deinitialize()
 
 	GetWorld()->OnWorldPartitionInitialized().RemoveAll(this);
 	GetWorld()->OnWorldPartitionUninitialized().RemoveAll(this);
+	FLevelStreamingDelegates::OnLevelStreamingTargetStateChanged.RemoveAll(this);
 
 	// At this point World Partition should be uninitialized
 	check(!GetWorldPartition() || !GetWorldPartition()->IsInitialized());
@@ -486,6 +490,25 @@ void UWorldPartitionSubsystem::OnWorldPartitionUninitialized(UWorldPartition* In
 	}
 }
 
+void UWorldPartitionSubsystem::OnLevelStreamingTargetStateChanged(UWorld* World, const ULevelStreaming* StreamingLevel, ULevel* LevelIfLoaded, ELevelStreamingState CurrentState, ELevelStreamingTargetState PrevTarget, ELevelStreamingTargetState NewTarget)
+{
+	if (World != GetWorld())
+	{
+		return;
+	}
+
+	// Make sure when a WorldPartiton is LevelStreamed that changing its state to remove it from world will update the target states of its Cells right away.
+	if(LevelIfLoaded && NewTarget != ELevelStreamingTargetState::LoadedVisible)
+	{
+		// At this point the StreamingLevel should not be visible or the global flag on the world should have been set to unload all streaming levels
+		ensure(!StreamingLevel->ShouldBeVisible() || World->GetShouldForceUnloadStreamingLevels());
+		if (UWorldPartition* WorldPartition = LevelIfLoaded->GetTypedOuter<UWorld>()->GetWorldPartition(); WorldPartition && WorldPartition->IsInitialized())
+		{
+			WorldPartition->UpdateStreamingState();
+		}
+	}
+}
+
 void UWorldPartitionSubsystem::RegisterStreamingSourceProvider(IWorldPartitionStreamingSourceProvider* StreamingSource)
 {
 	bool bIsAlreadyInSet = false;
diff --git a/Engine/Source/Runtime/Engine/Public/WorldPartition/WorldPartitionSubsystem.h b/Engine/Source/Runtime/Engine/Public/WorldPartition/WorldPartitionSubsystem.h
--- a/Engine/Source/Runtime/Engine/Public/WorldPartition/WorldPartitionSubsystem.h
+++ b/Engine/Source/Runtime/Engine/Public/WorldPartition/WorldPartitionSubsystem.h
@@ -16,6 +16,8 @@ class UActorDescContainer;
 class FWorldPartitionActorDesc;
 
 enum class EWorldPartitionRuntimeCellState : uint8;
+enum class ELevelStreamingState : uint8;
+enum class ELevelStreamingTargetState : uint8;
 
 /**
  * Helper to compute streaming source velocity based on position history.
@@ -160,6 +162,7 @@ private:
 
 	void OnWorldPartitionInitialized(UWorldPartition* InWorldPartition);
 	void OnWorldPartitionUninitialized(UWorldPartition* InWorldPartition);
+	void OnLevelStreamingTargetStateChanged(UWorld* World, const ULevelStreaming* StreamingLevel, ULevel* LevelIfLoaded, ELevelStreamingState CurrentState, ELevelStreamingTargetState PrevTarget, ELevelStreamingTargetState NewTarget);
 
 	UWorldPartition* GetWorldPartition();
 	const UWorldPartition* GetWorldPartition() const;
