World Partition Streaming Sources Updates:
- Allow streaming sources quantization even if only location or rotation is enabled (this allows disabling rotation from having an impact on streaming sources updates).
- Compute a simple distance between old and new values before quantization to implement simple streaming sources hysteresis.
- Defaults wp.Runtime.UpdateStreaming.ForceUpdateFrameCount to zero.
- Always enable streaming source optimizations on the server.
- Removed deprecated code from Release 5.1 version.

diff --git a/Engine/Source/Runtime/Engine/Private/WorldPartition/WorldPartitionStreamingPolicy.cpp b/Engine/Source/Runtime/Engine/Private/WorldPartition/WorldPartitionStreamingPolicy.cpp
--- a/Engine/Source/Runtime/Engine/Private/WorldPartition/WorldPartitionStreamingPolicy.cpp
+++ b/Engine/Source/Runtime/Engine/Private/WorldPartition/WorldPartitionStreamingPolicy.cpp
@@ -52,7 +52,7 @@ FAutoConsoleVariableRef UWorldPartitionStreamingPolicy::CVarUpdateOptimEnabled(
 	TEXT("if nothing relevant changed since last update."),
 	ECVF_Default);
 
-int32 UWorldPartitionStreamingPolicy::ForceUpdateFrameCount = 30;
+int32 UWorldPartitionStreamingPolicy::ForceUpdateFrameCount = 0;
 FAutoConsoleVariableRef UWorldPartitionStreamingPolicy::CVarForceUpdateFrameCount(
 	TEXT("wp.Runtime.UpdateStreaming.ForceUpdateFrameCount"),
 	UWorldPartitionStreamingPolicy::ForceUpdateFrameCount,
@@ -150,8 +150,8 @@ int32 UWorldPartitionStreamingPolicy::ComputeServerStreamingEnabledEpoch() const
 bool UWorldPartitionStreamingPolicy::IsUpdateStreamingOptimEnabled()
 {
 	return UWorldPartitionStreamingPolicy::IsUpdateOptimEnabled &&
-		(FWorldPartitionStreamingSource::GetLocationQuantization() > 0) &&
-		(FWorldPartitionStreamingSource::GetRotationQuantization() > 0);
+		((FWorldPartitionStreamingSource::GetLocationQuantization() > 0) ||
+		 (FWorldPartitionStreamingSource::GetRotationQuantization() > 0));
 };
 
 uint32 UWorldPartitionStreamingPolicy::ComputeUpdateStreamingHash(bool bCanOptimizeUpdate) const
@@ -274,7 +274,6 @@ void UWorldPartitionStreamingPolicy::UpdateStreamingState()
 		WorldPartition->RuntimeHash &&
 		!bForceFrameUpdate &&										// We garantee to update every N frame to force some internal updates like UpdateStreamingPerformance
 		IsUpdateStreamingOptimEnabled() &&							// Check CVars to see if optimization is enabled
-		!WorldPartition->IsServer() &&								// Don't optimize on the server
 		bLastUpdateCompletedLoadingAndActivation &&					// Don't optimize if last frame didn't process all cells to load/activate
 		!IsInBlockTillLevelStreamingCompleted() &&					// Don't optimize when inside UWorld::BlockTillLevelStreamingCompleted
 		ActivatedCells.GetPendingAddToWorldCells().IsEmpty(); 		// Don't optimize when remaining cells to add to world
diff --git a/Engine/Source/Runtime/Engine/Private/WorldPartition/WorldPartitionStreamingSource.cpp b/Engine/Source/Runtime/Engine/Private/WorldPartition/WorldPartitionStreamingSource.cpp
--- a/Engine/Source/Runtime/Engine/Private/WorldPartition/WorldPartitionStreamingSource.cpp
+++ b/Engine/Source/Runtime/Engine/Private/WorldPartition/WorldPartitionStreamingSource.cpp
@@ -1,6 +1,7 @@
 // Copyright Epic Games, Inc. All Rights Reserved.
 
 #include "WorldPartition/WorldPartitionStreamingSource.h"
+#include "WorldPartition/WorldPartitionStreamingPolicy.h"
 #include "Misc/HashBuilder.h"
 
 #include UE_INLINE_GENERATED_CPP_BY_NAME(WorldPartitionStreamingSource)
@@ -21,39 +22,46 @@ FAutoConsoleVariableRef FWorldPartitionStreamingSource::CVarRotationQuantization
 
 void FWorldPartitionStreamingSource::UpdateHash()
 {
-	FHashBuilder HashBuilder;
-	HashBuilder << Name;
-	HashBuilder << FMath::FloorToInt(Location.X / FWorldPartitionStreamingSource::LocationQuantization);
-	HashBuilder << FMath::FloorToInt(Location.Y / FWorldPartitionStreamingSource::LocationQuantization);
-	HashBuilder << FMath::FloorToInt(Rotation.Yaw / FWorldPartitionStreamingSource::RotationQuantization);
-	HashBuilder << TargetState;
-	HashBuilder << bBlockOnSlowLoading;
-	HashBuilder << bReplay;
-	HashBuilder << bRemote;
-	HashBuilder << Priority;
-	HashBuilder << TargetBehavior;
-	for (FName Grid : TargetGrids)
+	// Update old values when they are changing enough, to avoid the case where we are on the edge of a quantization unit.
+	if (!UWorldPartitionStreamingPolicy::IsUpdateStreamingOptimEnabled() || 
+		(FVector::Dist(Location, OldLocation) > FWorldPartitionStreamingSource::LocationQuantization))
 	{
-		HashBuilder << Grid;
+		OldLocation = Location;
 	}
-	for (const FSoftObjectPath& HLODLayer : TargetHLODLayers)
+	
+	if (!UWorldPartitionStreamingPolicy::IsUpdateStreamingOptimEnabled() ||
+		(FMath::Abs(Rotation.Pitch - OldRotation.Pitch) > FWorldPartitionStreamingSource::RotationQuantization) ||
+		(FMath::Abs(Rotation.Yaw - OldRotation.Yaw) > FWorldPartitionStreamingSource::RotationQuantization) ||
+		(FMath::Abs(Rotation.Roll - OldRotation.Roll) > FWorldPartitionStreamingSource::RotationQuantization))
 	{
-		HashBuilder << HLODLayer;
+		OldRotation = Rotation;
 	}
-	PRAGMA_DISABLE_DEPRECATION_WARNINGS
-	HashBuilder << TargetGrid;
-	HashBuilder << TargetHLODLayer;
-	PRAGMA_ENABLE_DEPRECATION_WARNINGS
-	for (const FStreamingSourceShape& Shape : Shapes)
+
+	FHashBuilder HashBuilder;
+	HashBuilder	<< Name << TargetState << bBlockOnSlowLoading << bReplay << bRemote << Priority << TargetBehavior << TargetGrids << TargetHLODLayers << Shapes;
+
+	if (LocationQuantization)
 	{
-		HashBuilder << Shape;
+		HashBuilder << FMath::FloorToInt(OldLocation.X / LocationQuantization) << FMath::FloorToInt(OldLocation.Y / LocationQuantization);
 	}
+
+	if (RotationQuantization > 0)
+	{
+		HashBuilder << FMath::FloorToInt(OldRotation.Yaw / RotationQuantization);
+	}
+
 	Hash2D = HashBuilder.GetHash();
-	
-	// Hash3D
-	HashBuilder << FMath::FloorToInt(Location.Z / FWorldPartitionStreamingSource::LocationQuantization);
-	HashBuilder << FMath::FloorToInt(Rotation.Pitch / FWorldPartitionStreamingSource::RotationQuantization);
-	HashBuilder << FMath::FloorToInt(Rotation.Roll / FWorldPartitionStreamingSource::RotationQuantization);
+
+	if (LocationQuantization)
+	{
+		HashBuilder << FMath::FloorToInt(OldLocation.Z / LocationQuantization);
+	}
+
+	if (RotationQuantization)
+	{
+		HashBuilder << FMath::FloorToInt(OldRotation.Pitch / RotationQuantization) << FMath::FloorToInt(OldRotation.Roll / RotationQuantization);
+	}
+
 	Hash3D = HashBuilder.GetHash();
 };
 
diff --git a/Engine/Source/Runtime/Engine/Public/WorldPartition/WorldPartitionStreamingSource.h b/Engine/Source/Runtime/Engine/Public/WorldPartition/WorldPartitionStreamingSource.h
--- a/Engine/Source/Runtime/Engine/Public/WorldPartition/WorldPartitionStreamingSource.h
+++ b/Engine/Source/Runtime/Engine/Public/WorldPartition/WorldPartitionStreamingSource.h
@@ -275,12 +275,6 @@ struct FWorldPartitionStreamingQuerySource
 		TargetGrids = Other.TargetGrids;
 		TargetHLODLayers = Other.TargetHLODLayers;
 		Shapes = Other.Shapes;
-
-		PRAGMA_DISABLE_DEPRECATION_WARNINGS
-		TargetGrid = Other.TargetGrid;
-		TargetHLODLayer = Other.TargetHLODLayer;
-		PRAGMA_ENABLE_DEPRECATION_WARNINGS
-
 		return *this;
 	}
 
@@ -338,17 +332,6 @@ struct FWorldPartitionStreamingQuerySource
 		{
 			FStreamingSourceShapeHelper::ForEachShape(InGridLoadingRange, bUseGridLoadingRange ? InGridLoadingRange : Radius, bInProjectIn2D, Location, Rotation, Shapes, InOperation);
 		}
-
-		PRAGMA_DISABLE_DEPRECATION_WARNINGS
-		if(!TargetGrid.IsNone() || !TargetHLODLayer.IsNull())
-		{
-
-			if (FStreamingSourceShapeHelper::IsSourceAffectingGrid({ TargetGrid }, { TargetHLODLayers }, TargetBehavior, InGridName, InGridHLODLayer))
-			{
-				FStreamingSourceShapeHelper::ForEachShape(InGridLoadingRange, bUseGridLoadingRange ? InGridLoadingRange : Radius, bInProjectIn2D, Location, Rotation, Shapes, InOperation);
-			}
-		}
-		PRAGMA_ENABLE_DEPRECATION_WARNINGS
 	}
 };
 
@@ -385,6 +368,8 @@ struct ENGINE_API FWorldPartitionStreamingSource
 		, bRemote(false)
 		, Hash2D(0)
 		, Hash3D(0)
+		, OldLocation(FVector::ZeroVector)
+		, OldRotation(FRotator::ZeroRotator)
 	{}
 
 	FWorldPartitionStreamingSource(FName InName, const FVector& InLocation, const FRotator& InRotation, EStreamingSourceTargetState InTargetState, bool bInBlockOnSlowLoading, EStreamingSourcePriority InPriority, bool bRemote, float InVelocity = 0.f)
@@ -401,6 +386,8 @@ struct ENGINE_API FWorldPartitionStreamingSource
 		, bRemote(bRemote)
 		, Hash2D(0)
 		, Hash3D(0)
+		, OldLocation(InLocation)
+		, OldRotation(InRotation)
 	{}
 
 	// Define Copy Constructor to avoid deprecation warnings
@@ -427,12 +414,6 @@ struct ENGINE_API FWorldPartitionStreamingSource
 		bRemote = Other.bRemote;
 		Hash2D = Other.Hash2D;
 		Hash3D = Other.Hash3D;
-
-		PRAGMA_DISABLE_DEPRECATION_WARNINGS
-		TargetGrid = Other.TargetGrid;
-		TargetHLODLayer = Other.TargetHLODLayer;
-		PRAGMA_ENABLE_DEPRECATION_WARNINGS
-
 		return *this;
 	}
 	
@@ -454,7 +435,7 @@ struct ENGINE_API FWorldPartitionStreamingSource
 
 	/** Source location. */
 	FVector Location;
-	
+
 	/** Source orientation (can impact streaming cell prioritization). */
 	FRotator Rotation;
 
@@ -476,15 +457,9 @@ struct ENGINE_API FWorldPartitionStreamingSource
 	/** Defines how TargetGrids/TargetHLODLayers will be applied to this streaming source. */
 	EStreamingSourceTargetBehavior TargetBehavior;
 
-	UE_DEPRECATED(5.1, "Use TargetGrids instead.")
-	FName TargetGrid;
-
 	/** When set, this will change how this streaming source is applied to the provided runtime streaming grids based on the TargetBehavior. */
 	TSet<FName> TargetGrids;
 	
-	UE_DEPRECATED(5.1, "Use TargetHLODLayers")
-	FSoftObjectPath TargetHLODLayer;
-
 	/** When set, this will change how this streaming source is applied to the provided HLODLayers based on the TargetBehavior. */
 	TSet<FSoftObjectPath> TargetHLODLayers;
 
@@ -515,16 +490,6 @@ struct ENGINE_API FWorldPartitionStreamingSource
 		{
 			FStreamingSourceShapeHelper::ForEachShape(InGridLoadingRange, InGridLoadingRange, bInProjectIn2D, Location, Rotation, Shapes, InOperation);
 		}
-
-		PRAGMA_DISABLE_DEPRECATION_WARNINGS
-		if (!TargetGrid.IsNone() || !TargetHLODLayer.IsNull())
-		{
-			if (FStreamingSourceShapeHelper::IsSourceAffectingGrid({ TargetGrid }, { TargetHLODLayers }, TargetBehavior, InGridName, InGridHLODLayer))
-			{
-				FStreamingSourceShapeHelper::ForEachShape(InGridLoadingRange, InGridLoadingRange, bInProjectIn2D, Location, Rotation, Shapes, InOperation);
-			}
-		}
-		PRAGMA_ENABLE_DEPRECATION_WARNINGS
 	}
 
 	FString ToString() const
@@ -547,7 +512,6 @@ struct ENGINE_API FWorldPartitionStreamingSource
 	static int32 GetRotationQuantization() { return RotationQuantization; }
 
 private:
-
 	static int32 LocationQuantization;
 	static int32 RotationQuantization;
 	static FAutoConsoleVariableRef CVarLocationQuantization;
@@ -556,6 +520,10 @@ private:
 	/** Hash of streaming source (used to detect changes) */
 	uint32 Hash2D;
 	uint32 Hash3D;
+
+	/** Source values used for hash computations. */
+	FVector OldLocation;
+	FRotator OldRotation;
 };
 
 /**
