Fix Non-unity

#rnx

diff --git a/Engine/Source/Runtime/Engine/Public/WorldPartition/WorldPartitionSubsystem.h b/Engine/Source/Runtime/Engine/Public/WorldPartition/WorldPartitionSubsystem.h
--- a/Engine/Source/Runtime/Engine/Public/WorldPartition/WorldPartitionSubsystem.h
+++ b/Engine/Source/Runtime/Engine/Public/WorldPartition/WorldPartitionSubsystem.h
@@ -11,6 +11,8 @@
 #include "Misc/Guid.h"
 #include "WorldPartitionSubsystem.generated.h"
 
+class ULevel;
+class ULevelStreaming;
 class UWorldPartition;
 class UActorDescContainer;
 class FWorldPartitionActorDesc;
