Undo changelist 25124114 ( doesn't work for interfaces ).

diff --git a/Engine/Source/Runtime/Engine/Classes/Engine/ActorInstanceHandle.h b/Engine/Source/Runtime/Engine/Classes/Engine/ActorInstanceHandle.h
--- a/Engine/Source/Runtime/Engine/Classes/Engine/ActorInstanceHandle.h
+++ b/Engine/Source/Runtime/Engine/Classes/Engine/ActorInstanceHandle.h
@@ -26,9 +26,6 @@ struct ENGINE_API FActorInstanceHandle
 
 	bool DoesRepresentClass(const UClass* OtherClass) const;
 
-	template<typename T>
-	bool DoesRepresentClass() const;
-
 	UClass* GetRepresentedClass() const;
 
 	FVector GetLocation() const;
@@ -96,25 +93,8 @@ private:
 	uint32 InstanceUID;
 };
 
-template<typename T>
-bool FActorInstanceHandle::DoesRepresentClass() const
-{
-	if constexpr (TIsIInterface<T>::Value)
-	{
-		return DoesRepresentClass(T::UClassType::StaticClass());
-	}
-	else
-	{
-		return DoesRepresentClass(T::StaticClass());
-	}
-}
-
 template <typename T>
 T* FActorInstanceHandle::FetchActor() const
 {
-	if (DoesRepresentClass<T>())
-	{
-		return Cast<T>(FetchActor());
-	}
-	return nullptr;
+	return Cast<T>(FetchActor());
 }
