PR #9836: Optimization - Avoid unnecessary audio chunk cache look-up at end-of-sound (Contributed by [REDACTED_PERSON])

#rnx

diff --git a/Engine/Source/Runtime/Engine/Private/AudioDecompress.cpp b/Engine/Source/Runtime/Engine/Private/AudioDecompress.cpp
--- a/Engine/Source/Runtime/Engine/Private/AudioDecompress.cpp
+++ b/Engine/Source/Runtime/Engine/Private/AudioDecompress.cpp
@@ -304,6 +304,7 @@ bool IStreamedCompressedInfo::StreamCompressedData(uint8* Destination, bool bLoo
 					{
 						RawPCMOffset += ZeroBuffer(Destination + RawPCMOffset, BufferSize - RawPCMOffset);
 						OutNumBytesStreamed = BufferSize - RawPCMOffset;
+						break;
 					}
 				}
 				else
