World Partition Debug Console Commands:
- wp.Editor.DumpActorDescs: sort actor descriptors bu guids when dumping in a csv file.
- wp.Editor.DumpClassDescs: new command to dump class descriptors in a csv file.

diff --git a/Engine/Source/Runtime/Engine/Private/WorldPartition/WorldPartition.cpp b/Engine/Source/Runtime/Engine/Private/WorldPartition/WorldPartition.cpp
--- a/Engine/Source/Runtime/Engine/Private/WorldPartition/WorldPartition.cpp
+++ b/Engine/Source/Runtime/Engine/Private/WorldPartition/WorldPartition.cpp
@@ -1737,7 +1737,7 @@ void UWorldPartition::DumpActorDescs(const FString& Path)
 		}
 		ActorDescs.Sort([](const FWorldPartitionActorDesc& A, const FWorldPartitionActorDesc& B)
 		{
-			return A.GetRuntimeBounds().GetExtent().GetMax() < B.GetRuntimeBounds().GetExtent().GetMax();
+			return A.GetGuid() < B.GetGuid();
 		});
 		for (const FWorldPartitionActorDesc* ActorDescIterator : ActorDescs)
 		{
diff --git a/Engine/Source/Runtime/Engine/Private/WorldPartition/WorldPartitionClassDescRegistry.cpp b/Engine/Source/Runtime/Engine/Private/WorldPartition/WorldPartitionClassDescRegistry.cpp
--- a/Engine/Source/Runtime/Engine/Private/WorldPartition/WorldPartitionClassDescRegistry.cpp
+++ b/Engine/Source/Runtime/Engine/Private/WorldPartition/WorldPartitionClassDescRegistry.cpp
@@ -44,6 +44,40 @@
  *	  - Create or update the class descriptor.
  */
 
+static FAutoConsoleCommand DumpClassDescs(
+	TEXT("wp.Editor.DumpClassDescs"),
+	TEXT("Dump the list of class descriptors in a CSV file."),
+	FConsoleCommandWithArgsDelegate::CreateLambda([](const TArray<FString>& Args)
+	{
+		if (Args.Num() > 0)
+		{
+			const FString& Path = Args[0];
+			if (FArchive* LogFile = IFileManager::Get().CreateFileWriter(*Path))
+			{
+				TArray<const FWorldPartitionActorDesc*> ClassDescs;
+				const FWorldPartitionClassDescRegistry& ClassDescRegistry = FWorldPartitionClassDescRegistry::Get();
+
+				for (FWorldPartitionClassDescRegistry::TConstIterator<> ClassDescIt(&ClassDescRegistry); ClassDescIt; ++ClassDescIt)
+				{
+					ClassDescs.Add(*ClassDescIt);
+				}
+				ClassDescs.Sort([](const FWorldPartitionActorDesc& A, const FWorldPartitionActorDesc& B)
+				{
+					return A.GetGuid() < B.GetGuid();
+				});
+				for (const FWorldPartitionActorDesc* ClassDescsIterator : ClassDescs)
+				{
+					FString LineEntry = ClassDescsIterator->ToString(FWorldPartitionActorDesc::EToStringMode::Full) + LINE_TERMINATOR;
+					LogFile->Serialize(TCHAR_TO_ANSI(*LineEntry), LineEntry.Len());
+				}
+
+				LogFile->Close();
+				delete LogFile;
+			}
+		}
+	})
+);
+
 /*
  * We have to deal with some different realities between non-cooked and cooked editor builds: for non-cooked editor builds, UBlueprint are considered assets
  * and will be the opnly thing visible from the asset registry standpoint. On the other hand, for cooked editor builds, UBlueprint are not considered assets
@@ -615,4 +649,4 @@ FTopLevelAssetPath FWorldPartitionClassDescRegistry::RedirectClassPath(const FTo
 
 	return ResolvedClassPath;
 }
-#endif
+#endif
\ No newline at end of file
