missing files from previous submit

diff --git a/Engine/Source/Runtime/Engine/Public/WorldPartition/ActorDescContainer.h b/Engine/Source/Runtime/Engine/Public/WorldPartition/ActorDescContainer.h
--- a/Engine/Source/Runtime/Engine/Public/WorldPartition/ActorDescContainer.h
+++ b/Engine/Source/Runtime/Engine/Public/WorldPartition/ActorDescContainer.h
@@ -22,7 +22,7 @@ class ENGINE_API UActorDescContainer : public UObject, public FActorDescList
 	friend struct FWorldPartitionHandleUtils;
 	friend class FWorldPartitionActorDesc;
 
-	using TNameActorDescMap = TMap<FName, TUniquePtr<FWorldPartitionActorDesc>*>;
+	using FNameActorDescMap = TMap<FName, TUniquePtr<FWorldPartitionActorDesc>*>;
 
 public:
 	/* Struct of parameters passed to Initialize function. */
@@ -111,7 +111,7 @@ public:
 	TArray<FAssetData> InvalidActors;
 
 protected:
-	TNameActorDescMap ActorsByName;
+	FNameActorDescMap ActorsByName;
 
 	//~ Begin UObject Interface
 	virtual void BeginDestroy() override;
diff --git a/Engine/Source/Runtime/Engine/Public/WorldPartition/ActorDescList.h b/Engine/Source/Runtime/Engine/Public/WorldPartition/ActorDescList.h
--- a/Engine/Source/Runtime/Engine/Public/WorldPartition/ActorDescList.h
+++ b/Engine/Source/Runtime/Engine/Public/WorldPartition/ActorDescList.h
@@ -19,7 +19,7 @@ class ENGINE_API FActorDescList
 		}
 	};
 
-	using TGuidActorDescMap = TMap<FGuid, TUniquePtr<FWorldPartitionActorDesc>*, FDefaultSetAllocator, FActorGuidKeyFuncs>;
+	using FGuidActorDescMap = TMap<FGuid, TUniquePtr<FWorldPartitionActorDesc>*, FDefaultSetAllocator, FActorGuidKeyFuncs>;
 
 public:
 	FActorDescList() {}
@@ -48,7 +48,7 @@ public:
 		static_assert(TIsDerivedFrom<ActorType, AActor>::IsDerived, "Type is not derived from AActor.");
 
 	protected:
-		using MapType = TGuidActorDescMap;
+		using MapType = FGuidActorDescMap;
 		using ValueType = typename FWorldPartitionActorDescType<ActorType>::Type;
 		using IteratorType = typename TChooseClass<bConst, MapType::TConstIterator, MapType::TIterator>::Result;
 		using ListType = typename TChooseClass<bConst, const FActorDescList*, FActorDescList*>::Result;
@@ -167,6 +167,6 @@ protected:
 
 	TChunkedArray<TUniquePtr<FWorldPartitionActorDesc>> ActorDescList;
 
-	TGuidActorDescMap ActorsByGuid;
+	FGuidActorDescMap ActorsByGuid;
 #endif
 };
\ No newline at end of file
diff --git a/Engine/Source/Runtime/Engine/Public/WorldPartition/WorldPartitionClassDescRegistry.h b/Engine/Source/Runtime/Engine/Public/WorldPartition/WorldPartitionClassDescRegistry.h
--- a/Engine/Source/Runtime/Engine/Public/WorldPartition/WorldPartitionClassDescRegistry.h
+++ b/Engine/Source/Runtime/Engine/Public/WorldPartition/WorldPartitionClassDescRegistry.h
@@ -10,11 +10,11 @@
 
 class UActorDescContainer;
 
-class ENGINE_API FWorldPartitionClassDescRegistry : FActorDescList
+class ENGINE_API FWorldPartitionClassDescRegistry : public FActorDescList
 {
-	using TNameClassDescMap = TMap<FTopLevelAssetPath, TUniquePtr<FWorldPartitionActorDesc>*>;
-	using TParentClassMap = TMap<FTopLevelAssetPath, FTopLevelAssetPath>;
-	using TRedirectClassMap = TMap<FTopLevelAssetPath, FTopLevelAssetPath>;
+	using FNameClassDescMap = TMap<FTopLevelAssetPath, TUniquePtr<FWorldPartitionActorDesc>*>;
+	using FParentClassMap = TMap<FTopLevelAssetPath, FTopLevelAssetPath>;
+	using FRedirectClassMap = TMap<FTopLevelAssetPath, FTopLevelAssetPath>;
 
 public:
 	static FWorldPartitionClassDescRegistry& Get();
@@ -29,7 +29,7 @@ public:
 
 	bool IsRegisteredClass(const FTopLevelAssetPath& InClassPath) const;
 
-	const TParentClassMap& GetParentClassMap() const { check(IsInitialized()); return ParentClassMap; }
+	const FParentClassMap& GetParentClassMap() const { check(IsInitialized()); return ParentClassMap; }
 
 private:
 	void PrefetchClassDesc(UClass* InClass);
@@ -62,8 +62,8 @@ private:
 
 	FTopLevelAssetPath RedirectClassPath(const FTopLevelAssetPath& InClassPath) const;
 
-	TNameClassDescMap ClassByPath;
-	TParentClassMap ParentClassMap;
-	TRedirectClassMap RedirectClassMap;
+	FNameClassDescMap ClassByPath;
+	FParentClassMap ParentClassMap;
+	FRedirectClassMap RedirectClassMap;
 };
 #endif
\ No newline at end of file
