Fix non-unity build.

diff --git a/Engine/Source/Runtime/Engine/Private/WorldPartition/WorldPartitionClassDescRegistry.cpp b/Engine/Source/Runtime/Engine/Private/WorldPartition/WorldPartitionClassDescRegistry.cpp
--- a/Engine/Source/Runtime/Engine/Private/WorldPartition/WorldPartitionClassDescRegistry.cpp
+++ b/Engine/Source/Runtime/Engine/Private/WorldPartition/WorldPartitionClassDescRegistry.cpp
@@ -7,6 +7,7 @@
 #include "Algo/ForEach.h"
 #include "Algo/RemoveIf.h"
 #include "Algo/Transform.h"
+#include "HAL/FileManager.h"
 #include "Misc/PackageName.h"
 #include "UObject/CoreRedirects.h"
 #include "UObject/ObjectSaveContext.h"
