Fix for UGameInstance::CreateLocalPlayer not correctly getting the next available local player to spawn if you pass in an invalid Platform User ID to create a new one.

In 5.0, the ControllerId was iterated by the ++ operator in the for loop. With the switch to FPlatformUserId, we need to manually assign the value of the user id instead.

#rnx

diff --git a/Engine/Source/Runtime/Engine/Private/GameInstance.cpp b/Engine/Source/Runtime/Engine/Private/GameInstance.cpp
--- a/Engine/Source/Runtime/Engine/Private/GameInstance.cpp
+++ b/Engine/Source/Runtime/Engine/Private/GameInstance.cpp
@@ -821,6 +821,16 @@ ULocalPlayer* UGameInstance::CreateLocalPlayer(FPlatformUserId UserId, FString&
 		{
 			for (int32 Id = 0; Id < MaxSplitscreenPlayers; ++Id)
 			{
+				// Iterate until we find a null player. We want the next available platform user ID
+				FInputDeviceId DummyDevice;
+				FPlatformUserId DummyId;
+				IPlatformInputDeviceMapper::Get().RemapControllerIdToPlatformUserAndDevice(Id, OUT DummyId, OUT DummyDevice);
+
+				if (DummyId.IsValid())
+				{
+					UserId = DummyId;
+				}
+				
 				if (FindLocalPlayerFromControllerId(Id) == nullptr)
 				{
 					break;
@@ -983,14 +993,14 @@ void UGameInstance::DebugCreatePlayer(int32 ControllerId)
 {
 #if !UE_BUILD_SHIPPING
 	FString Error;
-	CreateLocalPlayer(ControllerId, Error, true);
-	if (Error.Len() > 0)
+	const ULocalPlayer* LP = CreateLocalPlayer(ControllerId, Error, true);
+	if (Error.Len() > 0 || !LP)
 	{
 		UE_LOG(LogPlayerManagement, Error, TEXT("Failed to DebugCreatePlayer: %s"), *Error);
 	}
 	else
 	{
-		FPlatformUserId UserId = FGenericPlatformMisc::GetPlatformUserForUserIndex(ControllerId);
+		FPlatformUserId UserId = LP->GetPlatformUserId();
 		FInputDeviceId InputDevice = INPUTDEVICEID_NONE;
 		IPlatformInputDeviceMapper& DeviceMapper = IPlatformInputDeviceMapper::Get();
 		DeviceMapper.RemapControllerIdToPlatformUserAndDevice(ControllerId, UserId, InputDevice);
