New FSoftWorldReference blueprint type to enable content-defined structures and BPs to contain soft references to UWorld assets.

diff --git a/Engine/Source/Runtime/Engine/Classes/Engine/SoftWorldReference.h b/Engine/Source/Runtime/Engine/Classes/Engine/SoftWorldReference.h
new file mode 100644
--- /dev/null
+++ b/Engine/Source/Runtime/Engine/Classes/Engine/SoftWorldReference.h
@@ -0,0 +1,16 @@
+// Copyright Epic Games, Inc. All Rights Reserved.
+#pragma once
+
+#include "SoftWorldReference.generated.h"
+
+class UWorld;
+
+/** A simple wrapper type to enable content-defined structs to hold soft references to UWorld assets **/
+USTRUCT(BlueprintType)
+struct FSoftWorldReference
+{
+	GENERATED_BODY();
+
+	UPROPERTY(EditAnywhere, BlueprintReadWrite, Category="Asset")
+	TSoftObjectPtr<UWorld> WorldAsset;
+};
