Add SuggestProjectileVelocity_MovingTarget to GameplayStatics.
This function returns a launch velocity that will cause a projectile to hit a target actor after a certain number of seconds.

diff --git a/Engine/Source/Runtime/Engine/Classes/Kismet/GameplayStatics.h b/Engine/Source/Runtime/Engine/Classes/Kismet/GameplayStatics.h
--- a/Engine/Source/Runtime/Engine/Classes/Kismet/GameplayStatics.h
+++ b/Engine/Source/Runtime/Engine/Classes/Kismet/GameplayStatics.h
@@ -1353,6 +1353,20 @@ public:
 	UFUNCTION(BlueprintCallable, Category = "Game", DisplayName = "Suggest Projectile Velocity Custom Arc", meta = (WorldContext = "WorldContextObject", AdvancedDisplay = "OverrideGravityZ, ArcParam"))
 	static bool SuggestProjectileVelocity_CustomArc(const UObject* WorldContextObject, FVector& OutLaunchVelocity, FVector StartPos, FVector EndPos, float OverrideGravityZ = 0, float ArcParam = 0.5f);
 
+	/**
+	 * Returns a launch velocity need for a projectile to hit the TargetActor in TimeToTarget seconds based on the TargetActor's current velocity.
+	 * This assumes the projectile is only accelerated by gravity (i.e. no outside forces), and that the TargetActor is moving at a constant velocity.
+	 * 
+	 * @param OutLaunchVelocity			The launch velocity returned from this calculation
+	 * @param ProjectileStartLocation	Location the projectile is launched from
+	 * @param TargetActor				Actor that the projectile should hit in TimeToTarget seconds
+	 * @param TargetLocationOffset		Offset to apply to the location the projectile is aiming for
+	 * @param GravityZOverride			Optional override of WorldGravityZ
+	 * @param TimeToTarget				Time (in seconds) between the projectile being launched and the projectile hitting the TargetActor - clamped to be at least 0.1
+	 */
+	UFUNCTION(BlueprintCallable, Category = "Game", DisplayName = "Suggest Projectile Velocity Moving Target", meta = (WorldContext = "WorldContextObject", AdvancedDisplay = "DrawDebugType, DrawDebugTime, DrawDebugColor"))
+	static bool SuggestProjectileVelocity_MovingTarget(const UObject* WorldContextObject, FVector& OutLaunchVelocity, FVector ProjectileStartLocation, AActor* TargetActor, FVector TargetLocationOffset = FVector::ZeroVector, double GravityZOverride = 0.f, double TimeToTarget = 1.f, EDrawDebugTrace::Type DrawDebugType = EDrawDebugTrace::Type::None, float DrawDebugTime = 3.f, FLinearColor DrawDebugColor = FLinearColor::Red);
+
 	/** Returns world origin current location. */
 	UFUNCTION(BlueprintPure, Category="Game", meta=(WorldContext="WorldContextObject") )
 	static FIntVector GetWorldOriginLocation(const UObject* WorldContextObject);
diff --git a/Engine/Source/Runtime/Engine/Private/GameplayStatics.cpp b/Engine/Source/Runtime/Engine/Private/GameplayStatics.cpp
--- a/Engine/Source/Runtime/Engine/Private/GameplayStatics.cpp
+++ b/Engine/Source/Runtime/Engine/Private/GameplayStatics.cpp
@@ -10,6 +10,7 @@
 #include "Misc/PackageName.h"
 #include "Kismet/GameplayStaticsTypes.h"
 #include "Misc/EngineVersion.h"
+#include "GameFramework/CharacterMovementComponent.h"
 #include "GameFramework/DamageType.h"
 #include "Kismet/KismetSystemLibrary.h"
 #include "SceneView.h"
@@ -3021,6 +3022,107 @@ bool UGameplayStatics::SuggestProjectileVelocity_CustomArc(const UObject* WorldC
 	return false;
 }
 
+bool UGameplayStatics::SuggestProjectileVelocity_MovingTarget(const UObject* WorldContextObject, FVector& OutLaunchVelocity, FVector ProjectileStartLocation, AActor* TargetActor, FVector TargetLocationOffset /*= FVector::ZeroVector*/, double GravityZOverride /*= 0.f*/, double TimeToTarget /*= 1.f*/, EDrawDebugTrace::Type DrawDebugType /*= EDrawDebugTrace::Type::None*/, float DrawDebugTime /*= 3.f*/, FLinearColor DrawDebugColor /*= FLinearColor::Red*/)
+{
+	// Generally solved using the logic below. In code we also adjust for TargetLocationOffset and a GravityZOverride.
+	// 
+	// pp = projectile position
+	// pp0 = projectile initial position
+	// vp0 = projectile initial velocity (OutLaunchVelocity)
+	// g = gravity vector
+	// pt = target position
+	// pt0 = target initial position
+	// vt0 = target initial velocity
+	// 
+	// Projectile position as a function of time
+	// pp = pp0 + vp0*t + (g/2)*(t^2)
+	// 
+	// Target position as a function of time
+	// pt = pt0 + vt0*t
+	// 
+	// Since we want the projectile position and target position to be the same, we can set the equations equal to each other to get:
+	// vp0 = (pt0 + vt0*t - pp0 - (g/2)*(t^2)) / t
+	
+	if (!GEngine)
+	{
+		return false;
+	}
+
+	const UWorld* World = GEngine->GetWorldFromContextObject(WorldContextObject, EGetWorldErrorMode::LogAndReturnNull);
+	if (!World)
+	{
+		return false;
+	}
+
+	if (!TargetActor)
+	{
+		return false;
+	}
+
+	// Clamp TimeToTarget to ensure non-zero value
+	TimeToTarget = FMath::Clamp(TimeToTarget, 0.1, TimeToTarget);
+
+	// Find the target's velocity
+	// If the target is a character, GetVelocity() will return a vector with a Z value of 0, so we need to adjust for slope
+	FVector TargetVelocity = TargetActor->GetVelocity();
+	if (ACharacter* TargetCharacter = Cast<ACharacter>(TargetActor))
+	{
+		if (UCharacterMovementComponent* TargetMovementComp = TargetCharacter->GetCharacterMovement())
+		{
+			FHitResult FloorHit = TargetMovementComp->CurrentFloor.HitResult;
+			const double FloorAngleRadians = FMath::Acos(FloorHit.ImpactNormal.Z);
+
+			// z = x * tan(theta)
+			const double ZMagnitude = TargetVelocity.Size2D() * FMath::Tan(FloorAngleRadians);
+			TargetVelocity.Z = TargetVelocity.Dot(FloorHit.ImpactNormal) < 0.0 ? ZMagnitude : -ZMagnitude;
+		}
+	}
+
+	// Find projectile's velocity using the target's velocity
+	const FVector TargetCurrentLocationPlusOffset = TargetActor->GetActorLocation() + TargetLocationOffset;
+	const double GravityZ = FMath::IsNearlyZero(GravityZOverride) ? World->GetGravityZ() : GravityZOverride;
+	const FVector GravityVector = FVector(0.0, 0.0, GravityZ);
+	
+	OutLaunchVelocity = (TargetCurrentLocationPlusOffset + (TargetVelocity * TimeToTarget) - ProjectileStartLocation - (0.5 * GravityVector * FMath::Square(TimeToTarget))) / TimeToTarget;
+
+#if ENABLE_DRAW_DEBUG
+	if (DrawDebugType != EDrawDebugTrace::None)
+	{
+		const bool bPersistentLines = DrawDebugType == EDrawDebugTrace::Persistent;
+		const float LifeTime = (DrawDebugType == EDrawDebugTrace::ForDuration) ? DrawDebugTime : 0.f;
+		
+		// Draw bounding box of target at current location
+		const float BoxThickness = 2.f;
+		const FQuat CurrentRotation = TargetActor->GetActorRotation().Quaternion();
+		FBox BodyBox(ForceInit);
+		FVector BoundingBoxExtent = FVector::ZeroVector;
+		FVector BoundingBoxCenter = TargetActor->GetActorLocation();
+		if (UPrimitiveComponent* PrimitiveComponent = Cast<UPrimitiveComponent>(TargetActor->GetRootComponent()))
+		{
+			BodyBox = PrimitiveComponent->Bounds.GetBox();
+			BodyBox.GetCenterAndExtents(BoundingBoxCenter, BoundingBoxExtent);
+		}
+
+		DrawDebugBox(World, BoundingBoxCenter, BoundingBoxExtent, DrawDebugColor.ToFColor(true), bPersistentLines, LifeTime, 0, BoxThickness);
+
+		// Draw arrow indicating aim direction from start location
+		const float DirectionalArrowLength = 30.f;
+		const float DirectionalArrowsize = 5.f;
+		const float DirectionalArrowThickness = 2.f;
+		DrawDebugDirectionalArrow(World, ProjectileStartLocation, ProjectileStartLocation + OutLaunchVelocity.GetSafeNormal() * DirectionalArrowLength, DirectionalArrowsize, DrawDebugColor.ToFColor(true), bPersistentLines, LifeTime, 0, DirectionalArrowThickness);
+
+		// Draw sphere at the target's final location (where the projectile should pass through after TimeToTarget seconds)
+		const float SphereRadius = 16.f;
+		const int32 SphereSegments = 16;
+		const float SphereThickness = 2.f;
+		const FVector TargetFinalLocation = TargetCurrentLocationPlusOffset + (TargetVelocity * TimeToTarget);
+		DrawDebugSphere(World, TargetFinalLocation, SphereRadius, SphereSegments, DrawDebugColor.ToFColor(true), bPersistentLines, LifeTime, 0, SphereThickness);
+	}
+#endif //ENABLE_DRAW_DEBUG
+
+	return true;
+}
+
 FIntVector UGameplayStatics::GetWorldOriginLocation(const UObject* WorldContextObject)
 {
 	UWorld* World = GEngine->GetWorldFromContextObject(WorldContextObject, EGetWorldErrorMode::LogAndReturnNull);
