Add the concept of local gravity direction for character movement simulation.
#rnx

diff --git a/Engine/Source/Runtime/Engine/Classes/GameFramework/Character.h b/Engine/Source/Runtime/Engine/Classes/GameFramework/Character.h
--- a/Engine/Source/Runtime/Engine/Classes/GameFramework/Character.h
+++ b/Engine/Source/Runtime/Engine/Classes/GameFramework/Character.h
@@ -460,10 +460,15 @@ protected:
 	UPROPERTY(Replicated)
 	uint8 ReplicatedMovementMode;
 
+	/** CharacterMovement Custom gravity direction replicated for simulated proxies. */
+	UPROPERTY(Replicated)
+	FVector_NetQuantizeNormal ReplicatedGravityDirection;
+
 	/** Flag that we are receiving replication of the based movement. */
 	UPROPERTY()
 	bool bInBaseReplication;
 
+	FVector_NetQuantizeNormal PreNetReceivedGravityDirection;
 public:
 	UFUNCTION()
 	void OnRep_ReplayLastTransformUpdateTimeStamp();
@@ -494,6 +499,15 @@ public:
 	UFUNCTION(BlueprintCallable, Category=Character, meta=(DisplayName="Get Base Rotation Offset", ScriptName="GetBaseRotationOffset"))
 	FRotator GetBaseRotationOffsetRotator() const { return GetBaseRotationOffset().Rotator(); }
 
+	/** Returns vector direction of gravity */
+	virtual FVector GetGravityDirection() const override;
+
+	/** Returns a quaternion transforming from world space into gravity relative space */
+	virtual FQuat GetGravityTransform() const override;
+
+	/** Returns replicated gravity direction for simulated proxies */
+	virtual FVector GetReplicatedGravityDirection() const;
+
 	/** Default crouched eye height */
 	UPROPERTY(EditAnywhere, BlueprintReadWrite, Category=Camera)
 	float CrouchedEyeHeight;
diff --git a/Engine/Source/Runtime/Engine/Classes/GameFramework/CharacterMovementComponent.h b/Engine/Source/Runtime/Engine/Classes/GameFramework/CharacterMovementComponent.h
--- a/Engine/Source/Runtime/Engine/Classes/GameFramework/CharacterMovementComponent.h
+++ b/Engine/Source/Runtime/Engine/Classes/GameFramework/CharacterMovementComponent.h
@@ -188,6 +188,31 @@ private:
 	UPROPERTY(Category="Character Movement: Walking", VisibleAnywhere)
 	float WalkableFloorZ;
 
+	/**
+	 * A normalized vector representing the direction of gravity for gravity relative movement modes: walking, falling,
+	 * and custom movement modes. Gravity direction remaps player input as being within the plane defined by the gravity
+	 * direction. Movement simulation values like velocity and acceleration are maintained in their existing world coordinate
+	 * space but are transformed internally as gravity relative (for instance moving forward up a vertical wall that gravity is
+	 * defined to be perpendicular to and jump "up" from that wall). By default the character's capsule will not be oriented
+	 * to align with the gravity direction.
+	 */
+	UPROPERTY(Category="Character Movement: Gravity", VisibleAnywhere, BlueprintReadOnly, meta = (AllowPrivateAccess = true))
+	FVector GravityDirection;
+
+	/**
+	 * A cached quaternion representing the rotation from world space to gravity relative space defined by GravityDirection.
+	 */
+	UPROPERTY(Category="Character Movement: Gravity", VisibleAnywhere, BlueprintReadOnly, meta = (AllowPrivateAccess = true))
+	FQuat WorldToGravityTransform;
+
+	/**
+	 * A cached quaternion representing the inverse rotation from world space to gravity relative space defined by GravityDirection.
+	 */
+	UPROPERTY(Category="Character Movement: Gravity", VisibleAnywhere, BlueprintReadOnly, meta = (AllowPrivateAccess = true))
+	FQuat GravityToWorldTransform;
+
+	/** Whether the character has custom local gravity set. Cached in SetGravityDirection(). */
+	bool bHasCustomGravity;
 public:
 	
 	/**
@@ -211,6 +236,9 @@ public:
 	UPROPERTY(Category="Character Movement: MovementMode", BlueprintReadOnly)
 	uint8 CustomMovementMode;
 
+	/** The default direction that gravity points for movement simulation.  */
+	static const FVector DefaultGravityDirection;
+
 	/** Smoothing mode for simulated proxies in network game. */
 	UPROPERTY(Category="Character Movement (Networking)", EditAnywhere, BlueprintReadOnly)
 	ENetworkSmoothingMode NetworkSmoothingMode;
@@ -908,6 +936,10 @@ public:
 	UPROPERTY(Transient)
 	uint8 bNetworkMovementModeChanged:1;
 
+	/** True when the networked gravity direction has been replicated. */
+	UPROPERTY(Transient)
+	uint8 bNetworkGravityDirectionChanged:1;
+
 	/** 
 	 * If true, we should ignore server location difference checks for client error on this movement component.
 	 * This can be useful when character is moving at extreme speeds for a duration and you need it to look
@@ -1573,6 +1605,37 @@ public:
 	 */
 	virtual FVector GetAirControl(float DeltaTime, float TickAirControl, const FVector& FallAcceleration);
 
+	// Gravity Direction
+
+	/**
+	 * Set a custom, local gravity direction to use during movement simulation.
+	 * The gravity direction must be synchronized by external systems between the autonomous
+	 * and authority processes. The gravity direction will be corrected as part of movement
+	 * corrections should the movement state diverge.
+	 * SetGravityDirection is responsible for initializing cached values used to tranform to
+	 * from gravity relative space.
+	 * @param GravityDir		A non-zero vector representing the new gravity direction. The vector will be normalized.
+	 */
+	virtual void SetGravityDirection(const FVector& GravityDir);
+
+	/** Whether the gravity direction is different from UCharacterMovementComponent::DefaultGravityDirection. */
+	bool HasCustomGravity() const { return bHasCustomGravity; }
+
+	/** Returns the current gravity direction. */
+	FVector GetGravityDirection() const { return GravityDirection; }
+
+	/** Returns a quaternion transforming from world to gravity space. */
+	FQuat GetWorldToGravityTransform() const { return WorldToGravityTransform; }
+
+	/** Returns a quaternion transforming from gravity to world space. */
+	FQuat GetGravityToWorldTransform() const { return GravityToWorldTransform; }
+
+	/** Rotate a vector from world to gravity space. */
+	FVector RotateGravityToWorld(const FVector& World) const { return WorldToGravityTransform.RotateVector(World); }
+
+	/** Rotate a vector gravity to world space. */
+	FVector RotateWorldToGravity(const FVector& Gravity) const { return GravityToWorldTransform.RotateVector(Gravity); }
+
 protected:
 
 	/**
@@ -2491,7 +2554,7 @@ public:
 protected:
 
 	/** Event notification when client receives correction data from the server, before applying the data. Base implementation logs relevant data and draws debug info if "p.NetShowCorrections" is not equal to 0. */
-	virtual void OnClientCorrectionReceived(class FNetworkPredictionData_Client_Character& ClientData, float TimeStamp, FVector NewLocation, FVector NewVelocity, UPrimitiveComponent* NewBase, FName NewBaseBoneName, bool bHasBase, bool bBaseRelativePosition, uint8 ServerMovementMode);
+	virtual void OnClientCorrectionReceived(class FNetworkPredictionData_Client_Character& ClientData, float TimeStamp, FVector NewLocation, FVector NewVelocity, UPrimitiveComponent* NewBase, FName NewBaseBoneName, bool bHasBase, bool bBaseRelativePosition, uint8 ServerMovementMode, FVector ServerGravityDirection);
 
 	/**
 	 * Set custom struct used for client to server move RPC serialization.
diff --git a/Engine/Source/Runtime/Engine/Classes/GameFramework/CharacterMovementReplication.h b/Engine/Source/Runtime/Engine/Classes/GameFramework/CharacterMovementReplication.h
--- a/Engine/Source/Runtime/Engine/Classes/GameFramework/CharacterMovementReplication.h
+++ b/Engine/Source/Runtime/Engine/Classes/GameFramework/CharacterMovementReplication.h
@@ -258,6 +258,7 @@ public:
 		, NewLoc(ForceInitToZero)
 		, NewVel(ForceInitToZero)
 		, NewRot(ForceInitToZero)
+		, GravityDirection(FVector::DownVector)
 		, NewBase(NULL)
 		, NewBaseBoneName(NAME_None)
 		, bAckGoodMove(false)
@@ -272,6 +273,7 @@ public:
 	FVector NewLoc; // Note: if bBaseRelativePosition is set, this is a relative location to the Movement base.
 	FVector NewVel; // Note: if bBaseRelativeVelocity is set, this is a relative velocity to the Movement base.
 	FRotator NewRot;
+	FVector GravityDirection;
 	UPrimitiveComponent* NewBase;
 	FName NewBaseBoneName;
 	bool bAckGoodMove;
diff --git a/Engine/Source/Runtime/Engine/Classes/GameFramework/Pawn.h b/Engine/Source/Runtime/Engine/Classes/GameFramework/Pawn.h
--- a/Engine/Source/Runtime/Engine/Classes/GameFramework/Pawn.h
+++ b/Engine/Source/Runtime/Engine/Classes/GameFramework/Pawn.h
@@ -327,7 +327,10 @@ public:
 #endif
 
 	/** Returns vector direction of gravity */
-	FVector GetGravityDirection();
+	virtual FVector GetGravityDirection() const;
+
+	/** Returns a quaternion transforming from world space into gravity relative space */
+	virtual FQuat GetGravityTransform() const;
 
 	/** Make sure pawn properties are back to default. */
 	virtual void SetPlayerDefaults();
diff --git a/Engine/Source/Runtime/Engine/Private/Character.cpp b/Engine/Source/Runtime/Engine/Private/Character.cpp
--- a/Engine/Source/Runtime/Engine/Private/Character.cpp
+++ b/Engine/Source/Runtime/Engine/Private/Character.cpp
@@ -830,6 +830,35 @@ void ACharacter::SaveRelativeBasedMovement(const FVector& NewRelativeLocation, c
 	BasedMovement.bRelativeRotation = bRelativeRotation;
 }
 
+FVector ACharacter::GetGravityDirection() const
+{
+	FVector GravityDirection = UCharacterMovementComponent::DefaultGravityDirection;
+	const UCharacterMovementComponent* const MovementComponent = GetCharacterMovement();
+	if (MovementComponent)
+	{
+		GravityDirection = MovementComponent->GetGravityDirection();
+	}
+
+	return GravityDirection;
+}
+
+FQuat ACharacter::GetGravityTransform() const
+{
+	FQuat GravityTransform = FQuat::Identity;
+	const UCharacterMovementComponent* const MovementComponent = GetCharacterMovement();
+	if (MovementComponent)
+	{
+		GravityTransform = MovementComponent->GetWorldToGravityTransform();
+	}
+
+	return GravityTransform;
+}
+
+FVector ACharacter::GetReplicatedGravityDirection() const
+{
+	return ReplicatedGravityDirection;
+}
+
 FVector ACharacter::GetNavAgentLocation() const
 {
 	FVector AgentLocation = FNavigationSystem::InvalidLocation;
@@ -1140,6 +1169,7 @@ static uint8 SavedMovementMode;
 void ACharacter::PreNetReceive()
 {
 	SavedMovementMode = ReplicatedMovementMode;
+	PreNetReceivedGravityDirection = ReplicatedGravityDirection;
 	Super::PreNetReceive();
 }
 
@@ -1147,8 +1177,9 @@ void ACharacter::PostNetReceive()
 {
 	if (GetLocalRole() == ROLE_SimulatedProxy)
 	{
+		CharacterMovement->bNetworkGravityDirectionChanged = !PreNetReceivedGravityDirection.Equals(ReplicatedGravityDirection);
 		CharacterMovement->bNetworkMovementModeChanged |= ((SavedMovementMode != ReplicatedMovementMode) || (CharacterMovement->PackNetworkMovementMode() != ReplicatedMovementMode));
-		CharacterMovement->bNetworkUpdateReceived |= CharacterMovement->bNetworkMovementModeChanged || CharacterMovement->bJustTeleported;
+		CharacterMovement->bNetworkUpdateReceived |= CharacterMovement->bNetworkMovementModeChanged || CharacterMovement->bJustTeleported || CharacterMovement->bNetworkGravityDirectionChanged;
 	}
 
 	Super::PostNetReceive();
@@ -1520,7 +1551,8 @@ void ACharacter::PreReplication( IRepChangedPropertyTracker & ChangedPropertyTra
 	}
 
 	bProxyIsJumpForceApplied = (JumpForceTimeRemaining > 0.0f);
-	ReplicatedMovementMode = CharacterMovement->PackNetworkMovementMode();	
+	ReplicatedMovementMode = CharacterMovement->PackNetworkMovementMode();
+	ReplicatedGravityDirection = CharacterMovement->GetGravityDirection();
 
 	if(IsReplicatingMovement())
 	{
@@ -1592,6 +1624,7 @@ void ACharacter::GetLifetimeReplicatedProps( TArray< FLifetimeProperty > & OutLi
 	DOREPLIFETIME_CONDITION( ACharacter, bIsCrouched,						COND_SimulatedOnly );
 	DOREPLIFETIME_CONDITION( ACharacter, bProxyIsJumpForceApplied,			COND_SimulatedOnly );
 	DOREPLIFETIME_CONDITION( ACharacter, AnimRootMotionTranslationScale,	COND_SimulatedOnly );
+	DOREPLIFETIME_CONDITION( ACharacter, ReplicatedGravityDirection,		COND_SimulatedOnly );
 	DOREPLIFETIME_CONDITION( ACharacter, ReplayLastTransformUpdateTimeStamp, COND_ReplayOnly );
 }
 
diff --git a/Engine/Source/Runtime/Engine/Private/Components/CharacterMovementComponent.cpp b/Engine/Source/Runtime/Engine/Private/Components/CharacterMovementComponent.cpp
--- a/Engine/Source/Runtime/Engine/Private/Components/CharacterMovementComponent.cpp
+++ b/Engine/Source/Runtime/Engine/Private/Components/CharacterMovementComponent.cpp
@@ -83,6 +83,7 @@ namespace CharacterMovementConstants
 	const float VERTICAL_SLOPE_NORMAL_Z = 0.001f; // Slope is vertical if Abs(Normal.Z) <= this threshold. Accounts for precision problems that sometimes angle normals slightly off horizontal for vertical surface.
 }
 
+const FVector UCharacterMovementComponent::DefaultGravityDirection = FVector::DownVector;
 const float UCharacterMovementComponent::MIN_TICK_TIME = 1e-6f;
 const float UCharacterMovementComponent::MIN_FLOOR_DIST = 1.9f;
 const float UCharacterMovementComponent::MAX_FLOOR_DIST = 2.4f;
@@ -535,6 +536,11 @@ UCharacterMovementComponent::UCharacterMovementComponent(const FObjectInitialize
 	RotationRate = FRotator(0.f, 360.0f, 0.0f);
 	SetWalkableFloorZ(0.71f);
 
+	GravityDirection = DefaultGravityDirection;
+	WorldToGravityTransform = FQuat::Identity;
+	GravityToWorldTransform =  FQuat::Identity;
+	bHasCustomGravity = false;
+
 	MaxStepHeight = 45.0f;
 	PerchRadiusThreshold = 0.0f;
 	PerchAdditionalHeight = 40.f;
@@ -979,7 +985,17 @@ bool UCharacterMovementComponent::DoJump(bool bReplayingMoves)
 		// Don't jump if we can't move up/down.
 		if (!bConstrainToPlane || FMath::Abs(PlaneConstraintNormal.Z) != 1.f)
 		{
-			Velocity.Z = FMath::Max<FVector::FReal>(Velocity.Z, JumpZVelocity);
+			if (HasCustomGravity())
+			{
+				FVector GravityRelativeVelocity = RotateWorldToGravity(Velocity);
+				GravityRelativeVelocity.Z = FMath::Max<FVector::FReal>(GravityRelativeVelocity.Z, JumpZVelocity);
+				Velocity = RotateGravityToWorld(GravityRelativeVelocity);
+			}
+			else
+			{
+				Velocity.Z = FMath::Max<FVector::FReal>(Velocity.Z, JumpZVelocity);
+			}
+			
 			SetMovementMode(MOVE_Falling);
 			return true;
 		}
@@ -1066,7 +1082,18 @@ void UCharacterMovementComponent::JumpOff(AActor* MovementBaseActor)
 			{
 				Velocity = MaxSpeed * Velocity.GetSafeNormal();
 			}
-			Velocity.Z = JumpOffJumpZFactor * JumpZVelocity;
+
+			if (HasCustomGravity())
+			{
+				FVector GravityRelativeVelocity = RotateWorldToGravity(Velocity);
+				GravityRelativeVelocity.Z = JumpOffJumpZFactor * JumpZVelocity;
+				Velocity = RotateGravityToWorld(GravityRelativeVelocity);
+			}
+			else
+			{
+				Velocity.Z = JumpOffJumpZFactor * JumpZVelocity;
+			}
+			
 			SetMovementMode(MOVE_Falling);
 		}
 		bPerformingJumpOff = false;
@@ -1237,7 +1264,7 @@ void UCharacterMovementComponent::OnMovementModeChanged(EMovementMode PreviousMo
 	if (MovementMode == MOVE_Walking)
 	{
 		// Walking uses only XY velocity, and must be on a walkable floor, with a Base.
-		Velocity.Z = 0.f;
+		Velocity = FVector::VectorPlaneProject(Velocity, -GetGravityDirection());
 		bCrouchMaintainsBaseLocation = true;
 		GroundMovementMode = MovementMode;
 
@@ -1751,6 +1778,13 @@ void UCharacterMovementComponent::SimulatedTick(float DeltaSeconds)
 			bNetworkMovementModeChanged = false;
 		}
 
+		// Update replicated gravity direction
+		if (bNetworkGravityDirectionChanged)
+		{
+			SetGravityDirection(CharacterOwner->GetReplicatedGravityDirection());
+			bNetworkGravityDirectionChanged = false;
+		}
+
 		// Perform movement
 		PerformMovement(DeltaSeconds);
 
@@ -1772,6 +1806,7 @@ void UCharacterMovementComponent::SimulatedTick(float DeltaSeconds)
 			CharacterOwner->OnRep_ReplicatedMovement();
 			CharacterOwner->OnRep_ReplicatedBasedMovement();
 			ApplyNetworkMovementMode(GetCharacterOwner()->GetReplicatedMovementMode());
+			SetGravityDirection(CharacterOwner->GetReplicatedGravityDirection());
 		}
 
 		if (CharacterOwner->IsReplicatingMovement() && UpdatedComponent)
@@ -1793,6 +1828,13 @@ void UCharacterMovementComponent::SimulatedTick(float DeltaSeconds)
 						bNetworkMovementModeChanged = false;
 					}
 
+					// Update replicated gravity direction
+					if (bNetworkGravityDirectionChanged)
+					{
+						SetGravityDirection(CharacterOwner->GetReplicatedGravityDirection());
+						bNetworkGravityDirectionChanged = false;
+					}
+
 					PerformMovement(DeltaSeconds);
 				}
 				else
@@ -1864,6 +1906,13 @@ void UCharacterMovementComponent::SimulateRootMotion(float DeltaSeconds, const F
 			bNetworkMovementModeChanged = false;
 		}
 
+		// Update replicated gravity direction
+		if (bNetworkGravityDirectionChanged)
+		{
+			SetGravityDirection(CharacterOwner->GetReplicatedGravityDirection());
+			bNetworkGravityDirectionChanged = false;
+		}
+
 		NumJumpApexAttempts = 0;
 		StartNewPhysics(DeltaSeconds, 0);
 		// fixme laurent - simulate movement seems to have step up issues? investigate as that would be cheaper to use.
@@ -1903,7 +1952,17 @@ FVector UCharacterMovementComponent::ConstrainAnimRootMotionVelocity(const FVect
 	// Do not override Velocity.Z if in falling physics, we want to keep the effect of gravity.
 	if (IsFalling())
 	{
-		Result.Z = CurrentVelocity.Z;
+		if (HasCustomGravity())
+		{
+			FVector GravityRelativeResult = RotateWorldToGravity(Result);
+			const FVector GravityRelativeCurrentVelocity = RotateWorldToGravity(CurrentVelocity);
+			GravityRelativeResult.Z = GravityRelativeCurrentVelocity.Z;
+			Result = RotateGravityToWorld(GravityRelativeResult);
+		}
+		else
+		{
+			Result.Z = CurrentVelocity.Z;
+		}
 	}
 
 	return Result;
@@ -1952,6 +2011,12 @@ void UCharacterMovementComponent::SimulateMovement(float DeltaSeconds)
 				bNetworkUpdateReceived = false;
 				bHandledNetUpdate = true;
 				UE_LOG(LogCharacterMovement, Verbose, TEXT("Proxy %s received net update"), *CharacterOwner->GetName());
+				if (bNetworkGravityDirectionChanged)
+				{
+					SetGravityDirection(CharacterOwner->GetReplicatedGravityDirection());
+					bNetworkGravityDirectionChanged = false;
+				}
+
 				if (bNetworkMovementModeChanged)
 				{
 					ApplyNetworkMovementMode(CharacterOwner->GetReplicatedMovementMode());
@@ -2012,11 +2077,17 @@ void UCharacterMovementComponent::SimulateMovement(float DeltaSeconds)
 			// find floor and check if falling
 			if (IsMovingOnGround() || MovementMode == MOVE_Falling)
 			{
+				bool bShouldFindFloor = Velocity.Z <= 0.f;
+				if (HasCustomGravity())
+				{
+					bShouldFindFloor = RotateWorldToGravity(Velocity).Z <= 0.0;
+				}
+
 				if (StepDownResult.bComputedFloor)
 				{
 					CurrentFloor = StepDownResult.FloorResult;
 				}
-				else if (Velocity.Z <= 0.f)
+				else if (bShouldFindFloor)
 				{
 					FindFloor(UpdatedComponent->GetComponentLocation(), CurrentFloor, Velocity.IsZero(), NULL);
 				}
@@ -2031,7 +2102,15 @@ void UCharacterMovementComponent::SimulateMovement(float DeltaSeconds)
 				{
 					// Follows PhysWalking approach for encroachment on floor tests
 					FHitResult Hit(CurrentFloor.HitResult);
-					Hit.TraceEnd = Hit.TraceStart + FVector(0.f, 0.f, MAX_FLOOR_DIST);
+					if (HasCustomGravity())
+					{
+						Hit.TraceEnd = Hit.TraceStart - GetGravityDirection() * MAX_FLOOR_DIST;
+					}
+					else
+					{
+						Hit.TraceEnd = Hit.TraceStart + FVector(0.f, 0.f, MAX_FLOOR_DIST);
+					}
+					
 					const FVector RequestedAdjustment = GetPenetrationAdjustment(Hit);
 					const bool bResolved = ResolvePenetration(RequestedAdjustment, Hit, UpdatedComponent->GetComponentQuat());
 					bForceNextFloorCheck |= bResolved;
@@ -2041,7 +2120,14 @@ void UCharacterMovementComponent::SimulateMovement(float DeltaSeconds)
 					if (!bSimGravityDisabled)
 					{
 						// No floor, must fall.
-						if (Velocity.Z <= 0.f || bApplyGravityWhileJumping || !CharacterOwner->IsJumpProvidingForce())
+						if (HasCustomGravity())
+						{
+							if (RotateWorldToGravity(Velocity).Z <= 0.f || bApplyGravityWhileJumping || !CharacterOwner->IsJumpProvidingForce())
+							{
+								Velocity = NewFallVelocity(Velocity, -GetGravityDirection() * GetGravityZ(), DeltaSeconds);
+							}
+						}
+						else if (Velocity.Z <= 0.f || bApplyGravityWhileJumping || !CharacterOwner->IsJumpProvidingForce())
 						{
 							Velocity = NewFallVelocity(Velocity, FVector(0.f, 0.f, GetGravityZ()), DeltaSeconds);
 						}
@@ -2068,7 +2154,7 @@ void UCharacterMovementComponent::SimulateMovement(float DeltaSeconds)
 							if (!bSimGravityDisabled)
 							{
 								// Continue falling.
-								Velocity = NewFallVelocity(Velocity, FVector(0.f, 0.f, GetGravityZ()), DeltaSeconds);
+								Velocity = NewFallVelocity(Velocity, -GetGravityDirection() * GetGravityZ(), DeltaSeconds);
 							}
 							CurrentFloor.Clear();
 						}
@@ -3229,7 +3315,7 @@ float UCharacterMovementComponent::SlideAlongSurface(const FVector& Delta, float
 		return 0.f;
 	}
 
-	FVector Normal(InNormal);
+	FVector Normal(RotateWorldToGravity(InNormal));
 	if (IsMovingOnGround())
 	{
 		// We don't want to be pushed up an unwalkable surface.
@@ -3245,8 +3331,9 @@ float UCharacterMovementComponent::SlideAlongSurface(const FVector& Delta, float
 			// Don't push down into the floor when the impact is on the upper portion of the capsule.
 			if (CurrentFloor.FloorDist < MIN_FLOOR_DIST && CurrentFloor.bBlockingHit)
 			{
-				const FVector FloorNormal = CurrentFloor.HitResult.Normal;
-				const bool bFloorOpposedToMovement = (Delta | FloorNormal) < 0.f && (FloorNormal.Z < 1.f - UE_DELTA);
+				const FVector FloorNormal = RotateWorldToGravity(CurrentFloor.HitResult.Normal);
+
+				const bool bFloorOpposedToMovement = (RotateWorldToGravity(Delta) | FloorNormal) < 0.f && (FloorNormal.Z < 1.f - UE_DELTA);
 				if (bFloorOpposedToMovement)
 				{
 					Normal = FloorNormal;
@@ -3257,49 +3344,53 @@ float UCharacterMovementComponent::SlideAlongSurface(const FVector& Delta, float
 		}
 	}
 
-	return Super::SlideAlongSurface(Delta, Time, Normal, Hit, bHandleImpact);
+	return Super::SlideAlongSurface(Delta, Time, RotateGravityToWorld(Normal), Hit, bHandleImpact);
 }
 
 
 void UCharacterMovementComponent::TwoWallAdjust(FVector& Delta, const FHitResult& Hit, const FVector& OldHitNormal) const
 {
-	const FVector InDelta = Delta;
+	const FVector InDelta = RotateWorldToGravity(Delta);
 	Super::TwoWallAdjust(Delta, Hit, OldHitNormal);
 
+	FVector GravityRelativeDelta = RotateWorldToGravity(Delta);
 	if (IsMovingOnGround())
 	{
 		// Allow slides up walkable surfaces, but not unwalkable ones (treat those as vertical barriers).
-		if (Delta.Z > 0.f)
+		if (GravityRelativeDelta.Z > 0.f)
 		{
-			if ((Hit.Normal.Z >= WalkableFloorZ || IsWalkable(Hit)) && Hit.Normal.Z > UE_KINDA_SMALL_NUMBER)
+			const FVector GravityRelativeHitNormal = RotateWorldToGravity(Hit.Normal);
+			if ((GravityRelativeHitNormal.Z >= WalkableFloorZ || IsWalkable(Hit)) && GravityRelativeHitNormal.Z > UE_KINDA_SMALL_NUMBER)
 			{
 				// Maintain horizontal velocity
 				const float Time = (1.f - Hit.Time);
-				const FVector ScaledDelta = Delta.GetSafeNormal() * InDelta.Size();
-				Delta = FVector(InDelta.X, InDelta.Y, ScaledDelta.Z / Hit.Normal.Z) * Time;
+				const FVector ScaledDelta = GravityRelativeDelta.GetSafeNormal() * InDelta.Size();
+				GravityRelativeDelta = FVector(InDelta.X, InDelta.Y, ScaledDelta.Z / GravityRelativeHitNormal.Z) * Time;
 
 				// Should never exceed MaxStepHeight in vertical component, so rescale if necessary.
 				// This should be rare (Hit.Normal.Z above would have been very small) but we'd rather lose horizontal velocity than go too high.
-				if (Delta.Z > MaxStepHeight)
+				if (GravityRelativeDelta.Z > MaxStepHeight)
 				{
-					const float Rescale = MaxStepHeight / Delta.Z;
-					Delta *= Rescale;
+					const float Rescale = MaxStepHeight / GravityRelativeDelta.Z;
+					GravityRelativeDelta *= Rescale;
 				}
 			}
 			else
 			{
-				Delta.Z = 0.f;
+				GravityRelativeDelta.Z = 0.f;
 			}
 		}
-		else if (Delta.Z < 0.f)
+		else if (GravityRelativeDelta.Z < 0.f)
 		{
 			// Don't push down into the floor.
 			if (CurrentFloor.FloorDist < MIN_FLOOR_DIST && CurrentFloor.bBlockingHit)
 			{
-				Delta.Z = 0.f;
+				GravityRelativeDelta.Z = 0.f;
 			}
 		}
 	}
+
+	Delta = RotateGravityToWorld(GravityRelativeDelta);
 }
 
 
@@ -4381,10 +4472,11 @@ void UCharacterMovementComponent::NotifyJumpApex()
 FVector UCharacterMovementComponent::GetFallingLateralAcceleration(float DeltaTime)
 {
 	// No acceleration in Z
-	FVector FallAcceleration = FVector(Acceleration.X, Acceleration.Y, 0.f);
+	const FVector GravityRelativeAcceleration = RotateWorldToGravity(Acceleration);
+	FVector FallAcceleration = RotateGravityToWorld(FVector(GravityRelativeAcceleration.X, GravityRelativeAcceleration.Y, 0.f));
 
 	// bound acceleration, falling object has minimal ability to impact acceleration
-	if (!HasAnimRootMotion() && FallAcceleration.SizeSquared2D() > 0.f)
+	if (!HasAnimRootMotion() && GravityRelativeAcceleration.SizeSquared2D() > 0.f)
 	{
 		FallAcceleration = GetAirControl(DeltaTime, AirControl, FallAcceleration);
 		FallAcceleration = FallAcceleration.GetClampedToMaxSize(GetMaxAcceleration());
@@ -4396,7 +4488,8 @@ FVector UCharacterMovementComponent::GetFallingLateralAcceleration(float DeltaTi
 
 bool UCharacterMovementComponent::ShouldLimitAirControl(float DeltaTime, const FVector& FallAcceleration) const
 {
-	return (FallAcceleration.SizeSquared2D() > 0.f);
+	const FVector GravityRelativeFallAcceleration = RotateWorldToGravity(FallAcceleration);
+	return (GravityRelativeFallAcceleration.SizeSquared2D() > 0.f);
 }
 
 FVector UCharacterMovementComponent::GetAirControl(float DeltaTime, float TickAirControl, const FVector& FallAcceleration)
@@ -4410,11 +4503,27 @@ FVector UCharacterMovementComponent::GetAirControl(float DeltaTime, float TickAi
 	return TickAirControl * FallAcceleration;
 }
 
+void UCharacterMovementComponent::SetGravityDirection(const FVector& InNewGravityDir)
+{
+	FVector NewGravityDir = InNewGravityDir;
+	if (ensure(!NewGravityDir.IsNearlyZero()))
+	{
+		if (!GravityDirection.Equals(NewGravityDir))
+		{
+			UE_LOG(LogCharacterMovement, Verbose, TEXT("SetGravityDirection: From(%s) To(%s)"), *GravityDirection.ToCompactString(), *NewGravityDir.ToCompactString());
+			GravityDirection = NewGravityDir;
+			WorldToGravityTransform = FQuat::FindBetweenNormals(FVector::UpVector, -NewGravityDir);
+			GravityToWorldTransform = WorldToGravityTransform.Inverse();
+			bHasCustomGravity = !GravityDirection.Equals(DefaultGravityDirection);
+		}
+	}
+}
 
 float UCharacterMovementComponent::BoostAirControl(float DeltaTime, float TickAirControl, const FVector& FallAcceleration)
 {
 	// Allow a burst of initial acceleration
-	if (AirControlBoostMultiplier > 0.f && Velocity.SizeSquared2D() < FMath::Square(AirControlBoostVelocityThreshold))
+	const FVector GravityRelativeVelocity = RotateWorldToGravity(Velocity);
+	if (AirControlBoostMultiplier > 0.f && GravityRelativeVelocity.SizeSquared2D() < FMath::Square(AirControlBoostVelocityThreshold))
 	{
 		TickAirControl = FMath::Min(1.f, AirControlBoostMultiplier * TickAirControl);
 	}
@@ -4433,7 +4542,8 @@ void UCharacterMovementComponent::PhysFalling(float deltaTime, int32 Iterations)
 	}
 
 	FVector FallAcceleration = GetFallingLateralAcceleration(deltaTime);
-	FallAcceleration.Z = 0.f;
+	const FVector GravityRelativeFallAcceleration = RotateWorldToGravity(FallAcceleration);
+	FallAcceleration = RotateGravityToWorld(FVector(GravityRelativeFallAcceleration.X, GravityRelativeFallAcceleration.Y, 0));
 	const bool bHasLimitedAirControl = ShouldLimitAirControl(deltaTime, FallAcceleration);
 
 	float remainingTime = deltaTime;
@@ -4461,14 +4571,24 @@ void UCharacterMovementComponent::PhysFalling(float deltaTime, int32 Iterations)
 			{
 				// Acceleration = FallAcceleration for CalcVelocity(), but we restore it after using it.
 				TGuardValue<FVector> RestoreAcceleration(Acceleration, FallAcceleration);
-				Velocity.Z = 0.f;
-				CalcVelocity(timeTick, FallingLateralFriction, false, MaxDecel);
-				Velocity.Z = OldVelocity.Z;
+				if (HasCustomGravity())
+				{
+					Velocity = FVector::VectorPlaneProject(Velocity, RotateGravityToWorld(FVector::UpVector));
+					const FVector GravityRelativeOffset = OldVelocity - Velocity;
+					CalcVelocity(timeTick, FallingLateralFriction, false, MaxDecel);
+					Velocity += GravityRelativeOffset;
+				}
+				else
+				{
+					Velocity.Z = 0.f;
+					CalcVelocity(timeTick, FallingLateralFriction, false, MaxDecel);
+					Velocity.Z = OldVelocity.Z;
+				}
 			}
 		}
 
 		// Compute current gravity
-		const FVector Gravity(0.f, 0.f, GetGravityZ());
+		const FVector Gravity = -GetGravityDirection() * GetGravityZ();
 		float GravityTime = timeTick;
 
 		// If jump is providing force, gravity may be affected.
@@ -4496,21 +4616,31 @@ void UCharacterMovementComponent::PhysFalling(float deltaTime, int32 Iterations)
 		DecayFormerBaseVelocity(timeTick);
 
 		// See if we need to sub-step to exactly reach the apex. This is important for avoiding "cutting off the top" of the trajectory as framerate varies.
-		if (CharacterMovementCVars::ForceJumpPeakSubstep && OldVelocityWithRootMotion.Z > 0.f && Velocity.Z <= 0.f && NumJumpApexAttempts < MaxJumpApexAttemptsPerSimulation)
+		const FVector GravityRelativeOldVelocityWithRootMotion = RotateWorldToGravity(OldVelocityWithRootMotion);
+		if (CharacterMovementCVars::ForceJumpPeakSubstep && GravityRelativeOldVelocityWithRootMotion.Z > 0.f && RotateWorldToGravity(Velocity).Z <= 0.f && NumJumpApexAttempts < MaxJumpApexAttemptsPerSimulation)
 		{
 			const FVector DerivedAccel = (Velocity - OldVelocityWithRootMotion) / timeTick;
-			if (!FMath::IsNearlyZero(DerivedAccel.Z))
+			const FVector GravityRelativeDerivedAccel = RotateWorldToGravity(DerivedAccel);
+			if (!FMath::IsNearlyZero(GravityRelativeDerivedAccel.Z))
 			{
-				const float TimeToApex = -OldVelocityWithRootMotion.Z / DerivedAccel.Z;
+				const float TimeToApex = -GravityRelativeOldVelocityWithRootMotion.Z / GravityRelativeDerivedAccel.Z;
 				
 				// The time-to-apex calculation should be precise, and we want to avoid adding a substep when we are basically already at the apex from the previous iteration's work.
 				const float ApexTimeMinimum = 0.0001f;
 				if (TimeToApex >= ApexTimeMinimum && TimeToApex < timeTick)
 				{
 					const FVector ApexVelocity = OldVelocityWithRootMotion + (DerivedAccel * TimeToApex);
-					Velocity = ApexVelocity;
-					Velocity.Z = 0.f; // Should be nearly zero anyway, but this makes apex notifications consistent.
-
+					if (HasCustomGravity())
+					{
+						const FVector GravityRelativeApexVelocity = RotateWorldToGravity(ApexVelocity);
+						Velocity = RotateGravityToWorld(FVector(GravityRelativeApexVelocity.X, GravityRelativeApexVelocity.Y, 0)); // Should be nearly zero anyway, but this makes apex notifications consistent.
+					}
+					else
+					{
+						Velocity = ApexVelocity;
+						Velocity.Z = 0.f; // Should be nearly zero anyway, but this makes apex notifications consistent.
+					}
+					
 					// We only want to move the amount of time it takes to reach the apex, and refund the unused time for next iteration.
 					const float TimeToRefund = (timeTick - TimeToApex);
 
@@ -4529,7 +4659,7 @@ void UCharacterMovementComponent::PhysFalling(float deltaTime, int32 Iterations)
 			}
 		}
 
-		if (bNotifyApex && (Velocity.Z < 0.f))
+		if (bNotifyApex && (RotateWorldToGravity(Velocity).Z < 0.f))
 		{
 			// Just passed jump apex since now going down
 			bNotifyApex = false;
@@ -4613,9 +4743,20 @@ void UCharacterMovementComponent::PhysFalling(float deltaTime, int32 Iterations)
 						// Find velocity *without* acceleration.
 						TGuardValue<FVector> RestoreAcceleration(Acceleration, FVector::ZeroVector);
 						TGuardValue<FVector> RestoreVelocity(Velocity, OldVelocity);
-						Velocity.Z = 0.f;
-						CalcVelocity(timeTick, FallingLateralFriction, false, MaxDecel);
-						VelocityNoAirControl = FVector(Velocity.X, Velocity.Y, OldVelocity.Z);
+						if (HasCustomGravity())
+						{
+							Velocity = FVector::VectorPlaneProject(Velocity, RotateGravityToWorld(FVector::UpVector));
+							const FVector GravityRelativeOffset = OldVelocity - Velocity;
+							CalcVelocity(timeTick, FallingLateralFriction, false, MaxDecel);
+							VelocityNoAirControl = Velocity + GravityRelativeOffset;
+						}
+						else
+						{
+							Velocity.Z = 0.f;
+							CalcVelocity(timeTick, FallingLateralFriction, false, MaxDecel);
+							VelocityNoAirControl = FVector(Velocity.X, Velocity.Y, OldVelocity.Z);
+						}
+						
 						VelocityNoAirControl = NewFallVelocity(VelocityNoAirControl, Gravity, GravityTime);
 					}
 
@@ -4670,7 +4811,7 @@ void UCharacterMovementComponent::PhysFalling(float deltaTime, int32 Iterations)
 						}
 
 						// Act as if there was no air control on the last move when computing new deflection.
-						if (bHasLimitedAirControl && Hit.Normal.Z > CharacterMovementConstants::VERTICAL_SLOPE_NORMAL_Z)
+						if (bHasLimitedAirControl && RotateWorldToGravity(Hit.Normal).Z > CharacterMovementConstants::VERTICAL_SLOPE_NORMAL_Z)
 						{
 							const FVector LastMoveNoAirControl = VelocityNoAirControl * LastMoveTimeSlice;
 							Delta = ComputeSlideVector(LastMoveNoAirControl, 1.f, OldHitNormal, Hit);
@@ -4700,7 +4841,7 @@ void UCharacterMovementComponent::PhysFalling(float deltaTime, int32 Iterations)
 						}
 
 						// bDitch=true means that pawn is straddling two slopes, neither of which it can stand on
-						bool bDitch = ( (OldHitImpactNormal.Z > 0.f) && (Hit.ImpactNormal.Z > 0.f) && (FMath::Abs(Delta.Z) <= UE_KINDA_SMALL_NUMBER) && ((Hit.ImpactNormal | OldHitImpactNormal) < 0.f) );
+						bool bDitch = ( (RotateWorldToGravity(OldHitImpactNormal).Z > 0.f) && (RotateWorldToGravity(Hit.ImpactNormal).Z > 0.f) && (FMath::Abs(Delta.Z) <= UE_KINDA_SMALL_NUMBER) && ((Hit.ImpactNormal | OldHitImpactNormal) < 0.f) );
 						SafeMoveUpdatedComponent( Delta, PawnRotation, true, Hit);
 						if ( Hit.Time == 0.f )
 						{
@@ -4719,17 +4860,19 @@ void UCharacterMovementComponent::PhysFalling(float deltaTime, int32 Iterations)
 							ProcessLanded(Hit, remainingTime, Iterations);
 							return;
 						}
-						else if (GetPerchRadiusThreshold() > 0.f && Hit.Time == 1.f && OldHitImpactNormal.Z >= WalkableFloorZ)
+						else if (GetPerchRadiusThreshold() > 0.f && Hit.Time == 1.f && RotateWorldToGravity(OldHitImpactNormal).Z >= WalkableFloorZ)
 						{
 							// We might be in a virtual 'ditch' within our perch radius. This is rare.
 							const FVector PawnLocation = UpdatedComponent->GetComponentLocation();
-							const float ZMovedDist = FMath::Abs(PawnLocation.Z - OldLocation.Z);
-							const float MovedDist2DSq = (PawnLocation - OldLocation).SizeSquared2D();
+							const float ZMovedDist = FMath::Abs(RotateWorldToGravity(PawnLocation - OldLocation).Z);
+							const float MovedDist2DSq = FVector::VectorPlaneProject(PawnLocation - OldLocation, RotateGravityToWorld(FVector::UpVector)).Size2D();
 							if (ZMovedDist <= 0.2f * timeTick && MovedDist2DSq <= 4.f * timeTick)
 							{
-								Velocity.X += 0.25f * GetMaxSpeed() * (RandomStream.FRand() - 0.5f);
-								Velocity.Y += 0.25f * GetMaxSpeed() * (RandomStream.FRand() - 0.5f);
-								Velocity.Z = FMath::Max<float>(JumpZVelocity * 0.25f, 1.f);
+								FVector GravityRelativeVelocity = RotateWorldToGravity(Velocity);
+								GravityRelativeVelocity.X += 0.25f * GetMaxSpeed() * (RandomStream.FRand() - 0.5f);
+								GravityRelativeVelocity.Y += 0.25f * GetMaxSpeed() * (RandomStream.FRand() - 0.5f);
+								GravityRelativeVelocity.Z = FMath::Max<float>(JumpZVelocity * 0.25f, 1.f);
+								Velocity = RotateGravityToWorld(GravityRelativeVelocity);
 								Delta = Velocity * timeTick;
 								SafeMoveUpdatedComponent(Delta, PawnRotation, true, Hit);
 							}
@@ -4739,10 +4882,12 @@ void UCharacterMovementComponent::PhysFalling(float deltaTime, int32 Iterations)
 			}
 		}
 
-		if (Velocity.SizeSquared2D() <= UE_KINDA_SMALL_NUMBER * 10.f)
+		FVector GravityRelativeVelocity = RotateWorldToGravity(Velocity);
+		if (GravityRelativeVelocity.SizeSquared2D() <= UE_KINDA_SMALL_NUMBER * 10.f)
 		{
-			Velocity.X = 0.f;
-			Velocity.Y = 0.f;
+			GravityRelativeVelocity.X = 0.f;
+			GravityRelativeVelocity.Y = 0.f;
+			Velocity = RotateGravityToWorld(GravityRelativeVelocity);
 		}
 	}
 }
@@ -4918,26 +5063,52 @@ void UCharacterMovementComponent::RevertMove(const FVector& OldLocation, UPrimit
 
 FVector UCharacterMovementComponent::ComputeGroundMovementDelta(const FVector& Delta, const FHitResult& RampHit, const bool bHitFromLineTrace) const
 {
-	const FVector FloorNormal = RampHit.ImpactNormal;
-	const FVector ContactNormal = RampHit.Normal;
-
-	if (FloorNormal.Z < (1.f - UE_KINDA_SMALL_NUMBER) && FloorNormal.Z > UE_KINDA_SMALL_NUMBER && ContactNormal.Z > UE_KINDA_SMALL_NUMBER && !bHitFromLineTrace && IsWalkable(RampHit))
+	if (!HasCustomGravity())
 	{
-		// Compute a vector that moves parallel to the surface, by projecting the horizontal movement direction onto the ramp.
-		const float FloorDotDelta = (FloorNormal | Delta);
-		FVector RampMovement(Delta.X, Delta.Y, -FloorDotDelta / FloorNormal.Z);
-		
-		if (bMaintainHorizontalGroundVelocity)
+		const FVector FloorNormal = RampHit.ImpactNormal;
+		const FVector ContactNormal = RampHit.Normal;
+
+		if (FloorNormal.Z < (1.f - UE_KINDA_SMALL_NUMBER) && FloorNormal.Z > UE_KINDA_SMALL_NUMBER && ContactNormal.Z > UE_KINDA_SMALL_NUMBER && !bHitFromLineTrace && IsWalkable(RampHit))
 		{
-			return RampMovement;
+			// Compute a vector that moves parallel to the surface, by projecting the horizontal movement direction onto the ramp.
+			const float FloorDotDelta = (FloorNormal | Delta);
+			FVector RampMovement(Delta.X, Delta.Y, -FloorDotDelta / FloorNormal.Z);
+
+			if (bMaintainHorizontalGroundVelocity)
+			{
+				return RampMovement;
+			}
+			else
+			{
+				return RampMovement.GetSafeNormal() * Delta.Size();
+			}
 		}
-		else
+		return Delta;
+	}
+	else
+	{
+		const FVector GravityRelativeDelta = RotateWorldToGravity(Delta);
+		const FVector GravityRelativeFloorNormal = RotateWorldToGravity(RampHit.ImpactNormal);
+		const FVector GravityRelativeContactNormal = RotateWorldToGravity(RampHit.Normal);
+
+		if (GravityRelativeFloorNormal.Z < (1.f - UE_KINDA_SMALL_NUMBER) && GravityRelativeFloorNormal.Z > UE_KINDA_SMALL_NUMBER && GravityRelativeContactNormal.Z > UE_KINDA_SMALL_NUMBER && !bHitFromLineTrace && IsWalkable(RampHit))
 		{
-			return RampMovement.GetSafeNormal() * Delta.Size();
+			// Compute a vector that moves parallel to the surface, by projecting the horizontal movement direction onto the ramp.
+			const float FloorDotDelta = (GravityRelativeFloorNormal | GravityRelativeDelta);
+			FVector GravityRelativeRampMovement(GravityRelativeDelta.X, GravityRelativeDelta.Y, -FloorDotDelta / GravityRelativeFloorNormal.Z);
+
+			if (bMaintainHorizontalGroundVelocity)
+			{
+				return RotateGravityToWorld(GravityRelativeRampMovement);
+			}
+			else
+			{
+				return RotateGravityToWorld(GravityRelativeRampMovement.GetSafeNormal() * GravityRelativeDelta.Size());
+			}
 		}
-	}
 
-	return Delta;
+		return RotateGravityToWorld(GravityRelativeDelta);
+	}
 }
 
 void UCharacterMovementComponent::OnCharacterStuckInGeometry(const FHitResult* Hit)
@@ -4987,7 +5158,7 @@ void UCharacterMovementComponent::MoveAlongFloor(const FVector& InVelocity, floa
 	}
 
 	// Move along the current floor
-	const FVector Delta = FVector(InVelocity.X, InVelocity.Y, 0.f) * DeltaSeconds;
+	const FVector Delta = RotateGravityToWorld(RotateWorldToGravity(InVelocity) * FVector(1.0, 1.0, 0.0)) * DeltaSeconds;
 	FHitResult Hit(1.f);
 	FVector RampVector = ComputeGroundMovementDelta(Delta, CurrentFloor.HitResult, CurrentFloor.bLineTrace);
 	SafeMoveUpdatedComponent(RampVector, UpdatedComponent->GetComponentQuat(), true, Hit);
@@ -5026,7 +5197,7 @@ void UCharacterMovementComponent::MoveAlongFloor(const FVector& InVelocity, floa
 			{
 				// hit a barrier, try to step up
 				const FVector PreStepUpLocation = UpdatedComponent->GetComponentLocation();
-				const FVector GravDir(0.f, 0.f, -1.f);
+				const FVector GravDir = GetGravityDirection();
 				if (!StepUp(GravDir, Delta * (1.f - PercentTimeApplied), Hit, OutStepDownResult))
 				{
 					UE_LOG(LogCharacterMovement, Verbose, TEXT("- StepUp (ImpactNormal %s, Normal %s"), *Hit.ImpactNormal.ToString(), *Hit.Normal.ToString());
@@ -5044,7 +5215,7 @@ void UCharacterMovementComponent::MoveAlongFloor(const FVector& InVelocity, floa
 						if (!HasAnimRootMotion() && !CurrentRootMotion.HasOverrideVelocity() && StepUpTimeSlice >= UE_KINDA_SMALL_NUMBER)
 						{
 							Velocity = (UpdatedComponent->GetComponentLocation() - PreStepUpLocation) / StepUpTimeSlice;
-							Velocity.Z = 0;
+							Velocity = FVector::VectorPlaneProject(Velocity, -GravDir);
 						}
 					}
 				}
@@ -5061,19 +5232,22 @@ void UCharacterMovementComponent::MoveAlongFloor(const FVector& InVelocity, floa
 
 void UCharacterMovementComponent::MaintainHorizontalGroundVelocity()
 {
-	if (Velocity.Z != 0.f)
+	FVector GravityRelativeVelocity = RotateWorldToGravity(Velocity);
+	if (GravityRelativeVelocity.Z != 0.f)
 	{
 		if (bMaintainHorizontalGroundVelocity)
 		{
 			// Ramp movement already maintained the velocity, so we just want to remove the vertical component.
-			Velocity.Z = 0.f;
+			GravityRelativeVelocity.Z = 0.f;
 		}
 		else
 		{
 			// Rescale velocity to be horizontal but maintain magnitude of last update.
-			Velocity = Velocity.GetSafeNormal2D() * Velocity.Size();
+			GravityRelativeVelocity = GravityRelativeVelocity.GetSafeNormal2D() * GravityRelativeVelocity.Size();
 		}
 	}
+
+	Velocity = RotateGravityToWorld(GravityRelativeVelocity);
 }
 
 
@@ -5125,7 +5299,7 @@ void UCharacterMovementComponent::PhysWalking(float deltaTime, int32 Iterations)
 		// Ensure velocity is horizontal.
 		MaintainHorizontalGroundVelocity();
 		const FVector OldVelocity = Velocity;
-		Acceleration.Z = 0.f;
+		Acceleration = FVector::VectorPlaneProject(Acceleration, -GravityDirection);
 
 		// Apply acceleration
 		if( !HasAnimRootMotion() && !CurrentRootMotion.HasOverrideVelocity() )
@@ -5195,7 +5369,7 @@ void UCharacterMovementComponent::PhysWalking(float deltaTime, int32 Iterations)
 		if ( bCheckLedges && !CurrentFloor.IsWalkableFloor() )
 		{
 			// calculate possible alternate movement
-			const FVector GravDir = FVector(0.f,0.f,-1.f);
+			const FVector GravDir = GravityDirection;
 			const FVector NewDelta = bTriedLedgeMove ? FVector::ZeroVector : GetLedgeMove(OldLocation, Delta, GravDir);
 			if ( !NewDelta.IsZero() )
 			{
@@ -5251,7 +5425,7 @@ void UCharacterMovementComponent::PhysWalking(float deltaTime, int32 Iterations)
 				// The floor check failed because it started in penetration
 				// We do not want to try to move downward because the downward sweep failed, rather we'd like to try to pop out of the floor.
 				FHitResult Hit(CurrentFloor.HitResult);
-				Hit.TraceEnd = Hit.TraceStart + FVector(0.f, 0.f, MAX_FLOOR_DIST);
+				Hit.TraceEnd = Hit.TraceStart + RotateGravityToWorld(FVector(0.f, 0.f, MAX_FLOOR_DIST));
 				const FVector RequestedAdjustment = GetPenetrationAdjustment(Hit);
 				ResolvePenetration(RequestedAdjustment, Hit, UpdatedComponent->GetComponentQuat());
 				bForceNextFloorCheck = true;
@@ -5682,10 +5856,10 @@ void UCharacterMovementComponent::AdjustFloorHeight()
 	if (OldFloorDist < MIN_FLOOR_DIST || OldFloorDist > MAX_FLOOR_DIST)
 	{
 		FHitResult AdjustHit(1.f);
-		const float InitialZ = UpdatedComponent->GetComponentLocation().Z;
+		const double InitialZ = RotateWorldToGravity(UpdatedComponent->GetComponentLocation()).Z;
 		const float AvgFloorDist = (MIN_FLOOR_DIST + MAX_FLOOR_DIST) * 0.5f;
 		const float MoveDist = AvgFloorDist - OldFloorDist;
-		SafeMoveUpdatedComponent( FVector(0.f,0.f,MoveDist), UpdatedComponent->GetComponentQuat(), true, AdjustHit );
+		SafeMoveUpdatedComponent(RotateGravityToWorld(FVector(0.f,0.f,MoveDist)), UpdatedComponent->GetComponentQuat(), true, AdjustHit );
 		UE_LOG(LogCharacterMovement, VeryVerbose, TEXT("Adjust floor height %.3f (Hit = %d)"), MoveDist, AdjustHit.bBlockingHit);
 
 		if (!AdjustHit.IsValidBlockingHit())
@@ -5694,14 +5868,13 @@ void UCharacterMovementComponent::AdjustFloorHeight()
 		}
 		else if (MoveDist > 0.f)
 		{
-			const float CurrentZ = UpdatedComponent->GetComponentLocation().Z;
+			const double CurrentZ = RotateWorldToGravity(UpdatedComponent->GetComponentLocation()).Z;
 			CurrentFloor.FloorDist += CurrentZ - InitialZ;
 		}
 		else
 		{
 			checkSlow(MoveDist < 0.f);
-			const float CurrentZ = UpdatedComponent->GetComponentLocation().Z;
-			CurrentFloor.FloorDist = CurrentZ - AdjustHit.Location.Z;
+			CurrentFloor.FloorDist = RotateWorldToGravity(UpdatedComponent->GetComponentLocation() - AdjustHit.Location).Z;
 			if (IsWalkable(AdjustHit))
 			{
 				CurrentFloor.SetFromSweep(AdjustHit, CurrentFloor.FloorDist, true);
@@ -5780,7 +5953,7 @@ void UCharacterMovementComponent::SetPostLandedPhysics(const FHitResult& Hit)
 		}
 		else
 		{
-			const FVector PreImpactAccel = Acceleration + (IsFalling() ? FVector(0.f, 0.f, GetGravityZ()) : FVector::ZeroVector);
+			const FVector PreImpactAccel = Acceleration + (IsFalling() ? -GetGravityDirection() * GetGravityZ() : FVector::ZeroVector);
 			const FVector PreImpactVelocity = Velocity;
 
 			if (DefaultLandMovementMode == MOVE_Walking ||
@@ -6028,9 +6201,20 @@ void UCharacterMovementComponent::PhysicsRotation(float DeltaTime)
 
 	if (ShouldRemainVertical())
 	{
-		DesiredRotation.Pitch = 0.f;
-		DesiredRotation.Yaw = FRotator::NormalizeAxis(DesiredRotation.Yaw);
-		DesiredRotation.Roll = 0.f;
+		if (HasCustomGravity())
+		{
+			FRotator GravityRelativeDesiredRotation = (GravityToWorldTransform * DesiredRotation.Quaternion()).Rotator();
+			GravityRelativeDesiredRotation.Pitch = 0.f;
+			GravityRelativeDesiredRotation.Yaw = FRotator::NormalizeAxis(GravityRelativeDesiredRotation.Yaw);
+			GravityRelativeDesiredRotation.Roll = 0.f;
+			DesiredRotation = (WorldToGravityTransform * GravityRelativeDesiredRotation.Quaternion()).Rotator();
+		}
+		else
+		{
+			DesiredRotation.Pitch = 0.f;
+			DesiredRotation.Yaw = FRotator::NormalizeAxis(DesiredRotation.Yaw);
+			DesiredRotation.Roll = 0.f;
+		}
 	}
 	else
 	{
@@ -6245,9 +6429,11 @@ void UCharacterMovementComponent::MoveSmooth(const FVector& InVelocity, const fl
 				if (CanStepUp(Hit))
 				{
 					OutStepDownResult = NULL; // No need for a floor when not walking.
-					if (FMath::Abs(Hit.ImpactNormal.Z) < 0.2f)
+					bool bShouldAttemptStepUp = false;
+					bShouldAttemptStepUp = FMath::Abs(RotateWorldToGravity(Hit.ImpactNormal).Z) < 0.2;
+					if (bShouldAttemptStepUp)
 					{
-						const FVector GravDir = FVector(0.f,0.f,-1.f);
+						const FVector GravDir = GetGravityDirection();
 						const FVector DesiredDir = Delta.GetSafeNormal();
 						const float UpDown = GravDir | DesiredDir;
 						if ((UpDown < 0.5f) && (UpDown > -0.2f))
@@ -6284,7 +6470,8 @@ bool UCharacterMovementComponent::IsWalkable(const FHitResult& Hit) const
 	}
 
 	// Never walk up vertical surfaces.
-	if (Hit.ImpactNormal.Z < UE_KINDA_SMALL_NUMBER)
+	const FVector GravityRelativeImpactNormal = RotateWorldToGravity(Hit.ImpactNormal);
+	if (GravityRelativeImpactNormal.Z < UE_KINDA_SMALL_NUMBER)
 	{
 		return false;
 	}
@@ -6300,7 +6487,7 @@ bool UCharacterMovementComponent::IsWalkable(const FHitResult& Hit) const
 	}
 
 	// Can't walk on this surface if it is too steep.
-	if (Hit.ImpactNormal.Z < TestWalkableZ)
+	if (GravityRelativeImpactNormal.Z < TestWalkableZ)
 	{
 		return false;
 	}
@@ -6334,7 +6521,8 @@ float UCharacterMovementComponent::K2_GetWalkableFloorZ() const
 
 bool UCharacterMovementComponent::IsWithinEdgeTolerance(const FVector& CapsuleLocation, const FVector& TestImpactPoint, const float CapsuleRadius) const
 {
-	const float DistFromCenterSq = (TestImpactPoint - CapsuleLocation).SizeSquared2D();
+	const FVector GravityRelativeToTestImpactPoint = RotateWorldToGravity(TestImpactPoint - CapsuleLocation);
+	const float DistFromCenterSq = GravityRelativeToTestImpactPoint.SizeSquared2D();
 	const float ReducedRadiusSq = FMath::Square(FMath::Max(SWEEP_EDGE_REJECT_DISTANCE + UE_KINDA_SMALL_NUMBER, CapsuleRadius - SWEEP_EDGE_REJECT_DISTANCE));
 	return DistFromCenterSq < ReducedRadiusSq;
 }
@@ -6352,8 +6540,9 @@ void UCharacterMovementComponent::ComputeFloorDist(const FVector& CapsuleLocatio
 	if (DownwardSweepResult != NULL && DownwardSweepResult->IsValidBlockingHit())
 	{
 		// Only if the supplied sweep was vertical and downward.
-		if ((DownwardSweepResult->TraceStart.Z > DownwardSweepResult->TraceEnd.Z) &&
-			(DownwardSweepResult->TraceStart - DownwardSweepResult->TraceEnd).SizeSquared2D() <= UE_KINDA_SMALL_NUMBER)
+		const bool bIsDownward = RotateWorldToGravity(DownwardSweepResult->TraceStart - DownwardSweepResult->TraceEnd).Z > 0;
+		const bool bIsVertical = RotateWorldToGravity(DownwardSweepResult->TraceStart - DownwardSweepResult->TraceEnd).SizeSquared2D() <= UE_KINDA_SMALL_NUMBER;
+		if (bIsDownward && bIsVertical)
 		{
 			// Reject hits that are barely on the cusp of the radius of the capsule
 			if (IsWithinEdgeTolerance(DownwardSweepResult->Location, DownwardSweepResult->ImpactPoint, PawnRadius))
@@ -6362,7 +6551,7 @@ void UCharacterMovementComponent::ComputeFloorDist(const FVector& CapsuleLocatio
 				bSkipSweep = true;
 
 				const bool bIsWalkable = IsWalkable(*DownwardSweepResult);
-				const float FloorDist = (CapsuleLocation.Z - DownwardSweepResult->Location.Z);
+				const float FloorDist = RotateWorldToGravity(CapsuleLocation - DownwardSweepResult->Location).Z;
 				OutFloorResult.SetFromSweep(*DownwardSweepResult, FloorDist, bIsWalkable);
 				
 				if (bIsWalkable)
@@ -6399,7 +6588,7 @@ void UCharacterMovementComponent::ComputeFloorDist(const FVector& CapsuleLocatio
 		FCollisionShape CapsuleShape = FCollisionShape::MakeCapsule(SweepRadius, PawnHalfHeight - ShrinkHeight);
 
 		FHitResult Hit(1.f);
-		bBlockingHit = FloorSweepTest(Hit, CapsuleLocation, CapsuleLocation + FVector(0.f,0.f,-TraceDist), CollisionChannel, CapsuleShape, QueryParams, ResponseParam);
+		bBlockingHit = FloorSweepTest(Hit, CapsuleLocation, CapsuleLocation + RotateGravityToWorld(FVector(0.f,0.f,-TraceDist)), CollisionChannel, CapsuleShape, QueryParams, ResponseParam);
 
 		if (bBlockingHit)
 		{
@@ -6417,7 +6606,7 @@ void UCharacterMovementComponent::ComputeFloorDist(const FVector& CapsuleLocatio
 					CapsuleShape.Capsule.HalfHeight = FMath::Max(PawnHalfHeight - ShrinkHeight, CapsuleShape.Capsule.Radius);
 					Hit.Reset(1.f, false);
 
-					bBlockingHit = FloorSweepTest(Hit, CapsuleLocation, CapsuleLocation + FVector(0.f,0.f,-TraceDist), CollisionChannel, CapsuleShape, QueryParams, ResponseParam);
+					bBlockingHit = FloorSweepTest(Hit, CapsuleLocation, CapsuleLocation + RotateGravityToWorld(FVector(0.f,0.f,-TraceDist)), CollisionChannel, CapsuleShape, QueryParams, ResponseParam);
 				}
 			}
 
@@ -6453,7 +6642,7 @@ void UCharacterMovementComponent::ComputeFloorDist(const FVector& CapsuleLocatio
 		const float ShrinkHeight = PawnHalfHeight;
 		const FVector LineTraceStart = CapsuleLocation;	
 		const float TraceDist = LineDistance + ShrinkHeight;
-		const FVector Down = FVector(0.f, 0.f, -TraceDist);
+		const FVector Down = RotateGravityToWorld(FVector(0.f, 0.f, -TraceDist));
 		QueryParams.TraceTag = SCENE_QUERY_STAT_NAME_ONLY(FloorLineTrace);
 
 		FHitResult Hit(1.f);
@@ -6632,7 +6821,7 @@ bool UCharacterMovementComponent::FloorSweepTest(
 		const FCollisionShape BoxShape = FCollisionShape::MakeBox(FVector(CapsuleRadius * 0.707f, CapsuleRadius * 0.707f, CapsuleHeight));
 
 		// First test with the box rotated so the corners are along the major axes (ie rotated 45 degrees).
-		bBlockingHit = GetWorld()->SweepSingleByChannel(OutHit, Start, End, FQuat(FVector(0.f, 0.f, -1.f), UE_PI * 0.25f), TraceChannel, BoxShape, Params, ResponseParam);
+		bBlockingHit = GetWorld()->SweepSingleByChannel(OutHit, Start, End, FQuat(RotateGravityToWorld(FVector(0.f, 0.f, -1.f)), UE_PI * 0.25f), TraceChannel, BoxShape, Params, ResponseParam);
 
 		if (!bBlockingHit)
 		{
@@ -6666,8 +6855,10 @@ bool UCharacterMovementComponent::IsValidLandingSpot(const FVector& CapsuleLocat
 		CharacterOwner->GetCapsuleComponent()->GetScaledCapsuleSize(PawnRadius, PawnHalfHeight);
 
 		// Reject hits that are above our lower hemisphere (can happen when sliding down a vertical surface).
-		const float LowerHemisphereZ = Hit.Location.Z - PawnHalfHeight + PawnRadius;
-		if (Hit.ImpactPoint.Z >= LowerHemisphereZ)
+		const FVector GravityRelativeHitLocation = RotateWorldToGravity(Hit.Location);
+		const FVector GravityRelativeHitImpactPoint = RotateWorldToGravity(Hit.ImpactPoint);
+		const float LowerHemisphereZ = GravityRelativeHitLocation.Z - PawnHalfHeight + PawnRadius;
+		if (GravityRelativeHitImpactPoint.Z >= LowerHemisphereZ)
 		{
 			return false;
 		}
@@ -6775,7 +6966,7 @@ bool UCharacterMovementComponent::ComputePerchResult(const float TestRadius, con
 	CharacterOwner->GetCapsuleComponent()->GetScaledCapsuleSize(PawnRadius, PawnHalfHeight);
 	const FVector CapsuleLocation = (bUseFlatBaseForFloorChecks ? InHit.TraceStart : InHit.Location);
 
-	const float InHitAboveBase = FMath::Max<float>(0.f, InHit.ImpactPoint.Z - (CapsuleLocation.Z - PawnHalfHeight));
+	const float InHitAboveBase = FMath::Max<float>(0.f, RotateWorldToGravity(InHit.ImpactPoint - CapsuleLocation).Z + PawnHalfHeight);
 	const float PerchLineDist = FMath::Max(0.f, InMaxFloorDist - InHitAboveBase);
 	const float PerchSweepDist = FMath::Max(0.f, InMaxFloorDist);
 
@@ -7073,7 +7264,8 @@ void UCharacterMovementComponent::HandleImpact(const FHitResult& Impact, float T
 
 	if (bEnablePhysicsInteraction)
 	{
-		const FVector ForceAccel = Acceleration + (IsFalling() ? FVector(0.f, 0.f, GetGravityZ()) : FVector::ZeroVector);
+		const FVector FallingAcceleration = -GetGravityDirection() * GetGravityZ();
+		const FVector ForceAccel = Acceleration + (IsFalling() ? FallingAcceleration : FVector::ZeroVector);
 		ApplyImpactPhysicsForces(Impact, ForceAccel, Velocity);
 	}
 }
@@ -7364,9 +7556,10 @@ void UCharacterMovementComponent::ForceClientAdjustment()
 FVector UCharacterMovementComponent::ConstrainInputAcceleration(const FVector& InputAcceleration) const
 {
 	// walking or falling pawns ignore up/down sliding
-	if (InputAcceleration.Z != 0.f && (IsMovingOnGround() || IsFalling()))
+	const double InputAccelerationDotGravityNormal = FVector::DotProduct(InputAcceleration, -GetGravityDirection());
+	if (!FMath::IsNearlyZero(InputAccelerationDotGravityNormal) && (IsMovingOnGround() || IsFalling()))
 	{
-		return FVector(InputAcceleration.X, InputAcceleration.Y, 0.f);
+		return FVector::VectorPlaneProject(InputAcceleration, -GetGravityDirection());
 	}
 
 	return InputAcceleration;
@@ -7388,7 +7581,6 @@ FVector UCharacterMovementComponent::RoundAcceleration(FVector InAccel) const
 	return InAccel;
 }
 
-
 float UCharacterMovementComponent::ComputeAnalogInputModifier() const
 {
 	const float MaxAccel = GetMaxAcceleration();
@@ -9039,28 +9231,28 @@ void FCharacterNetworkMoveData::ClientFillNetworkMoveData(const FSavedMove_Chara
 	CompressedMoveFlags = ClientMove.GetCompressedFlags();
 	MovementMode = ClientMove.EndPackedMovementMode;
 
-		// Determine if we send absolute or relative location
-		UPrimitiveComponent* ClientMovementBase = ClientMove.EndBase.Get();
+	// Determine if we send absolute or relative location
+	UPrimitiveComponent* ClientMovementBase = ClientMove.EndBase.Get();
 
-		const bool bSendBaseRelativeLocation     = MovementBaseUtility::UseRelativeLocation(ClientMovementBase);
-		const bool bSendBaseRelativeAcceleration = CharacterMovementCVars::NetUseBaseRelativeAcceleration && bSendBaseRelativeLocation;
+	const bool bSendBaseRelativeLocation     = MovementBaseUtility::UseRelativeLocation(ClientMovementBase);
+	const bool bSendBaseRelativeAcceleration = CharacterMovementCVars::NetUseBaseRelativeAcceleration && bSendBaseRelativeLocation;
 
-		const FVector SendLocation     = bSendBaseRelativeLocation ? ClientMove.SavedRelativeLocation : FRepMovement::RebaseOntoZeroOrigin(ClientMove.SavedLocation, ClientMove.CharacterOwner->GetCharacterMovement());
-		const FVector SendAcceleration = bSendBaseRelativeAcceleration ? ClientMove.SavedRelativeAcceleration : ClientMove.Acceleration;
+	const FVector SendLocation     = bSendBaseRelativeLocation ? ClientMove.SavedRelativeLocation : FRepMovement::RebaseOntoZeroOrigin(ClientMove.SavedLocation, ClientMove.CharacterOwner->GetCharacterMovement());
+	const FVector SendAcceleration = bSendBaseRelativeAcceleration ? ClientMove.SavedRelativeAcceleration : ClientMove.Acceleration;
 
-		Location = SendLocation;
-		Acceleration = SendAcceleration;
+	Location = SendLocation;
+	Acceleration = SendAcceleration;
 
-		if (bSendBaseRelativeLocation || bSendBaseRelativeAcceleration)
-		{
-			MovementBase = ClientMovementBase;
-			MovementBaseBoneName = ClientMove.EndBoneName;
-		}
-		else
-		{
-			MovementBase = nullptr;
-			MovementBaseBoneName = NAME_None;
-		}
+	if (bSendBaseRelativeLocation || bSendBaseRelativeAcceleration)
+	{
+		MovementBase = ClientMovementBase;
+		MovementBaseBoneName = ClientMove.EndBoneName;
+	}
+	else
+	{
+		MovementBase = nullptr;
+		MovementBaseBoneName = NAME_None;
+	}
 }
 
 
@@ -9260,8 +9452,8 @@ void UCharacterMovementComponent::ServerMove_PerformMovement(const FCharacterNet
 			MoveAutonomous(ClientTimeStamp, DeltaTime, ClientMoveFlags, ClientAccel);
 		}
 
-		UE_CLOG(CharacterOwner && UpdatedComponent, LogNetPlayerMovement, VeryVerbose, TEXT("ServerMove Time %f Acceleration %s Velocity %s Position %s Rotation %s DeltaTime %f Mode %s MovementBase %s.%s (Dynamic:%d)"),
-			ClientTimeStamp, *ClientAccel.ToString(), *Velocity.ToString(), *UpdatedComponent->GetComponentLocation().ToString(), *UpdatedComponent->GetComponentRotation().ToCompactString(), DeltaTime, *GetMovementName(),
+		UE_CLOG(CharacterOwner && UpdatedComponent, LogNetPlayerMovement, VeryVerbose, TEXT("ServerMove Time %f Acceleration %s Velocity %s Position %s Rotation %s GravityDirection %s DeltaTime %f Mode %s MovementBase %s.%s (Dynamic:%d)"),
+			ClientTimeStamp, *ClientAccel.ToString(), *Velocity.ToString(), *UpdatedComponent->GetComponentLocation().ToString(), *UpdatedComponent->GetComponentRotation().ToCompactString(), *GravityDirection.ToCompactString(), DeltaTime, *GetMovementName(),
 			*GetNameSafe(GetMovementBase()), *CharacterOwner->GetBasedMovement().BoneName.ToString(), MovementBaseUtility::IsDynamicBase(GetMovementBase()) ? 1 : 0);
 	}
 
@@ -9593,6 +9785,7 @@ void UCharacterMovementComponent::ServerMoveHandleClientError(float ClientTimeSt
 		ServerData->PendingAdjustment.NewBaseBoneName = MovementBaseBoneName;
 		ServerData->PendingAdjustment.NewLoc = FRepMovement::RebaseOntoZeroOrigin(ServerLoc, this);
 		ServerData->PendingAdjustment.NewRot = UpdatedComponent->GetComponentRotation();
+		ServerData->PendingAdjustment.GravityDirection = GravityDirection;
 
 		ServerData->PendingAdjustment.bBaseRelativePosition = (bDeferServerCorrectionsWhenFalling && bUseLastBase) || MovementBaseUtility::UseRelativeLocation(MovementBase);
 		ServerData->PendingAdjustment.bBaseRelativeVelocity = false;
@@ -10049,6 +10242,7 @@ bool FCharacterMoveResponseDataContainer::Serialize(UCharacterMovementComponent&
 
 		ClientAdjustment.NewLoc.NetSerialize(Ar, PackageMap, bLocalSuccess);
 		ClientAdjustment.NewVel.NetSerialize(Ar, PackageMap, bLocalSuccess);
+		NetSerializeOptionalValue(bIsSaving, Ar, ClientAdjustment.GravityDirection, UCharacterMovementComponent::DefaultGravityDirection, PackageMap);
 
 		if (bHasRotation)
 		{
@@ -10407,7 +10601,7 @@ void UCharacterMovementComponent::ClientAdjustPosition_Implementation
 	}
 
 	// Trigger event
-	OnClientCorrectionReceived(*ClientData, TimeStamp, WorldShiftedNewLocation, NewVelocity, NewBase, NewBaseBoneName, bHasBase, bBaseRelativePosition, ServerMovementMode);
+	OnClientCorrectionReceived(*ClientData, TimeStamp, WorldShiftedNewLocation, NewVelocity, NewBase, NewBaseBoneName, bHasBase, bBaseRelativePosition, ServerMovementMode, DefaultGravityDirection);
 
 	// Trust the server's positioning.
 	if (UpdatedComponent)
@@ -10476,7 +10670,7 @@ void UCharacterMovementComponent::ClientAdjustRootMotionPosition(float TimeStamp
 	CharacterOwner->ClientAdjustRootMotionPosition(TimeStamp, ServerMontageTrackPosition, ServerLoc, ServerRotation, ServerVelZ, ServerBase, ServerBoneName, bHasBase, bBaseRelativePosition, ServerMovementMode);
 }
 
-void UCharacterMovementComponent::OnClientCorrectionReceived(FNetworkPredictionData_Client_Character& ClientData, float TimeStamp, FVector NewLocation, FVector NewVelocity, UPrimitiveComponent* NewBase, FName NewBaseBoneName, bool bHasBase, bool bBaseRelativePosition, uint8 ServerMovementMode)
+void UCharacterMovementComponent::OnClientCorrectionReceived(FNetworkPredictionData_Client_Character& ClientData, float TimeStamp, FVector NewLocation, FVector NewVelocity, UPrimitiveComponent* NewBase, FName NewBaseBoneName, bool bHasBase, bool bBaseRelativePosition, uint8 ServerMovementMode, FVector ServerGravityDirection)
 {
 #if !UE_BUILD_SHIPPING
 	if (CharacterMovementCVars::NetShowCorrections != 0)
@@ -10834,7 +11028,7 @@ void UCharacterMovementComponent::ApplyDownwardForce(float DeltaSeconds)
 	if (StandingDownwardForceScale != 0.0f && CurrentFloor.HitResult.IsValidBlockingHit())
 	{
 		UPrimitiveComponent* BaseComp = CurrentFloor.HitResult.GetComponent();
-		const FVector Gravity = FVector(0.0f, 0.0f, GetGravityZ());
+		const FVector Gravity = -GetGravityDirection() * GetGravityZ();
 
 		if (BaseComp && BaseComp->IsAnySimulatingPhysics() && !Gravity.IsZero())
 		{
@@ -10945,10 +11139,12 @@ void UCharacterMovementComponent::ApplyRepulsionForce(float DeltaSeconds)
 
 void UCharacterMovementComponent::ApplyAccumulatedForces(float DeltaSeconds)
 {
-	if (PendingImpulseToApply.Z != 0.f || PendingForceToApply.Z != 0.f)
+	const FVector GravityRelativePendingImpulseToApply = RotateWorldToGravity(PendingImpulseToApply);
+	const FVector GravityRelativePendingForceToApply = RotateWorldToGravity(PendingForceToApply);
+	if (GravityRelativePendingImpulseToApply.Z != 0.0 || GravityRelativePendingForceToApply.Z != 0.0)
 	{
 		// check to see if applied momentum is enough to overcome gravity
-		if ( IsMovingOnGround() && (PendingImpulseToApply.Z + (PendingForceToApply.Z * DeltaSeconds) + (GetGravityZ() * DeltaSeconds) > UE_SMALL_NUMBER))
+		if ( IsMovingOnGround() && (GravityRelativePendingImpulseToApply.Z + (GravityRelativePendingForceToApply.Z * DeltaSeconds) + (GetGravityZ() * DeltaSeconds) > UE_SMALL_NUMBER))
 		{
 			SetMovementMode(MOVE_Falling);
 		}
diff --git a/Engine/Source/Runtime/Engine/Private/Pawn.cpp b/Engine/Source/Runtime/Engine/Private/Pawn.cpp
--- a/Engine/Source/Runtime/Engine/Private/Pawn.cpp
+++ b/Engine/Source/Runtime/Engine/Private/Pawn.cpp
@@ -306,11 +306,16 @@ UPawnNoiseEmitterComponent* APawn::GetPawnNoiseEmitterComponent() const
 	 return NoiseEmitterComponent;
 }
 
-FVector APawn::GetGravityDirection()
+FVector APawn::GetGravityDirection() const
 {
 	return FVector(0.f,0.f,-1.f);
 }
 
+FQuat APawn::GetGravityTransform() const
+{
+	return FQuat::Identity;
+}
+
 bool APawn::ShouldTickIfViewportsOnly() const 
 { 
 	return IsLocallyControlled() && Cast<APlayerController>(GetController()); 
