Fix a static analysis warning

diff --git a/Engine/Source/Runtime/Engine/Private/Materials/MaterialHLSLGenerator.cpp b/Engine/Source/Runtime/Engine/Private/Materials/MaterialHLSLGenerator.cpp
--- a/Engine/Source/Runtime/Engine/Private/Materials/MaterialHLSLGenerator.cpp
+++ b/Engine/Source/Runtime/Engine/Private/Materials/MaterialHLSLGenerator.cpp
@@ -444,7 +444,7 @@ const UE::HLSLTree::FExpression* FMaterialHLSLGenerator::AcquireExpression(UE::H
 		for (int32 Index = FunctionCallStack.Num() - 1; Index >= 0; --Index)
 		{
 			const FFunctionCallEntry* FunctionCall = FunctionCallStack[Index];
-			Hasher.AppendData(&FunctionCall->MaterialFunction, sizeof(FunctionCall->MaterialFunction));
+			Hasher.AppendData(&FunctionCall->MaterialFunction, sizeof(FunctionCall->MaterialFunction)); //-V568
 			Hasher.AppendData(&FunctionCall->ParameterAssociation, sizeof(FunctionCall->ParameterAssociation));
 			Hasher.AppendData(&FunctionCall->ParameterIndex, sizeof(FunctionCall->ParameterIndex));
 
