Add CFLAG_CullBeforeFetch and CFLAG_WarpCulling to control warp culling and CBF hardware optimizations

diff --git a/Engine/Source/Runtime/Engine/Private/ShaderCompiler/ShaderCompiler.cpp b/Engine/Source/Runtime/Engine/Private/ShaderCompiler/ShaderCompiler.cpp
--- a/Engine/Source/Runtime/Engine/Private/ShaderCompiler/ShaderCompiler.cpp
+++ b/Engine/Source/Runtime/Engine/Private/ShaderCompiler/ShaderCompiler.cpp
@@ -1148,6 +1148,22 @@ static TAutoConsoleVariable<int32> CVarOpenGLForceDXC(
 	TEXT(" 1: Force new compiler for all shaders (default)"),
 	ECVF_ReadOnly);
 
+static TAutoConsoleVariable<int32> CVarWarpCulling(
+	TEXT("r.WarpCulling"),
+	0,
+	TEXT("Enable Warp Culling optimization for platforms that support it.\n")
+	TEXT(" 0: Disable (default)\n")
+	TEXT(" 1: Enable"),
+	ECVF_ReadOnly);
+
+static TAutoConsoleVariable<int32> CVarCullBeforeFetch(
+	TEXT("r.CullBeforeFetch"),
+	0,
+	TEXT("Enable Cull-Before-Fetch optimization for platforms that support it.\n")
+	TEXT(" 0: Disable (default)\n")
+	TEXT(" 1: Enable"),
+	ECVF_ReadOnly);
+
 ENGINE_API int32 GCreateShadersOnLoad = 0;
 static FAutoConsoleVariableRef CVarCreateShadersOnLoad(
 	TEXT("r.CreateShadersOnLoad"),
@@ -6131,6 +6147,18 @@ void GlobalBeginCompileShader(
 		}
 
 		Input.Environment.SetDefine(TEXT("USE_SCENE_DEPTH_AUX"), MobileRequiresSceneDepthAux(ShaderPlatform) ? 1 : 0);
+
+		static FShaderPlatformCachedIniValue<bool> EnableCullBeforeFetchIniValue(TEXT("r.CullBeforeFetch"));
+		if (EnableCullBeforeFetchIniValue.Get((EShaderPlatform)Target.Platform) == 1)
+		{
+			Input.Environment.CompilerFlags.Add(CFLAG_CullBeforeFetch);
+		}
+
+		static FShaderPlatformCachedIniValue<bool> EnableWarpCullingIniValue(TEXT("r.WarpCulling"));
+		if (EnableWarpCullingIniValue.Get((EShaderPlatform)Target.Platform) == 1)
+		{
+			Input.Environment.CompilerFlags.Add(CFLAG_WarpCulling);
+		}
 	}
 
 	if (ShaderPlatform == SP_VULKAN_ES3_1_ANDROID || ShaderPlatform == SP_VULKAN_SM5_ANDROID)
diff --git a/Engine/Source/Runtime/RenderCore/Private/Shader.cpp b/Engine/Source/Runtime/RenderCore/Private/Shader.cpp
--- a/Engine/Source/Runtime/RenderCore/Private/Shader.cpp
+++ b/Engine/Source/Runtime/RenderCore/Private/Shader.cpp
@@ -1730,6 +1730,19 @@ void ShaderMapAppendKeyString(EShaderPlatform Platform, FString& KeyString)
 			KeyString += IsMobileDistanceFieldEnabled(Platform) ? TEXT("_MobSDF") : TEXT("");
 		}
 
+		{
+			static FShaderPlatformCachedIniValue<bool> EnableCullBeforeFetchIniValue(TEXT("r.CullBeforeFetch"));
+			if (EnableCullBeforeFetchIniValue.Get(Platform) == 1)
+			{
+				KeyString += TEXT("_CBF");
+			}
+			static FShaderPlatformCachedIniValue<bool> EnableWarpCullingIniValue(TEXT("r.WarpCulling"));
+			if (EnableWarpCullingIniValue.Get(Platform) == 1)
+			{
+				KeyString += TEXT("_WC");
+			}
+		}
+
 		{
 			KeyString += MobileUsesFullDepthPrepass(Platform) ? TEXT("_MobFDP") : TEXT("");
 		}
diff --git a/Engine/Source/Runtime/RenderCore/Public/ShaderCompilerCore.h b/Engine/Source/Runtime/RenderCore/Public/ShaderCompilerCore.h
--- a/Engine/Source/Runtime/RenderCore/Public/ShaderCompilerCore.h
+++ b/Engine/Source/Runtime/RenderCore/Public/ShaderCompilerCore.h
@@ -131,7 +131,10 @@ enum ECompilerFlags
 	CFLAG_RemoveDeadCode,
 	// Execute shader preprocessing with MCPP (instead of the new STB preprocessor)
 	CFLAG_UseLegacyPreprocessor,
-
+	// Enable CullBeforeFetch optimization on supported platforms
+	CFLAG_CullBeforeFetch,
+	// Enable WarpCulling optimization on supported platforms
+	CFLAG_WarpCulling,
 	CFLAG_Max,
 };
 static_assert(CFLAG_Max < 64, "Out of bitfield space! Modify FShaderCompilerFlags");
