LWI - make templated FActorInstanceHandle::FetchActor not spawn actor if template class is not the right type ( to prevent unnecessary actor spawns ).
*Resubmit after backed out after not properly handling interfaces. Verified this version works with both FetchActor<UClass> and FetchActor<IInterface>.

diff --git a/Engine/Source/Runtime/Engine/Classes/Engine/ActorInstanceHandle.h b/Engine/Source/Runtime/Engine/Classes/Engine/ActorInstanceHandle.h
--- a/Engine/Source/Runtime/Engine/Classes/Engine/ActorInstanceHandle.h
+++ b/Engine/Source/Runtime/Engine/Classes/Engine/ActorInstanceHandle.h
@@ -26,6 +26,9 @@ struct ENGINE_API FActorInstanceHandle
 
 	bool DoesRepresentClass(const UClass* OtherClass) const;
 
+	template<typename T>
+	bool DoesRepresent() const;
+
 	UClass* GetRepresentedClass() const;
 
 	FVector GetLocation() const;
@@ -93,8 +96,29 @@ private:
 	uint32 InstanceUID;
 };
 
+template<typename T>
+bool FActorInstanceHandle::DoesRepresent() const
+{
+	if (const UClass* RepresentedClass = GetRepresentedClass())
+	{
+		if constexpr (TIsIInterface<T>::Value)
+		{
+			return RepresentedClass->ImplementsInterface(T::UClassType::StaticClass());
+		}
+		else
+		{
+			return RepresentedClass->IsChildOf(T::StaticClass());
+		}
+	}
+	return false;
+}
+
 template <typename T>
 T* FActorInstanceHandle::FetchActor() const
 {
-	return Cast<T>(FetchActor());
+	if (DoesRepresent<T>())
+	{
+		return CastChecked<T>(FetchActor());
+	}
+	return nullptr;
 }
diff --git a/Engine/Source/Runtime/Engine/Private/Engine/ActorInstanceHandle.cpp b/Engine/Source/Runtime/Engine/Private/Engine/ActorInstanceHandle.cpp
--- a/Engine/Source/Runtime/Engine/Private/Engine/ActorInstanceHandle.cpp
+++ b/Engine/Source/Runtime/Engine/Private/Engine/ActorInstanceHandle.cpp
@@ -56,21 +56,10 @@ bool FActorInstanceHandle::IsValid() const
 
 bool FActorInstanceHandle::DoesRepresentClass(const UClass* OtherClass) const
 {
-	if (OtherClass == nullptr)
+	if (const UClass* RepresentedClass = GetRepresentedClass())
 	{
-		return false;
+		return RepresentedClass->IsChildOf(OtherClass);
 	}
-
-	if (IsActorValid())
-	{
-		return Actor->IsA(OtherClass);
-	}
-
-	if (Manager.IsValid())
-	{
-		return Manager->DoesRepresentClass(OtherClass);
-	}
-
 	return false;
 }
 
