Implemented an optimized path for creating a SparseVolumeTexture from a dense source volume. Also did a little bit of renaming of some of the involved functions/interfaces. Removed unorm-remapping option when importing a .vdb file. The usefulness of this feature is debatable and it doesn't work properly on .vdb sequences.

#rnx

diff --git a/Engine/Plugins/FX/Niagara/Source/Niagara/Private/NiagaraDataInterfaceRenderTargetVolume.cpp b/Engine/Plugins/FX/Niagara/Source/Niagara/Private/NiagaraDataInterfaceRenderTargetVolume.cpp
--- a/Engine/Plugins/FX/Niagara/Source/Niagara/Private/NiagaraDataInterfaceRenderTargetVolume.cpp
+++ b/Engine/Plugins/FX/Niagara/Source/Niagara/Private/NiagaraDataInterfaceRenderTargetVolume.cpp
@@ -657,76 +657,6 @@ UObject* UNiagaraDataInterfaceRenderTargetVolume::SimCacheBeginWrite(UObject* In
 	return nullptr;
 }
 
-
-
-class FSparseVolumeTextureDataConstructionVolumeRenderTargetAdapter : public ISparseVolumeTextureDataConstructionAdapter
-{
-public:
-
-	FSparseVolumeTextureDataConstructionVolumeRenderTargetAdapter(TArray<FFloat16Color>& SourceTextureData, EPixelFormat SourceFormat, FIntVector3 SourceGridResolution)
-		: TextureData(SourceTextureData), Format(SourceFormat), GridResolution(SourceGridResolution)
-	{}
-
-	virtual void GetAttributesInfo(FAttributesInfo& OutInfoA, FAttributesInfo& OutInfoB) const
-	{
-		OutInfoA.Format = Format;
-		OutInfoA.FallbackValue = FVector4f(0, 0, 0, 0);
-		OutInfoA.NormalizeScale = FVector4f(0, 0, 0, 0);
-		OutInfoA.NormalizeBias = FVector4f(0, 0, 0, 0);
-		OutInfoA.bNormalized = false;
-
-		OutInfoB.Format = PF_Unknown;
-		OutInfoB.FallbackValue = FVector4f(0, 0, 0, 0);
-		OutInfoB.NormalizeScale = FVector4f(0, 0, 0, 0);
-		OutInfoB.NormalizeBias = FVector4f(0, 0, 0, 0);
-		OutInfoB.bNormalized = false;
-	}
-
-	virtual FIntVector3 GetAABBMin() const
-	{
-		return FIntVector3(0, 0, 0);
-	}
-
-	virtual FIntVector3 GetAABBMax() const 
-	{
-		return GridResolution;
-	}
-
-	virtual FIntVector3 GetResolution() const
-	{
-		return GridResolution;
-	}
-
-	virtual void IteratePhysicalSource(TFunctionRef<void(const FIntVector3& Coord, int32 AttributesIdx, int32 ComponentIdx, float VoxelValue)> OnVisit) const
-	{
-		for (int32 Z = 0; Z < GridResolution.Z - 1; ++Z)
-		{
-			for (int32 Y = 0; Y < GridResolution.Y - 1; ++Y)
-			{
-				for (int32 X = 0; X < GridResolution.X - 1; ++X)
-				{
-					FIntVector3 CurrCoord(X, Y, Z);
-
-					FFloat16Color CurrCellValue;
-					CurrCellValue = TextureData[Y * GridResolution.X + Z * GridResolution.X * GridResolution.Y + X];
-
-					OnVisit(CurrCoord, 0, 0, CurrCellValue.R);
-					OnVisit(CurrCoord, 0, 1, CurrCellValue.G);
-					OnVisit(CurrCoord, 0, 2, CurrCellValue.B);
-					OnVisit(CurrCoord, 0, 3, CurrCellValue.A);
-				}
-			}
-		}
-	}
-
-	virtual ~FSparseVolumeTextureDataConstructionVolumeRenderTargetAdapter() = default;
-
-private:		
-	const TArray<FFloat16Color>& TextureData;
-	EPixelFormat Format;
-	FIntVector3 GridResolution;
-};
-
 bool UNiagaraDataInterfaceRenderTargetVolume::SimCacheWriteFrame(UObject* StorageObject, int FrameIndex, FNiagaraSystemInstance* SystemInstance, const void* OptionalPerInstanceData, FNiagaraSimCacheFeedbackContext& FeedbackContext) const
 {
 	check(OptionalPerInstanceData);
@@ -778,11 +708,17 @@ bool UNiagaraDataInterfaceRenderTargetVolume::SimCacheWriteFrame(UObject* Storag
 			{
 #if WITH_EDITOR
 				UAnimatedSparseVolumeTexture* CurrCache = CastChecked<UAnimatedSparseVolumeTexture>(StorageObject);
-								
-				FSparseVolumeTextureData SparseTextureData{};
-				FSparseVolumeTextureDataConstructionVolumeRenderTargetAdapter Adapter(TextureData, PF_FloatRGBA, InstanceData_GT->Size);
 				
-				SparseTextureData.Construct(Adapter);
+				FSparseVolumeTextureDataCreateInfo SVTCreateInfo;
+				SVTCreateInfo.VirtualVolumeAABBMin = FIntVector3::ZeroValue;
+				SVTCreateInfo.VirtualVolumeAABBMax = InstanceData_GT->Size;
+				SVTCreateInfo.FallbackValues[0] = FVector4f(0, 0, 0, 0);
+				SVTCreateInfo.FallbackValues[1] = FVector4f(0, 0, 0, 0);
+				SVTCreateInfo.AttributesFormats[0] = PF_FloatRGBA;
+				SVTCreateInfo.AttributesFormats[1] = PF_Unknown;
+
+				FSparseVolumeTextureData SparseTextureData{};
+				SparseTextureData.CreateFromDense(SVTCreateInfo, MakeArrayView<uint8>((uint8*)TextureData.GetData(), TextureData.Num() * sizeof(TextureData[0])), TArrayView<uint8>());
 				
 				CurrCache->AppendFrame(SparseTextureData);
 #endif
diff --git a/Engine/Source/Editor/SparseVolumeTexture/Private/OpenVDBImportOptions.h b/Engine/Source/Editor/SparseVolumeTexture/Private/OpenVDBImportOptions.h
--- a/Engine/Source/Editor/SparseVolumeTexture/Private/OpenVDBImportOptions.h
+++ b/Engine/Source/Editor/SparseVolumeTexture/Private/OpenVDBImportOptions.h
@@ -22,7 +22,6 @@ struct FOpenVDBSparseVolumeAttributesDesc
 {
 	TStaticArray<FOpenVDBSparseVolumeComponentMapping, 4> Mappings;
 	ESparseVolumeAttributesFormat Format;
-	bool bRemapInputForUnorm; // Maps the input from its minimum and maximum value into the [0-1] range. Clamps to [0-1] otherwise.
 };
 
 struct FOpenVDBImportOptions
diff --git a/Engine/Source/Editor/SparseVolumeTexture/Private/OpenVDBImportWindow.cpp b/Engine/Source/Editor/SparseVolumeTexture/Private/OpenVDBImportWindow.cpp
--- a/Engine/Source/Editor/SparseVolumeTexture/Private/OpenVDBImportWindow.cpp
+++ b/Engine/Source/Editor/SparseVolumeTexture/Private/OpenVDBImportWindow.cpp
@@ -345,7 +345,7 @@ void SOpenVDBComponentPicker::Construct(const FArguments& InArgs)
 	OpenVDBGridComponentInfo = InArgs._OpenVDBGridComponentInfo;
 	
 	check(ComponentIndex < 4);
-	const TCHAR* ComponentLabels[] = { TEXT("X"), TEXT("Y"), TEXT("Z"), TEXT("W") };
+	const TCHAR* ComponentLabels[] = { TEXT("R"), TEXT("G"), TEXT("B"), TEXT("A") };
 
 	this->ChildSlot
 		[
@@ -438,7 +438,7 @@ void SOpenVDBAttributesConfigurator::Construct(const FArguments& InArgs)
 				]
 
 				+ SHorizontalBox::Slot()
-				.FillWidth(1.0f)
+				.FillWidth(0.5f)
 				.VAlign(VAlign_Center)
 				.Padding(2.0f)
 				[
@@ -465,37 +465,6 @@ void SOpenVDBAttributesConfigurator::Construct(const FArguments& InArgs)
 						]
 					]
 				]
-				+ SHorizontalBox::Slot()
-				.VAlign(VAlign_Center)
-				.HAlign(HAlign_Right)
-				.AutoWidth()
-				.Padding(2.0f)
-				[
-					SNew(STextBlock)
-					.Text(LOCTEXT("UnormRemapCheckBoxLabel", "Unorm Remap"))
-					.IsEnabled_Lambda([this]()
-					{
-						return AttributesDesc->Format == ESparseVolumeAttributesFormat::Unorm8;
-					})
-				]
-				+ SHorizontalBox::Slot()
-				.VAlign(VAlign_Center)
-				.HAlign(HAlign_Right)
-				.AutoWidth()
-				.Padding(2.0f)
-				[
-					SAssignNew(RemapUnormCheckBox, SCheckBox)
-					.OnCheckStateChanged_Lambda([this](ECheckBoxState CheckBoxState)
-					{
-						AttributesDesc->bRemapInputForUnorm = CheckBoxState == ECheckBoxState::Checked;
-					})
-					.IsEnabled_Lambda([this]()
-					{
-						return AttributesDesc->Format == ESparseVolumeAttributesFormat::Unorm8;
-					})
-					.ToolTipText(LOCTEXT("UnormRemapCheckBoxTooltip", "Remaps input values for unorm formats into the [0-1] range instead of clamping values outside this range."))
-					.IsChecked(false)
-				]
 			]
 
 			+ SVerticalBox::Slot()
@@ -550,7 +519,6 @@ void SOpenVDBAttributesConfigurator::RefreshUIFromData()
 	{
 		ComponentPickers[i]->RefreshUIFromData();
 	}
-	RemapUnormCheckBox->SetIsChecked(AttributesDesc->bRemapInputForUnorm);
 }
 
 void SOpenVDBGridInfoTableRow::Construct(const FArguments& InArgs, const TSharedRef<STableViewBase>& OwnerTableView)
diff --git a/Engine/Source/Editor/SparseVolumeTexture/Private/OpenVDBImportWindow.h b/Engine/Source/Editor/SparseVolumeTexture/Private/OpenVDBImportWindow.h
--- a/Engine/Source/Editor/SparseVolumeTexture/Private/OpenVDBImportWindow.h
+++ b/Engine/Source/Editor/SparseVolumeTexture/Private/OpenVDBImportWindow.h
@@ -76,7 +76,6 @@ private:
 	TSharedPtr<SOpenVDBComponentPicker> ComponentPickers[4];
 	const TArray<TSharedPtr<ESparseVolumeAttributesFormat>>* OpenVDBSupportedTargetFormats;
 	TSharedPtr<SComboBox<TSharedPtr<ESparseVolumeAttributesFormat>>> FormatComboBox;
-	TSharedPtr<SCheckBox> RemapUnormCheckBox;
 };
 
 class SOpenVDBImportWindow : public SCompoundWidget
diff --git a/Engine/Source/Editor/SparseVolumeTexture/Private/SparseVolumeTextureFactory.cpp b/Engine/Source/Editor/SparseVolumeTexture/Private/SparseVolumeTextureFactory.cpp
--- a/Engine/Source/Editor/SparseVolumeTexture/Private/SparseVolumeTextureFactory.cpp
+++ b/Engine/Source/Editor/SparseVolumeTexture/Private/SparseVolumeTextureFactory.cpp
@@ -42,7 +42,6 @@ static void ComputeDefaultOpenVDBGridAssignment(const TArray<TSharedPtr<FOpenVDB
 			Mapping.SourceComponentIndex = INDEX_NONE;
 		}
 		AttributesDesc.Format = ESparseVolumeAttributesFormat::Float32;
-		AttributesDesc.bRemapInputForUnorm = false;
 	}
 
 	// Assign the components of the input grids to the components of the output SVT.
diff --git a/Engine/Source/Editor/SparseVolumeTexture/Private/SparseVolumeTextureOpenVDBUtility.cpp b/Engine/Source/Editor/SparseVolumeTexture/Private/SparseVolumeTextureOpenVDBUtility.cpp
--- a/Engine/Source/Editor/SparseVolumeTexture/Private/SparseVolumeTextureOpenVDBUtility.cpp
+++ b/Engine/Source/Editor/SparseVolumeTexture/Private/SparseVolumeTextureOpenVDBUtility.cpp
@@ -281,7 +281,7 @@ static EPixelFormat GetMultiComponentFormat(ESparseVolumeAttributesFormat Format
 
 #if OPENVDB_AVAILABLE
 
-class FSparseVolumeTextureDataConstructionOpenVDBAdapter : public ISparseVolumeTextureDataConstructionAdapter
+class FSparseVolumeTextureDataProviderOpenVDB : public ISparseVolumeTextureDataProvider
 {
 public:
 
@@ -361,10 +361,8 @@ public:
 			}
 		}
 
-		AttributesInfo[0].Format = MultiCompFormat[0];
-		AttributesInfo[1].Format = MultiCompFormat[1];
-		AttributesInfo[0].bNormalized = Attributes[0].bRemapInputForUnorm;
-		AttributesInfo[1].bNormalized = Attributes[1].bRemapInputForUnorm;
+		SVTCreateInfo.AttributesFormats[0] = MultiCompFormat[0];
+		SVTCreateInfo.AttributesFormats[1] = MultiCompFormat[1];
 
 		FIntVector3 SmallestAABBMin = FIntVector3(INT32_MAX);
 		FIntVector3 LargestAABBMax = FIntVector3(INT32_MIN);
@@ -376,7 +374,6 @@ public:
 		{
 			for (int32 CompIdx = 0; CompIdx < 4; ++CompIdx)
 			{
-				AttributesInfo[AttributesIdx].NormalizeScale[CompIdx] = 1.0f;
 				const uint32 SourceGridIndex = Attributes[AttributesIdx].Mappings[CompIdx].SourceGridIndex;
 				const uint32 SourceComponentIndex = Attributes[AttributesIdx].Mappings[CompIdx].SourceComponentIndex;
 				if (SourceGridIndex == INDEX_NONE)
@@ -410,16 +407,7 @@ public:
 				LargestAABBMax.Y = FMath::Max(LargestAABBMax.Y, GridInfo.VolumeActiveAABBMax.Y);
 				LargestAABBMax.Z = FMath::Max(LargestAABBMax.Z, GridInfo.VolumeActiveAABBMax.Z);
 
-				AttributesInfo[AttributesIdx].FallbackValue[CompIdx] = UniqueGridAdapters[SourceGridIndex]->GetBackgroundValue(SourceComponentIndex);
-				if (bNormalizedFormat[AttributesIdx] && Attributes[AttributesIdx].bRemapInputForUnorm)
-				{
-					float MinVal = 0.0f;
-					float MaxVal = 0.0f;
-					UniqueGridAdapters[SourceGridIndex]->GetMinMaxValue(SourceComponentIndex, &MinVal, &MaxVal);
-					const float Diff = MaxVal - MinVal;
-					AttributesInfo[AttributesIdx].NormalizeScale[CompIdx] = MaxVal > SMALL_NUMBER ? (1.0f / Diff) : 1.0f;
-					AttributesInfo[AttributesIdx].NormalizeBias[CompIdx] = -MinVal * AttributesInfo[AttributesIdx].NormalizeScale[CompIdx];
-				}
+				SVTCreateInfo.FallbackValues[AttributesIdx][CompIdx] = UniqueGridAdapters[SourceGridIndex]->GetBackgroundValue(SourceComponentIndex);
 
 				FSingleGridToComponentMapping Mapping{};
 				Mapping.AttributesIdx = (int32)AttributesIdx;
@@ -429,32 +417,15 @@ public:
 			}
 		}
 
-		AABBMin = SmallestAABBMin - VolumeBoundsMin;
-		AABBMax = LargestAABBMax - VolumeBoundsMin;
-		Resolution = LargestAABBMax - SmallestAABBMin;
+		SVTCreateInfo.VirtualVolumeAABBMin= SmallestAABBMin - VolumeBoundsMin;
+		SVTCreateInfo.VirtualVolumeAABBMax = LargestAABBMax - VolumeBoundsMin;
 
 		return true;
 	}
 
-	void GetAttributesInfo(FAttributesInfo& OutInfoA, FAttributesInfo& OutInfoB) const override
-	{
-		OutInfoA = AttributesInfo[0];
-		OutInfoB = AttributesInfo[1];
-	}
-
-	FIntVector3 GetAABBMin() const override
-	{
-		return AABBMin;
-	}
-
-	FIntVector3 GetAABBMax() const override
-	{
-		return AABBMax;
-	}
-
-	FIntVector3 GetResolution() const override
+	virtual FSparseVolumeTextureDataCreateInfo GetCreateInfo() const override
 	{
-		return Resolution;
+		return SVTCreateInfo;
 	}
 
 	void IteratePhysicalSource(TFunctionRef<void(const FIntVector3& Coord, int32 AttributesIdx, int32 ComponentIdx, float VoxelValue)> OnVisit) const override
@@ -491,11 +462,8 @@ private:
 	TArray<TArray<FSingleGridToComponentMapping, TInlineAllocator<4>>> GridToComponentMappings;
 	TStaticArray<FOpenVDBSparseVolumeAttributesDesc, NumAttributesDescs> Attributes;
 	TStaticArray<uint32, NumAttributesDescs> NumComponents;
-	TStaticArray<FAttributesInfo, NumAttributesDescs> AttributesInfo;
-	FIntVector3 AABBMin;
-	FIntVector3 AABBMax;
+	FSparseVolumeTextureDataCreateInfo SVTCreateInfo;
 	FIntVector3 VolumeBoundsMin;
-	FIntVector3 Resolution;
 };
 
 #endif // OPENVDB_AVAILABLE
@@ -503,12 +471,12 @@ private:
 bool ConvertOpenVDBToSparseVolumeTexture(TArray64<uint8>& SourceFile, const FOpenVDBImportOptions& ImportOptions, const FIntVector3& VolumeBoundsMin, FSparseVolumeTextureData& OutResult)
 {
 #if OPENVDB_AVAILABLE
-	FSparseVolumeTextureDataConstructionOpenVDBAdapter Adapter;
-	if (!Adapter.Initialize(SourceFile, ImportOptions, VolumeBoundsMin))
+	FSparseVolumeTextureDataProviderOpenVDB DataProvider;
+	if (!DataProvider.Initialize(SourceFile, ImportOptions, VolumeBoundsMin))
 	{
 		return false;
 	}
-	if (!OutResult.Construct(Adapter))
+	if (!OutResult.Create(DataProvider))
 	{
 		return false;
 	}
diff --git a/Engine/Source/Runtime/Engine/Classes/SparseVolumeTexture/SparseVolumeTexture.h b/Engine/Source/Runtime/Engine/Classes/SparseVolumeTexture/SparseVolumeTexture.h
--- a/Engine/Source/Runtime/Engine/Classes/SparseVolumeTexture/SparseVolumeTexture.h
+++ b/Engine/Source/Runtime/Engine/Classes/SparseVolumeTexture/SparseVolumeTexture.h
@@ -39,6 +39,8 @@ struct ENGINE_API FSparseVolumeTextureHeader
 	TStaticArray<FVector4f, 2> NullTileValues = TStaticArray<FVector4f, 2>(InPlace, FVector4f());
 	TStaticArray<FVector4f, 2> NullTileValuesQuantized = TStaticArray<FVector4f, 2>(InPlace, FVector4f()); // Values after converting/quantizing original fallback values to the given attribute formats
 
+	FSparseVolumeTextureHeader() = default;
+	FSparseVolumeTextureHeader(const FIntVector3& AABBMin, const FIntVector3& AABBMax, EPixelFormat FormatA, EPixelFormat FormatB, const FVector4f& FallbackValueA, const FVector4f& FallbackValueB);
 	void Serialize(FArchive& Ar);
 };
 
diff --git a/Engine/Source/Runtime/Engine/Private/SparseVolumeTexture/SparseVolumeTexture.cpp b/Engine/Source/Runtime/Engine/Private/SparseVolumeTexture/SparseVolumeTexture.cpp
--- a/Engine/Source/Runtime/Engine/Private/SparseVolumeTexture/SparseVolumeTexture.cpp
+++ b/Engine/Source/Runtime/Engine/Private/SparseVolumeTexture/SparseVolumeTexture.cpp
@@ -53,6 +53,31 @@ static int32 SVTComputeNumMipLevels(const FIntVector3& InResolution)
 
 ////////////////////////////////////////////////////////////////////////////////////////////////
 
+FSparseVolumeTextureHeader::FSparseVolumeTextureHeader(const FIntVector3& AABBMin, const FIntVector3& AABBMax, EPixelFormat FormatA, EPixelFormat FormatB, const FVector4f& FallbackValueA, const FVector4f& FallbackValueB)
+{
+	VirtualVolumeAABBMin = AABBMin;
+	VirtualVolumeAABBMax = AABBMax;
+	VirtualVolumeResolution = VirtualVolumeAABBMax - VirtualVolumeAABBMin;
+
+	PageTableVolumeAABBMin = VirtualVolumeAABBMin / SPARSE_VOLUME_TILE_RES;
+	PageTableVolumeAABBMax = (VirtualVolumeAABBMax + FIntVector3(SPARSE_VOLUME_TILE_RES - 1)) / SPARSE_VOLUME_TILE_RES;
+	PageTableVolumeResolution = PageTableVolumeAABBMax - PageTableVolumeAABBMin;
+
+	// We need to ensure a power of two resolution for the page table in order to fit all mips of the page table into the physical mips of the texture resource.
+	PageTableVolumeResolution.X = FMath::RoundUpToPowerOfTwo(PageTableVolumeResolution.X);
+	PageTableVolumeResolution.Y = FMath::RoundUpToPowerOfTwo(PageTableVolumeResolution.Y);
+	PageTableVolumeResolution.Z = FMath::RoundUpToPowerOfTwo(PageTableVolumeResolution.Z);
+	PageTableVolumeAABBMax = PageTableVolumeAABBMin + PageTableVolumeResolution;
+
+	AttributesFormats[0] = FormatA;
+	AttributesFormats[1] = FormatB;
+
+	NullTileValues[0] = FallbackValueA;
+	NullTileValues[1] = FallbackValueB;
+	NullTileValuesQuantized[0] = FallbackValueA;
+	NullTileValuesQuantized[1] = FallbackValueB;
+}
+
 void FSparseVolumeTextureHeader::Serialize(FArchive& Ar)
 {
 	Ar << Version;
@@ -434,7 +459,7 @@ bool UStreamableSparseVolumeTexture::EndInitialize(int32 InNumMipLevels)
 	{
 		UE_LOG(LogSparseVolumeTexture, Warning, TEXT("SVT has zero frames! Adding a dummy frame. SVT: %s"), *GetName());
 		FSparseVolumeTextureData DummyFrame;
-		DummyFrame.ConstructDefault();
+		DummyFrame.CreateDefault();
 		AppendFrame(DummyFrame);
 	}
 
diff --git a/Engine/Source/Runtime/Engine/Private/SparseVolumeTexture/SparseVolumeTextureData.cpp b/Engine/Source/Runtime/Engine/Private/SparseVolumeTexture/SparseVolumeTextureData.cpp
--- a/Engine/Source/Runtime/Engine/Private/SparseVolumeTexture/SparseVolumeTextureData.cpp
+++ b/Engine/Source/Runtime/Engine/Private/SparseVolumeTexture/SparseVolumeTextureData.cpp
@@ -4,12 +4,16 @@
 #include "SparseVolumeTexture/SparseVolumeTextureUtility.h"
 #include "Async/ParallelFor.h"
 #include "Hash/CityHash.h"
+#include <atomic>
 
 DEFINE_LOG_CATEGORY_STATIC(LogSparseVolumeTextureData, Log, All);
 
 // Enabling this ensures proper bilinear filtering between physical pages and empty pages by tagging neighboring empty pages as resident/physical. This causes more physical tiles to be generated though.
 #define SVT_CORRECT_TILE_ALLOCATION_FOR_LINEAR_FILTERING 0
 
+// Enabling this ensures that the order of the physical tiles will be deterministic.
+#define SVT_DETERMINISTIC_PAGE_TABLE_GENERATION 1
+
 ////////////////////////////////////////////////////////////////////////////////////////////////
 
 void FSparseVolumeTextureData::Serialize(FArchive& Ar)
@@ -43,47 +47,29 @@ void FSparseVolumeTextureData::Serialize(FArchive& Ar)
 	}
 }
 
-bool FSparseVolumeTextureData::Construct(const ISparseVolumeTextureDataConstructionAdapter& Adapter)
+bool FSparseVolumeTextureData::Create(const ISparseVolumeTextureDataProvider& DataProvider)
 {
-	ISparseVolumeTextureDataConstructionAdapter::FAttributesInfo AttributesInfo[2];
-	Adapter.GetAttributesInfo(AttributesInfo[0], AttributesInfo[1]);
-
-	MipMaps.SetNum(1);
-	Header.AttributesFormats[0] = AttributesInfo[0].Format;
-	Header.AttributesFormats[1] = AttributesInfo[1].Format;
-
-	Header.VirtualVolumeResolution = Adapter.GetResolution();
-	Header.VirtualVolumeAABBMin = Adapter.GetAABBMin();
-	Header.VirtualVolumeAABBMax = Adapter.GetAABBMax();
-	Header.PageTableVolumeAABBMin = Header.VirtualVolumeAABBMin / SPARSE_VOLUME_TILE_RES;
-	Header.PageTableVolumeAABBMax = (Header.VirtualVolumeAABBMax + FIntVector3(SPARSE_VOLUME_TILE_RES - 1)) / SPARSE_VOLUME_TILE_RES;
-
-	FIntVector3 PageTableVolumeResolution = Header.PageTableVolumeAABBMax - Header.PageTableVolumeAABBMin;
-	
-	// We need to ensure a power of two resolution for the page table in order to fit all mips of the page table into the physical mips of the texture resource.
-	PageTableVolumeResolution.X = FMath::RoundUpToPowerOfTwo(PageTableVolumeResolution.X);
-	PageTableVolumeResolution.Y = FMath::RoundUpToPowerOfTwo(PageTableVolumeResolution.Y);
-	PageTableVolumeResolution.Z = FMath::RoundUpToPowerOfTwo(PageTableVolumeResolution.Z);
+	const FSparseVolumeTextureDataCreateInfo CreateInfo = DataProvider.GetCreateInfo();
 
-	Header.PageTableVolumeAABBMax = Header.PageTableVolumeAABBMin + PageTableVolumeResolution;
+	Header = FSparseVolumeTextureHeader(CreateInfo.VirtualVolumeAABBMin, CreateInfo.VirtualVolumeAABBMax, CreateInfo.AttributesFormats[0], CreateInfo.AttributesFormats[1], CreateInfo.FallbackValues[0], CreateInfo.FallbackValues[1]);
 
-	if (PageTableVolumeResolution.X > UE::SVT::SVTMaxVolumeTextureDim
-		|| PageTableVolumeResolution.Y > UE::SVT::SVTMaxVolumeTextureDim
-		|| PageTableVolumeResolution.Z > UE::SVT::SVTMaxVolumeTextureDim)
+	if (Header.PageTableVolumeResolution.X > UE::SVT::SVTMaxVolumeTextureDim
+		|| Header.PageTableVolumeResolution.Y > UE::SVT::SVTMaxVolumeTextureDim
+		|| Header.PageTableVolumeResolution.Z > UE::SVT::SVTMaxVolumeTextureDim)
 	{
 		UE_LOG(LogSparseVolumeTextureData, Warning, TEXT("SparseVolumeTexture page table texture dimensions exceed limit (%ix%ix%i): %ix%ix%i"),
 			UE::SVT::SVTMaxVolumeTextureDim, UE::SVT::SVTMaxVolumeTextureDim, UE::SVT::SVTMaxVolumeTextureDim,
-			PageTableVolumeResolution.X, PageTableVolumeResolution.Y, PageTableVolumeResolution.Z);
+			Header.PageTableVolumeResolution.X, Header.PageTableVolumeResolution.Y, Header.PageTableVolumeResolution.Z);
 		return false;
 	}
 
-	Header.PageTableVolumeResolution = PageTableVolumeResolution;
+	MipMaps.SetNum(1);
 
 	// Tag all pages with valid data
 	TBitArray PagesWithData(false, Header.PageTableVolumeResolution.Z * Header.PageTableVolumeResolution.Y * Header.PageTableVolumeResolution.X);
-	Adapter.IteratePhysicalSource([&](const FIntVector3& Coord, int32 AttributesIdx, int32 ComponentIdx, float VoxelValue)
+	DataProvider.IteratePhysicalSource([&](const FIntVector3& Coord, int32 AttributesIdx, int32 ComponentIdx, float VoxelValue)
 		{
-			const bool bIsFallbackValue = (VoxelValue == AttributesInfo[AttributesIdx].FallbackValue[ComponentIdx]);
+			const bool bIsFallbackValue = (VoxelValue == CreateInfo.FallbackValues[AttributesIdx][ComponentIdx]);
 			if (!bIsFallbackValue)
 			{
 				check(Coord.X >= Header.VirtualVolumeAABBMin.X && Coord.Y >= Header.VirtualVolumeAABBMin.Y && Coord.Z >= Header.VirtualVolumeAABBMin.Z);
@@ -163,21 +149,9 @@ bool FSparseVolumeTextureData::Construct(const ISparseVolumeTextureDataConstruct
 	{
 		// Fill all tiles with null tile data. See below for an explanation of why this is necessary.
 		{
-			// Compute potentially normalized fallback values
-			FVector4f FallbackValues[] = { AttributesInfo[0].FallbackValue, AttributesInfo[1].FallbackValue };
-			for (int32 AttributesIdx = 0; AttributesIdx < 2; ++AttributesIdx)
-			{
-				if (AttributesInfo[AttributesIdx].bNormalized)
-				{
-					FallbackValues[AttributesIdx] = FallbackValues[AttributesIdx] * AttributesInfo[AttributesIdx].NormalizeScale + AttributesInfo[AttributesIdx].NormalizeBias;
-				}
-			}
-			Header.NullTileValues[0] = FallbackValues[0];
-			Header.NullTileValues[1] = FallbackValues[1];
-
 			// Fill all tiles with fallback values/null tile values. We can't guarantee that all voxels of a page will be written to in the next step, so if we don't do this,
 			// we might end up with physical tiles that have only been partially filled with valid voxels.
-			ParallelFor(NumAllocatedPages, [&FallbackValues, this](int32 TileIndex)
+			ParallelFor(NumAllocatedPages, [&](int32 TileIndex)
 				{
 					for (int32 Z = 0; Z < SPARSE_VOLUME_TILE_RES_PADDED; ++Z)
 					{
@@ -186,8 +160,8 @@ bool FSparseVolumeTextureData::Construct(const ISparseVolumeTextureDataConstruct
 							for (int32 X = 0; X < SPARSE_VOLUME_TILE_RES_PADDED; ++X)
 							{
 								const FIntVector3 WriteCoord = FIntVector3(X - 1, Y - 1, Z - 1);
-								WriteTileDataVoxel(TileIndex, WriteCoord, 0 /*MipLevel*/, 0 /*AttributesIndex*/, FallbackValues[0]);
-								WriteTileDataVoxel(TileIndex, WriteCoord, 0 /*MipLevel*/, 1 /*AttributesIndex*/, FallbackValues[1]);
+								WriteTileDataVoxel(TileIndex, WriteCoord, 0 /*MipLevel*/, 0 /*AttributesIndex*/, Header.NullTileValues[0]);
+								WriteTileDataVoxel(TileIndex, WriteCoord, 0 /*MipLevel*/, 1 /*AttributesIndex*/, Header.NullTileValues[1]);
 							}
 						}
 					}
@@ -207,7 +181,7 @@ bool FSparseVolumeTextureData::Construct(const ISparseVolumeTextureDataConstruct
 	}
 
 	// Write physical tile data
-	Adapter.IteratePhysicalSource([&](const FIntVector3& Coord, int32 AttributesIdx, int32 ComponentIdx, float VoxelValue)
+	DataProvider.IteratePhysicalSource([&](const FIntVector3& Coord, int32 AttributesIdx, int32 ComponentIdx, float VoxelValue)
 		{
 			const FIntVector3 GridCoord = Coord;
 			check(GridCoord.X >= 0 && GridCoord.Y >= 0 && GridCoord.Z >= 0);
@@ -223,31 +197,182 @@ bool FSparseVolumeTextureData::Construct(const ISparseVolumeTextureDataConstruct
 				return;
 			}
 
-			float WriteValue = VoxelValue;
-			if (AttributesInfo[AttributesIdx].bNormalized)
-			{
-				WriteValue = WriteValue * AttributesInfo[AttributesIdx].NormalizeScale[ComponentIdx] + AttributesInfo[AttributesIdx].NormalizeBias[ComponentIdx];
-			}
-
 			const int32 TileIndex = (int32)MipMaps[0].PageTable[PageIndex];
 			const FIntVector3 TileLocalCoord = GridCoord % SPARSE_VOLUME_TILE_RES;
-			WriteTileDataVoxel(TileIndex - 1, TileLocalCoord, 0 /*MipLevel*/, AttributesIdx, FVector4f(WriteValue, WriteValue, WriteValue, WriteValue), ComponentIdx);  // -1 because 0 is reserved for the null tile which we don't actually store in memory
+			WriteTileDataVoxel(TileIndex - 1, TileLocalCoord, 0 /*MipLevel*/, AttributesIdx, FVector4f(VoxelValue, VoxelValue, VoxelValue, VoxelValue), ComponentIdx);  // -1 because 0 is reserved for the null tile which we don't actually store in memory
 		});
 
 	return true;
 }
 
-void FSparseVolumeTextureData::ConstructDefault()
+bool FSparseVolumeTextureData::CreateFromDense(const FSparseVolumeTextureDataCreateInfo& CreateInfo, const TArrayView<uint8>& VoxelDataA, const TArrayView<uint8>& VoxelDataB)
+{	
+	Header = FSparseVolumeTextureHeader(CreateInfo.VirtualVolumeAABBMin, CreateInfo.VirtualVolumeAABBMax, CreateInfo.AttributesFormats[0], CreateInfo.AttributesFormats[1], CreateInfo.FallbackValues[0], CreateInfo.FallbackValues[1]);
+
+	if (Header.PageTableVolumeResolution.X > UE::SVT::SVTMaxVolumeTextureDim
+		|| Header.PageTableVolumeResolution.Y > UE::SVT::SVTMaxVolumeTextureDim
+		|| Header.PageTableVolumeResolution.Z > UE::SVT::SVTMaxVolumeTextureDim)
+	{
+		UE_LOG(LogSparseVolumeTextureData, Warning, TEXT("SparseVolumeTexture page table texture dimensions exceed limit (%ix%ix%i): %ix%ix%i"),
+			UE::SVT::SVTMaxVolumeTextureDim, UE::SVT::SVTMaxVolumeTextureDim, UE::SVT::SVTMaxVolumeTextureDim,
+			Header.PageTableVolumeResolution.X, Header.PageTableVolumeResolution.Y, Header.PageTableVolumeResolution.Z);
+		return false;
+	}
+
+	const int64 FormatSize[] = { GPixelFormats[Header.AttributesFormats[0]].BlockBytes, GPixelFormats[Header.AttributesFormats[1]].BlockBytes };
+	const FIntVector3 FractionalPageOffset = Header.VirtualVolumeAABBMin % SPARSE_VOLUME_TILE_RES; // Offset from page origin to source data volume origin
+	uint8 NullTileValuesU8[2][sizeof(float) * 4] = {};
+	for (int32 i = 0; i < 2; ++i)
+	{
+		if (CreateInfo.AttributesFormats[i] != PF_Unknown)
+		{
+			UE::SVT::WriteVoxel(0, NullTileValuesU8[i], CreateInfo.AttributesFormats[i], CreateInfo.FallbackValues[i]);
+			Header.NullTileValuesQuantized[i] = UE::SVT::ReadVoxel(0, NullTileValuesU8[i], CreateInfo.AttributesFormats[i]);
+		}
+	}
+
+	MipMaps.SetNum(1);
+
+	MipMaps[0].PageTable.SetNumZeroed(Header.PageTableVolumeResolution.X * Header.PageTableVolumeResolution.Y * Header.PageTableVolumeResolution.Z);
+	std::atomic<uint32> NumAllocatedPagesAtomic = 0;
+
+	// This is a two-pass algorithm; first marking pages and allocating physical tiles and then in a second pass writing to the physical tiles.
+	// Most of the setup of these two passes is very similar, so we use a single lambda for them.
+	auto ProcessBlock = [&](int32 JobIndex, bool bWriteTileData)
+	{
+		// Make sure we only process voxels inside a single page in the Y and Z dimension so we don't end up reading/writing the same entries in the page table from multiple threads.
+		const FIntVector3 PageBegin = FIntVector3(0, JobIndex % Header.PageTableVolumeResolution.Y, JobIndex / Header.PageTableVolumeResolution.Y) + Header.PageTableVolumeAABBMin;
+		const FIntVector3 PageEnd = FIntVector3(Header.PageTableVolumeAABBMax.X, FMath::Min(PageBegin.Y + 1, Header.PageTableVolumeAABBMax.Y), FMath::Min(PageBegin.Z + 1, Header.PageTableVolumeAABBMax.Z));
+		const FIntVector3 PageVoxelBegin = PageBegin * SPARSE_VOLUME_TILE_RES;
+		const FIntVector3 PageVoxelEnd = PageEnd * SPARSE_VOLUME_TILE_RES;
+		const FIntVector3 ActiveVoxelBegin = PageVoxelBegin + FractionalPageOffset;
+		const FIntVector3 ActiveVoxelEnd = FIntVector3(FMath::Min(PageVoxelEnd.X, Header.VirtualVolumeAABBMax.X), FMath::Min(PageVoxelEnd.Y, Header.VirtualVolumeAABBMax.Y), FMath::Min(PageVoxelEnd.Z, Header.VirtualVolumeAABBMax.Z));
+
+		// Iterate over voxels in the YZ plane and over pages along the X axis.
+		// For both page table filling and tile writing, we can process an entire page length worth of voxels (along the X axis) at a time.
+		for (int64 Z = ActiveVoxelBegin.Z; Z < ActiveVoxelEnd.Z; ++Z)
+		{
+			for (int64 Y = ActiveVoxelBegin.Y; Y < ActiveVoxelEnd.Y; ++Y)
+			{
+				for (int64 PageX = PageBegin.X; PageX < PageEnd.X; ++PageX)
+				{
+					const FIntVector3 GridCoord = FIntVector3(PageX * SPARSE_VOLUME_TILE_RES + FractionalPageOffset.X, Y, Z); // Global/absolute voxel coordinate
+					check(GridCoord.X >= 0 && GridCoord.Y >= 0 && GridCoord.Z >= 0);
+					check(GridCoord.X < Header.VirtualVolumeAABBMax.X && GridCoord.Y < Header.VirtualVolumeAABBMax.Y && GridCoord.Z < Header.VirtualVolumeAABBMax.Z);
+					const FIntVector3 PageCoord = (GridCoord / SPARSE_VOLUME_TILE_RES) - Header.PageTableVolumeAABBMin;
+					check(PageCoord.X >= 0 && PageCoord.Y >= 0 && PageCoord.Z >= 0);
+					check(PageCoord.X < Header.PageTableVolumeResolution.X && PageCoord.Y < Header.PageTableVolumeResolution.Y && PageCoord.Z < Header.PageTableVolumeResolution.Z);
+
+					const int32 PageIndex = PageCoord.Z * (Header.PageTableVolumeResolution.Y * Header.PageTableVolumeResolution.X) + PageCoord.Y * Header.PageTableVolumeResolution.X + PageCoord.X;
+					const uint32 TileIndex = MipMaps[0].PageTable[PageIndex];
+
+					// Only write tile data if a physical tile has been allocated.
+					// Only try to allocate a physical tile if there is none already.
+					if ((bWriteTileData && TileIndex != 0) || (!bWriteTileData && TileIndex == 0))
+					{
+						// Coordinates into the dense source volume data
+						const int64 BeginSrcX = PageX * SPARSE_VOLUME_TILE_RES - Header.VirtualVolumeAABBMin.X;
+						const int64 EndSrcX = FMath::Min(BeginSrcX + SPARSE_VOLUME_TILE_RES, Header.VirtualVolumeResolution.X);
+						const int64 SrcY = Y - Header.VirtualVolumeAABBMin.Y;
+						const int64 SrcZ = Z - Header.VirtualVolumeAABBMin.Z;
+						const int64 NumVoxelsToProcess = EndSrcX - BeginSrcX;
+
+						// Compute flat source and destination voxel index. Make sure to account for the fact that the voxels in the dst buffer have padding/a border.
+						const FInt64Vector3 LocalCoordPadded = FInt64Vector3(GridCoord % SPARSE_VOLUME_TILE_RES) + FInt64Vector3(SPARSE_VOLUME_TILE_BORDER);
+						const int64 DstVoxelIndex = ((int64)(TileIndex - 1) * (int64)UE::SVT::SVTNumVoxelsPerPaddedTile) 
+							+ (LocalCoordPadded.Z * (SPARSE_VOLUME_TILE_RES_PADDED * SPARSE_VOLUME_TILE_RES_PADDED))
+							+ (LocalCoordPadded.Y * SPARSE_VOLUME_TILE_RES_PADDED)
+							+ LocalCoordPadded.X;
+						const int64 SrcVoxelIndex = (SrcZ * (Header.VirtualVolumeResolution.Y * Header.VirtualVolumeResolution.X)) + (SrcY * Header.VirtualVolumeResolution.X) + BeginSrcX;
+						
+						if (bWriteTileData)
+						{
+							// Copy the whole row at once!
+							if (FormatSize[0])
+							{
+								FMemory::Memcpy(MipMaps[0].PhysicalTileDataA.GetData() + DstVoxelIndex * FormatSize[0], &VoxelDataA[SrcVoxelIndex * FormatSize[0]], FormatSize[0] * NumVoxelsToProcess);
+							}
+							if (FormatSize[1])
+							{
+								FMemory::Memcpy(MipMaps[0].PhysicalTileDataB.GetData() + DstVoxelIndex * FormatSize[1], &VoxelDataB[SrcVoxelIndex * FormatSize[1]], FormatSize[1] * NumVoxelsToProcess);
+							}
+						}
+						else
+						{
+							// Compare each voxel of the row against the quantized fallback value and only allocate a physical tile if there are non-fallback voxels present.
+							bool bHasNonFallbackValue = false;
+							for (int64 i = 0; i < NumVoxelsToProcess && !bHasNonFallbackValue; ++i)
+							{
+								bHasNonFallbackValue |= (FormatSize[0] != 0) && (FMemory::Memcmp(&VoxelDataA[(SrcVoxelIndex + i) * FormatSize[0]], NullTileValuesU8[0], FormatSize[0]) != 0);
+								bHasNonFallbackValue |= (FormatSize[1] != 0) && (FMemory::Memcmp(&VoxelDataB[(SrcVoxelIndex + i) * FormatSize[1]], NullTileValuesU8[1], FormatSize[1]) != 0);
+							}
+							if (bHasNonFallbackValue)
+							{
+								MipMaps[0].PageTable[PageIndex] = NumAllocatedPagesAtomic.fetch_add(1) + 1u;
+							}
+						}
+					}
+				}
+			}
+		}
+	};
+
+	const int32 NumJobs = Header.PageTableVolumeResolution.Z * Header.PageTableVolumeResolution.Y;
+
+	// Create page table
+	ParallelFor(NumJobs, [&](int32 Index) { ProcessBlock(Index, false /*bWriteTileData*/); });
+	const int32 NumAllocatedPages = (int32)NumAllocatedPagesAtomic.load();
+
+#if SVT_DETERMINISTIC_PAGE_TABLE_GENERATION
+	{
+		uint32 NextPageTableEntry = 1;
+		for (uint32& PageTableEntry : MipMaps[0].PageTable)
+		{
+			if (PageTableEntry != 0)
+			{
+				PageTableEntry = NextPageTableEntry++;
+			}
+		}
+		check((NumAllocatedPages + 1) == NextPageTableEntry);
+	}
+#endif // SVT_DETERMINISTIC_PAGE_TABLE_GENERATION
+
+	// Allocate tile data
+	MipMaps[0].PhysicalTileDataA.SetNum((int64)NumAllocatedPages * (int64)UE::SVT::SVTNumVoxelsPerPaddedTile * (int64)FormatSize[0]);
+	MipMaps[0].PhysicalTileDataB.SetNum((int64)NumAllocatedPages * (int64)UE::SVT::SVTNumVoxelsPerPaddedTile * (int64)FormatSize[1]);
+	MipMaps[0].NumPhysicalTiles = NumAllocatedPages;
+
+	// Clear all physical tiles to the fallback value. We can't guarantee that every voxel will be written to in the next step, so we need to ensure that tiles are always fully written to with valid data.
+	// SVT_TODO: Make this smarter. We don't need to write ALL voxels in this stage because we know which voxels will be overwritten anyways.
+	if (FractionalPageOffset != FIntVector3::ZeroValue || (Header.VirtualVolumeAABBMax % SPARSE_VOLUME_TILE_RES) != FIntVector3::ZeroValue)
+	{
+		ParallelFor(NumJobs, [&](int32 TileIndex)
+			{
+				for (int32 VoxelIndex = 0; VoxelIndex < UE::SVT::SVTNumVoxelsPerPaddedTile; ++VoxelIndex)
+				{
+					const int64 VoxelIndexGlobal = (int64)TileIndex * (int64)UE::SVT::SVTNumVoxelsPerPaddedTile + (int64)VoxelIndex;
+					if (FormatSize[0])
+					{
+						FMemory::Memcpy(MipMaps[0].PhysicalTileDataA.GetData() + VoxelIndexGlobal * FormatSize[0], NullTileValuesU8[0], FormatSize[0]);
+					}
+					if (FormatSize[1])
+					{
+						FMemory::Memcpy(MipMaps[0].PhysicalTileDataB.GetData() + VoxelIndexGlobal * FormatSize[1], NullTileValuesU8[1], FormatSize[1]);
+					}
+				}
+			});
+	}
+
+	// Write physical tile data
+	ParallelFor(NumJobs, [&](int32 Index) { ProcessBlock(Index, true /*bWriteTileData*/); });
+
+	return true;
+}
+
+void FSparseVolumeTextureData::CreateDefault()
 {
 	// Store a single 1x1x1 mip level with the page table pointing to a (zero-valued) null tile.
-
-	Header.VirtualVolumeResolution = FIntVector3(1, 1, 1);
-	Header.VirtualVolumeAABBMin = FIntVector3::ZeroValue;
-	Header.VirtualVolumeAABBMax = Header.VirtualVolumeResolution;
-	Header.PageTableVolumeResolution = FIntVector3(1, 1, 1);
-	Header.PageTableVolumeAABBMin = FIntVector3::ZeroValue;
-	Header.PageTableVolumeAABBMax = Header.PageTableVolumeResolution;
-	Header.AttributesFormats[0] = PF_R8;
+	Header = FSparseVolumeTextureHeader(FIntVector3::ZeroValue, FIntVector3(1, 1, 1), PF_R8, PF_Unknown, FVector4f(), FVector4f());
 
 	MipMaps.SetNum(1);
 	MipMaps[0].PageTable.SetNumZeroed(1);
diff --git a/Engine/Source/Runtime/Engine/Public/SparseVolumeTexture/SparseVolumeTextureData.h b/Engine/Source/Runtime/Engine/Public/SparseVolumeTexture/SparseVolumeTextureData.h
--- a/Engine/Source/Runtime/Engine/Public/SparseVolumeTexture/SparseVolumeTextureData.h
+++ b/Engine/Source/Runtime/Engine/Public/SparseVolumeTexture/SparseVolumeTextureData.h
@@ -5,24 +5,20 @@
 #include "CoreMinimal.h"
 #include "SparseVolumeTexture/SparseVolumeTexture.h"
 
-class ENGINE_API ISparseVolumeTextureDataConstructionAdapter
+struct ENGINE_API FSparseVolumeTextureDataCreateInfo
 {
-public:
-	struct ENGINE_API FAttributesInfo
-	{
-		EPixelFormat Format;
-		FVector4f FallbackValue;
-		FVector4f NormalizeScale;
-		FVector4f NormalizeBias;
-		bool bNormalized;
-	};
+	FIntVector3 VirtualVolumeAABBMin = FIntVector3(INT32_MAX, INT32_MAX, INT32_MAX);
+	FIntVector3 VirtualVolumeAABBMax = FIntVector3(INT32_MIN, INT32_MIN, INT32_MIN);
+	TStaticArray<EPixelFormat, 2> AttributesFormats = TStaticArray<EPixelFormat, 2>(InPlace, PF_Unknown); // Currently supported: PF_R8, PF_R8G8, PF_R8G8B8A8, PF_R16F, PF_G16R16F, PF_FloatRGBA, PF_R32_FLOAT, PF_G32R32F, PF_A32B32G32R32F
+	TStaticArray<FVector4f, 2> FallbackValues = TStaticArray<FVector4f, 2>(InPlace, FVector4f());
+};
 
-	virtual void GetAttributesInfo(FAttributesInfo& OutInfoA, FAttributesInfo& OutInfoB) const = 0;
-	virtual FIntVector3 GetAABBMin() const = 0;
-	virtual FIntVector3 GetAABBMax() const = 0;
-	virtual FIntVector3 GetResolution() const = 0;
+class ENGINE_API ISparseVolumeTextureDataProvider
+{
+public:
+	virtual FSparseVolumeTextureDataCreateInfo GetCreateInfo() const = 0;
 	virtual void IteratePhysicalSource(TFunctionRef<void(const FIntVector3& Coord, int32 AttributesIdx, int32 ComponentIdx, float VoxelValue)> OnVisit) const = 0;
-	virtual ~ISparseVolumeTextureDataConstructionAdapter() = default;
+	virtual ~ISparseVolumeTextureDataProvider() = default;
 };
 
 struct ENGINE_API FSparseVolumeTextureDataAddressingInfo
@@ -54,8 +50,9 @@ struct ENGINE_API FSparseVolumeTextureData
 
 	void Serialize(FArchive& Ar);
 	
-	bool Construct(const ISparseVolumeTextureDataConstructionAdapter& Adapter);
-	void ConstructDefault();
+	bool Create(const ISparseVolumeTextureDataProvider& DataProvider);
+	bool CreateFromDense(const FSparseVolumeTextureDataCreateInfo& CreateInfo, const TArrayView<uint8>& VoxelDataA, const TArrayView<uint8>& VoxelDataB);
+	void CreateDefault();
 	FVector4f Load(const FIntVector3& VolumeCoord, int32 MipLevel, int32 AttributesIdx, const FSparseVolumeTextureDataAddressingInfo& AddressingInfo) const;
 	bool BuildDerivedData(const FSparseVolumeTextureDataAddressingInfo& AddressingInfo, int32 NumMipLevels, bool bMoveMip0FromThis, FSparseVolumeTextureData& OutDerivedData);
 
