Scene Capture:  expose boolean to enable near clip override to blueprints (UDN request).

diff --git a/Engine/Source/Runtime/Engine/Classes/Components/SceneCaptureComponent2D.h b/Engine/Source/Runtime/Engine/Classes/Components/SceneCaptureComponent2D.h
--- a/Engine/Source/Runtime/Engine/Classes/Components/SceneCaptureComponent2D.h
+++ b/Engine/Source/Runtime/Engine/Classes/Components/SceneCaptureComponent2D.h
@@ -49,7 +49,7 @@ public:
 	UPROPERTY(interp, Category=PostProcessVolume, BlueprintReadWrite, meta=(UIMin = "0.0", UIMax = "1.0"))
 	float PostProcessBlendWeight;
 
-	UPROPERTY(EditAnywhere, Category = Projection, meta = (InlineEditConditionToggle))
+	UPROPERTY(EditAnywhere, BlueprintReadWrite, Category = Projection, meta = (InlineEditConditionToggle))
 	uint32 bOverride_CustomNearClippingPlane : 1;
 
 	/** 
