Added current time accessors to anim notify events, plus library functions to allow for easy access to relative times and ratios

diff --git a/Engine/Source/Runtime/Engine/Private/Animation/AnimNotifyLibrary.cpp b/Engine/Source/Runtime/Engine/Private/Animation/AnimNotifyLibrary.cpp
--- a/Engine/Source/Runtime/Engine/Private/Animation/AnimNotifyLibrary.cpp
+++ b/Engine/Source/Runtime/Engine/Private/Animation/AnimNotifyLibrary.cpp
@@ -4,6 +4,7 @@
 #include "Animation/AnimNodeMessages.h"
 #include "Animation/AnimNotifyQueue.h"
 #include "Animation/AnimNotifyEndDataContext.h"
+#include "Animation/AnimSequenceBase.h"
 
 #include UE_INLINE_GENERATED_CPP_BY_NAME(AnimNotifyLibrary)
 
@@ -19,3 +20,43 @@ bool UAnimNotifyLibrary::NotifyStateReachedEnd(const FAnimNotifyEventReference&
 	return false;
 }
 
+float UAnimNotifyLibrary::GetCurrentAnimationTime(const FAnimNotifyEventReference& EventReference)
+{
+	return EventReference.GetCurrentAnimationTime();
+}
+
+float UAnimNotifyLibrary::GetCurrentAnimationTimeRatio(const FAnimNotifyEventReference& EventReference)
+{
+	if(const UAnimSequenceBase* AnimSequenceBase = Cast<UAnimSequenceBase>(EventReference.GetSourceObject()))
+	{
+		if(AnimSequenceBase->GetPlayLength() > 0.0f)
+		{
+			return FMath::Clamp(EventReference.GetCurrentAnimationTime() / AnimSequenceBase->GetPlayLength(), 0.0f, 1.0f);
+		}
+	}
+
+	return 0.0f;
+}
+
+float UAnimNotifyLibrary::GetCurrentAnimationNotifyStateTime(const FAnimNotifyEventReference& EventReference)
+{
+	if(const FAnimNotifyEvent* NotifyEvent = EventReference.GetNotify())
+	{
+		return FMath::Clamp(EventReference.GetCurrentAnimationTime() - NotifyEvent->GetTriggerTime(), 0.0f, NotifyEvent->GetDuration());
+	}
+
+	return 0.0f;
+}
+
+float UAnimNotifyLibrary::GetCurrentAnimationNotifyStateTimeRatio(const FAnimNotifyEventReference& EventReference)
+{
+	if(const FAnimNotifyEvent* NotifyEvent = EventReference.GetNotify())
+	{
+		if(NotifyEvent->GetDuration() > 0.0f)
+		{
+			return FMath::Clamp((EventReference.GetCurrentAnimationTime() - NotifyEvent->GetTriggerTime()) / NotifyEvent->GetDuration(), 0.0f, 1.0f);
+		}
+	}
+
+	return 0.0f;
+}
diff --git a/Engine/Source/Runtime/Engine/Private/Animation/AnimNotifyQueue.cpp b/Engine/Source/Runtime/Engine/Private/Animation/AnimNotifyQueue.cpp
--- a/Engine/Source/Runtime/Engine/Private/Animation/AnimNotifyQueue.cpp
+++ b/Engine/Source/Runtime/Engine/Private/Animation/AnimNotifyQueue.cpp
@@ -18,6 +18,7 @@ bool operator==(const FAnimNotifyEventReference& Lhs, const FAnimNotifyEvent& Rh
 void FAnimNotifyEventReference::GatherTickRecordData(const FAnimTickRecord& InTickRecord)
 {
 	ContextData = InTickRecord.ContextData;
+	CurrentAnimTime = InTickRecord.TimeAccumulator != nullptr ? *InTickRecord.TimeAccumulator : 0.0f;
 }
 
 bool FAnimNotifyQueue::PassesFiltering(const FAnimNotifyEvent* Notify) const
diff --git a/Engine/Source/Runtime/Engine/Public/Animation/AnimNotifyLibrary.h b/Engine/Source/Runtime/Engine/Public/Animation/AnimNotifyLibrary.h
--- a/Engine/Source/Runtime/Engine/Public/Animation/AnimNotifyLibrary.h
+++ b/Engine/Source/Runtime/Engine/Public/Animation/AnimNotifyLibrary.h
@@ -25,4 +25,42 @@ public:
 	*/
 	UFUNCTION(BlueprintPure, Category = "Utilities|Animation|Notifies" , meta = (ScriptMethod))
 	static bool NotifyStateReachedEnd(const FAnimNotifyEventReference& EventReference);
+
+	/**
+	 * Get the current anim notify time in seconds for when this notify was fired
+	 *
+	 * @param EventReference		The event to inspect
+	 * @return the time in seconds through the current animation for when this notify was fired
+	 */
+	UFUNCTION(BlueprintPure, Category = "Utilities|Animation|Notifies" , meta = (ScriptMethod))
+	static float GetCurrentAnimationTime(const FAnimNotifyEventReference& EventReference);
+
+	/**
+	 * Get the current anim notify time as a ratio (0 -> 1) through the animation for when this notify was fired
+	 *
+	 * @param EventReference		The event to inspect
+	 * @return the time as a ratio (0 -> 1) through the animation for when this notify was fired
+	 */
+	UFUNCTION(BlueprintPure, Category = "Utilities|Animation|Notifies" , meta = (ScriptMethod))
+	static float GetCurrentAnimationTimeRatio(const FAnimNotifyEventReference& EventReference);
+	
+	/**
+	 * Gets the current time in seconds relative to the start of the notify state, clamped to the range of the notify
+	 * state
+	 *
+	 * @param EventReference		The event to inspect
+	 * @return  the current time in seconds relative to the start of the notify state, clamped to the range of the
+	 *			notify state
+	 */
+	UFUNCTION(BlueprintPure, Category = "Utilities|Animation|Notifies" , meta = (ScriptMethod))
+	static float GetCurrentAnimationNotifyStateTime(const FAnimNotifyEventReference& EventReference);
+
+	/**
+	 * Gets the current time as a ratio (0 -> 1) relative to the start of the notify state
+	 *
+	 * @param EventReference		The event to inspect
+	 * @return  the current time as a ratio (0 -> 1) relative to the start of the notify state
+	 */
+	UFUNCTION(BlueprintPure, Category = "Utilities|Animation|Notifies" , meta = (ScriptMethod))
+	static float GetCurrentAnimationNotifyStateTimeRatio(const FAnimNotifyEventReference& EventReference);
 };
diff --git a/Engine/Source/Runtime/Engine/Public/Animation/AnimNotifyQueue.h b/Engine/Source/Runtime/Engine/Public/Animation/AnimNotifyQueue.h
--- a/Engine/Source/Runtime/Engine/Public/Animation/AnimNotifyQueue.h
+++ b/Engine/Source/Runtime/Engine/Public/Animation/AnimNotifyQueue.h
@@ -21,19 +21,9 @@ struct FAnimNotifyEventReference
 {
 	GENERATED_BODY()
 
-	FAnimNotifyEventReference()
-		: Notify(nullptr)
-		, MirrorTable(nullptr)
-		, NotifySource(nullptr)
-
-	{}
+	FAnimNotifyEventReference() = default;
 
-	FAnimNotifyEventReference(const FAnimNotifyEventReference& rhs)
-		: ContextData(rhs.ContextData)
-		, Notify(rhs.Notify)
-		, MirrorTable(rhs.MirrorTable)
-		, NotifySource(rhs.NotifySource)
-	{}
+	FAnimNotifyEventReference(const FAnimNotifyEventReference& rhs) = default;
 
 	FAnimNotifyEventReference(const FAnimNotifyEvent* InNotify, const UObject* InNotifySource)
 		: Notify(InNotify)
@@ -108,18 +98,33 @@ struct FAnimNotifyEventReference
 		ContextData->Add(MakeUnique<Type>(Forward<TArgs>(Args)...));
 	}
 
+	// Gets the source object of this notify (e.g. anim sequence), if any
+	const UObject* GetSourceObject() const
+	{
+		return NotifySource;
+	}
+
+	// Gets the current animation's time that this notify was fired at
+	float GetCurrentAnimationTime() const
+	{
+		return CurrentAnimTime;
+	}
+
 private:
 	// Context data gleaned from the tick record
 	TSharedPtr<TArray<TUniquePtr<const UE::Anim::IAnimNotifyEventContextDataInterface>>> ContextData;
 	
-	const FAnimNotifyEvent* Notify;
+	const FAnimNotifyEvent* Notify = nullptr;
 
 	// If set, the Notify has been mirrored.  The mirrored name can be found in MirrorTable->AnimNotifyToMirrorAnimNotifyMap
 	UPROPERTY(transient)
-	TObjectPtr<const UMirrorDataTable> MirrorTable; 
+	TObjectPtr<const UMirrorDataTable> MirrorTable = nullptr; 
 
 	UPROPERTY(transient)
-	TObjectPtr<const UObject> NotifySource;
+	TObjectPtr<const UObject> NotifySource = nullptr;
+
+	// The recorded time from the tick record that this notify event was fired at
+	float CurrentAnimTime = 0.0f;
 };
 
 USTRUCT()
