Fixed mismatched copy constructor/assignment operator in FAnimationBaseContext and FAnimationUpdateContext.

diff --git a/Engine/Source/Runtime/Engine/Classes/Animation/AnimNodeBase.h b/Engine/Source/Runtime/Engine/Classes/Animation/AnimNodeBase.h
--- a/Engine/Source/Runtime/Engine/Classes/Animation/AnimNodeBase.h
+++ b/Engine/Source/Runtime/Engine/Classes/Animation/AnimNodeBase.h
@@ -174,8 +174,12 @@ protected:
 	ENGINE_API FAnimationBaseContext(FAnimInstanceProxy* InAnimInstanceProxy, FAnimationUpdateSharedContext* InSharedContext = nullptr);
 
 public:
-	// we define a copy constructor here simply to avoid deprecation warnings with clang
-	ENGINE_API FAnimationBaseContext(const FAnimationBaseContext& InContext);
+	PRAGMA_DISABLE_DEPRECATION_WARNINGS
+	FAnimationBaseContext(FAnimationBaseContext&&) = default;
+	FAnimationBaseContext(const FAnimationBaseContext&) = default;
+	FAnimationBaseContext& operator=(FAnimationBaseContext&&) = default;
+	FAnimationBaseContext& operator=(const FAnimationBaseContext&) = default;
+	PRAGMA_ENABLE_DEPRECATION_WARNINGS
 
 public:
 	// Get the Blueprint IAnimClassInterface associated with this context, if there is one.
@@ -345,8 +349,6 @@ public:
 	}
 
 
-	FAnimationUpdateContext(const FAnimationUpdateContext& Copy) = default;
-
 	FAnimationUpdateContext(const FAnimationUpdateContext& Copy, FAnimInstanceProxy* InAnimInstanceProxy)
 		: FAnimationBaseContext(InAnimInstanceProxy, Copy.SharedContext)
 		, CurrentWeight(Copy.CurrentWeight)
