Missed file from CL# 25180858.

diff --git a/Engine/Source/Runtime/Engine/Private/Animation/AnimNodeBase.cpp b/Engine/Source/Runtime/Engine/Private/Animation/AnimNodeBase.cpp
--- a/Engine/Source/Runtime/Engine/Private/Animation/AnimNodeBase.cpp
+++ b/Engine/Source/Runtime/Engine/Private/Animation/AnimNodeBase.cpp
@@ -76,14 +76,6 @@ FAnimationBaseContext::FAnimationBaseContext(FAnimInstanceProxy* InAnimInstanceP
 {
 }
 
-FAnimationBaseContext::FAnimationBaseContext(const FAnimationBaseContext& InContext)
-	: AnimInstanceProxy(InContext.AnimInstanceProxy)
-	, SharedContext(InContext.SharedContext)
-	, CurrentNodeId(InContext.CurrentNodeId)
-	, PreviousNodeId(InContext.PreviousNodeId)
-{
-}
-
 IAnimClassInterface* FAnimationBaseContext::GetAnimClass() const
 {
 	return AnimInstanceProxy ? AnimInstanceProxy->GetAnimClassInterface() : nullptr;
