Fixed mismatched copy constructor/assignment operator in FAnimNotifyEventReference.

diff --git a/Engine/Source/Runtime/Engine/Public/Animation/AnimNotifyQueue.h b/Engine/Source/Runtime/Engine/Public/Animation/AnimNotifyQueue.h
--- a/Engine/Source/Runtime/Engine/Public/Animation/AnimNotifyQueue.h
+++ b/Engine/Source/Runtime/Engine/Public/Animation/AnimNotifyQueue.h
@@ -23,8 +23,6 @@ struct FAnimNotifyEventReference
 
 	FAnimNotifyEventReference() = default;
 
-	FAnimNotifyEventReference(const FAnimNotifyEventReference& rhs) = default;
-
 	FAnimNotifyEventReference(const FAnimNotifyEvent* InNotify, const UObject* InNotifySource)
 		: Notify(InNotify)
 		, MirrorTable(nullptr)
