Fixed mismatched copy constructor/assignment operator in FAnimCompressContext.

diff --git a/Engine/Source/Runtime/Engine/Classes/Animation/AnimCompress.h b/Engine/Source/Runtime/Engine/Classes/Animation/AnimCompress.h
--- a/Engine/Source/Runtime/Engine/Classes/Animation/AnimCompress.h
+++ b/Engine/Source/Runtime/Engine/Classes/Animation/AnimCompress.h
@@ -254,6 +254,9 @@ public:
 		, bOutput(Rhs.bOutput)
 	{}
 
+	// Unlike the copy constructor, this will copy the CompressionSummary, but the class is deprecated anyway
+	FAnimCompressContext& operator=(const FAnimCompressContext&) = default;
+
 	friend class FAnimationUtils;
 	friend class FDerivedDataAnimationCompression;
 	friend class UAnimSequence;
