Temp fix for errors on level save for soft references to HLOD packages which does not exist.
#rnx

diff --git a/Engine/Source/Runtime/Engine/Private/HLOD/HLODEngineSubsystem.cpp b/Engine/Source/Runtime/Engine/Private/HLOD/HLODEngineSubsystem.cpp
--- a/Engine/Source/Runtime/Engine/Private/HLOD/HLODEngineSubsystem.cpp
+++ b/Engine/Source/Runtime/Engine/Private/HLOD/HLODEngineSubsystem.cpp
@@ -197,6 +197,8 @@ void UHLODEngineSubsystem::OnPreSaveWorld(UWorld* InWorld, FObjectPreSaveContext
 
 void UHLODEngineSubsystem::OnWorldCollectSaveReferences(UWorld* World, FArchive& Ar)
 {
+	// Disabled for now as this introduce validation errors
+#if 0
 	// Record a dependency on the paths to HLODProxy packages that would be associated with
 	// this level if they exist
 	FHierarchicalLODUtilitiesModule& Module = FModuleManager::LoadModuleChecked<FHierarchicalLODUtilitiesModule>("HierarchicalLODUtilities");
@@ -219,6 +221,7 @@ void UHLODEngineSubsystem::OnWorldCollectSaveReferences(UWorld* World, FArchive&
 			}
 		}
 	}
+#endif
 }
 
 #endif // WITH_EDITOR
