Added preemptive ConditionalPostLoad calls until zenloader load ordering improvements can be rolled out

diff --git a/Engine/Source/Runtime/Engine/Private/Components/StaticMeshComponent.cpp b/Engine/Source/Runtime/Engine/Private/Components/StaticMeshComponent.cpp
--- a/Engine/Source/Runtime/Engine/Private/Components/StaticMeshComponent.cpp
+++ b/Engine/Source/Runtime/Engine/Private/Components/StaticMeshComponent.cpp
@@ -2667,11 +2667,18 @@ UMaterialInterface* UStaticMeshComponent::GetMaterial(int32 MaterialIndex, bool
 		OutMaterial = GetStaticMesh()->GetMaterial(MaterialIndex);
 	}
 
-	// If we have a nanite override, use that
-	if (OutMaterial != nullptr && UseNaniteOverrideMaterials(bDoingNaniteMaterialAudit))
+	if (OutMaterial)
 	{
-		UMaterialInterface* NaniteOverride = OutMaterial->GetNaniteOverride();
-		OutMaterial = NaniteOverride != nullptr ? NaniteOverride : OutMaterial;
+		//@note FH: temporary preemptive PostLoad until zenloader load ordering improvements
+		OutMaterial->ConditionalPostLoad();
+
+		// If we have a nanite override, use that
+		if (UseNaniteOverrideMaterials(bDoingNaniteMaterialAudit))
+		{
+			UMaterialInterface* NaniteOverride = OutMaterial->GetNaniteOverride();
+			OutMaterial = NaniteOverride != nullptr ? NaniteOverride : OutMaterial;
+		}
+
 	}
 
 	return OutMaterial;
diff --git a/Engine/Source/Runtime/Engine/Private/Materials/Material.cpp b/Engine/Source/Runtime/Engine/Private/Materials/Material.cpp
--- a/Engine/Source/Runtime/Engine/Private/Materials/Material.cpp
+++ b/Engine/Source/Runtime/Engine/Private/Materials/Material.cpp
@@ -2037,6 +2037,9 @@ void UMaterial::UpdateCachedExpressionData()
 {
 	COOK_STAT(FScopedDurationTimer BlockingTimer(MaterialCookStats::MaterialUpdateCachedExpressionDataSec));
 
+	//@note FH: temporary preemptive PostLoad until zenloader load ordering improvements
+	ConditionalPostLoad();
+
 	if (bLoadedCachedExpressionData)
 	{
 		// Don't need to rebuild cached data if it was serialized
