LevelInstance: Allow discards even if some actors were saved in the edit session

diff --git a/Engine/Source/Runtime/Engine/Private/LevelInstance/LevelInstanceEditorObject.cpp b/Engine/Source/Runtime/Engine/Private/LevelInstance/LevelInstanceEditorObject.cpp
--- a/Engine/Source/Runtime/Engine/Private/LevelInstance/LevelInstanceEditorObject.cpp
+++ b/Engine/Source/Runtime/Engine/Private/LevelInstance/LevelInstanceEditorObject.cpp
@@ -30,16 +30,6 @@ ULevelInstanceEditorObject::ULevelInstanceEditorObject(const FObjectInitializer&
 
 bool ULevelInstanceEditorObject::CanDiscard(FText* OutReason /*= nullptr*/) const
 {
-	if (bCommittedChanges)
-	{
-		if (OutReason)
-		{
-			*OutReason = LOCTEXT("CantDiscardCommittedChanges", "Can't discard edit because some packages were saved while editing.");
-		}
-
-		return false;
-	}
-
 	if (bMovedActors)
 	{
 		if(OutReason)
