Fix for crash in Dense Volume -> SparseVolumeTexture conversion if the page table needed to be padded to fit the power of two requirement.

#rnx

diff --git a/Engine/Source/Runtime/Engine/Private/SparseVolumeTexture/SparseVolumeTextureData.cpp b/Engine/Source/Runtime/Engine/Private/SparseVolumeTexture/SparseVolumeTextureData.cpp
--- a/Engine/Source/Runtime/Engine/Private/SparseVolumeTexture/SparseVolumeTextureData.cpp
+++ b/Engine/Source/Runtime/Engine/Private/SparseVolumeTexture/SparseVolumeTextureData.cpp
@@ -242,7 +242,10 @@ bool FSparseVolumeTextureData::CreateFromDense(const FSparseVolumeTextureDataCre
 	{
 		// Make sure we only process voxels inside a single page in the Y and Z dimension so we don't end up reading/writing the same entries in the page table from multiple threads.
 		const FIntVector3 PageBegin = FIntVector3(0, JobIndex % Header.PageTableVolumeResolution.Y, JobIndex / Header.PageTableVolumeResolution.Y) + Header.PageTableVolumeAABBMin;
-		const FIntVector3 PageEnd = FIntVector3(Header.PageTableVolumeAABBMax.X, FMath::Min(PageBegin.Y + 1, Header.PageTableVolumeAABBMax.Y), FMath::Min(PageBegin.Z + 1, Header.PageTableVolumeAABBMax.Z));
+		// PageTableVolumeAABBMax is padded to a power of 2 which may be larger than the tightest possible page table for the source volume.
+		// We only want to iterate over pages overlapping the source volume, so we recompute a tight fit here.
+		const int32 PageEndX = (Header.VirtualVolumeAABBMax.X + (SPARSE_VOLUME_TILE_RES - 1)) / SPARSE_VOLUME_TILE_RES;
+		const FIntVector3 PageEnd = FIntVector3(PageEndX, FMath::Min(PageBegin.Y + 1, Header.PageTableVolumeAABBMax.Y), FMath::Min(PageBegin.Z + 1, Header.PageTableVolumeAABBMax.Z));
 		const FIntVector3 PageVoxelBegin = PageBegin * SPARSE_VOLUME_TILE_RES;
 		const FIntVector3 PageVoxelEnd = PageEnd * SPARSE_VOLUME_TILE_RES;
 		const FIntVector3 ActiveVoxelBegin = PageVoxelBegin + FractionalPageOffset;
