Relink function call inputs before generating a syntax tree. Otherwise, we may get invalid inputs and they will cause errors when generating the tree

diff --git a/Engine/Source/Runtime/Engine/Private/Materials/Material.cpp b/Engine/Source/Runtime/Engine/Private/Materials/Material.cpp
--- a/Engine/Source/Runtime/Engine/Private/Materials/Material.cpp
+++ b/Engine/Source/Runtime/Engine/Private/Materials/Material.cpp
@@ -2050,6 +2050,9 @@ void UMaterial::UpdateCachedExpressionData()
 	FMaterialCachedHLSLTree* LocalCachedTree = nullptr;
 	if (IsUsingNewHLSLGenerator())
 	{
+		// Relinks function call inputs. Otherwise, we can get invalid inputs and they will cause errors when generating the syntax tree
+		UpdateTransientExpressionData();
+
 		LocalCachedTree = new FMaterialCachedHLSLTree();
 		LocalCachedTree->GenerateTree(this, nullptr, nullptr);
 		LocalCachedExpressionData->UpdateForCachedHLSLTree(*LocalCachedTree, nullptr);
