World partition : Child actors are now returning the same value as their parent actor for AActor::IsForceExternalActorLevelReferenceForPIE.

diff --git a/Engine/Source/Runtime/Engine/Classes/GameFramework/Actor.h b/Engine/Source/Runtime/Engine/Classes/GameFramework/Actor.h
--- a/Engine/Source/Runtime/Engine/Classes/GameFramework/Actor.h
+++ b/Engine/Source/Runtime/Engine/Classes/GameFramework/Actor.h
@@ -2346,10 +2346,7 @@ public:
 	virtual bool SupportsLayers() const;
 
 	/** Returns if level should keep a reference to the external actor for PIE (used for always loaded actors). */
-	bool IsForceExternalActorLevelReferenceForPIE() const
-	{
-		return bForceExternalActorLevelReferenceForPIE;
-	}
+	bool IsForceExternalActorLevelReferenceForPIE() const;
 
 	void SetForceExternalActorLevelReferenceForPIE(bool bValue)
 	{
diff --git a/Engine/Source/Runtime/Engine/Private/ActorEditor.cpp b/Engine/Source/Runtime/Engine/Private/ActorEditor.cpp
--- a/Engine/Source/Runtime/Engine/Private/ActorEditor.cpp
+++ b/Engine/Source/Runtime/Engine/Private/ActorEditor.cpp
@@ -989,6 +989,16 @@ bool AActor::SupportsLayers() const
 	return bIsValid;
 }
 
+bool AActor::IsForceExternalActorLevelReferenceForPIE() const
+{
+	if (const AActor* ParentActor = GetParentActor())
+	{
+		return ParentActor->IsForceExternalActorLevelReferenceForPIE();
+	}
+
+	return bForceExternalActorLevelReferenceForPIE;
+}
+
 bool AActor::IsEditable() const
 {
 	return bEditable;
diff --git a/Engine/Source/Runtime/Engine/Private/WorldPartition/WorldPartitionRuntimeHash.cpp b/Engine/Source/Runtime/Engine/Private/WorldPartition/WorldPartitionRuntimeHash.cpp
--- a/Engine/Source/Runtime/Engine/Private/WorldPartition/WorldPartitionRuntimeHash.cpp
+++ b/Engine/Source/Runtime/Engine/Private/WorldPartition/WorldPartitionRuntimeHash.cpp
@@ -153,15 +153,6 @@ bool UWorldPartitionRuntimeHash::ConditionalRegisterAlwaysLoadedActorsForPIE(con
 			if (AActor* AlwaysLoadedActor = FindObject<AActor>(nullptr, *ActorDescView.GetActorSoftPath().ToString()))
 			{
 				AlwaysLoadedActorsForPIE.Emplace(Reference, AlwaysLoadedActor);
-
-				// Handle child actors
-				AlwaysLoadedActor->ForEachComponent<UChildActorComponent>(true, [this, &Reference](UChildActorComponent* ChildActorComponent)
-				{
-					if (AActor* ChildActor = ChildActorComponent->GetChildActor())
-					{
-						AlwaysLoadedActorsForPIE.Emplace(Reference, ChildActor);
-					}
-				});
 			}
 		}
 
