Fixed mismatched copy constructor/assignment operator in FPropertyAccessLibrary.

diff --git a/Engine/Source/Runtime/Engine/Private/PropertyAccess.cpp b/Engine/Source/Runtime/Engine/Private/PropertyAccess.cpp
--- a/Engine/Source/Runtime/Engine/Private/PropertyAccess.cpp
+++ b/Engine/Source/Runtime/Engine/Private/PropertyAccess.cpp
@@ -675,6 +675,9 @@ struct FPropertyAccessSystem
 	}	
 };
 
+// Unlike the copy assignment operator, the copy constructor doesn't take a lock
+FPropertyAccessLibrary::FPropertyAccessLibrary(const FPropertyAccessLibrary& Other) = default;
+
 const FPropertyAccessLibrary& FPropertyAccessLibrary::operator =(const FPropertyAccessLibrary& Other)
 {
 	// Need to perform a lock as copying can race with PatchPropertyOffsets in async loading thread-enabled builds
diff --git a/Engine/Source/Runtime/Engine/Public/PropertyAccess.h b/Engine/Source/Runtime/Engine/Public/PropertyAccess.h
--- a/Engine/Source/Runtime/Engine/Public/PropertyAccess.h
+++ b/Engine/Source/Runtime/Engine/Public/PropertyAccess.h
@@ -401,6 +401,7 @@ struct FPropertyAccessLibrary
 
 	FPropertyAccessLibrary() = default;
 
+	ENGINE_API FPropertyAccessLibrary(const FPropertyAccessLibrary& Other);
 	ENGINE_API const FPropertyAccessLibrary& operator =(const FPropertyAccessLibrary& Other);
 
 private:
