Fixed mismatched copy constructor/assignment operator in FCancelCompressionSignal.

diff --git a/Engine/Source/Runtime/Engine/Public/Animation/AnimCompressionTypes.h b/Engine/Source/Runtime/Engine/Public/Animation/AnimCompressionTypes.h
--- a/Engine/Source/Runtime/Engine/Public/Animation/AnimCompressionTypes.h
+++ b/Engine/Source/Runtime/Engine/Public/Animation/AnimCompressionTypes.h
@@ -179,6 +179,9 @@ private:
 	FThreadSafeBool Signal;
 
 public:
+	FCancelCompressionSignal() = default;
+	FCancelCompressionSignal(const FCancelCompressionSignal&) = default;
+
 	FCancelCompressionSignal& operator=(const FCancelCompressionSignal& Other)
 	{
 		Signal = (bool)Other.Signal;
