Added cook stats to track critical path, game thread time to submit material compile jobs.

- On Lyra this accounts for ~10.4% of a full cook with `-noshaderddc` and a warm cache for everything else.
- On CitySample this accounts for ~8.9% of a full cook with `-noshaderddc` and a warm cache for everything else.

diff --git a/Engine/Source/Runtime/Engine/Private/Materials/MaterialShader.cpp b/Engine/Source/Runtime/Engine/Private/Materials/MaterialShader.cpp
--- a/Engine/Source/Runtime/Engine/Private/Materials/MaterialShader.cpp
+++ b/Engine/Source/Runtime/Engine/Private/Materials/MaterialShader.cpp
@@ -57,11 +57,15 @@ namespace MaterialShaderCookStats
 {
 	static FCookStats::FDDCResourceUsageStats UsageStats;
 	static int32 MaterialShadersCompiled = 0;
+	static double MaterialShaderMapSubmitCompileJobsTime = 0.0;
+	static int32 MaterialShaderMapSubmitCompileCalls = 0;
 	static FCookStatsManager::FAutoRegisterCallback RegisterCookStats([](FCookStatsManager::AddStatFuncRef AddStat)
 	{
 		UsageStats.LogStats(AddStat, TEXT("MaterialShader.Usage"), TEXT(""));
 		AddStat(TEXT("Material"), FCookStatsManager::CreateKeyValueArray(
-			TEXT("MaterialShadersCompiled"), MaterialShadersCompiled
+			TEXT("MaterialShadersCompiled"), MaterialShadersCompiled,
+			TEXT("MaterialShaderMapSubmitCompileJobsTime"), MaterialShaderMapSubmitCompileJobsTime,
+			TEXT("MaterialShaderMapSubmitCompileCalls"), MaterialShaderMapSubmitCompileCalls
 			));
 	});
 }
@@ -1828,6 +1832,9 @@ int32 FMaterialShaderMap::SubmitCompileJobs(uint32 CompilingShaderMapId,
 	EShaderCompileJobPriority InPriority) const
 {
 	TRACE_CPUPROFILER_EVENT_SCOPE(FMaterialShaderMap::SubmitCompileJobs);
+	COOK_STAT(MaterialShaderCookStats::MaterialShaderMapSubmitCompileCalls++);
+	COOK_STAT(FScopedDurationTimer DurationTimer(MaterialShaderCookStats::MaterialShaderMapSubmitCompileJobsTime));
+
 	check(CompilingShaderMapId != 0u);
 	check(MaterialEnvironment);
 
