Fixed mismatched copy constructor/assignment operator in BoneTrackPair.

diff --git a/Engine/Source/Runtime/Engine/Public/AnimEncoding.h b/Engine/Source/Runtime/Engine/Public/AnimEncoding.h
--- a/Engine/Source/Runtime/Engine/Public/AnimEncoding.h
+++ b/Engine/Source/Runtime/Engine/Public/AnimEncoding.h
@@ -39,13 +39,6 @@ struct BoneTrackPair
 	int32 AtomIndex;
 	int32 TrackIndex;
 
-	BoneTrackPair& operator=(const BoneTrackPair& Other)
-	{
-		this->AtomIndex = Other.AtomIndex;
-		this->TrackIndex = Other.TrackIndex;
-		return *this;
-	}
-
 	BoneTrackPair(){}
 	BoneTrackPair(int32 Atom, int32 Track):AtomIndex(Atom),TrackIndex(Track){}
 };
