Optimize required bones calculation by sharing the bone container

diff --git a/Engine/Source/Runtime/Engine/Classes/Components/SkeletalMeshComponent.h b/Engine/Source/Runtime/Engine/Classes/Components/SkeletalMeshComponent.h
--- a/Engine/Source/Runtime/Engine/Classes/Components/SkeletalMeshComponent.h
+++ b/Engine/Source/Runtime/Engine/Classes/Components/SkeletalMeshComponent.h
@@ -337,6 +337,9 @@ public:
 	UFUNCTION(BlueprintPure, Category = "Components|SkeletalMesh")
 	USkeletalMesh* GetSkeletalMeshAsset() const;
 
+	/** Gets the shared bone container used between all owned anim instances. Creates it on the first call. */
+	TSharedPtr<struct FBoneContainer> GetSharedRequiredBones();
+
 #if WITH_EDITORONLY_DATA
 	/** The blueprint for creating an AnimationScript. */
 	UPROPERTY()
@@ -421,6 +424,9 @@ private:
 	UPROPERTY(transient)
 	TArray<TObjectPtr<UAnimInstance>> LinkedInstances;
 
+	/** Shared bone container between all anim instances owned by this skeletal mesh component */
+	TSharedPtr<struct FBoneContainer> SharedRequiredBones;
+
 	// Update Rate
 
 	/** Cached BoneSpaceTransforms for Update Rate optimization. */
diff --git a/Engine/Source/Runtime/Engine/Private/Animation/AnimInstanceProxy.cpp b/Engine/Source/Runtime/Engine/Private/Animation/AnimInstanceProxy.cpp
--- a/Engine/Source/Runtime/Engine/Private/Animation/AnimInstanceProxy.cpp
+++ b/Engine/Source/Runtime/Engine/Private/Animation/AnimInstanceProxy.cpp
@@ -65,6 +65,7 @@ FAnimInstanceProxy::FAnimInstanceProxy()
 	, RootMotionMode(ERootMotionMode::NoRootMotionExtraction)
 	, FrameCounterForUpdate(0)
 	, FrameCounterForNodeUpdate(0)
+	, RequiredBones(MakeShared<FBoneContainer>())	// We sometime query for this before RecalcRequiredBones has been called
 	, CacheBonesRecursionCounter(0)
 	, bUpdatingRoot(false)
 	, bBoneCachesInvalidated(false)
@@ -92,6 +93,7 @@ FAnimInstanceProxy::FAnimInstanceProxy(UAnimInstance* Instance)
 	, RootMotionMode(ERootMotionMode::NoRootMotionExtraction)
 	, FrameCounterForUpdate(0)
 	, FrameCounterForNodeUpdate(0)
+	, RequiredBones(MakeShared<FBoneContainer>())	// We sometime query for this before RecalcRequiredBones has been called
 	, CacheBonesRecursionCounter(0)
 	, bUpdatingRoot(false)
 	, bBoneCachesInvalidated(false)
@@ -1116,22 +1118,29 @@ void FAnimInstanceProxy::RecalcRequiredBones(USkeletalMeshComponent* Component,
 {
 	DECLARE_SCOPE_HIERARCHICAL_COUNTER_FUNC()
 
-	RequiredBones.InitializeTo(Component->RequiredBones, Component->GetCurveFilterSettings(), *Asset);
+	// Use the shared bone container
+	RequiredBones = Component->GetSharedRequiredBones();
 
-	// If there is a ref pose override, we want to replace ref pose in RequiredBones
-	// Update ref pose in required bones structure (either set it, or clear it, depending on if one is set on the Component)
-	RequiredBones.SetRefPoseOverride(Component->GetRefPoseOverride());
+	// The first anim instance will initialize the required bones, all others will re-use it
+	if (!RequiredBones->IsValid())
+	{
+		RequiredBones->InitializeTo(Component->RequiredBones, Component->GetCurveFilterSettings(), *Asset);
+
+		// If there is a ref pose override, we want to replace ref pose in RequiredBones
+		// Update ref pose in required bones structure (either set it, or clear it, depending on if one is set on the Component)
+		RequiredBones->SetRefPoseOverride(Component->GetRefPoseOverride());
+	}
 
 	// If this instance can accept input poses, initialise the input pose container
-	if(DefaultLinkedInstanceInputNode)
+	if (DefaultLinkedInstanceInputNode)
 	{
-		DefaultLinkedInstanceInputNode->CachedInputPose.SetBoneContainer(&RequiredBones);
+		DefaultLinkedInstanceInputNode->CachedInputPose.SetBoneContainer(RequiredBones.Get());
 		
 		// SetBoneContainer allocates space for bone data but leaves it uninitalized.
 		DefaultLinkedInstanceInputNode->bIsCachedInputPoseInitialized = false;
 	}
 
-	// When RequiredBones mapping has changed, AnimNodes need to update their bones caches. 
+	// When RequiredBones mapping has changed, AnimNodes need to update their bones caches.
 	bBoneCachesInvalidated = true;
 }
 
@@ -1139,7 +1148,7 @@ void FAnimInstanceProxy::RecalcRequiredCurves(const UE::Anim::FCurveFilterSettin
 {
 	DECLARE_SCOPE_HIERARCHICAL_COUNTER_FUNC()
 
-	RequiredBones.CacheRequiredAnimCurves(CurveFilterSettings);
+	RequiredBones->CacheRequiredAnimCurves(CurveFilterSettings);
 	bBoneCachesInvalidated = true;
 }
 
@@ -1148,7 +1157,7 @@ void FAnimInstanceProxy::RecalcRequiredCurves(const FCurveEvaluationOption& Curv
 	DECLARE_SCOPE_HIERARCHICAL_COUNTER_FUNC()
 
 	PRAGMA_DISABLE_DEPRECATION_WARNINGS
-	RequiredBones.CacheRequiredAnimCurveUids(CurveEvalOption);
+	RequiredBones->CacheRequiredAnimCurveUids(CurveEvalOption);
 	PRAGMA_ENABLE_DEPRECATION_WARNINGS
 	bBoneCachesInvalidated = true;
 }
@@ -1480,7 +1489,7 @@ void FAnimInstanceProxy::SlotEvaluatePoseWithBlendProfiles(const FName& SlotNode
 	TArray<float>& PerBoneWeightTotals = BlendProfileScratchData.PerBoneWeightTotals;
 	TArray<float>& PerBoneWeightTotalsAdditive = BlendProfileScratchData.PerBoneWeightTotalsAdditive;
 	TArray<float>& BoneBlendProfileScales = BlendProfileScratchData.BoneBlendProfileScales;
-	const int32 NumBones = RequiredBones.GetCompactPoseNumBones();
+	const int32 NumBones = RequiredBones->GetCompactPoseNumBones();
 	PerBoneWeightTotals.Reset(NumBones);
 	PerBoneWeightTotals.AddUninitialized(NumBones);
 	for (int32 BoneIndex = 0; BoneIndex < NumBones; ++BoneIndex)
@@ -1531,8 +1540,8 @@ void FAnimInstanceProxy::SlotEvaluatePoseWithBlendProfiles(const FName& SlotNode
 
 			// Bone array has to be allocated prior to calling GetPoseFromAnimTrack.
 			FSlotEvaluationPose NewPose(EvalState.BlendInfo.GetBlendedValue(), AdditiveAnimType);
-			NewPose.Pose.SetBoneContainer(&RequiredBones);
-			NewPose.Curve.InitFrom(RequiredBones);
+			NewPose.Pose.SetBoneContainer(RequiredBones.Get());
+			NewPose.Curve.InitFrom(*RequiredBones);
 
 			// Extract pose from Track.
 			FAnimExtractContext ExtractionContext(static_cast<double>(EvalState.MontagePosition), Montage->HasRootMotion() && RootMotionMode != ERootMotionMode::NoRootMotionExtraction, EvalState.DeltaTimeRecord);
@@ -1541,7 +1550,7 @@ void FAnimInstanceProxy::SlotEvaluatePoseWithBlendProfiles(const FName& SlotNode
 
 			// Add montage curves.
 			FBlendedCurve MontageCurve;
-			MontageCurve.InitFrom(RequiredBones);
+			MontageCurve.InitFrom(*RequiredBones);
 			Montage->EvaluateCurveData(MontageCurve, EvalState.MontagePosition);
 			NewPose.Curve.Combine(MontageCurve);
 
@@ -1565,7 +1574,7 @@ void FAnimInstanceProxy::SlotEvaluatePoseWithBlendProfiles(const FName& SlotNode
 			const float PoseWeight = EvalState.BlendInfo.GetBlendedValue();
 			if (BlendProfile)
 			{
-				BlendProfile->FillBoneScalesArray(BoneBlendProfileScales, RequiredBones);
+				BlendProfile->FillBoneScalesArray(BoneBlendProfileScales, *RequiredBones);
 				for (int32 BoneIndex = 0; BoneIndex < NumBones; ++BoneIndex)
 				{
 					PerBoneWeights[CurrentPoseIndex][BoneIndex] = BlendProfile->CalculateBoneWeight(BoneBlendProfileScales[BoneIndex], BlendProfile->Mode, EvalState.BlendInfo, EvalState.BlendStartAlpha, NewPose.Weight, false);
@@ -1870,8 +1879,8 @@ void FAnimInstanceProxy::SlotEvaluatePose(const FName& SlotNodeName, const FAnim
 			FSlotEvaluationPose NewPose(MontageWeight, AdditiveAnimType);
 			
 			// Bone array has to be allocated prior to calling GetPoseFromAnimTrack
-			NewPose.Pose.SetBoneContainer(&RequiredBones);
-			NewPose.Curve.InitFrom(RequiredBones);
+			NewPose.Pose.SetBoneContainer(RequiredBones.Get());
+			NewPose.Curve.InitFrom(*RequiredBones);
 
 			// Extract pose from Track
 			FAnimExtractContext ExtractionContext(static_cast<double>(EvalState.MontagePosition), Montage->HasRootMotion() && RootMotionMode != ERootMotionMode::NoRootMotionExtraction, EvalState.DeltaTimeRecord);
@@ -1881,7 +1890,7 @@ void FAnimInstanceProxy::SlotEvaluatePose(const FName& SlotNodeName, const FAnim
 
 			// add montage curves 
 			FBlendedCurve MontageCurve;
-			MontageCurve.InitFrom(RequiredBones);
+			MontageCurve.InitFrom(*RequiredBones);
 			Montage->EvaluateCurveData(MontageCurve, EvalState.MontagePosition);
 			NewPose.Curve.Combine(MontageCurve);
 
@@ -3436,7 +3445,7 @@ void FAnimInstanceProxy::UpdateCurvesToEvaluationContext(const FAnimationEvaluat
 		AnimationCurves[(uint8)EAnimCurveType::AttributeCurve].Add(InCurveElement.Name, InCurveElement.Value);
 	});
 
-	UE::Anim::FNamedValueArrayUtils::Intersection(InContext.Curve, RequiredBones.GetCurveFlags(), 
+	UE::Anim::FNamedValueArrayUtils::Intersection(InContext.Curve, RequiredBones->GetCurveFlags(), 
 		[this](const UE::Anim::FCurveElement& InCurveElement, const UE::Anim::FCurveElementFlags& InCurveFlagsElement)
 		{
 			if(EnumHasAnyFlags(InCurveFlagsElement.Flags | InCurveElement.Flags, UE::Anim::ECurveElementFlags::MorphTarget))
diff --git a/Engine/Source/Runtime/Engine/Private/BoneContainer.cpp b/Engine/Source/Runtime/Engine/Private/BoneContainer.cpp
--- a/Engine/Source/Runtime/Engine/Private/BoneContainer.cpp
+++ b/Engine/Source/Runtime/Engine/Private/BoneContainer.cpp
@@ -89,6 +89,37 @@ void FBoneContainer::InitializeTo(const TArray<FBoneIndexType>& InRequiredBoneIn
 	Initialize(CurveFilterSettings);
 }
 
+void FBoneContainer::Reset()
+{
+	Asset = nullptr;
+	AssetSkeletalMesh = nullptr;
+	AssetSkeleton = nullptr;
+	RefSkeleton = nullptr;
+	UIDToArrayIndexLUTValidCount = 0;
+	RefPoseOverride = nullptr;
+
+#if DO_CHECK
+	CalculatedForLOD = INDEX_NONE;
+#endif
+
+	bDisableRetargeting = false;
+	bUseRAWData = false;
+	bUseSourceData = false;
+
+	BoneIndicesArray.Empty(0);
+	BoneSwitchArray.Empty(0);
+	SkeletonToPoseBoneIndexArray.Empty(0);
+	PoseToSkeletonBoneIndexArray.Empty(0);
+	CompactPoseToSkeletonIndex.Empty(0);
+	SkeletonToCompactPose.Empty(0);
+	CompactPoseParentBones.Empty(0);
+	VirtualBoneCompactPoseData.Empty(0);
+	UIDToArrayIndexLUT.Empty(0);
+
+	// Container changed and is no longer valid, preserve but update our serial number
+	RegenerateSerialNumber();
+}
+
 struct FBoneContainerScratchArea : public TThreadSingleton<FBoneContainerScratchArea>
 {
 	TArray<int32> MeshIndexToCompactPoseIndex;
diff --git a/Engine/Source/Runtime/Engine/Private/Components/SkeletalMeshComponent.cpp b/Engine/Source/Runtime/Engine/Private/Components/SkeletalMeshComponent.cpp
--- a/Engine/Source/Runtime/Engine/Private/Components/SkeletalMeshComponent.cpp
+++ b/Engine/Source/Runtime/Engine/Private/Components/SkeletalMeshComponent.cpp
@@ -292,6 +292,16 @@ USkeletalMesh* USkeletalMeshComponent::GetSkeletalMeshAsset() const
 	return Cast<USkeletalMesh>(GetSkinnedAsset());
 }
 
+TSharedPtr<FBoneContainer> USkeletalMeshComponent::GetSharedRequiredBones()
+{
+	if (!SharedRequiredBones)
+	{
+		SharedRequiredBones = MakeShared<FBoneContainer>();
+	}
+
+	return SharedRequiredBones;
+}
+
 void USkeletalMeshComponent::Serialize(FArchive& Ar)
 {
 	PRAGMA_DISABLE_DEPRECATION_WARNINGS
@@ -1988,9 +1998,18 @@ void USkeletalMeshComponent::RecalcRequiredBones(int32 LODIndex)
 	}
 
 	ComputeRequiredBones(RequiredBones, FillComponentSpaceTransformsRequiredBones, LODIndex, /*bIgnorePhysicsAsset=*/ false);
+
+	// Reset our animated pose to the reference pose
 	PRAGMA_DISABLE_DEPRECATION_WARNINGS
 	BoneSpaceTransforms = GetSkeletalMeshAsset()->GetRefSkeleton().GetRefBonePose();
 	PRAGMA_ENABLE_DEPRECATION_WARNINGS
+
+	// If we had cached our shared bone container, reset it
+	if (SharedRequiredBones)
+	{
+		SharedRequiredBones->Reset();
+	}
+
 	// make sure animation requiredBone to mark as dirty
 	if (AnimScriptInstance)
 	{
@@ -2007,7 +2026,7 @@ void USkeletalMeshComponent::RecalcRequiredBones(int32 LODIndex)
 		PostProcessAnimInstance->RecalcRequiredBones();
 	}
 
-	// when recalc requiredbones happend
+	// when RecalcRequiredBones happened
 	// this should always happen
 	MarkRequiredCurveUpToDate();
 	bRequiredBonesUpToDate = true;
diff --git a/Engine/Source/Runtime/Engine/Public/Animation/AnimInstanceProxy.h b/Engine/Source/Runtime/Engine/Public/Animation/AnimInstanceProxy.h
--- a/Engine/Source/Runtime/Engine/Public/Animation/AnimInstanceProxy.h
+++ b/Engine/Source/Runtime/Engine/Public/Animation/AnimInstanceProxy.h
@@ -333,13 +333,13 @@ public:
 	/** const access to required bones array */
 	const FBoneContainer& GetRequiredBones() const 
 	{ 
-		return RequiredBones; 
+		return *RequiredBones; 
 	}
 
 	/** access to required bones array */
 	FBoneContainer& GetRequiredBones() 
 	{ 
-		return RequiredBones;
+		return *RequiredBones;
 	}
 
 	/** Access to LODLevel */
@@ -747,6 +747,7 @@ protected:
 	/** 
 	 * Recalculate Required Bones [RequiredBones]
 	 * Is called when bRequiredBonesUpToDate = false
+	 * The provided asset must be a SkeletalMesh or a Skeleton.
 	 */
 	void RecalcRequiredBones(USkeletalMeshComponent* Component, UObject* Asset);
 
@@ -1084,8 +1085,11 @@ private:
 	// Root motion extracted from animation since the last time ConsumeExtractedRootMotion was called
 	FRootMotionMovementParams ExtractedRootMotion;
 
-	/** Temporary array of bone indices required this frame. Should be subset of Skeleton and Mesh's RequiredBones */
-	FBoneContainer RequiredBones;
+	/**
+	  * Temporary array of bone indices required this frame. Should be subset of Skeleton and Mesh's RequiredBones.
+	  * Shared between all anim instances on a skeletal mesh component if it has a skeletal mesh and a skeleton present.
+	  */
+	TSharedPtr<FBoneContainer> RequiredBones;
 
 	/** LODLevel used by RequiredBones */
 	int32 LODLevel = 0;
diff --git a/Engine/Source/Runtime/Engine/Public/BoneContainer.h b/Engine/Source/Runtime/Engine/Public/BoneContainer.h
--- a/Engine/Source/Runtime/Engine/Public/BoneContainer.h
+++ b/Engine/Source/Runtime/Engine/Public/BoneContainer.h
@@ -244,6 +244,9 @@ public:
 	UE_DEPRECATED(5.3, "Please use InitializeTo that takes a FCurveFilterSettings.")
 	void InitializeTo(const TArray<FBoneIndexType>& InRequiredBoneIndexArray, const FCurveEvaluationOption& CurveEvalOption, UObject& InAsset);
 
+	/** Resets the container and reclaims all allocated memory but preserve the serial number. */
+	void Reset();
+
 	/** Returns true if FBoneContainer is Valid. Needs an Asset, a RefPoseArray, and a RequiredBonesArray. */
 	const bool IsValid() const
 	{
