Refactored world partition's debug 2D view to show all instanced partitioned worlds by default.
If there's only 1 world partition or if wp.Runtime.DrawWorldPartitionIndex is set to a valid index, then it will only display a detailed 2D debug view of the specified partition world.
Setting wp.Runtime.DrawWorldPartitionIndex to a negative value will default back to showing all world partitions.

diff --git a/Engine/Source/Runtime/Engine/Private/WorldPartition/WorldPartition.cpp b/Engine/Source/Runtime/Engine/Private/WorldPartition/WorldPartition.cpp
--- a/Engine/Source/Runtime/Engine/Private/WorldPartition/WorldPartition.cpp
+++ b/Engine/Source/Runtime/Engine/Private/WorldPartition/WorldPartition.cpp
@@ -1551,35 +1551,13 @@ void UWorldPartition::OnCellHidden(const UWorldPartitionRuntimeCell* InCell)
 	}
 }
 
-bool UWorldPartition::CanDebugDraw() const
+bool UWorldPartition::DrawRuntimeHash2D(FWorldPartitionDraw2DContext& DrawContext)
 {
-#if WITH_EDITOR
-	const ENetMode NetMode = GetWorld()->GetNetMode();
-	const bool bIsDedicatedServer = (NetMode == NM_DedicatedServer);
-	const bool bIsClient = (NetMode == NM_Client);
-	if ((bIsDedicatedServer && !UWorldPartition::DebugDedicatedServerStreaming) ||
-		(bIsClient && UWorldPartition::DebugDedicatedServerStreaming))
-	{
-		return false;
-	}
-
-	if (!bIsDedicatedServer && UWorldPartition::IsSimulating(false) && (!GEditor || (GetWorld() != GEditor->PlayWorld)))
-	{
-		return false;
-	}
-#endif
-	return GetWorld()->IsGameWorld();
-}
-
-bool UWorldPartition::DrawRuntimeHash2D(class UCanvas* Canvas, const FVector2D& PartitionCanvasSize, const FVector2D& Offset, FVector2D& OutUsedCanvasSize)
-{
-	check(CanDebugDraw());
-	return StreamingPolicy->DrawRuntimeHash2D(Canvas, PartitionCanvasSize, Offset, OutUsedCanvasSize);
+	return StreamingPolicy->DrawRuntimeHash2D(DrawContext);
 }
 
 void UWorldPartition::DrawRuntimeHash3D()
 {
-	check(CanDebugDraw());
 	StreamingPolicy->DrawRuntimeHash3D();
 }
 
@@ -1588,11 +1566,6 @@ void UWorldPartition::DrawRuntimeCellsDetails(class UCanvas* Canvas, FVector2D&
 	StreamingPolicy->DrawRuntimeCellsDetails(Canvas, Offset);
 }
 
-void UWorldPartition::DrawStreamingStatusLegend(class UCanvas* Canvas, FVector2D& Offset)
-{
-	StreamingPolicy->DrawStreamingStatusLegend(Canvas, Offset);
-}
-
 EWorldPartitionStreamingPerformance UWorldPartition::GetStreamingPerformance() const
 {
 	return StreamingPolicy->GetStreamingPerformance();
diff --git a/Engine/Source/Runtime/Engine/Private/WorldPartition/WorldPartitionDraw2DContext.cpp b/Engine/Source/Runtime/Engine/Private/WorldPartition/WorldPartitionDraw2DContext.cpp
new file mode 100644
--- /dev/null
+++ b/Engine/Source/Runtime/Engine/Private/WorldPartition/WorldPartitionDraw2DContext.cpp
@@ -0,0 +1,129 @@
+// Copyright Epic Games, Inc. All Rights Reserved.
+
+#include "WorldPartition/WorldPartitionDraw2DContext.h"
+#include "WorldPartition/WorldPartitionDebugHelper.h"
+#include "Math/UnrealMathUtility.h"
+#include "GlobalRenderResources.h"
+#include "BatchedElements.h"
+#include "Engine/Canvas.h"
+#include "Engine/Engine.h"
+
+FWorldPartitionDraw2DCanvas::FWorldPartitionDraw2DCanvas(UCanvas* InCanvas)
+: Canvas(InCanvas)
+{
+	check(Canvas);
+	FCanvas* CanvasObject = Canvas->Canvas;
+	CanvasLineBatchedElements = CanvasObject->GetBatchedElements(FCanvas::ET_Line);
+	CanvasTriangleBatchedElements = CanvasObject->GetBatchedElements(FCanvas::ET_Triangle, nullptr, GWhiteTexture, SE_BLEND_Translucent);
+}
+
+void FWorldPartitionDraw2DCanvas::PrepareDraw(int32 InLineCount, int32 InBoxCount)
+{
+	check(InLineCount >= 0);
+	check(InBoxCount >= 0);
+
+	CanvasLineBatchedElements->AddReserveLines(InLineCount);
+	CanvasLineBatchedElements->AddReserveLines(InLineCount, false, true);
+
+	CanvasTriangleBatchedElements->ReserveVertices(4 * InBoxCount);
+	CanvasTriangleBatchedElements->AddReserveTriangles(InBoxCount * 2, GWhiteTexture, SE_BLEND_Translucent);
+}
+
+void FWorldPartitionDraw2DCanvas::Draw(const FWorldPartitionCanvasItems& Items)
+{
+	TRACE_CPUPROFILER_EVENT_SCOPE(FWorldPartitionDraw2DCanvas::Draw);
+
+	const FHitProxyId CanvasHitProxyId = Canvas->Canvas->GetHitProxyId();
+
+	for (const FWorldPartitionCanvasBoxItem& Box : Items.Boxes)
+	{
+		int32 V0 = CanvasTriangleBatchedElements->AddVertexf(FVector4f(Box.Points[0].X, Box.Points[0].Y, 0, 1), FVector2f::ZeroVector, Box.Color, CanvasHitProxyId);
+		int32 V1 = CanvasTriangleBatchedElements->AddVertexf(FVector4f(Box.Points[1].X, Box.Points[1].Y, 0, 1), FVector2f::ZeroVector, Box.Color, CanvasHitProxyId);
+		int32 V2 = CanvasTriangleBatchedElements->AddVertexf(FVector4f(Box.Points[2].X, Box.Points[2].Y, 0, 1), FVector2f::ZeroVector, Box.Color, CanvasHitProxyId);
+		int32 V3 = CanvasTriangleBatchedElements->AddVertexf(FVector4f(Box.Points[3].X, Box.Points[3].Y, 0, 1), FVector2f::ZeroVector, Box.Color, CanvasHitProxyId);
+		CanvasTriangleBatchedElements->AddTriangle(V0, V1, V2, GWhiteTexture, SE_BLEND_Translucent);
+		CanvasTriangleBatchedElements->AddTriangle(V2, V3, V0, GWhiteTexture, SE_BLEND_Translucent);
+	}
+
+	for (const FWorldPartitionLineCanvasItem& Line : Items.Lines)
+	{
+		CanvasLineBatchedElements->AddLine(Line.Start, Line.End, Line.Color, CanvasHitProxyId, Line.Thickness);
+	}
+
+	for (const FWorldPartitionCanvasMultiLineTextItem& Item : Items.MultiLineTexts)
+	{
+		FVector2D Position = Item.Position;
+		for (const FWorldPartitionCanvasText& Text : Item.MultiLineText)
+		{
+			FWorldPartitionDebugHelper::DrawText(Canvas, Text.Key, GEngine->GetSmallFont(), Text.Value.ToFColor(true), Position);
+		}
+	}
+}
+
+static bool SegmentIntersection2D(const FVector2D& InSegmentStartA, const FVector2D& InSegmentEndA, const FVector2D& InSegmentStartB, const FVector2D& InSegmentEndB, FVector2D& OutIntersection)
+{
+	FVector Intersection;
+	if (FMath::SegmentIntersection2D(FVector(InSegmentStartA, 0), FVector(InSegmentEndA, 0), FVector(InSegmentStartB, 0), FVector(InSegmentEndB, 0), Intersection))
+	{
+		OutIntersection = FVector2D(Intersection.X, Intersection.Y);
+		return true;
+	}
+	return false;
+}
+
+bool FWorldPartitionDraw2DContext::BoxSegmentIntersect(const FBox2D& InBox, const FVector2D& InSegmentStart, const FVector2D& InSegmentEnd, TArray<FVector2D>& OutIntersections)
+{
+	if (InBox.IsInside(InSegmentStart))
+	{
+		OutIntersections.Add(InSegmentStart);
+	}
+	if (InBox.IsInside(InSegmentEnd))
+	{
+		OutIntersections.Add(InSegmentEnd);
+	}
+	if (OutIntersections.Num() == 2)
+	{
+		return true;
+	}
+
+	FVector2D BoxSize = InBox.GetSize();
+	FVector2D BoxA = InBox.Min;
+	FVector2D BoxB = InBox.Min + FVector2D(BoxSize.X, 0);
+	FVector2D BoxC = InBox.Min + BoxSize;
+	FVector2D BoxD = InBox.Min + FVector2D(0, BoxSize.Y);
+
+	FVector2D Intersection;
+	if (SegmentIntersection2D(InSegmentStart, InSegmentEnd, BoxA, BoxB, Intersection))
+	{
+		OutIntersections.Add(Intersection);
+		if (OutIntersections.Num() == 2)
+		{
+			return true;
+		}
+	}
+	if (SegmentIntersection2D(InSegmentStart, InSegmentEnd, BoxB, BoxC, Intersection))
+	{
+		OutIntersections.Add(Intersection);
+		if (OutIntersections.Num() == 2)
+		{
+			return true;
+		}
+	}
+	if (SegmentIntersection2D(InSegmentStart, InSegmentEnd, BoxC, BoxD, Intersection))
+	{
+		OutIntersections.Add(Intersection);
+		if (OutIntersections.Num() == 2)
+		{
+			return true;
+		}
+	}
+	if (SegmentIntersection2D(InSegmentStart, InSegmentEnd, BoxA, BoxD, Intersection))
+	{
+		OutIntersections.Add(Intersection);
+		if (OutIntersections.Num() == 2)
+		{
+			return true;
+		}
+	}
+	return false;
+}
\ No newline at end of file
diff --git a/Engine/Source/Runtime/Engine/Private/WorldPartition/WorldPartitionDraw2DContext.h b/Engine/Source/Runtime/Engine/Private/WorldPartition/WorldPartitionDraw2DContext.h
new file mode 100644
--- /dev/null
+++ b/Engine/Source/Runtime/Engine/Private/WorldPartition/WorldPartitionDraw2DContext.h
@@ -0,0 +1,212 @@
+// Copyright Epic Games, Inc. All Rights Reserved.
+#pragma once
+
+#include "CoreMinimal.h"
+
+class UCanvas;
+class FBatchedElements;
+class FWorldPartitionCanvasItems;
+
+using FWorldPartitionCanvasText = TPair<FString, FLinearColor>;
+using FWorldPartitionCanvasMultiLineText = TArray<FWorldPartitionCanvasText>;
+
+class FWorldPartitionCanvasMultiLineTextItem
+{
+public:
+	FWorldPartitionCanvasMultiLineTextItem(const FVector2D& InPosition, const FWorldPartitionCanvasMultiLineText& InMultiLineText)
+		: Position(InPosition)
+		, MultiLineText(InMultiLineText)
+	{
+	}
+
+	FVector2D Position;
+	FWorldPartitionCanvasMultiLineText MultiLineText;
+};
+
+class FWorldPartitionLineCanvasItem
+{
+public:
+	FWorldPartitionLineCanvasItem(const FVector2D& InStart, const FVector2D& InEnd, const FLinearColor& InColor = FLinearColor::White, float InThickness = 1)
+		: Color(InColor)
+		, Start(InStart, 0)
+		, End(InEnd, 0)
+		, Thickness(InThickness)
+	{
+	}
+	FLinearColor Color;
+	FVector Start;
+	FVector End;
+	float Thickness;
+};
+
+class FWorldPartitionCanvasBoxItem
+{
+public:
+	FWorldPartitionCanvasBoxItem(const FVector2D& InPointA, const FVector2D& InPointB, const FVector2D& InPointC, const FVector2D& InPointD, const FLinearColor& InColor = FLinearColor::White)
+		: Color(InColor)
+	{
+		Points[0] = InPointA;
+		Points[1] = InPointB;
+		Points[2] = InPointC;
+		Points[3] = InPointD;
+	};
+
+	FLinearColor Color;
+	FVector2D Points[4];
+};
+
+class FWorldPartitionDraw2DCanvas
+{
+public:
+	FWorldPartitionDraw2DCanvas(UCanvas* InCanvas);
+	void PrepareDraw(int32 InLineCount, int32 InBoxCount);
+	void Draw(const FWorldPartitionCanvasItems& Items);
+
+private:
+	UCanvas* Canvas;
+	FBatchedElements* CanvasLineBatchedElements;
+	FBatchedElements* CanvasTriangleBatchedElements;
+};
+
+class FWorldPartitionCanvasItems
+{
+public:
+	FWorldPartitionCanvasItems()
+	{
+		Lines.Reserve(1024);
+		Boxes.Reserve(1024);
+		MultiLineTexts.Reserve(128);
+	}
+
+	void AddLine(FWorldPartitionLineCanvasItem& InLine)
+	{
+		Lines.Add(MoveTemp(InLine));
+	}
+
+	void AddBox(FWorldPartitionCanvasBoxItem& InBox)
+	{
+		Boxes.Add(MoveTemp(InBox));
+	}
+
+	void AddText(FWorldPartitionCanvasMultiLineTextItem& InMultiLineText)
+	{
+		MultiLineTexts.Add(MoveTemp(InMultiLineText));
+	}
+
+	void Reset()
+	{
+		Lines.Reset();
+		Boxes.Reset();
+		MultiLineTexts.Reset();
+	}
+
+	int32 GetLineCount() const { return Lines.Num(); }
+	int32 GetBoxCount() const { return Boxes.Num(); }
+
+	void Draw(FWorldPartitionDraw2DCanvas& InCanvas);
+
+private:
+	TArray<FWorldPartitionLineCanvasItem> Lines;
+	TArray<FWorldPartitionCanvasBoxItem> Boxes;
+	TArray<FWorldPartitionCanvasMultiLineTextItem> MultiLineTexts;
+
+	friend class FWorldPartitionDraw2DCanvas;
+};
+
+class FWorldPartitionDraw2DContext
+{
+public:
+	FWorldPartitionDraw2DContext()
+	{
+		Initialize(FBox2D(ForceInit), FBox2D(ForceInit));
+	}
+
+	void Initialize(const FBox2D& InCanvasRegion, const FBox2D& InWorldRegion)
+	{
+		CanvasRegion = InCanvasRegion;
+		WorldRegion = InWorldRegion;
+		bDrawGridBounds = false;
+		bDrawGridAxis = false;
+		bIsDetailedMode = false;
+		UsedCanvasBounds = FBox2D(ForceInit);
+		DesiredWorldBounds = FBox2D(ForceInit);
+		CanvasItems.Reset();
+	}
+
+	void SetIsDetailedMode(bool bValue) { bIsDetailedMode = bValue; }
+	bool IsDetailedMode() const { return bIsDetailedMode; }
+
+	void SetDrawGridBounds(bool bValue) { bDrawGridBounds = bValue; }
+	bool GetDrawGridBounds() const { return bDrawGridBounds; }
+
+	void SetDrawGridAxis(bool bValue) { bDrawGridAxis = bValue; }
+	bool GetDrawGridAxis() const { return bDrawGridAxis; }
+
+	void SetUsedCanvasBounds(const FBox2D& InUsedCanvasBounds) { UsedCanvasBounds = InUsedCanvasBounds; }
+	const FBox2D& GetUsedCanvasBounds() const { return UsedCanvasBounds; }
+
+	void SetDesiredWorldBounds(const FBox2D& InDesiredWorldBounds) { DesiredWorldBounds = InDesiredWorldBounds; }
+	const FBox2D& GetDesiredWorldBounds() const { return DesiredWorldBounds; }
+
+	const FBox2D& GetCanvasRegion() const { check(CanvasRegion.bIsValid); return CanvasRegion; }
+	const FBox2D& GetWorldRegion() const { check(WorldRegion.bIsValid); return WorldRegion; }
+	int32 GetLineCount() const { return CanvasItems.GetLineCount(); }
+	int32 GetBoxCount() const { return CanvasItems.GetBoxCount(); }
+
+	bool PushDrawSegment(const FBox2D& InBounds, const FVector2D& A, const FVector2D& B, const FLinearColor& InColor, float InLineThickness)
+	{
+		TArray<FVector2D> Intersections;
+		if (BoxSegmentIntersect(InBounds, A, B, Intersections))
+		{
+			FWorldPartitionLineCanvasItem Line(Intersections[0], Intersections[1], InColor, InLineThickness);
+			CanvasItems.AddLine(Line);
+			return true;
+		}
+		return false;
+	};
+
+	bool PushDrawBox(const FBox2D& InBounds, const FVector2D& A, const FVector2D& B, const FVector2D& C, const FVector2D& D, const FLinearColor& InColor, float InLineThickness)
+	{
+		PushDrawSegment(InBounds, A, B, InColor, InLineThickness);
+		PushDrawSegment(InBounds, B, C, InColor, InLineThickness);
+		PushDrawSegment(InBounds, C, D, InColor, InLineThickness);
+		PushDrawSegment(InBounds, A, D, InColor, InLineThickness);
+		return true;
+	};
+
+	bool PushDrawTile(const FBox2D& InBounds, const FVector2D& A, const FVector2D& B, const FVector2D& C, const FVector2D& D, const FLinearColor& InColor)
+	{
+		if (!InBounds.IsInside(A) ||
+			!InBounds.IsInside(B) ||
+			!InBounds.IsInside(C) ||
+			!InBounds.IsInside(D))
+		{
+			return false;
+		}
+
+		FWorldPartitionCanvasBoxItem Box(A, B, C, D, InColor);
+		CanvasItems.AddBox(Box);
+		return true;
+	};
+
+	bool PushDrawText(FWorldPartitionCanvasMultiLineTextItem& InMultiLineText)
+	{
+		CanvasItems.AddText(InMultiLineText);
+		return true;
+	}
+
+	const FWorldPartitionCanvasItems& GetCanvasItems() const { return CanvasItems; }
+
+private:
+
+	bool BoxSegmentIntersect(const FBox2D& Box, const FVector2D& a1, const FVector2D& a2, TArray<FVector2D>& Intersections);
+
+	bool bDrawGridBounds;
+	bool bDrawGridAxis;
+	bool bIsDetailedMode;
+	FBox2D CanvasRegion;
+	FBox2D WorldRegion;
+	FBox2D UsedCanvasBounds;
+	FBox2D DesiredWorldBounds;
+	FWorldPartitionCanvasItems CanvasItems;
+};
\ No newline at end of file
diff --git a/Engine/Source/Runtime/Engine/Private/WorldPartition/WorldPartitionLevelStreamingPolicy.cpp b/Engine/Source/Runtime/Engine/Private/WorldPartition/WorldPartitionLevelStreamingPolicy.cpp
--- a/Engine/Source/Runtime/Engine/Private/WorldPartition/WorldPartitionLevelStreamingPolicy.cpp
+++ b/Engine/Source/Runtime/Engine/Private/WorldPartition/WorldPartitionLevelStreamingPolicy.cpp
@@ -314,45 +314,4 @@ void UWorldPartitionLevelStreamingPolicy::DrawRuntimeCellsDetails(UCanvas* Canva
 	}
 
 	Offset.Y = MaxPosY;
-}
-
-/**
- * Debug Draw Streaming Status Legend
- */
-void UWorldPartitionLevelStreamingPolicy::DrawStreamingStatusLegend(UCanvas* Canvas, FVector2D& Offset)
-{
-	TRACE_CPUPROFILER_EVENT_SCOPE(UWorldPartitionLevelStreamingPolicy::DrawStreamingStatusLegend);
-
-	check(Canvas);
-
-	// Cumulate counter stats
-	int32 StatusCount[(int32)LEVEL_StreamingStatusCount] = { 0 };
-	ForEachActiveRuntimeCell([&StatusCount](const UWorldPartitionRuntimeCell* Cell)
-	{
-		StatusCount[(int32)Cell->GetStreamingStatus()]++;
-	});
-
-	// @todo_ow: This is not exactly the good value, as we could have pending unload level from Level Instances, etc.
-	//           We could modify GetNumLevelsPendingPurge to return the number of pending purge levels from the grid, 
-	//           bu that will do for now.
-	StatusCount[LEVEL_UnloadedButStillAround] = FLevelStreamingGCHelper::GetNumLevelsPendingPurge();
-
-	// Draw legend
-	FVector2D Pos = Offset;
-	float MaxTextWidth = 0.f;
-	FWorldPartitionDebugHelper::DrawText(Canvas, TEXT("Streaming Status Legend"), GEngine->GetSmallFont(), FColor::Yellow, Pos, &MaxTextWidth);
-	
-	for (int32 i = 0; i < (int32)LEVEL_StreamingStatusCount; ++i)
-	{
-		EStreamingStatus Status = (EStreamingStatus)i;
-		const FColor& StatusColor = ULevelStreaming::GetLevelStreamingStatusColor(Status);
-		FString DebugString = *FString::Printf(TEXT("%d) %s"), i, ULevelStreaming::GetLevelStreamingStatusDisplayName(Status));
-		if (Status != LEVEL_Unloaded)
-		{
-			DebugString += *FString::Printf(TEXT(" (%d)"), StatusCount[(int32)Status]);
-		}
-		FWorldPartitionDebugHelper::DrawLegendItem(Canvas, *DebugString, GEngine->GetSmallFont(), StatusColor, FColor::White, Pos, &MaxTextWidth);
-	}
-
-	Offset.X += MaxTextWidth + 10;
-}
+}
\ No newline at end of file
diff --git a/Engine/Source/Runtime/Engine/Private/WorldPartition/WorldPartitionRuntimeSpatialHash.cpp b/Engine/Source/Runtime/Engine/Private/WorldPartition/WorldPartitionRuntimeSpatialHash.cpp
--- a/Engine/Source/Runtime/Engine/Private/WorldPartition/WorldPartitionRuntimeSpatialHash.cpp
+++ b/Engine/Source/Runtime/Engine/Private/WorldPartition/WorldPartitionRuntimeSpatialHash.cpp
@@ -13,6 +13,7 @@
 #include "WorldPartition/WorldPartitionLog.h"
 #include "WorldPartition/WorldPartitionStreamingGeneration.h"
 #include "WorldPartition/WorldPartitionSubsystem.h"
+#include "WorldPartition/WorldPartitionDraw2DContext.h"
 #include "WorldPartition/RuntimeSpatialHash/RuntimeSpatialHashGridHelper.h"
 #include "WorldPartition/ContentBundle/ContentBundleWorldSubsystem.h"
 #include "WorldPartition/ContentBundle/ContentBundleDescriptor.h"
@@ -27,6 +28,8 @@
 #include "DrawDebugHelpers.h"
 #include "GlobalRenderResources.h"
 #include "Math/TranslationMatrix.h"
+#include "Math/UnrealMathUtility.h"
+#include "Math/TransformCalculus2D.h"
 #include "Misc/ArchiveMD5.h"
 #include "Misc/HierarchicalLogArchive.h"
 #include "UObject/UnrealType.h"
@@ -741,26 +744,52 @@ void FSpatialHashStreamingGrid::DrawStreamingSource3D(UWorld* World, const FSphe
 	}
 }
 
-void FSpatialHashStreamingGrid::Draw2D(UCanvas* Canvas, const UWorldPartitionRuntimeSpatialHash* Owner, const TArray<FWorldPartitionStreamingSource>& Sources, const FBox& Region, const FBox2D& GridScreenBounds, TFunctionRef<FVector2D(const FVector2D&)> WorldToScreen) const
+void FSpatialHashStreamingGrid::Draw2D(const UWorldPartitionRuntimeSpatialHash* Owner, const FBox2D& Region2D, const FBox2D& GridScreenBounds, TFunctionRef<FVector2D(const FVector2D&, bool)> InWorldToScreen, FWorldPartitionDraw2DContext& DrawContext) const
 {
 	TRACE_CPUPROFILER_EVENT_SCOPE(FSpatialHashStreamingGrid::Draw2D);
 
 	const UWorldPartition* WorldPartition = Owner->GetOuterUWorldPartition();
-	UWorld* OwningWorld = WorldPartition->GetWorld();
+	const TArray<FWorldPartitionStreamingSource>& Sources = WorldPartition->GetStreamingSources();
+
+	const FBox Region(FVector(Region2D.Min.X, Region2D.Min.Y, 0), FVector(Region2D.Max.X, Region2D.Max.Y, 0));
+	const UWorld* OwningWorld = WorldPartition->GetWorld();
 	const EWorldPartitionRuntimeCellVisualizeMode VisualizeMode = GetStreamingCellVisualizeMode();
 	const UContentBundleManager* ContentBundleManager = OwningWorld->ContentBundleManager;
 	TMap<FName, FColor> DataLayerDebugColors = GetDataLayerDebugColors(WorldPartition);
 
-	FCanvas* CanvasObject = Canvas->Canvas;
 	int32 MinGridLevel = FMath::Clamp<int32>(GShowRuntimeSpatialHashGridLevel, 0, GridLevels.Num() - 1);
 	int32 MaxGridLevel = FMath::Clamp<int32>(MinGridLevel + GShowRuntimeSpatialHashGridLevelCount - 1, 0, GridLevels.Num() - 1);
 
-	// Precompute a cell coordinate text with/height using a generic coordinate
-	// This will be used later to filter out drawing of cell coordinates (avoid doing expesive calls to Canvas->StrLen)
-	const FString CellCoordString = UWorldPartitionRuntimeSpatialHash::GetCellCoordString(FGridCellCoord(88,88,88));
-	float MaxCellCoordTextWidth, MaxCellCoordTextHeight;
-	Canvas->StrLen(GEngine->GetTinyFont(), CellCoordString, MaxCellCoordTextWidth, MaxCellCoordTextHeight);
-	
+	auto WorldToScreen = [&](const FVector2D& Pos, bool bIsLocal = true)
+	{
+		return InWorldToScreen(Pos, bIsLocal);
+	};
+
+	auto LocalDrawTile = [&](const FVector2D& Min, const FVector2D& Size, const FLinearColor& Color)
+	{
+		FVector2D A = WorldToScreen(Min);
+		FVector2D B = WorldToScreen(Min + FVector2D(Size.X, 0));
+		FVector2D C = WorldToScreen(Min + Size);
+		FVector2D D = WorldToScreen(Min + FVector2D(0, Size.Y));
+		DrawContext.PushDrawTile(GridScreenBounds, A, B, C, D, Color);
+	};
+
+	auto LocalDrawBox = [&](const FVector2D& Min, const FVector2D& Size, const FLinearColor& Color, float LineThickness, const FBox2D* CustomGridScreenBounds = nullptr)
+	{
+		FVector2D A = WorldToScreen(Min);
+		FVector2D B = WorldToScreen(Min + FVector2D(Size.X, 0));
+		FVector2D C = WorldToScreen(Min + Size);
+		FVector2D D = WorldToScreen(Min + FVector2D(0, Size.Y));
+		DrawContext.PushDrawBox(GridScreenBounds, A, B, C, D, Color, LineThickness);
+	};
+
+	auto LocalDrawSegment = [&](const FVector2D& Start, const FVector2D& End, const FLinearColor& Color, float LineThickness)
+	{
+		FVector2D A = WorldToScreen(Start);
+		FVector2D B = WorldToScreen(End);
+		DrawContext.PushDrawSegment(GridScreenBounds, A, B, Color, LineThickness);
+	};
+
 	TArray<const UWorldPartitionRuntimeCell*> FilteredCells;
 	for (int32 GridLevel = MinGridLevel; GridLevel <= MaxGridLevel; ++GridLevel)
 	{
@@ -771,31 +800,7 @@ void FSpatialHashStreamingGrid::Draw2D(UCanvas* Canvas, const UWorldPartitionRun
 			FBox2D CellWorldBounds;
 			FGridCellCoord Coords(Coords2D.X, Coords2D.Y, GridLevel);
 			Helper.Levels[GridLevel].GetCellBounds(Coords2D, CellWorldBounds);
-			FBox2D CellScreenBounds = FBox2D(WorldToScreen(CellWorldBounds.Min), WorldToScreen(CellWorldBounds.Max));
-			// Clamp inside grid bounds
-			if (!GridScreenBounds.IsInside(CellScreenBounds))
-			{
-				CellScreenBounds.Min.X = FMath::Clamp(CellScreenBounds.Min.X, GridScreenBounds.Min.X, GridScreenBounds.Max.X);
-				CellScreenBounds.Min.Y = FMath::Clamp(CellScreenBounds.Min.Y, GridScreenBounds.Min.Y, GridScreenBounds.Max.Y);
-				CellScreenBounds.Max.X = FMath::Clamp(CellScreenBounds.Max.X, GridScreenBounds.Min.X, GridScreenBounds.Max.X);
-				CellScreenBounds.Max.Y = FMath::Clamp(CellScreenBounds.Max.Y, GridScreenBounds.Min.Y, GridScreenBounds.Max.Y);
-			}
-			else
-			{
-				FGridCellCoord CellGlobalCoords;
-				verify(Helper.GetCellGlobalCoords(Coords, CellGlobalCoords));
-				const FString CellCoordString = UWorldPartitionRuntimeSpatialHash::GetCellCoordString(CellGlobalCoords);
-				const FVector2D CellBoundsSize = CellScreenBounds.GetSize();
-				if (MaxCellCoordTextWidth < CellBoundsSize.X && MaxCellCoordTextHeight < CellBoundsSize.Y)
-				{
-					float CellCoordTextWidth, CellCoordTextHeight;
-					Canvas->StrLen(GEngine->GetTinyFont(), CellCoordString, CellCoordTextWidth, CellCoordTextHeight);
-					FVector2D GridInfoPos = CellScreenBounds.GetCenter() - FVector2D(CellCoordTextWidth / 2, CellCoordTextHeight / 2);
-					Canvas->SetDrawColor(255, 255, 0);
-					Canvas->DrawText(GEngine->GetTinyFont(), CellCoordString, GridInfoPos.X, GridInfoPos.Y);
-				}
-			}
-
+			
 			FilteredCells.Reset();
 			ForEachLayerCell(Coords, [&](const FSpatialHashStreamingGridLayerCell* LayerCell)
 			{
@@ -804,121 +809,93 @@ void FSpatialHashStreamingGrid::Draw2D(UCanvas* Canvas, const UWorldPartitionRun
 
 			if (FilteredCells.Num())
 			{
-				FVector2D CellBoundsSize = CellScreenBounds.GetSize();
+				FVector2D CellBoundsSize = CellWorldBounds.GetSize();
 				CellBoundsSize.Y /= FilteredCells.Num();
 				FVector2D CellOffset(0.f, 0.f);
 				for (const UWorldPartitionRuntimeCell* Cell : FilteredCells)
 				{
 					// Draw Cell using its debug color
-					FVector2D StartPos = CellScreenBounds.Min + CellOffset;
-					FCanvasTileItem Item(StartPos, GWhiteTexture, CellBoundsSize, Cell->GetDebugColor(VisualizeMode));
-					Item.BlendMode = SE_BLEND_Translucent;
-					Canvas->DrawItem(Item);
+					FVector2D StartPos = CellWorldBounds.Min + CellOffset;
+					{
+						LocalDrawTile(StartPos, CellBoundsSize, Cell->GetDebugColor(VisualizeMode));
+					}
+
 					CellOffset.Y += CellBoundsSize.Y;
 
-					// Draw Cell's DataLayers colored boxes
-					if (DataLayerDebugColors.Num() && Cell->GetDataLayers().Num() > 0)
+					if (DrawContext.IsDetailedMode())
 					{
-						FVector2D DataLayerOffset(0, 0);
-						FVector2D DataLayerColoredBoxSize = CellBoundsSize;
-						DataLayerColoredBoxSize.X /= (5 * Cell->GetDataLayers().Num()); // Use 20% of cell's width
-						for (const FName& DataLayer : Cell->GetDataLayers())
+						// Draw Cell's DataLayers colored boxes
+						if (DataLayerDebugColors.Num() && Cell->GetDataLayers().Num() > 0)
 						{
-							const FColor& DataLayerColor = DataLayerDebugColors[DataLayer];
-							FCanvasTileItem DataLayerItem(StartPos + DataLayerOffset, GWhiteTexture, DataLayerColoredBoxSize, DataLayerColor);
-							Canvas->DrawItem(DataLayerItem);
-							DataLayerOffset.X += DataLayerColoredBoxSize.X;
+							FVector2D DataLayerOffset(0, 0);
+							FVector2D DataLayerColoredBoxSize = CellBoundsSize;
+							DataLayerColoredBoxSize.X /= (5 * Cell->GetDataLayers().Num()); // Use 20% of cell's width
+							for (const FName& DataLayer : Cell->GetDataLayers())
+							{
+								const FColor& DataLayerColor = DataLayerDebugColors[DataLayer];
+								LocalDrawTile(StartPos + DataLayerOffset, DataLayerColoredBoxSize, DataLayerColor);
+								DataLayerOffset.X += DataLayerColoredBoxSize.X;
+							}
 						}
-					}
 
-					// Draw Content Bundle colored boxes
-					if (ContentBundleManager && FWorldPartitionDebugHelper::CanDrawContentBundles() && Cell->GetContentBundleID().IsValid())
-					{
-						if (const FContentBundleBase* ContentBundle = ContentBundleManager->GetContentBundle(OwningWorld, Cell->GetContentBundleID()))
+						// Draw Content Bundle colored boxes
+						if (ContentBundleManager && FWorldPartitionDebugHelper::CanDrawContentBundles() && Cell->GetContentBundleID().IsValid())
 						{
-							check(ContentBundle->GetDescriptor());
-							FVector2D BoxSize = CellBoundsSize;
-							BoxSize.X /= 5; // Use 20% of cell's width
-							FVector2D ContentBundleOffset(CellBoundsSize.X - BoxSize.X, 0);
-							FCanvasTileItem ContentBundleItem(StartPos + ContentBundleOffset, GWhiteTexture, BoxSize, ContentBundle->GetDescriptor()->GetDebugColor());
-							Canvas->DrawItem(ContentBundleItem);
+							if (const FContentBundleBase* ContentBundle = ContentBundleManager->GetContentBundle(OwningWorld, Cell->GetContentBundleID()))
+							{
+								check(ContentBundle->GetDescriptor());
+								FVector2D BoxSize = CellBoundsSize;
+								BoxSize.X /= 5; // Use 20% of cell's width
+								FVector2D ContentBundleOffset(CellBoundsSize.X - BoxSize.X, 0);
+								LocalDrawTile(StartPos + ContentBundleOffset, BoxSize, ContentBundle->GetDescriptor()->GetDebugColor());
+							}
 						}
 					}
 				}
-			}
-
-			// Draw cell bounds
-			FCanvasBoxItem Box(CellScreenBounds.Min, CellScreenBounds.GetSize());
-			Box.SetColor(FLinearColor::Black);
-			Box.BlendMode = SE_BLEND_Translucent;
-			Canvas->DrawItem(Box);
-		});
 
-		// Draw X/Y Axis
-		{
-			FCanvasLineItem Axis;
-			Axis.LineThickness = 3;
-			{
-				Axis.SetColor(FLinearColor::Red);
-				FVector2D LineStart = WorldToScreen(FVector2D(-1638400.f, 0.f));
-				FVector2D LineEnd = WorldToScreen(FVector2D(1638400.f, 0.f));
-				LineStart.X = FMath::Clamp(LineStart.X, GridScreenBounds.Min.X, GridScreenBounds.Max.X);
-				LineStart.Y = FMath::Clamp(LineStart.Y, GridScreenBounds.Min.Y, GridScreenBounds.Max.Y);
-				LineEnd.X = FMath::Clamp(LineEnd.X, GridScreenBounds.Min.X, GridScreenBounds.Max.X);
-				LineEnd.Y = FMath::Clamp(LineEnd.Y, GridScreenBounds.Min.Y, GridScreenBounds.Max.Y);
-				Axis.Draw(CanvasObject, LineStart, LineEnd);
+				// Draw cell bounds
+				LocalDrawBox(CellWorldBounds.Min, CellWorldBounds.GetSize(), FLinearColor::Black, 1);
 			}
-			{
-				Axis.SetColor(FLinearColor::Green);
-				FVector2D LineStart = WorldToScreen(FVector2D(0.f, -1638400.f));
-				FVector2D LineEnd = WorldToScreen(FVector2D(0.f, 1638400.f));
-				LineStart.X = FMath::Clamp(LineStart.X, GridScreenBounds.Min.X, GridScreenBounds.Max.X);
-				LineStart.Y = FMath::Clamp(LineStart.Y, GridScreenBounds.Min.Y, GridScreenBounds.Max.Y);
-				LineEnd.X = FMath::Clamp(LineEnd.X, GridScreenBounds.Min.X, GridScreenBounds.Max.X);
-				LineEnd.Y = FMath::Clamp(LineEnd.Y, GridScreenBounds.Min.Y, GridScreenBounds.Max.Y);
-				Axis.Draw(CanvasObject, LineStart, LineEnd);
-			}
-		}
+		});
+	}
 
-		// Draw Streaming Sources
-		const float GridLoadingRange = GetLoadingRange();
-		for (const FWorldPartitionStreamingSource& Source : Sources)
-		{
-			const FColor Color = Source.GetDebugColor();
-			Source.ForEachShape(GridLoadingRange, GridName, HLODLayer, /*bProjectIn2D*/ true, [&Color, &Canvas, &WorldToScreen, this](const FSphericalSector& Shape)
-			{
-				DrawStreamingSource2D(Canvas, Shape, WorldToScreen, Color);
-			});
-		}
+	// Draw X/Y Axis
+	if (DrawContext.GetDrawGridAxis())
+	{
+		DrawContext.PushDrawSegment(GridScreenBounds, WorldToScreen(FVector2D(-1638400.f, 0.f), false), WorldToScreen(FVector2D(1638400.f, 0.f), false), FLinearColor::Red, 3);
+		DrawContext.PushDrawSegment(GridScreenBounds, WorldToScreen(FVector2D(0.f, -1638400.f), false), WorldToScreen(FVector2D(0.f, 1638400.f), false), FLinearColor::Green, 3);
+	}
 
-		FCanvasBoxItem Box(GridScreenBounds.Min, GridScreenBounds.GetSize());
-		Box.SetColor(DebugColor);
-		Canvas->DrawItem(Box);
+	// Draw Grid Bounds
+	if (DrawContext.GetDrawGridBounds())
+	{
+		FBox2D Bounds = GridScreenBounds.ExpandBy(FVector2D(10));
+		FVector2D Size = GridScreenBounds.GetSize();
+		DrawContext.PushDrawBox(Bounds, GridScreenBounds.Min, GridScreenBounds.Min + FVector2D(Size.X, 0), GridScreenBounds.Max, GridScreenBounds.Min + FVector2D(0, Size.Y), DebugColor, 1);
 	}
 
 	// Draw WorldBounds
-	FBox2D WorldScreenBounds = FBox2D(WorldToScreen(FVector2D(WorldBounds.Min)), WorldToScreen(FVector2D(WorldBounds.Max)));
-	if (!GridScreenBounds.IsInside(WorldScreenBounds))
+	if (DrawContext.IsDetailedMode())
+	{
+		LocalDrawBox(FVector2D(WorldBounds.Min), FVector2D(WorldBounds.GetSize()), FLinearColor::Yellow, 1);
+	}
+
+	// Draw Streaming Sources
+	const float GridLoadingRange = GetLoadingRange();
+	for (const FWorldPartitionStreamingSource& Source : Sources)
 	{
-		WorldScreenBounds.Min.X = FMath::Clamp(WorldScreenBounds.Min.X, GridScreenBounds.Min.X, GridScreenBounds.Max.X);
-		WorldScreenBounds.Min.Y = FMath::Clamp(WorldScreenBounds.Min.Y, GridScreenBounds.Min.Y, GridScreenBounds.Max.Y);
-		WorldScreenBounds.Max.X = FMath::Clamp(WorldScreenBounds.Max.X, GridScreenBounds.Min.X, GridScreenBounds.Max.X);
-		WorldScreenBounds.Max.Y = FMath::Clamp(WorldScreenBounds.Max.Y, GridScreenBounds.Min.Y, GridScreenBounds.Max.Y);
+		const FColor Color = Source.GetDebugColor();
+		Source.ForEachShape(GridLoadingRange, GridName, HLODLayer, /*bProjectIn2D*/ true, [&Color, &WorldToScreen, &GridScreenBounds, &DrawContext, this](const FSphericalSector& Shape)
+		{
+			DrawStreamingSource2D(GridScreenBounds, Shape, WorldToScreen, Color, DrawContext);
+		});
 	}
-	FCanvasBoxItem Box(WorldScreenBounds.Min, WorldScreenBounds.GetSize());
-	Box.SetColor(FColor::Yellow);
-	Box.BlendMode = SE_BLEND_Translucent;
-	Canvas->DrawItem(Box);
 }
 
-void FSpatialHashStreamingGrid::DrawStreamingSource2D(UCanvas* Canvas, const FSphericalSector& Shape, TFunctionRef<FVector2D(const FVector2D&)> WorldToScreen, const FColor& Color) const
+void FSpatialHashStreamingGrid::DrawStreamingSource2D(const FBox2D& GridScreenBounds, const FSphericalSector& Shape, TFunctionRef<FVector2D(const FVector2D&)> WorldToScreen, const FColor& Color, FWorldPartitionDraw2DContext& DrawContext) const
 {
 	check(!Shape.IsNearlyZero())
 
-	FCanvasLineItem LineItem;
-	LineItem.LineThickness = 2;
-	LineItem.SetColor(Color);
-
 	// Spherical Sector
 	const FVector2D Center2D(FVector2D(Shape.GetCenter()));
 	const FSphericalSector::FReal Angle = Shape.GetAngle();
@@ -926,26 +903,25 @@ void FSpatialHashStreamingGrid::DrawStreamingSource2D(UCanvas* Canvas, const FSp
 	const float AngleIncrement = Angle / MaxSegments;
 	const FVector2D Axis = FVector2D(Shape.GetAxis());
 	const FVector Startup = FRotator(0, -0.5f * Angle, 0).RotateVector(Shape.GetScaledAxis());
-	FCanvas* CanvasObject = Canvas->Canvas;
 
 	FVector2D LineStart = FVector2D(Startup);
 	if (!Shape.IsSphere())
 	{
 		// Draw sector start axis
-		LineItem.Draw(CanvasObject, WorldToScreen(Center2D), WorldToScreen(Center2D + LineStart));
+		DrawContext.PushDrawSegment(GridScreenBounds, WorldToScreen(Center2D), WorldToScreen(Center2D + LineStart), Color, 2);
 	}
 	// Draw sector Arc
 	for (int32 i = 1; i <= MaxSegments; i++)
 	{
 		FVector2D LineEnd = FVector2D(FRotator(0, AngleIncrement * i, 0).RotateVector(Startup));
-		LineItem.Draw(CanvasObject, WorldToScreen(Center2D + LineStart), WorldToScreen(Center2D + LineEnd));
+		DrawContext.PushDrawSegment(GridScreenBounds, WorldToScreen(Center2D + LineStart), WorldToScreen(Center2D + LineEnd), Color, 2);
 		LineStart = LineEnd;
 	}
 	// If sphere, close circle, else draw sector end axis
-	LineItem.Draw(CanvasObject, WorldToScreen(Center2D + LineStart), WorldToScreen(Center2D + (Shape.IsSphere() ? FVector2D(Startup) : FVector2D::ZeroVector)));
+	DrawContext.PushDrawSegment(GridScreenBounds, WorldToScreen(Center2D + LineStart), WorldToScreen(Center2D + (Shape.IsSphere() ? FVector2D(Startup) : FVector2D::ZeroVector)), Color, 2);
 
 	// Draw direction vector
-	LineItem.Draw(CanvasObject, WorldToScreen(Center2D), WorldToScreen(Center2D + Axis * Shape.GetRadius()));
+	DrawContext.PushDrawSegment(GridScreenBounds, WorldToScreen(Center2D), WorldToScreen(Center2D + Axis * Shape.GetRadius()), Color, 2);
 }
 
 // ------------------------------------------------------------------------------------------------
@@ -1728,10 +1704,12 @@ EWorldPartitionStreamingPerformance UWorldPartitionRuntimeSpatialHash::GetStream
 	return EWorldPartitionStreamingPerformance::Good;
 }
 
-bool UWorldPartitionRuntimeSpatialHash::Draw2D(UCanvas* Canvas, const TArray<FWorldPartitionStreamingSource>& Sources, const FVector2D& PartitionCanvasSize, const FVector2D& Offset, FVector2D& OutUsedCanvasSize) const
+bool UWorldPartitionRuntimeSpatialHash::Draw2D(FWorldPartitionDraw2DContext& DrawContext) const
 {
 	TRACE_CPUPROFILER_EVENT_SCOPE(UWorldPartitionRuntimeSpatialHash::Draw2D);
-	OutUsedCanvasSize = FVector2D::ZeroVector;
+
+	const UWorldPartition* WorldPartition = GetOuterUWorldPartition();
+	const TArray<FWorldPartitionStreamingSource>& Sources = WorldPartition->GetStreamingSources();
 
 	TArray<const FSpatialHashStreamingGrid*> FilteredStreamingGrids = GetFilteredStreamingGrids();
 	if (FilteredStreamingGrids.Num() == 0 || Sources.Num() == 0)
@@ -1739,59 +1717,65 @@ bool UWorldPartitionRuntimeSpatialHash::Draw2D(UCanvas* Canvas, const TArray<FWo
 		return false;
 	}
 
-	UWorldPartition* WorldPartition = GetOuterUWorldPartition();
-	const float CanvasMaxScreenSize = PartitionCanvasSize.X;
-	const float GridMaxScreenSize = CanvasMaxScreenSize / FilteredStreamingGrids.Num();
-	const float GridEffectiveScreenRatio = 1.f;
-	const float GridEffectiveScreenSize = FMath::Min(GridMaxScreenSize, PartitionCanvasSize.Y) - 10.f;
-	const float GridViewMinimumSizeInCellCount = 5.f;
-	const FVector2D GridScreenExtent = FVector2D(GridEffectiveScreenSize, GridEffectiveScreenSize);
-	const FVector2D GridScreenHalfExtent = 0.5f * GridScreenExtent;
-	const FVector2D GridScreenInitialOffset = Offset;
-
 	// Sort streaming grids to render them sorted by loading range
 	FilteredStreamingGrids.Sort([](const FSpatialHashStreamingGrid& A, const FSpatialHashStreamingGrid& B) { return A.LoadingRange < B.LoadingRange; });
 
+	const FTransform WorldPartitionTransform = WorldPartition->GetInstanceTransform();
+	const FTransform2D LocalWPToGlobal(FQuat2D(FMath::DegreesToRadians(WorldPartitionTransform.Rotator().Yaw)), FVector2D(WorldPartitionTransform.GetLocation()));
+	const FTransform2D GlobalToLocalWP = LocalWPToGlobal.Inverse();
+	const FBox2D& WorldRegion = DrawContext.GetWorldRegion();
+	const FVector2D WorldRegionSize = WorldRegion.GetSize();
+	TArray<FVector2D> Points;
+	Points.Add(GlobalToLocalWP.TransformPoint(WorldRegion.Min));
+	Points.Add(GlobalToLocalWP.TransformPoint(WorldRegion.Min + FVector2D(WorldRegionSize.X, 0)));
+	Points.Add(GlobalToLocalWP.TransformPoint(WorldRegion.Min + WorldRegionSize));
+	Points.Add(GlobalToLocalWP.TransformPoint(WorldRegion.Min + FVector2D(0, WorldRegionSize.Y)));
+	const FBox2D LocalRegion(Points);
+
+	const FBox2D& CanvasRegion = DrawContext.GetCanvasRegion();
+	const int32 GridScreenWidthDivider = DrawContext.IsDetailedMode() ? FilteredStreamingGrids.Num() : 1;
+	const float GridScreenWidthShrinkSize = GridScreenWidthDivider > 1 ? 20.f : 0.f;
+	const float CanvasMaxScreenWidth = CanvasRegion.GetSize().X;
+	const float GridMaxScreenWidth = CanvasMaxScreenWidth / GridScreenWidthDivider;
+	const float GridEffectiveScreenWidth = FMath::Min(GridMaxScreenWidth, CanvasRegion.GetSize().Y) - GridScreenWidthShrinkSize;
+	const FVector2D PartitionCanvasSize = FVector2D(CanvasRegion.GetSize().GetMin());
+	const FVector2D GridScreenExtent = FVector2D(GridEffectiveScreenWidth, GridEffectiveScreenWidth);
+	const FVector2D GridScreenHalfExtent = 0.5f * GridScreenExtent;
+	const FVector2D GridScreenInitialOffset = CanvasRegion.Min;
+
+	FBox GridsShapeBounds(ForceInit);
 	FBox2D GridsBounds(ForceInit);
 	int32 GridIndex = 0;
 	for (const FSpatialHashStreamingGrid* StreamingGrid : FilteredStreamingGrids)
 	{
-		// Display view sides based on extended grid loading range (minimum of N cells)
-		// Take into consideration GShowRuntimeSpatialHashGridLevel when using CellSize
-		int32 MinGridLevel = FMath::Clamp<int32>(GShowRuntimeSpatialHashGridLevel, 0, StreamingGrid->GridLevels.Num() - 1);
-		int64 CellSize = (1LL << MinGridLevel) * (int64)StreamingGrid->CellSize;
-		const FVector MinExtent(CellSize * GridViewMinimumSizeInCellCount);
-		FBox Region(ForceInit);
-		// Keep rendered region around source location if loading range is zero
-		float LoadingRange = FMath::Max<float>(StreamingGrid->GetLoadingRange(), 1.f);
 		for (const FWorldPartitionStreamingSource& Source : Sources)
 		{
-			Region += Source.CalcBounds(LoadingRange, StreamingGrid->GridName, StreamingGrid->HLODLayer, /*bCalcIn2D*/ true);
+			Source.ForEachShape(StreamingGrid->GetLoadingRange(), StreamingGrid->GridName, StreamingGrid->HLODLayer, /*bProjectIn2D*/ true, [&GridsShapeBounds](const FSphericalSector& Shape) { GridsShapeBounds += Shape.CalcBounds(); });
 		}
-		Region += FBox(Region.GetCenter() - MinExtent, Region.GetCenter() + MinExtent);
-		const FVector2D GridReferenceWorldPos = FVector2D(Region.GetCenter());
-		const float RegionExtent = FVector2D(Region.GetExtent()).Size();
-		const FVector2D GridScreenOffset = GridScreenInitialOffset + ((float)GridIndex * FVector2D(GridMaxScreenSize, 0.f)) + GridScreenHalfExtent;
+
+		const FVector2D GridReferenceWorldPos = FVector2D(WorldRegion.GetCenter());
+		const FVector2D WorldRegionExtent = WorldRegion.GetExtent();
+		const FVector2D GridScreenOffset = GridScreenInitialOffset + ((float)GridIndex * FVector2D(GridMaxScreenWidth, 0.f)) + GridScreenHalfExtent + FVector2D(GridScreenWidthShrinkSize * 0.5f);
+		const FVector2D WorldToScreenScale = GridScreenHalfExtent / WorldRegionExtent;
 		const FBox2D GridScreenBounds(GridScreenOffset - GridScreenHalfExtent, GridScreenOffset + GridScreenHalfExtent);
-		const float WorldToScreenScale = (0.5f * GridEffectiveScreenSize) / RegionExtent;
-		auto WorldToScreen = [&](const FVector2D& WorldPos) { return (WorldToScreenScale * (WorldPos - GridReferenceWorldPos)) + GridScreenOffset; };
 
-		StreamingGrid->Draw2D(Canvas, this, Sources, Region, GridScreenBounds, WorldToScreen);
+		auto WorldToScreen = [&](const FVector2D& LocalWorldPos, bool bIsLocalWorldPos = true)
+		{
+			FVector2D GlocalPos = bIsLocalWorldPos ? LocalWPToGlobal.TransformPoint(LocalWorldPos) : LocalWorldPos;
+			return (WorldToScreenScale * (GlocalPos - GridReferenceWorldPos)) + GridScreenOffset;
+		};
+
+		StreamingGrid->Draw2D(this, LocalRegion, GridScreenBounds, WorldToScreen, DrawContext);
 		GridsBounds += GridScreenBounds;
 
-		// Draw WorldPartition name
-		FVector2D GridInfoPos = GridScreenOffset - GridScreenHalfExtent;
+		if (DrawContext.IsDetailedMode())
 		{
-			const FString GridInfoText = UWorld::RemovePIEPrefix(FPaths::GetBaseFilename(WorldPartition->GetPackage()->GetName()));
-			float TextWidth, TextHeight;
-			Canvas->StrLen(GEngine->GetTinyFont(), GridInfoText, TextWidth, TextHeight);
-			Canvas->SetDrawColor(255, 255, 255);
-			Canvas->DrawText(GEngine->GetTinyFont(), GridInfoText, GridInfoPos.X, GridInfoPos.Y);
-			GridInfoPos.Y += TextHeight + 1;
-		}
+			// Draw WorldPartition name
+			FVector2D GridInfoPos = GridScreenOffset - GridScreenHalfExtent;
 
-		// Draw Grid name, loading range
-		{
+			FWorldPartitionCanvasMultiLineText MultiLineText;
+			MultiLineText.Emplace(UWorld::RemovePIEPrefix(FPaths::GetBaseFilename(WorldPartition->GetPackage()->GetName())), FLinearColor::White);
+			// Draw Grid name, loading range
 			FString GridInfoText = FString::Printf(TEXT("%s | %d m"), *StreamingGrid->GridName.ToString(), int32(StreamingGrid->GetLoadingRange() * 0.01f));
 			if (StreamingGrid->bClientOnlyVisible)
 			{
@@ -1803,17 +1787,28 @@ bool UWorldPartitionRuntimeSpatialHash::Draw2D(UCanvas* Canvas, const TArray<FWo
 				GridInfoText += FString::Printf(TEXT(" | GridLevelFilter %d"), GFilterRuntimeSpatialHashGridLevel);
 			}
 #endif
-			Canvas->SetDrawColor(255, 255, 0);
-			Canvas->DrawText(GEngine->GetTinyFont(), GridInfoText, GridInfoPos.X, GridInfoPos.Y);
-		}
+			MultiLineText.Emplace(GridInfoText, FLinearColor::Yellow);
+			FWorldPartitionCanvasMultiLineTextItem Item(GridInfoPos, MultiLineText);
+			DrawContext.PushDrawText(Item);
 
-		++GridIndex;
+			++GridIndex;
+		}
 	}
 
-	if (GridsBounds.bIsValid)
-	{
-		OutUsedCanvasSize = GridsBounds.GetSize();
+	FBox2D DesiredWorldBounds(ForceInit);
+	if (GridsShapeBounds.IsValid)
+	{
+		// Convert to 2D
+		FBox2D GridsShapeBounds2D = FBox2D(FVector2D(GridsShapeBounds.Min.X, GridsShapeBounds.Min.Y), FVector2D(GridsShapeBounds.Max.X, GridsShapeBounds.Max.Y));
+		// Use max extent of X/Y
+		GridsShapeBounds2D = FBox2D(GridsShapeBounds2D.GetCenter() - FVector2D(GridsShapeBounds2D.GetExtent().GetMax()), GridsShapeBounds2D.GetCenter() + FVector2D(GridsShapeBounds2D.GetExtent().GetMax()));
+		// Transform to global space
+		GridsShapeBounds2D = FBox2D(LocalWPToGlobal.TransformPoint(GridsShapeBounds2D.Min), LocalWPToGlobal.TransformPoint(GridsShapeBounds2D.Max));
+		// Expand by 10% computed bounds
+		DesiredWorldBounds = GridsShapeBounds2D.ExpandBy(GridsShapeBounds2D.GetExtent() * 0.1f);
 	}
+	DrawContext.SetDesiredWorldBounds(DesiredWorldBounds);
+	DrawContext.SetUsedCanvasBounds(GridsBounds);
 
 	return true;
 }
diff --git a/Engine/Source/Runtime/Engine/Private/WorldPartition/WorldPartitionStreamingPolicy.cpp b/Engine/Source/Runtime/Engine/Private/WorldPartition/WorldPartitionStreamingPolicy.cpp
--- a/Engine/Source/Runtime/Engine/Private/WorldPartition/WorldPartitionStreamingPolicy.cpp
+++ b/Engine/Source/Runtime/Engine/Private/WorldPartition/WorldPartitionStreamingPolicy.cpp
@@ -948,11 +948,11 @@ bool UWorldPartitionStreamingPolicy::IsStreamingCompleted(EWorldPartitionRuntime
 	return bResult;
 }
 
-bool UWorldPartitionStreamingPolicy::DrawRuntimeHash2D(class UCanvas* Canvas, const FVector2D& PartitionCanvasSize, const FVector2D& Offset, FVector2D& OutUsedCanvasSize)
+bool UWorldPartitionStreamingPolicy::DrawRuntimeHash2D(FWorldPartitionDraw2DContext& DrawContext)
 {
 	if (StreamingSources.Num() > 0 && WorldPartition->RuntimeHash)
 	{
-		return WorldPartition->RuntimeHash->Draw2D(Canvas, StreamingSources, PartitionCanvasSize, Offset, OutUsedCanvasSize);
+		return WorldPartition->RuntimeHash->Draw2D(DrawContext);
 	}
 	return false;
 }
diff --git a/Engine/Source/Runtime/Engine/Private/WorldPartition/WorldPartitionSubsystem.cpp b/Engine/Source/Runtime/Engine/Private/WorldPartition/WorldPartitionSubsystem.cpp
--- a/Engine/Source/Runtime/Engine/Private/WorldPartition/WorldPartitionSubsystem.cpp
+++ b/Engine/Source/Runtime/Engine/Private/WorldPartition/WorldPartitionSubsystem.cpp
@@ -10,17 +10,22 @@
 #include "WorldPartition/WorldPartitionRuntimeHash.h"
 #include "WorldPartition/WorldPartitionStreamingSource.h"
 #include "WorldPartition/WorldPartitionReplay.h"
+#include "WorldPartition/WorldPartitionDraw2DContext.h"
 #include "WorldPartition/DataLayer/WorldDataLayersActorDesc.h"
 #include "WorldPartition/DataLayer/DataLayerManager.h"
 #include "WorldPartition/DataLayer/DataLayerAsset.h"
 #include "WorldPartition/ContentBundle/ContentBundleWorldSubsystem.h"
+#include "WorldPartition/WorldPartitionLevelStreamingDynamic.h"
 #include "Engine/Canvas.h"
 #include "Engine/CoreSettings.h"
 #include "Engine/LevelStreaming.h"
+#include "Engine/LevelStreamingGCHelper.h"
 #include "Streaming/LevelStreamingDelegates.h"
 #include "Engine/LevelBounds.h"
 #include "Debug/DebugDrawService.h"
 #include "GameFramework/PlayerController.h"
+#include "Async/ParallelFor.h"
+#include "Algo/ForEach.h"
 
 #if WITH_EDITOR
 #include "Editor.h"
@@ -32,10 +37,11 @@
 extern int32 GBlockOnSlowStreaming;
 static const FName NAME_WorldPartitionRuntimeHash("WorldPartitionRuntimeHash");
 
-static int32 GDrawWorldPartitionIndex = 0;
+static int32 GDrawWorldPartitionIndex = -1;
 static FAutoConsoleCommand CVarDrawWorldPartitionIndex(
 	TEXT("wp.Runtime.DrawWorldPartitionIndex"),
-	TEXT("Sets the index of the wanted world partition to display debug draw."),
+	TEXT("Sets the index of the wanted partitioned world to display debug draw.\n")
+	TEXT("Sets < 0 to display debug draw all registered partitioned worlds.\n"),
 	FConsoleCommandWithArgsDelegate::CreateLambda([](const TArray<FString>& Args)
 	{
 		if (Args.Num() == 1)
@@ -550,7 +556,7 @@ void UWorldPartitionSubsystem::Tick(float DeltaSeconds)
 	{
 		RegisteredWorldPartition->Tick(DeltaSeconds);
 
-		if (GDrawRuntimeHash3D && RegisteredWorldPartition->CanDebugDraw())
+		if (GDrawRuntimeHash3D && CanDebugDraw())
 		{
 			RegisteredWorldPartition->DrawRuntimeHash3D();
 		}
@@ -766,29 +772,56 @@ void UWorldPartitionSubsystem::UpdateStreamingState()
 	}
 }
 
+bool UWorldPartitionSubsystem::CanDebugDraw() const
+{
+#if WITH_EDITOR
+	const ENetMode NetMode = GetWorld()->GetNetMode();
+	const bool bIsDedicatedServer = (NetMode == NM_DedicatedServer);
+	const bool bIsClient = (NetMode == NM_Client);
+	if ((bIsDedicatedServer && !UWorldPartition::DebugDedicatedServerStreaming) ||
+		(bIsClient && UWorldPartition::DebugDedicatedServerStreaming))
+	{
+		return false;
+	}
+
+	if (!bIsDedicatedServer && UWorldPartition::IsSimulating(false) && (!GEditor || (GetWorld() != GEditor->PlayWorld)))
+	{
+		return false;
+	}
+#endif
+	return GetWorld()->IsGameWorld();
+}
+
 void UWorldPartitionSubsystem::Draw(UCanvas* Canvas, class APlayerController* PC)
 {
 	TRACE_CPUPROFILER_EVENT_SCOPE(UWorldPartitionSubsystem::Draw);
-	if (!Canvas || !Canvas->SceneView || !RegisteredWorldPartitions.IsValidIndex(GDrawWorldPartitionIndex))
+	if (!Canvas || !Canvas->SceneView)
+	{
+		return;
+	}
+
+	if (!CanDebugDraw())
 	{
 		return;
 	}
 
-	UWorldPartition* WorldPartition = RegisteredWorldPartitions[GDrawWorldPartitionIndex];
-	if (!WorldPartition->CanDebugDraw())
+	if ((GDrawWorldPartitionIndex >= 0) && !RegisteredWorldPartitions.IsValidIndex(GDrawWorldPartitionIndex))
 	{
 		return;
 	}
 
 	// Filter out views that don't match our world
-	if (!WorldPartition->GetWorld()->IsNetMode(NM_DedicatedServer) && !UWorldPartition::IsSimulating(false) &&
+	if (!GetWorld()->IsNetMode(NM_DedicatedServer) && !UWorldPartition::IsSimulating(false) &&
 		(Canvas->SceneView->ViewActor == nullptr || Canvas->SceneView->ViewActor->GetWorld() != GetWorld()))
 	{
 		return;
 	}
 
-	const FVector2D CanvasTopLeftPadding(10.f, 10.f);
+	// When GDrawWorldPartitionIndex < 0, draw all world partitions
+	const int32 WPDrawCount = (GDrawWorldPartitionIndex < 0) ? RegisteredWorldPartitions.Num() : 1;
+	UWorldPartition* SingleWorldPartition = (WPDrawCount == 1) ? RegisteredWorldPartitions[FMath::Max(GDrawWorldPartitionIndex, 0)] : nullptr;
 
+	const FVector2D CanvasTopLeftPadding(10.f, 10.f);
 	FVector2D CurrentOffset(CanvasTopLeftPadding);
 
 	if (GDrawRuntimeHash2D)
@@ -797,18 +830,81 @@ void UWorldPartitionSubsystem::Draw(UCanvas* Canvas, class APlayerController* PC
 		const FVector2D CanvasBottomRightPadding(10.f, 10.f);
 		const FVector2D CanvasMinimumSize(100.f, 100.f);
 		const FVector2D CanvasMaxScreenSize = FVector2D::Max(MaxScreenRatio*FVector2D(Canvas->ClipX, Canvas->ClipY) - CanvasBottomRightPadding - CurrentOffset, CanvasMinimumSize);
+		const FVector2D PartitionCanvasSize = FVector2D(CanvasMaxScreenSize.X, CanvasMaxScreenSize.Y);
 
-		FVector2D PartitionCanvasSize = FVector2D(CanvasMaxScreenSize.X, CanvasMaxScreenSize.Y);
-		FVector2D UsedCanvasSize = FVector2D::ZeroVector;
-		if (WorldPartition->DrawRuntimeHash2D(Canvas, PartitionCanvasSize, CurrentOffset, UsedCanvasSize))
+		FBox2D WorldRegion(ForceInit);
+		WorldPartitionsDraw2DContext.SetNum(RegisteredWorldPartitions.Num(), false);
+		Algo::ForEach(WorldPartitionsDraw2DContext, [&](FWorldPartitionDraw2DContext& Context) { if (const FBox2D& Bounds = Context.GetDesiredWorldBounds(); Bounds.bIsValid) { WorldRegion += Bounds; } });
+		if (!WorldRegion.bIsValid)
 		{
+			// Fallback on a default region centered around centroid of all streaming sources
+			const float DefaultWorldRegionExtent = 25600.f;
+			FVector2D Extent = FVector2D(DefaultWorldRegionExtent);
+			FVector2D Pos = FVector2D::ZeroVector;
+			for (const FWorldPartitionStreamingSource& Source : StreamingSources)
+			{
+				Pos += FVector2D(Source.Location.X, Source.Location.Y);
+			}
+			Pos /= StreamingSources.Num();
+			WorldRegion = FBox2D(Pos - Extent, Pos + Extent);
+		}
+
+		const FBox2D CanvasBox = FBox2D(CurrentOffset, CurrentOffset + PartitionCanvasSize);
+		Algo::ForEach(WorldPartitionsDraw2DContext, [&](FWorldPartitionDraw2DContext& Context) { Context.Initialize(CanvasBox, WorldRegion); });
+
+		std::atomic<int32> LineCount;
+		std::atomic<int32> BoxCount;
+		bool bSucceeded = false;
+		
+		ParallelFor(WPDrawCount, [&](int32 ProcessIndex)
+		{
+			UWorldPartition* CurrentWorldPartition = SingleWorldPartition ? SingleWorldPartition : RegisteredWorldPartitions[ProcessIndex].Get();
+			FWorldPartitionDraw2DContext& Context = WorldPartitionsDraw2DContext[ProcessIndex];
+			Context.SetIsDetailedMode(WPDrawCount == 1);
+			Context.SetDrawGridBounds(ProcessIndex == (WPDrawCount-1));
+			Context.SetDrawGridAxis(ProcessIndex == (WPDrawCount-1));
+			if (CurrentWorldPartition->DrawRuntimeHash2D(Context))
+			{
+				LineCount += Context.GetLineCount();
+				BoxCount += Context.GetBoxCount();
+			}
+		});
+
+		if (WPDrawCount > 1)
+		{
+			FWorldPartitionCanvasMultiLineText Text;
+			Text.Emplace(FString::Printf(TEXT("Showing %d world partitions"), WPDrawCount), FLinearColor::White);
+			FWorldPartitionCanvasMultiLineTextItem Item(CanvasBox.Min, Text);
+			WorldPartitionsDraw2DContext[WPDrawCount - 1].PushDrawText(Item);
+		}
+
+		int32 TotalLineCount = LineCount.load();
+		int32 TotalBoxCount = BoxCount.load();
+		if (TotalLineCount + TotalBoxCount > 0)
+		{
+			TRACE_CPUPROFILER_EVENT_SCOPE(UWorldPartitionSubsystem::DrawAllCanvasItems);
+			FWorldPartitionDraw2DCanvas WorldPartitionCanvas(Canvas);
+			WorldPartitionCanvas.PrepareDraw(TotalLineCount, TotalBoxCount);
+
+			FBox2D MaxUsedBounds(ForceInit);
+			for (FWorldPartitionDraw2DContext& Context : WorldPartitionsDraw2DContext)
+			{
+				WorldPartitionCanvas.Draw(Context.GetCanvasItems());
+				
+				if (Context.GetUsedCanvasBounds().bIsValid)
+				{
+					MaxUsedBounds += Context.GetUsedCanvasBounds();
+				}
+			}
+			
 			CurrentOffset.X = CanvasBottomRightPadding.X;
-			CurrentOffset.Y += UsedCanvasSize.Y;
+			CurrentOffset.Y += MaxUsedBounds.bIsValid ? MaxUsedBounds.Max.Y : CanvasBox.GetSize().Y;
 		}
 	}
 	
 	if (GDrawStreamingPerfs || GDrawRuntimeHash2D)
 	{
+		TRACE_CPUPROFILER_EVENT_SCOPE(UWorldPartitionSubsystem::DrawStreamingStatus);
 		{
 			FString StatusText;
 			if (IsIncrementalPurgePending()) { StatusText += TEXT("(Purging) "); }
@@ -817,9 +913,9 @@ void UWorldPartitionSubsystem::Draw(UCanvas* Canvas, class APlayerController* PC
 			if (StatusText.IsEmpty()) { StatusText = TEXT("(Idle) "); }
 
 			FString DebugWorldText = FString::Printf(TEXT("(%s)"), *GetDebugStringForWorld(GetWorld()));
-			if (WorldPartition->IsServer())
+			if (SingleWorldPartition && SingleWorldPartition->IsServer())
 			{
-				DebugWorldText += FString::Printf(TEXT(" (Server Streaming %s)"), WorldPartition->IsServerStreamingEnabled() ? TEXT("Enabled") : TEXT("Disabled"));
+				DebugWorldText += FString::Printf(TEXT(" (Server Streaming %s)"), SingleWorldPartition->IsServerStreamingEnabled() ? TEXT("Enabled") : TEXT("Disabled"));
 			}
 			
 			const FString Text = FString::Printf(TEXT("Streaming Status for %s: %s"), *DebugWorldText, *StatusText);
@@ -827,8 +923,25 @@ void UWorldPartitionSubsystem::Draw(UCanvas* Canvas, class APlayerController* PC
 		}
 
 		{
+			auto GetStreamingPerformance = [&]()
+			{
+				EWorldPartitionStreamingPerformance StreamingPerformance = SingleWorldPartition ? SingleWorldPartition->GetStreamingPerformance() : EWorldPartitionStreamingPerformance::Good;
+				if (!SingleWorldPartition)
+				{
+					for (UWorldPartition* RegisteredWorldPartition : RegisteredWorldPartitions)
+					{
+						StreamingPerformance = FMath::Max(StreamingPerformance, RegisteredWorldPartition->GetStreamingPerformance());
+						if (StreamingPerformance == EWorldPartitionStreamingPerformance::Critical)
+						{
+							break;
+						}
+					}
+				}
+				return StreamingPerformance;
+			};
+
 			FString StatusText;
-			EWorldPartitionStreamingPerformance StreamingPerformance = WorldPartition->GetStreamingPerformance();
+			EWorldPartitionStreamingPerformance StreamingPerformance = GetStreamingPerformance();
 			switch (StreamingPerformance)
 			{
 			case EWorldPartitionStreamingPerformance::Good:
@@ -853,7 +966,7 @@ void UWorldPartitionSubsystem::Draw(UCanvas* Canvas, class APlayerController* PC
 	{
 		TRACE_CPUPROFILER_EVENT_SCOPE(UWorldPartitionSubsystem::DrawStreamingSources);
 
-		const TArray<FWorldPartitionStreamingSource>& LocalStreamingSources = WorldPartition->GetStreamingSources();
+		const TArray<FWorldPartitionStreamingSource>& LocalStreamingSources = SingleWorldPartition ? SingleWorldPartition->GetStreamingSources() : StreamingSources;
 		if (LocalStreamingSources.Num() > 0)
 		{
 			FString Title(TEXT("Streaming Sources"));
@@ -882,27 +995,80 @@ void UWorldPartitionSubsystem::Draw(UCanvas* Canvas, class APlayerController* PC
 	if (GDrawLegends || GDrawRuntimeHash2D)
 	{
 		// Streaming Status Legend
-		WorldPartition->DrawStreamingStatusLegend(Canvas, CurrentOffset);
+		DrawStreamingStatusLegend(Canvas, CurrentOffset, SingleWorldPartition);
 	}
 
-	if (GDrawDataLayers || GDrawDataLayersLoadTime || GDrawRuntimeHash2D)
+	if (SingleWorldPartition)
 	{
-		if (UDataLayerManager* DataLayerManager = WorldPartition->GetDataLayerManager())
+		if (GDrawDataLayers || GDrawDataLayersLoadTime || GDrawRuntimeHash2D)
+		{
+			if (UDataLayerManager* DataLayerManager = SingleWorldPartition->GetDataLayerManager())
+			{
+				DataLayerManager->DrawDataLayersStatus(Canvas, CurrentOffset);
+			}
+		}
+
+		UContentBundleManager* ContentBundleManager = GetWorld()->ContentBundleManager;
+		if (ContentBundleManager && FWorldPartitionDebugHelper::CanDrawContentBundles() && GDrawRuntimeHash2D)
+		{
+			ContentBundleManager->DrawContentBundlesStatus(GetWorld(), Canvas, CurrentOffset);
+		}
+
+		if (GDrawRuntimeCellsDetails)
 		{
-			DataLayerManager->DrawDataLayersStatus(Canvas, CurrentOffset);
+			SingleWorldPartition->DrawRuntimeCellsDetails(Canvas, CurrentOffset);
 		}
 	}
+}
+
+void UWorldPartitionSubsystem::DrawStreamingStatusLegend(class UCanvas* Canvas, FVector2D& Offset, const UWorldPartition* InWorldPartition)
+{
+	TRACE_CPUPROFILER_EVENT_SCOPE(UWorldPartitionSubsystem::DrawStreamingStatusLegend);
+
+	check(Canvas);
 
-	UContentBundleManager* ContentBundleManager = GetWorld()->ContentBundleManager;
-	if (ContentBundleManager && (FWorldPartitionDebugHelper::CanDrawContentBundles() && GDrawRuntimeHash2D))
+	TArray<int32> StreamingStatus;
+	StreamingStatus.SetNumZeroed((int32)LEVEL_StreamingStatusCount);
+	
+	const UWorld* StreamingWorld = InWorldPartition ? InWorldPartition->GetTypedOuter<UWorld>() : nullptr;
+	// Cumulate counter stats
+	for (ULevelStreaming* LevelStreaming : GetWorld()->GetStreamingLevels())
 	{
-		ContentBundleManager->DrawContentBundlesStatus(GetWorld(), Canvas, CurrentOffset);
+		if (!StreamingWorld || (StreamingWorld == LevelStreaming->GetStreamingWorld()))
+		{
+			if (UWorldPartitionLevelStreamingDynamic* WorldPartitionLevelStreaming = Cast<UWorldPartitionLevelStreamingDynamic>(LevelStreaming))
+			{
+				if (const UWorldPartitionRuntimeCell* Cell = WorldPartitionLevelStreaming->GetWorldPartitionRuntimeCell())
+				{
+					StreamingStatus[(int32)Cell->GetStreamingStatus()]++;
+				}
+			}
+		}
 	}
 
-	if (GDrawRuntimeCellsDetails)
+	// @todo_ow: This is not exactly the good value, as we could have pending unload level from Level Instances, etc.
+	//           We could modify GetNumLevelsPendingPurge to return the number of pending purge levels from the grid, 
+	//           but that will do for now.
+	StreamingStatus[LEVEL_UnloadedButStillAround] = FLevelStreamingGCHelper::GetNumLevelsPendingPurge();
+
+	// Draw legend
+	FVector2D Pos = Offset;
+	float MaxTextWidth = 0.f;
+	FWorldPartitionDebugHelper::DrawText(Canvas, TEXT("Streaming Status Legend"), GEngine->GetSmallFont(), FColor::Yellow, Pos, &MaxTextWidth);
+
+	for (int32 i = 0; i < (int32)LEVEL_StreamingStatusCount; ++i)
 	{
-		WorldPartition->DrawRuntimeCellsDetails(Canvas, CurrentOffset);
+		EStreamingStatus Status = (EStreamingStatus)i;
+		const FColor& StatusColor = ULevelStreaming::GetLevelStreamingStatusColor(Status);
+		FString DebugString = *FString::Printf(TEXT("%d) %s"), i, ULevelStreaming::GetLevelStreamingStatusDisplayName(Status));
+		if (Status != LEVEL_Unloaded)
+		{
+			DebugString += *FString::Printf(TEXT(" (%d)"), StreamingStatus[(int32)Status]);
+		}
+		FWorldPartitionDebugHelper::DrawLegendItem(Canvas, *DebugString, GEngine->GetSmallFont(), StatusColor, FColor::White, Pos, &MaxTextWidth);
 	}
+
+	Offset.X += MaxTextWidth + 10;
 }
 
 /*
diff --git a/Engine/Source/Runtime/Engine/Public/WorldPartition/WorldPartition.h b/Engine/Source/Runtime/Engine/Public/WorldPartition/WorldPartition.h
--- a/Engine/Source/Runtime/Engine/Public/WorldPartition/WorldPartition.h
+++ b/Engine/Source/Runtime/Engine/Public/WorldPartition/WorldPartition.h
@@ -42,6 +42,7 @@ class UHLODLayer;
 class UCanvas;
 class ULevel;
 class FAutoConsoleVariableRef;
+class FWorldPartitionDraw2DContext;
 
 struct IWorldPartitionStreamingSourceProvider;
 
@@ -338,11 +339,9 @@ public:
 	const TArray<FWorldPartitionStreamingSource>& GetStreamingSources() const;
 
 	// Debugging
-	bool CanDebugDraw() const;
-	bool DrawRuntimeHash2D(UCanvas* Canvas, const FVector2D& PartitionCanvasSize, const FVector2D& Offset, FVector2D& OutUsedCanvasSize);
+	bool DrawRuntimeHash2D(FWorldPartitionDraw2DContext& DrawContext);
 	void DrawRuntimeHash3D();
 	void DrawRuntimeCellsDetails(UCanvas* Canvas, FVector2D& Offset);
-	void DrawStreamingStatusLegend(UCanvas* Canvas, FVector2D& Offset);
 
 	void OnCellShown(const UWorldPartitionRuntimeCell* InCell);
 	void OnCellHidden(const UWorldPartitionRuntimeCell* InCell);
diff --git a/Engine/Source/Runtime/Engine/Public/WorldPartition/WorldPartitionLevelStreamingPolicy.h b/Engine/Source/Runtime/Engine/Public/WorldPartition/WorldPartitionLevelStreamingPolicy.h
--- a/Engine/Source/Runtime/Engine/Public/WorldPartition/WorldPartitionLevelStreamingPolicy.h
+++ b/Engine/Source/Runtime/Engine/Public/WorldPartition/WorldPartitionLevelStreamingPolicy.h
@@ -23,7 +23,6 @@ class UWorldPartitionLevelStreamingPolicy : public UWorldPartitionStreamingPolic
 
 public:
 	virtual void DrawRuntimeCellsDetails(class UCanvas* Canvas, FVector2D& Offset) override;
-	virtual void DrawStreamingStatusLegend(UCanvas* Canvas, FVector2D& Offset) override;
 	virtual bool IsStreamingCompleted(const TArray<FWorldPartitionStreamingSource>* InStreamingSources) const override;
 
 #if WITH_EDITOR
diff --git a/Engine/Source/Runtime/Engine/Public/WorldPartition/WorldPartitionRuntimeHash.h b/Engine/Source/Runtime/Engine/Public/WorldPartition/WorldPartitionRuntimeHash.h
--- a/Engine/Source/Runtime/Engine/Public/WorldPartition/WorldPartitionRuntimeHash.h
+++ b/Engine/Source/Runtime/Engine/Public/WorldPartition/WorldPartitionRuntimeHash.h
@@ -19,6 +19,7 @@
 #include "WorldPartitionRuntimeHash.generated.h"
 
 struct FHierarchicalLogArchive;
+class FWorldPartitionDraw2DContext;
 
 UENUM()
 enum class EWorldPartitionStreamingPerformance : uint8
@@ -141,7 +142,7 @@ public:
 	virtual bool InjectExternalStreamingObject(URuntimeHashExternalStreamingObjectBase* ExternalStreamingObject) { return false; }
 	virtual bool RemoveExternalStreamingObject(URuntimeHashExternalStreamingObjectBase* ExternalStreamingObject) { return false; }
 
-	virtual bool Draw2D(class UCanvas* Canvas, const TArray<FWorldPartitionStreamingSource>& Sources, const FVector2D& PartitionCanvasSize, const FVector2D& Offset, FVector2D& OutUsedCanvasSize) const { return false; }
+	virtual bool Draw2D(FWorldPartitionDraw2DContext& DrawContext) const { return false; }
 	virtual void Draw3D(const TArray<FWorldPartitionStreamingSource>& Sources) const {}
 	virtual bool ContainsRuntimeHash(const FString& Name) const { return false; }
 	virtual bool IsStreaming3D() const { return true; }
diff --git a/Engine/Source/Runtime/Engine/Public/WorldPartition/WorldPartitionRuntimeSpatialHash.h b/Engine/Source/Runtime/Engine/Public/WorldPartition/WorldPartitionRuntimeSpatialHash.h
--- a/Engine/Source/Runtime/Engine/Public/WorldPartition/WorldPartitionRuntimeSpatialHash.h
+++ b/Engine/Source/Runtime/Engine/Public/WorldPartition/WorldPartitionRuntimeSpatialHash.h
@@ -5,6 +5,7 @@
 #include "UObject/ObjectMacros.h"
 #include "UObject/Object.h"
 #include "Misc/HierarchicalLogArchive.h"
+#include "Math/MathFwd.h"
 #include "GameFramework/Info.h"
 
 #include "WorldPartition.h"
@@ -16,6 +17,7 @@
 
 typedef UE::Math::TIntVector3<int64> FGridCellCoord;
 typedef UE::Math::TIntVector2<int64> FGridCellCoord2;
+class FWorldPartitionDraw2DContext;
 
 USTRUCT()
 struct FSpatialHashStreamingGridLayerCell
@@ -143,7 +145,7 @@ struct ENGINE_API FSpatialHashStreamingGrid
 	void GetCells(const FWorldPartitionStreamingQuerySource& QuerySource, TSet<const UWorldPartitionRuntimeCell*>& OutCells, bool bEnableZCulling, FWorldPartitionQueryCache* QueryCache = nullptr) const;
 	void GetCells(const TArray<FWorldPartitionStreamingSource>& Sources, UWorldPartitionRuntimeHash::FStreamingSourceCells& OutActivateCells, UWorldPartitionRuntimeHash::FStreamingSourceCells& OutLoadCells, bool bEnableZCulling) const;
 	void GetNonSpatiallyLoadedCells(TSet<const UWorldPartitionRuntimeCell*>& OutActivateCells, TSet<const UWorldPartitionRuntimeCell*>& OutLoadCells) const;
-	void Draw2D(UCanvas* Canvas, const class UWorldPartitionRuntimeSpatialHash* Owner, const TArray<FWorldPartitionStreamingSource>& Sources, const FBox& Region, const FBox2D& GridScreenBounds, TFunctionRef<FVector2D(const FVector2D&)> WorldToScreen) const;
+	void Draw2D(const class UWorldPartitionRuntimeSpatialHash* Owner, const FBox2D& Region2D, const FBox2D& GridScreenBounds, TFunctionRef<FVector2D(const FVector2D&, bool)> WorldToScreen, FWorldPartitionDraw2DContext& DrawContext) const;
 	void Draw3D(const class UWorldPartitionRuntimeSpatialHash* Owner, const TArray<FWorldPartitionStreamingSource>& Sources, const FTransform& Transform) const;
 	void ForEachRuntimeCell(TFunctionRef<bool(const UWorldPartitionRuntimeCell*)> Func) const;
 	const FSquare2DGridHelper& GetGridHelper() const;
@@ -151,7 +153,7 @@ struct ENGINE_API FSpatialHashStreamingGrid
 private:
 	void ForEachRuntimeCell(const FGridCellCoord& Coords, TFunctionRef<void(const UWorldPartitionRuntimeCell*)> Func) const;
 	void ForEachLayerCell(const FGridCellCoord& Coords, TFunctionRef<void(const FSpatialHashStreamingGridLayerCell*)> Func) const;
-	void DrawStreamingSource2D(UCanvas* Canvas, const FSphericalSector& Shape, TFunctionRef<FVector2D(const FVector2D&)> WorldToScreen, const FColor& Color) const;
+	void DrawStreamingSource2D(const FBox2D& GridScreenBounds, const FSphericalSector& Shape, TFunctionRef<FVector2D(const FVector2D&)> WorldToScreen, const FColor& Color, FWorldPartitionDraw2DContext& DrawContext) const;
 	void DrawStreamingSource3D(UWorld* World, const FSphericalSector& Shape, const FTransform& Transform, const FColor& Color) const;
 	void GetFilteredCellsForDebugDraw(const FSpatialHashStreamingGridLayerCell* LayerCell, TArray<const UWorldPartitionRuntimeCell*>& FilteredCells) const;
 	EWorldPartitionRuntimeCellVisualizeMode GetStreamingCellVisualizeMode() const;
@@ -335,7 +337,7 @@ protected:
 	mutable bool bIsNameToGridMappingDirty;
 
 private:
-	virtual bool Draw2D(class UCanvas* Canvas, const TArray<FWorldPartitionStreamingSource>& Sources, const FVector2D& PartitionCanvasSize, const FVector2D& Offset, FVector2D& OutUsedCanvasSize) const override;
+	virtual bool Draw2D(FWorldPartitionDraw2DContext& DrawContext) const override;
 	virtual void Draw3D(const TArray<FWorldPartitionStreamingSource>& Sources) const override;
 	virtual bool ContainsRuntimeHash(const FString& Name) const override;
 	virtual bool IsStreaming3D() const override;
diff --git a/Engine/Source/Runtime/Engine/Public/WorldPartition/WorldPartitionStreamingPolicy.h b/Engine/Source/Runtime/Engine/Public/WorldPartition/WorldPartitionStreamingPolicy.h
--- a/Engine/Source/Runtime/Engine/Public/WorldPartition/WorldPartitionStreamingPolicy.h
+++ b/Engine/Source/Runtime/Engine/Public/WorldPartition/WorldPartitionStreamingPolicy.h
@@ -17,6 +17,7 @@
 #include "WorldPartitionStreamingPolicy.generated.h"
 
 class UWorldPartition;
+class FWorldPartitionDraw2DContext;
 
 UCLASS(Abstract, Within = WorldPartition)
 class UWorldPartitionStreamingPolicy : public UObject
@@ -27,11 +28,10 @@ public:
 	virtual bool GetIntersectingCells(const TArray<FWorldPartitionStreamingQuerySource>& InSources, TArray<const IWorldPartitionCell*>& OutCells) const;
 	virtual void UpdateStreamingState();
 	virtual bool CanAddLoadedLevelToWorld(class ULevel* InLevel) const;
-	virtual bool DrawRuntimeHash2D(class UCanvas* Canvas, const FVector2D& PartitionCanvasSize, const FVector2D& Offset, FVector2D& OutUsedCanvasSize);
+	virtual bool DrawRuntimeHash2D(FWorldPartitionDraw2DContext& DrawContext);
 	virtual void DrawRuntimeHash3D();
 	virtual void DrawRuntimeCellsDetails(class UCanvas* Canvas, FVector2D& Offset) {}
-	virtual void DrawStreamingStatusLegend(class UCanvas* Canvas, FVector2D& Offset) {}
-
+	
 	virtual bool IsStreamingCompleted(const TArray<FWorldPartitionStreamingSource>* InStreamingSources) const;
 	virtual bool IsStreamingCompleted(EWorldPartitionRuntimeCellState QueryState, const TArray<FWorldPartitionStreamingQuerySource>& QuerySources, bool bExactState = true) const;
 
diff --git a/Engine/Source/Runtime/Engine/Public/WorldPartition/WorldPartitionStreamingSource.h b/Engine/Source/Runtime/Engine/Public/WorldPartition/WorldPartitionStreamingSource.h
--- a/Engine/Source/Runtime/Engine/Public/WorldPartition/WorldPartitionStreamingSource.h
+++ b/Engine/Source/Runtime/Engine/Public/WorldPartition/WorldPartitionStreamingSource.h
@@ -494,15 +494,14 @@ struct ENGINE_API FWorldPartitionStreamingSource
 
 	FString ToString() const
 	{
-		const FVector Direction = Rotation.Euler();
 		return FString::Printf(
-			TEXT("Priority: %d | %s | %s | %s | Pos: X=%lld,Y=%lld,Z=%lld | Rot: X=%d,Y=%d,Z=%d | Vel: %3.2f m/s (%d mph)"), 
+			TEXT("Priority: %d | %s | %s | %s | Pos: X=%lld,Y=%lld,Z=%lld | Rot: %s | Vel: %3.2f m/s (%d mph)"), 
 			Priority, 
 			bRemote ? TEXT("Remote") : TEXT("Local"),
 			GetStreamingSourceTargetStateName(TargetState),
 			bBlockOnSlowLoading ? TEXT("Blocking") : TEXT("NonBlocking"),
 			(int64)Location.X, (int64)Location.Y, (int64)Location.Z, 
-			(int32)Direction.X, (int32)Direction.Y, (int32)Direction.Z, 
+			*Rotation.ToCompactString(),
 			Velocity, 
 			(int32)(Velocity*2.23694f)
 		);
diff --git a/Engine/Source/Runtime/Engine/Public/WorldPartition/WorldPartitionSubsystem.h b/Engine/Source/Runtime/Engine/Public/WorldPartition/WorldPartitionSubsystem.h
--- a/Engine/Source/Runtime/Engine/Public/WorldPartition/WorldPartitionSubsystem.h
+++ b/Engine/Source/Runtime/Engine/Public/WorldPartition/WorldPartitionSubsystem.h
@@ -16,6 +16,7 @@ class ULevelStreaming;
 class UWorldPartition;
 class UActorDescContainer;
 class FWorldPartitionActorDesc;
+class FWorldPartitionDraw2DContext;
 
 enum class EWorldPartitionRuntimeCellState : uint8;
 enum class ELevelStreamingState : uint8;
@@ -168,7 +169,9 @@ private:
 
 	UWorldPartition* GetWorldPartition();
 	const UWorldPartition* GetWorldPartition() const;
+	bool CanDebugDraw() const;
 	void Draw(class UCanvas* Canvas, class APlayerController* PC);
+	void DrawStreamingStatusLegend(class UCanvas* Canvas, FVector2D& Offset, const UWorldPartition* InWorldPartition);
 	friend class UWorldPartition;
 	friend class UWorldPartitionStreamingPolicy;
 
@@ -183,6 +186,7 @@ private:
 	TMap<FName, FStreamingSourceVelocity> StreamingSourcesVelocity;
 	uint32 StreamingSourcesHash;
 
+	TArray<FWorldPartitionDraw2DContext> WorldPartitionsDraw2DContext;
 	FDelegateHandle	DrawHandle;
 
 	// GC backup values
