Optimize collision Events in physics scene
Each registered event need to find its associated proxy. This optimization look first by the smallest list, and use a hash find its linked one. (Map/Set).

diff --git a/Engine/Source/Runtime/Engine/Private/PhysicsEngine/Experimental/PhysScene_Chaos.cpp b/Engine/Source/Runtime/Engine/Private/PhysicsEngine/Experimental/PhysScene_Chaos.cpp
--- a/Engine/Source/Runtime/Engine/Private/PhysicsEngine/Experimental/PhysScene_Chaos.cpp
+++ b/Engine/Source/Runtime/Engine/Private/PhysicsEngine/Experimental/PhysScene_Chaos.cpp
@@ -21,7 +21,6 @@
 #include "PhysicsProxy/ClusterUnionPhysicsProxy.h"
 #include "PhysicsProxy/SkeletalMeshPhysicsProxy.h"
 #include "PhysicsProxy/StaticMeshPhysicsProxy.h"
-#include "EventsData.h"
 #include "Chaos/PhysicsSolverBaseImpl.h"
 
 #if WITH_EDITOR
@@ -559,7 +558,7 @@ AActor* FPhysScene_Chaos::GetSolverActor() const
 
 void FPhysScene_Chaos::RegisterForCollisionEvents(UPrimitiveComponent* Component)
 {
-	CollisionEventRegistrations.AddUnique(Component);
+	CollisionEventRegistrations.Add(Component);
 }
 
 void FPhysScene_Chaos::UnRegisterForCollisionEvents(UPrimitiveComponent* Component)
@@ -808,6 +807,86 @@ FCollisionNotifyInfo& FPhysScene_Chaos::GetPendingCollisionForContactPair(const
 	return PendingCollisionNotifies[NewIdx];
 }
 
+FORCEINLINE void FPhysScene_Chaos::HandleEachCollisionEvent(const TArray<int32>& CollisionIndices, IPhysicsProxyBase* PhysicsProxy0, Chaos::FCollisionDataArray const& CollisionData, Chaos::FReal MinDeltaVelocityThreshold)
+{
+	for (int32 EncodedCollisionIdx : CollisionIndices)
+	{
+		bool bSwapOrder;
+		int32 CollisionIdx = Chaos::FEventManager::DecodeCollisionIndex(EncodedCollisionIdx, bSwapOrder);
+
+		Chaos::FCollidingData const& CollisionDataItem = CollisionData[CollisionIdx];
+		IPhysicsProxyBase* const PhysicsProxy1 = bSwapOrder ? CollisionDataItem.Proxy1 : CollisionDataItem.Proxy2;
+
+		// Are the proxies pending destruction? If they are no longer tracked by the PhysScene, the proxy is deleted or pending deletion.
+		if (GetOwningComponent<UPrimitiveComponent>(PhysicsProxy0) == nullptr || GetOwningComponent<UPrimitiveComponent>(PhysicsProxy1) == nullptr)
+		{
+			return;
+		}
+
+		bool bNewEntry = false;
+		FCollisionNotifyInfo& NotifyInfo = GetPendingCollisionForContactPair(PhysicsProxy0, PhysicsProxy1, CollisionDataItem.SolverTime, bNewEntry);
+
+		// #note: we only notify on the first contact, though we will still accumulate the impulse data from subsequent contacts
+		const FVector NormalImpulse = FVector::DotProduct(CollisionDataItem.AccumulatedImpulse, CollisionDataItem.Normal) * CollisionDataItem.Normal;	// project impulse along normal
+		const FVector FrictionImpulse = FVector(CollisionDataItem.AccumulatedImpulse) - NormalImpulse; // friction is component not along contact normal
+		NotifyInfo.RigidCollisionData.TotalNormalImpulse += NormalImpulse;
+		NotifyInfo.RigidCollisionData.TotalFrictionImpulse += FrictionImpulse;
+
+		// Get swapped velocity deltas
+		const FVector& DeltaVelocity1 = bSwapOrder
+			? CollisionDataItem.DeltaVelocity2
+			: CollisionDataItem.DeltaVelocity1;
+		const FVector& DeltaVelocity2 = bSwapOrder
+			? CollisionDataItem.DeltaVelocity1
+			: CollisionDataItem.DeltaVelocity2;
+
+		// Populate additional contact information for additional hits
+		FRigidBodyContactInfo& NewContact = NotifyInfo.RigidCollisionData.ContactInfos.AddZeroed_GetRef();
+		NewContact.ContactNormal = CollisionDataItem.Normal;
+		NewContact.ContactPosition = CollisionDataItem.Location;
+		NewContact.ContactPenetration = CollisionDataItem.PenetrationDepth;
+		NewContact.bContactProbe = CollisionDataItem.bProbe;
+		NotifyInfo.RigidCollisionData.bIsVelocityDeltaUnderThreshold =
+			DeltaVelocity1.IsNearlyZero(MinDeltaVelocityThreshold) &&
+			DeltaVelocity2.IsNearlyZero(MinDeltaVelocityThreshold);
+
+		const Chaos::FChaosPhysicsMaterial* InternalMat1 = CollisionDataItem.Mat1.Get();
+		const Chaos::FChaosPhysicsMaterial* InternalMat2 = CollisionDataItem.Mat2.Get();
+
+		NewContact.PhysMaterial[0] = InternalMat1 ? FPhysicsUserData::Get<UPhysicalMaterial>(InternalMat1->UserData) : nullptr;
+		NewContact.PhysMaterial[1] = InternalMat2 ? FPhysicsUserData::Get<UPhysicalMaterial>(InternalMat2->UserData) : nullptr;
+
+		if (bSwapOrder)
+		{
+			NewContact.SwapOrder();
+		}
+
+		if (bNewEntry)
+		{
+			// fill in legacy contact data
+			NotifyInfo.bCallEvent0 = true;
+			// if Comp1 wants this event too, it will get its own pending collision entry, so we leave it false
+
+			NotifyInfo.SolverTime = CollisionDataItem.SolverTime;
+
+			NotifyInfo.Info0.SetFrom(GetBodyInstanceFromProxyAndShape(PhysicsProxy0, CollisionDataItem.ShapeIndex1), DeltaVelocity1);
+			NotifyInfo.Info1.SetFrom(GetBodyInstanceFromProxyAndShape(PhysicsProxy1, CollisionDataItem.ShapeIndex2), DeltaVelocity2);
+
+			// in some case ( like with geometry collections ) we don't have a body instance so the component part will null, we need to handle that 
+			if (NotifyInfo.Info0.Component == nullptr)
+			{
+				NotifyInfo.Info0.Component = GetOwningComponent<UPrimitiveComponent>(PhysicsProxy0);
+				NotifyInfo.Info0.Actor = (NotifyInfo.Info0.Component != nullptr) ? NotifyInfo.Info0.Component->GetOwner() : nullptr;
+			}
+			if (NotifyInfo.Info1.Component == nullptr)
+			{
+				NotifyInfo.Info1.Component = GetOwningComponent<UPrimitiveComponent>(PhysicsProxy1);
+				NotifyInfo.Info1.Actor = (NotifyInfo.Info1.Component != nullptr) ? NotifyInfo.Info1.Component->GetOwner() : nullptr;
+			}
+		}
+	}
+}
+
 void FPhysScene_Chaos::HandleCollisionEvents(const Chaos::FCollisionEventData& Event)
 {
 	QUICK_SCOPE_CYCLE_COUNTER(STAT_HandleCollisionEvents);
@@ -822,96 +901,39 @@ void FPhysScene_Chaos::HandleCollisionEvents(const Chaos::FCollisionEventData& E
 	{
 		LastEventDispatchTime = Event.CollisionData.TimeCreated;
 
-		// look through all the components that someone is interested in, and see if they had a collision
-		// note that we only need to care about the interaction from the POV of the registered component,
-		// since if anyone wants notifications for the other component it hit, it's also registered and we'll get to that elsewhere in the list
-		for (TArray<UPrimitiveComponent*>::TIterator It(CollisionEventRegistrations); It; ++It)
+		// Iterate through the smallest between events registration and physics proxies
+		if (PhysicsProxyToCollisionIndicesMap.Num() <= CollisionEventRegistrations.Num())
 		{
-			UPrimitiveComponent* const Comp0 = *It;
-			const TArray<IPhysicsProxyBase*>* PhysicsProxyArray = GetOwnedPhysicsProxies(Comp0);
+			for (TPair<IPhysicsProxyBase*, TArray<int32>> Pair : PhysicsProxyToCollisionIndicesMap)
+			{
+				IPhysicsProxyBase*  PhysicsProxy0 = Pair.Key;
+				TArray<int32> CollisionIndices = Pair.Value;
 
-			if (PhysicsProxyArray)
+				UPrimitiveComponent* Comp0 = GetOwningComponent<UPrimitiveComponent>(PhysicsProxy0);
+
+				if (Comp0 != nullptr && CollisionEventRegistrations.Contains(Comp0))
+				{
+					HandleEachCollisionEvent(CollisionIndices, PhysicsProxy0, CollisionData, MinDeltaVelocityThreshold);
+				}
+			}
+		}
+		else
+		{
+			// look through all the components that someone is interested in, and see if they had a collision
+			// note that we only need to care about the interaction from the POV of the registered component,
+			// since if anyone wants notifications for the other component it hit, it's also registered and we'll get to that elsewhere in the list
+			for (UPrimitiveComponent* const Comp0 : CollisionEventRegistrations)
 			{
-				for (IPhysicsProxyBase* PhysicsProxy0 : *PhysicsProxyArray)
+				const TArray<IPhysicsProxyBase*>* PhysicsProxyArray = GetOwnedPhysicsProxies(Comp0);
+
+				if (PhysicsProxyArray)
 				{
-					TArray<int32> const* const CollisionIndices = PhysicsProxyToCollisionIndicesMap.Find(PhysicsProxy0);
-					if (CollisionIndices)
+					for (IPhysicsProxyBase* PhysicsProxy0 : *PhysicsProxyArray)
 					{
-						for (int32 EncodedCollisionIdx : *CollisionIndices)
+						TArray<int32> const* const CollisionIndices = PhysicsProxyToCollisionIndicesMap.Find(PhysicsProxy0);
+						if (CollisionIndices)
 						{
-							bool bSwapOrder;
-							int32 CollisionIdx = Chaos::FEventManager::DecodeCollisionIndex(EncodedCollisionIdx, bSwapOrder);
-
-							Chaos::FCollidingData const& CollisionDataItem = CollisionData[CollisionIdx];
-							IPhysicsProxyBase* const PhysicsProxy1 = bSwapOrder ? CollisionDataItem.Proxy1 : CollisionDataItem.Proxy2;
-
-							// Are the proxies pending destruction? If they are no longer tracked by the PhysScene, the proxy is deleted or pending deletion.
-							if (GetOwningComponent<UPrimitiveComponent>(PhysicsProxy0) == nullptr || GetOwningComponent<UPrimitiveComponent>(PhysicsProxy1) == nullptr)
-							{
-								continue;
-							}
-
-							bool bNewEntry = false;
-							FCollisionNotifyInfo& NotifyInfo = GetPendingCollisionForContactPair(PhysicsProxy0, PhysicsProxy1, CollisionDataItem.SolverTime, bNewEntry);
-
-							// #note: we only notify on the first contact, though we will still accumulate the impulse data from subsequent contacts
-							const FVector NormalImpulse = FVector::DotProduct(CollisionDataItem.AccumulatedImpulse, CollisionDataItem.Normal) * CollisionDataItem.Normal;	// project impulse along normal
-							const FVector FrictionImpulse = FVector(CollisionDataItem.AccumulatedImpulse) - NormalImpulse; // friction is component not along contact normal
-							NotifyInfo.RigidCollisionData.TotalNormalImpulse += NormalImpulse;
-							NotifyInfo.RigidCollisionData.TotalFrictionImpulse += FrictionImpulse;
-
-							// Get swapped velocity deltas
-							const FVector& DeltaVelocity1 = bSwapOrder
-								? CollisionDataItem.DeltaVelocity2
-								: CollisionDataItem.DeltaVelocity1;
-							const FVector& DeltaVelocity2 = bSwapOrder
-								? CollisionDataItem.DeltaVelocity1
-								: CollisionDataItem.DeltaVelocity2;
-
-							// Populate additional contact information for additional hits
-							FRigidBodyContactInfo& NewContact = NotifyInfo.RigidCollisionData.ContactInfos.AddZeroed_GetRef();
-							NewContact.ContactNormal = CollisionDataItem.Normal;
-							NewContact.ContactPosition = CollisionDataItem.Location;
-							NewContact.ContactPenetration = CollisionDataItem.PenetrationDepth;
-							NewContact.bContactProbe = CollisionDataItem.bProbe;
-							NotifyInfo.RigidCollisionData.bIsVelocityDeltaUnderThreshold =
-								DeltaVelocity1.IsNearlyZero(MinDeltaVelocityThreshold) &&
-								DeltaVelocity2.IsNearlyZero(MinDeltaVelocityThreshold);
-								
-							const Chaos::FChaosPhysicsMaterial* InternalMat1 = CollisionDataItem.Mat1.Get();
-							const Chaos::FChaosPhysicsMaterial* InternalMat2 = CollisionDataItem.Mat2.Get();
-								
-							NewContact.PhysMaterial[0] = InternalMat1 ? FPhysicsUserData::Get<UPhysicalMaterial>(InternalMat1->UserData) : nullptr;
-							NewContact.PhysMaterial[1] = InternalMat2 ? FPhysicsUserData::Get<UPhysicalMaterial>(InternalMat2->UserData) : nullptr;
-
-							if(bSwapOrder)
-							{
-								NewContact.SwapOrder();
-							}
-
-							if (bNewEntry)
-							{
-								// fill in legacy contact data
-								NotifyInfo.bCallEvent0 = true;
-								// if Comp1 wants this event too, it will get its own pending collision entry, so we leave it false
-
-								NotifyInfo.SolverTime = CollisionDataItem.SolverTime;
-									
-								NotifyInfo.Info0.SetFrom(GetBodyInstanceFromProxyAndShape(PhysicsProxy0, CollisionDataItem.ShapeIndex1), DeltaVelocity1);
-								NotifyInfo.Info1.SetFrom(GetBodyInstanceFromProxyAndShape(PhysicsProxy1, CollisionDataItem.ShapeIndex2), DeltaVelocity2);
-
-								// in some case ( like with geometry collections ) we don't have a body instance so the component part will null, we need to handle that 
-								if (NotifyInfo.Info0.Component == nullptr)
-								{
-									NotifyInfo.Info0.Component = GetOwningComponent<UPrimitiveComponent>(PhysicsProxy0);
-									NotifyInfo.Info0.Actor = (NotifyInfo.Info0.Component != nullptr) ? NotifyInfo.Info0.Component->GetOwner() : nullptr;
-								}
-								if (NotifyInfo.Info1.Component == nullptr)
-								{
-									NotifyInfo.Info1.Component = GetOwningComponent<UPrimitiveComponent>(PhysicsProxy1);
-									NotifyInfo.Info1.Actor = (NotifyInfo.Info1.Component != nullptr) ? NotifyInfo.Info1.Component->GetOwner() : nullptr;
-								}
-							}
+							HandleEachCollisionEvent(*CollisionIndices, PhysicsProxy0, CollisionData, MinDeltaVelocityThreshold);
 						}
 					}
 				}
diff --git a/Engine/Source/Runtime/Engine/Public/Physics/Experimental/PhysScene_Chaos.h b/Engine/Source/Runtime/Engine/Public/Physics/Experimental/PhysScene_Chaos.h
--- a/Engine/Source/Runtime/Engine/Public/Physics/Experimental/PhysScene_Chaos.h
+++ b/Engine/Source/Runtime/Engine/Public/Physics/Experimental/PhysScene_Chaos.h
@@ -4,6 +4,7 @@
 
 #include "CoreMinimal.h"
 #include "Tickable.h"
+#include "EventsData.h"
 #include "Physics/PhysScene.h"
 #include "GameFramework/Actor.h"
 #include "PhysicsPublic.h"
@@ -246,7 +247,7 @@ public:
 
 private:
 	UPROPERTY()
-	TArray<UPrimitiveComponent*> CollisionEventRegistrations;
+	TSet<UPrimitiveComponent*> CollisionEventRegistrations;
 
 	// contains the set of properties that uniquely identifies a reported collision
 	// Note that order matters, { Body0, Body1 } is not the same as { Body1, Body0 }
@@ -274,6 +275,7 @@ private:
 	TArray<FCollisionNotifyInfo> PendingCollisionNotifies;
 
 	// Chaos Event Handlers
+	void HandleEachCollisionEvent(const TArray<int32>& CollisionIndices, IPhysicsProxyBase* PhysicsProxy0, Chaos::FCollisionDataArray const& CollisionData, Chaos::FReal MinDeltaVelocityThreshold);
 	void HandleCollisionEvents(const Chaos::FCollisionEventData& CollisionData);
 
 	void DispatchPendingCollisionNotifies();
