Fix build as a result of merge differences between branches

diff --git a/Engine/Source/Runtime/Engine/Private/BoneContainer.cpp b/Engine/Source/Runtime/Engine/Private/BoneContainer.cpp
--- a/Engine/Source/Runtime/Engine/Private/BoneContainer.cpp
+++ b/Engine/Source/Runtime/Engine/Private/BoneContainer.cpp
@@ -95,7 +95,6 @@ void FBoneContainer::Reset()
 	AssetSkeletalMesh = nullptr;
 	AssetSkeleton = nullptr;
 	RefSkeleton = nullptr;
-	UIDToArrayIndexLUTValidCount = 0;
 	RefPoseOverride = nullptr;
 
 #if DO_CHECK
@@ -114,7 +113,8 @@ void FBoneContainer::Reset()
 	SkeletonToCompactPose.Empty(0);
 	CompactPoseParentBones.Empty(0);
 	VirtualBoneCompactPoseData.Empty(0);
-	UIDToArrayIndexLUT.Empty(0);
+	CurveFilter = UE::Anim::FCurveFilter();
+	CurveFlags = UE::Anim::FBulkCurveFlags();
 
 	// Container changed and is no longer valid, preserve but update our serial number
 	RegenerateSerialNumber();
