Update comments and ordering of gravity direction application for custom gravity direction support
#rnx

diff --git a/Engine/Source/Runtime/Engine/Classes/GameFramework/Character.h b/Engine/Source/Runtime/Engine/Classes/GameFramework/Character.h
--- a/Engine/Source/Runtime/Engine/Classes/GameFramework/Character.h
+++ b/Engine/Source/Runtime/Engine/Classes/GameFramework/Character.h
@@ -468,6 +468,7 @@ protected:
 	UPROPERTY()
 	bool bInBaseReplication;
 
+	/** Cached version of the replicated gravity direction before replication. Used to compare if the value was changed as a result of replication. */
 	FVector_NetQuantizeNormal PreNetReceivedGravityDirection;
 public:
 	UFUNCTION()
diff --git a/Engine/Source/Runtime/Engine/Classes/GameFramework/CharacterMovementComponent.h b/Engine/Source/Runtime/Engine/Classes/GameFramework/CharacterMovementComponent.h
--- a/Engine/Source/Runtime/Engine/Classes/GameFramework/CharacterMovementComponent.h
+++ b/Engine/Source/Runtime/Engine/Classes/GameFramework/CharacterMovementComponent.h
@@ -193,8 +193,8 @@ private:
 	 * and custom movement modes. Gravity direction remaps player input as being within the plane defined by the gravity
 	 * direction. Movement simulation values like velocity and acceleration are maintained in their existing world coordinate
 	 * space but are transformed internally as gravity relative (for instance moving forward up a vertical wall that gravity is
-	 * defined to be perpendicular to and jump "up" from that wall). By default the character's capsule will not be oriented
-	 * to align with the gravity direction.
+	 * defined to be perpendicular to and jump "up" from that wall). If ShouldRemainVertical() is true the character's capsule
+	 * will be oriented to align with the gravity direction.
 	 */
 	UPROPERTY(Category="Character Movement: Gravity", VisibleAnywhere, BlueprintReadOnly, meta = (AllowPrivateAccess = true))
 	FVector GravityDirection;
@@ -1952,7 +1952,7 @@ protected:
 	virtual float SlideAlongSurface(const FVector& Delta, float Time, const FVector& Normal, FHitResult& Hit, bool bHandleImpact) override;
 
 	/** Custom version that allows upwards slides when walking if the surface is walkable. */
-	virtual void TwoWallAdjust(FVector& Delta, const FHitResult& Hit, const FVector& OldHitNormal) const override;
+	virtual void TwoWallAdjust(FVector& WorldSpaceDelta, const FHitResult& Hit, const FVector& OldHitNormal) const override;
 
 	/**
 	 * Calculate slide vector along a surface.
diff --git a/Engine/Source/Runtime/Engine/Private/Components/CharacterMovementComponent.cpp b/Engine/Source/Runtime/Engine/Private/Components/CharacterMovementComponent.cpp
--- a/Engine/Source/Runtime/Engine/Private/Components/CharacterMovementComponent.cpp
+++ b/Engine/Source/Runtime/Engine/Private/Components/CharacterMovementComponent.cpp
@@ -1771,13 +1771,6 @@ void UCharacterMovementComponent::SimulatedTick(float DeltaSeconds)
 			CharacterOwner->RootMotionRepMoves.Reset();
 		}
 
-		// Update replicated movement mode.
-		if (bNetworkMovementModeChanged)
-		{
-			ApplyNetworkMovementMode(CharacterOwner->GetReplicatedMovementMode());
-			bNetworkMovementModeChanged = false;
-		}
-
 		// Update replicated gravity direction
 		if (bNetworkGravityDirectionChanged)
 		{
@@ -1785,6 +1778,13 @@ void UCharacterMovementComponent::SimulatedTick(float DeltaSeconds)
 			bNetworkGravityDirectionChanged = false;
 		}
 
+		// Update replicated movement mode.
+		if (bNetworkMovementModeChanged)
+		{
+			ApplyNetworkMovementMode(CharacterOwner->GetReplicatedMovementMode());
+			bNetworkMovementModeChanged = false;
+		}
+
 		// Perform movement
 		PerformMovement(DeltaSeconds);
 
@@ -1805,8 +1805,8 @@ void UCharacterMovementComponent::SimulatedTick(float DeltaSeconds)
 			CharacterOwner->RootMotionRepMoves.Empty();
 			CharacterOwner->OnRep_ReplicatedMovement();
 			CharacterOwner->OnRep_ReplicatedBasedMovement();
-			ApplyNetworkMovementMode(GetCharacterOwner()->GetReplicatedMovementMode());
 			SetGravityDirection(CharacterOwner->GetReplicatedGravityDirection());
+			ApplyNetworkMovementMode(GetCharacterOwner()->GetReplicatedMovementMode());
 		}
 
 		if (CharacterOwner->IsReplicatingMovement() && UpdatedComponent)
@@ -1821,13 +1821,6 @@ void UCharacterMovementComponent::SimulatedTick(float DeltaSeconds)
 				const FScopedPreventAttachedComponentMove PreventMeshMovement(bPreventMeshMovement ? Mesh : nullptr);
 				if (CharacterOwner->IsPlayingRootMotion())
 				{
-					// Update replicated movement mode.
-					if (bNetworkMovementModeChanged)
-					{
-						ApplyNetworkMovementMode(CharacterOwner->GetReplicatedMovementMode());
-						bNetworkMovementModeChanged = false;
-					}
-
 					// Update replicated gravity direction
 					if (bNetworkGravityDirectionChanged)
 					{
@@ -1835,6 +1828,13 @@ void UCharacterMovementComponent::SimulatedTick(float DeltaSeconds)
 						bNetworkGravityDirectionChanged = false;
 					}
 
+					// Update replicated movement mode.
+					if (bNetworkMovementModeChanged)
+					{
+						ApplyNetworkMovementMode(CharacterOwner->GetReplicatedMovementMode());
+						bNetworkMovementModeChanged = false;
+					}
+
 					PerformMovement(DeltaSeconds);
 				}
 				else
@@ -1899,13 +1899,6 @@ void UCharacterMovementComponent::SimulateRootMotion(float DeltaSeconds, const F
 		AnimRootMotionVelocity = CalcAnimRootMotionVelocity(WorldSpaceRootMotionTransform.GetTranslation(), DeltaSeconds, Velocity);
 		Velocity = ConstrainAnimRootMotionVelocity(AnimRootMotionVelocity, Velocity);
 
-		// Update replicated movement mode.
-		if (bNetworkMovementModeChanged)
-		{
-			ApplyNetworkMovementMode(CharacterOwner->GetReplicatedMovementMode());
-			bNetworkMovementModeChanged = false;
-		}
-
 		// Update replicated gravity direction
 		if (bNetworkGravityDirectionChanged)
 		{
@@ -1913,6 +1906,13 @@ void UCharacterMovementComponent::SimulateRootMotion(float DeltaSeconds, const F
 			bNetworkGravityDirectionChanged = false;
 		}
 
+		// Update replicated movement mode.
+		if (bNetworkMovementModeChanged)
+		{
+			ApplyNetworkMovementMode(CharacterOwner->GetReplicatedMovementMode());
+			bNetworkMovementModeChanged = false;
+		}
+
 		NumJumpApexAttempts = 0;
 		StartNewPhysics(DeltaSeconds, 0);
 		// fixme laurent - simulate movement seems to have step up issues? investigate as that would be cheaper to use.
@@ -3348,12 +3348,11 @@ float UCharacterMovementComponent::SlideAlongSurface(const FVector& Delta, float
 }
 
 
-void UCharacterMovementComponent::TwoWallAdjust(FVector& Delta, const FHitResult& Hit, const FVector& OldHitNormal) const
+void UCharacterMovementComponent::TwoWallAdjust(FVector& WorldSpaceDelta, const FHitResult& Hit, const FVector& OldHitNormal) const
 {
-	const FVector InDelta = RotateWorldToGravity(Delta);
-	Super::TwoWallAdjust(Delta, Hit, OldHitNormal);
+	Super::TwoWallAdjust(WorldSpaceDelta, Hit, OldHitNormal);
 
-	FVector GravityRelativeDelta = RotateWorldToGravity(Delta);
+	FVector GravityRelativeDelta = RotateWorldToGravity(WorldSpaceDelta);
 	if (IsMovingOnGround())
 	{
 		// Allow slides up walkable surfaces, but not unwalkable ones (treat those as vertical barriers).
@@ -3364,8 +3363,8 @@ void UCharacterMovementComponent::TwoWallAdjust(FVector& Delta, const FHitResult
 			{
 				// Maintain horizontal velocity
 				const float Time = (1.f - Hit.Time);
-				const FVector ScaledDelta = GravityRelativeDelta.GetSafeNormal() * InDelta.Size();
-				GravityRelativeDelta = FVector(InDelta.X, InDelta.Y, ScaledDelta.Z / GravityRelativeHitNormal.Z) * Time;
+				const FVector ScaledDelta = GravityRelativeDelta.GetSafeNormal() * GravityRelativeDelta.Size();
+				GravityRelativeDelta = FVector(GravityRelativeDelta.X, GravityRelativeDelta.Y, ScaledDelta.Z / GravityRelativeHitNormal.Z) * Time;
 
 				// Should never exceed MaxStepHeight in vertical component, so rescale if necessary.
 				// This should be rare (Hit.Normal.Z above would have been very small) but we'd rather lose horizontal velocity than go too high.
@@ -3390,7 +3389,7 @@ void UCharacterMovementComponent::TwoWallAdjust(FVector& Delta, const FHitResult
 		}
 	}
 
-	Delta = RotateGravityToWorld(GravityRelativeDelta);
+	WorldSpaceDelta = RotateGravityToWorld(GravityRelativeDelta);
 }
 
 
