Adds an editor only BP utility function to add a row to a given data table

diff --git a/Engine/Source/Runtime/Engine/Classes/Kismet/DataTableFunctionLibrary.h b/Engine/Source/Runtime/Engine/Classes/Kismet/DataTableFunctionLibrary.h
--- a/Engine/Source/Runtime/Engine/Classes/Kismet/DataTableFunctionLibrary.h
+++ b/Engine/Source/Runtime/Engine/Classes/Kismet/DataTableFunctionLibrary.h
@@ -139,5 +139,10 @@ class ENGINE_API UDataTableFunctionLibrary : public UBlueprintFunctionLibrary
 	 */
 	UFUNCTION(BlueprintCallable, Category = "Editor Scripting | DataTable", DisplayName = "Fill Data Table from JSON File")
 	static bool FillDataTableFromJSONFile(UDataTable* DataTable, const FString& JSONFilePath, UScriptStruct* ImportRowStruct = nullptr);
+
+	/** Add a row to a Data Table with the provided name and data. */
+    UFUNCTION(BlueprintCallable, CustomThunk, Category = "Editor Scripting | DataTable", meta=(AutoCreateRefTerm="RowName", CustomStructureParam="RowData"))
+	static void AddDataTableRow(UDataTable* const DataTable, const FName& RowName, const FTableRowBase& RowData);
+    DECLARE_FUNCTION(execAddDataTableRow);
 #endif //WITH_EDITOR
 };
diff --git a/Engine/Source/Runtime/Engine/Private/DataTableFunctionLibrary.cpp b/Engine/Source/Runtime/Engine/Private/DataTableFunctionLibrary.cpp
--- a/Engine/Source/Runtime/Engine/Private/DataTableFunctionLibrary.cpp
+++ b/Engine/Source/Runtime/Engine/Private/DataTableFunctionLibrary.cpp
@@ -197,4 +197,70 @@ bool UDataTableFunctionLibrary::FillDataTableFromJSONFile(UDataTable* DataTable,
 	ImportFactory->AutomatedImportSettings.ImportRowStruct = ImportRowStruct;
 	return ImportFactory->ReimportCSV(DataTable) == EReimportResult::Succeeded;
 }
+
+void UDataTableFunctionLibrary::AddDataTableRow(UDataTable* const DataTable, const FName& RowName, const FTableRowBase& RowData)
+{
+	if (!DataTable)
+	{
+		UE_LOG(LogDataTable, Error, TEXT("AddDataTableRow - The DataTable is invalid."));
+		return;
+	}
+
+	DataTable->AddRow(RowName, RowData);
+}
+
+DEFINE_FUNCTION(UDataTableFunctionLibrary::execAddDataTableRow)
+{
+	P_GET_OBJECT(UDataTable, DataTable);
+	P_GET_PROPERTY(FNameProperty, RowName);
+
+	Stack.MostRecentPropertyAddress = nullptr;
+	Stack.MostRecentPropertyContainer = nullptr;
+	Stack.StepCompiledIn<FStructProperty>(nullptr);
+
+	const FTableRowBase* const RowData = reinterpret_cast<FTableRowBase*>(Stack.MostRecentPropertyAddress);
+	const FStructProperty* const StructProp = CastField<FStructProperty>(Stack.MostRecentProperty);
+
+	P_FINISH;
+
+	if (!DataTable)
+	{
+		FBlueprintExceptionInfo ExceptionInfo(
+			EBlueprintExceptionType::AccessViolation,
+			NSLOCTEXT("GetDataTableRow", "MissingTableInput", "Failed to resolve the DataTable parameter.")
+		);
+		FBlueprintCoreDelegates::ThrowScriptException(P_THIS, Stack, ExceptionInfo);
+	}
+	else if (StructProp && RowData)
+	{
+		const UScriptStruct* const RowType = StructProp->Struct;
+		const UScriptStruct* const TableType = DataTable->GetRowStruct();
+
+		// If the row type is compatible with the table type ...
+		const bool bIsTableRow = RowType->IsChildOf(FTableRowBase::StaticStruct());
+		const bool bMatchesTableType = RowType == TableType;
+		if (bIsTableRow && (bMatchesTableType || (RowType->IsChildOf(TableType) && FStructUtils::TheSameLayout(RowType, TableType))))
+		{
+			P_NATIVE_BEGIN;
+			AddDataTableRow(DataTable, RowName, *RowData);
+			P_NATIVE_END;
+		}
+		else
+		{
+			FBlueprintExceptionInfo ExceptionInfo(
+				EBlueprintExceptionType::AccessViolation,
+				NSLOCTEXT("GetDataTableRow", "IncompatibleProperty", "The data table type is incompatible with the RowData parameter.")
+			);
+			FBlueprintCoreDelegates::ThrowScriptException(P_THIS, Stack, ExceptionInfo);
+		}
+	}
+	else
+	{
+		FBlueprintExceptionInfo ExceptionInfo(
+			EBlueprintExceptionType::AccessViolation,
+			NSLOCTEXT("GetDataTableRow", "MissingOutputProperty", "Failed to resolve the RowData parameter.")
+		);
+		FBlueprintCoreDelegates::ThrowScriptException(P_THIS, Stack, ExceptionInfo);
+	}
+}
 #endif //WITH_EDITOR
