Fix FMaterialCachedExpressionData::bHasVertexInterpolator not being set when the new HLSL generator is used

diff --git a/Engine/Source/Runtime/Engine/Private/Materials/MaterialHLSLTree.cpp b/Engine/Source/Runtime/Engine/Private/Materials/MaterialHLSLTree.cpp
--- a/Engine/Source/Runtime/Engine/Private/Materials/MaterialHLSLTree.cpp
+++ b/Engine/Source/Runtime/Engine/Private/Materials/MaterialHLSLTree.cpp
@@ -1394,6 +1394,7 @@ void FEmitData::AddInterpolator(const FExpression* Expression, const FRequestedT
 		InterpolatorIndex = VertexInterpolators.Emplace(Expression);
 	}
 
+	bool bAnyComponentRequested = false;
 	const Shader::FType LocalType = PreparedType.GetResultType();
 	FVertexInterpolator& Interpolator = VertexInterpolators[InterpolatorIndex];
 	Interpolator.RequestedType = FRequestedType(RequestedType, false);
@@ -1407,9 +1408,15 @@ void FEmitData::AddInterpolator(const FExpression* Expression, const FRequestedT
 			{
 				// Only request components that need 'Shader' evaluation
 				Interpolator.RequestedType.SetComponentRequest(ComponentIndex);
+				bAnyComponentRequested = true;
 			}
 		}
 	}
+
+	if (bAnyComponentRequested && CachedExpressionData)
+	{
+		CachedExpressionData->bHasVertexInterpolator = true;
+	}
 }
 
 void FEmitData::PrepareInterpolators(FEmitContext& Context, FEmitScope& Scope)
