Fix a static analysis warning

diff --git a/Engine/Source/Runtime/Engine/Private/HLSLTree/HLSLTreeEmit.cpp b/Engine/Source/Runtime/Engine/Private/HLSLTree/HLSLTreeEmit.cpp
--- a/Engine/Source/Runtime/Engine/Private/HLSLTree/HLSLTreeEmit.cpp
+++ b/Engine/Source/Runtime/Engine/Private/HLSLTree/HLSLTreeEmit.cpp
@@ -926,7 +926,7 @@ FEmitShaderExpression* FEmitContext::EmitPreshaderOrConstant(FEmitScope& Scope,
 	}
 
 	FXxHash64Builder Hasher;
-	Hasher.Update(&Type.StructType, sizeof(Type.StructType));
+	Hasher.Update(&Type.StructType, sizeof(Type.StructType)); //-V568
 	Hasher.Update(&Type.ObjectType, sizeof(Type.ObjectType));
 	Hasher.Update(&Type.ValueType, sizeof(Type.ValueType));
 	LocalPreshader.AppendHash(Hasher);
