LevelInstance: ResetLoaders on PreSaveWorld

diff --git a/Engine/Source/Runtime/Engine/Private/LevelInstance/LevelInstanceSubsystem.cpp b/Engine/Source/Runtime/Engine/Private/LevelInstance/LevelInstanceSubsystem.cpp
--- a/Engine/Source/Runtime/Engine/Private/LevelInstance/LevelInstanceSubsystem.cpp
+++ b/Engine/Source/Runtime/Engine/Private/LevelInstance/LevelInstanceSubsystem.cpp
@@ -85,6 +85,7 @@ void ULevelInstanceSubsystem::Initialize(FSubsystemCollectionBase& Collection)
 	{
 		ILevelInstanceEditorModule& EditorModule = FModuleManager::LoadModuleChecked<ILevelInstanceEditorModule>("LevelInstanceEditor");
 		FEditorDelegates::OnAssetsPreDelete.AddUObject(this, &ULevelInstanceSubsystem::OnAssetsPreDelete);
+		FEditorDelegates::PreSaveWorldWithContext.AddUObject(this, &ULevelInstanceSubsystem::OnPreSaveWorldWithContext);
 	}
 #endif
 }
@@ -93,6 +94,7 @@ void ULevelInstanceSubsystem::Deinitialize()
 {
 #if WITH_EDITOR
 	FEditorDelegates::OnAssetsPreDelete.RemoveAll(this);
+	FEditorDelegates::PreSaveWorldWithContext.RemoveAll(this);
 #endif
 }
 
@@ -353,6 +355,17 @@ void ULevelInstanceSubsystem::OnAssetsPreDelete(const TArray<UObject*>& Objects)
 	}
 }
 
+void ULevelInstanceSubsystem::OnPreSaveWorldWithContext(UWorld* InWorld, FObjectPreSaveContext ObjectSaveContext)
+{
+	if (!(ObjectSaveContext.GetSaveFlags() & SAVE_FromAutosave) && !ObjectSaveContext.IsProceduralSave())
+	{
+		if (const UPackage* WorldPackage = InWorld->GetPackage())
+		{
+			ResetLoadersForWorldAssetInternal(WorldPackage->GetName());
+		}
+	}
+}
+
 void ULevelInstanceSubsystem::RegisterLoadedLevelStreamingLevelInstanceEditor(ULevelStreamingLevelInstanceEditor* LevelStreaming)
 {
 	if (!bIsCreatingLevelInstance)
diff --git a/Engine/Source/Runtime/Engine/Public/LevelInstance/LevelInstanceSubsystem.h b/Engine/Source/Runtime/Engine/Public/LevelInstance/LevelInstanceSubsystem.h
--- a/Engine/Source/Runtime/Engine/Public/LevelInstance/LevelInstanceSubsystem.h
+++ b/Engine/Source/Runtime/Engine/Public/LevelInstance/LevelInstanceSubsystem.h
@@ -143,6 +143,7 @@ private:
 
 	void ResetLoadersForWorldAssetInternal(const FString& WorldAsset);
 	void OnAssetsPreDelete(const TArray<UObject*>& Objects);
+	void OnPreSaveWorldWithContext(UWorld* InWorld, FObjectPreSaveContext ObjectSaveContext);
 
 	void RegisterLoadedLevelStreamingLevelInstanceEditor(ULevelStreamingLevelInstanceEditor* LevelStreaming);
 
