Made skylight create from EnvLightMixer use realtime capture by default.

diff --git a/Engine/Source/Editor/EnvironmentLightingViewer/Private/SEnvironmentLightingViewer.cpp b/Engine/Source/Editor/EnvironmentLightingViewer/Private/SEnvironmentLightingViewer.cpp
--- a/Engine/Source/Editor/EnvironmentLightingViewer/Private/SEnvironmentLightingViewer.cpp
+++ b/Engine/Source/Editor/EnvironmentLightingViewer/Private/SEnvironmentLightingViewer.cpp
@@ -371,6 +371,7 @@ FReply SEnvironmentLightingViewer::OnButtonCreateSkyLight()
 	const FTransform Transform(FVector(0.0f, 0.0f, 0.0f));
 	ASkyLight* SkyLight = Cast<ASkyLight>(GEditor->AddActor(World->GetCurrentLevel(), ASkyLight::StaticClass(), Transform));
 	SkyLight->GetLightComponent()->SetMobility(EComponentMobility::Movable);
+	SkyLight->GetLightComponent()->SetRealTimeCaptureEnabled(true);
 
 	return FReply::Handled();
 }
diff --git a/Engine/Source/Runtime/Engine/Classes/Components/SkyLightComponent.h b/Engine/Source/Runtime/Engine/Classes/Components/SkyLightComponent.h
--- a/Engine/Source/Runtime/Engine/Classes/Components/SkyLightComponent.h
+++ b/Engine/Source/Runtime/Engine/Classes/Components/SkyLightComponent.h
@@ -294,6 +294,7 @@ public:
 	/** Whether sky occlusion is supported by current feature level */
 	bool IsOcclusionSupported() const;
 
+	void SetRealTimeCaptureEnabled(bool bNewRealTimeCaptureEnabled);
 	bool IsRealTimeCaptureEnabled() const;
 
 	/** 
diff --git a/Engine/Source/Runtime/Engine/Private/Components/SkyLightComponent.cpp b/Engine/Source/Runtime/Engine/Private/Components/SkyLightComponent.cpp
--- a/Engine/Source/Runtime/Engine/Private/Components/SkyLightComponent.cpp
+++ b/Engine/Source/Runtime/Engine/Private/Components/SkyLightComponent.cpp
@@ -1113,6 +1113,12 @@ bool USkyLightComponent::IsRealTimeCaptureEnabled() const
 	return bRealTimeCapture && (Mobility == EComponentMobility::Movable || Mobility == EComponentMobility::Stationary) && GSkylightRealTimeReflectionCapture >0 && !bIsMobile;
 }
 
+void USkyLightComponent::SetRealTimeCaptureEnabled(bool bNewRealTimeCaptureEnabled)
+{
+	bRealTimeCapture = bNewRealTimeCaptureEnabled;
+	MarkRenderStateDirty();
+}
+
 void USkyLightComponent::OnVisibilityChanged()
 {
 	Super::OnVisibilityChanged();
