World Partition: added more descriptive assert when updating an actor descriptor with bad data.

diff --git a/Engine/Source/Runtime/Engine/Private/WorldPartition/WorldPartitionActorDescUtils.cpp b/Engine/Source/Runtime/Engine/Private/WorldPartition/WorldPartitionActorDescUtils.cpp
--- a/Engine/Source/Runtime/Engine/Private/WorldPartition/WorldPartitionActorDescUtils.cpp
+++ b/Engine/Source/Runtime/Engine/Private/WorldPartition/WorldPartitionActorDescUtils.cpp
@@ -158,9 +158,17 @@ void FWorldPartitionActorDescUtils::UpdateActorDescriptorFromActorDescriptor(TUn
 
 void FWorldPartitionActorDescUtils::ReplaceActorDescriptorPointerFromActor(const AActor* InOldActor, AActor* InNewActor, FWorldPartitionActorDesc* InActorDesc)
 {
-	check(!InNewActor || (InOldActor->GetActorGuid() == InNewActor->GetActorGuid()));
-	check(!InNewActor || (InNewActor->GetActorGuid() == InActorDesc->GetGuid()));
-	check(!InActorDesc->ActorPtr.IsValid() || (InActorDesc->ActorPtr == InOldActor));
+	if (InNewActor)
+	{
+		checkf(InOldActor->GetActorGuid() == InNewActor->GetActorGuid(), TEXT("Mismatching new actor GUID: old=%s new=%s"), *InOldActor->GetActorGuid().ToString(), *InNewActor->GetActorGuid().ToString());
+		checkf(InNewActor->GetActorGuid() == InActorDesc->GetGuid(), TEXT("Mismatching desc actor GUID: desc=%s new=%s"), *InActorDesc->GetGuid().ToString(), *InNewActor->GetActorGuid().ToString());
+	}
+
+	if (InActorDesc->ActorPtr.IsValid())
+	{
+		checkf(InActorDesc->ActorPtr == InOldActor, TEXT("Mismatching old desc actor: desc=%s old=%s"), *InActorDesc->ActorPtr->GetActorNameOrLabel(), *InOldActor->GetActorNameOrLabel());
+	}
+
 	InActorDesc->ActorPtr = InNewActor;
 }
 
@@ -215,5 +223,4 @@ bool FWorldPartitionActorDescUtils::FixupRedirectedAssetPath(FName& InOutAssetPa
 	InOutAssetPath = FName(*SoftObjectPath.ToString());
 	return true;
 }
-
 #endif
