Moving bCanEverAffectNavigation property into the 'Navigation' category with the others navigation properties.

diff --git a/Engine/Source/Runtime/Engine/Classes/Components/ActorComponent.h b/Engine/Source/Runtime/Engine/Classes/Components/ActorComponent.h
--- a/Engine/Source/Runtime/Engine/Classes/Components/ActorComponent.h
+++ b/Engine/Source/Runtime/Engine/Classes/Components/ActorComponent.h
@@ -247,7 +247,7 @@ public:
 
 protected:
 	/** Whether this component can potentially influence navigation */
-	UPROPERTY(EditAnywhere, Category = Collision, AdvancedDisplay, config)
+	UPROPERTY(EditAnywhere, Category = Navigation, AdvancedDisplay, config)
 	uint8 bCanEverAffectNavigation : 1;
 
 public:
