World Partition Actor Descriptors: save if the actor bounds is valid or not to handle the case where an actor returns an invalid streaming bounds and is considered at the origin.

diff --git a/Engine/Source/Editor/SceneOutliner/Private/ActorDescTreeItem.cpp b/Engine/Source/Editor/SceneOutliner/Private/ActorDescTreeItem.cpp
--- a/Engine/Source/Editor/SceneOutliner/Private/ActorDescTreeItem.cpp
+++ b/Engine/Source/Editor/SceneOutliner/Private/ActorDescTreeItem.cpp
@@ -287,7 +287,11 @@ void FActorDescTreeItem::FocusActorBounds() const
 	if (FWorldPartitionActorDesc const* ActorDesc = ActorDescHandle.Get())
 	{
 		const bool bActiveViewportOnly = true;
-		GEditor->MoveViewportCamerasToBox(ActorDesc->GetEditorBounds(), bActiveViewportOnly, 0.5f);
+		const FBox EditorBounds = ActorDesc->GetEditorBounds();
+		if (EditorBounds.IsValid)
+		{
+			GEditor->MoveViewportCamerasToBox(EditorBounds, bActiveViewportOnly, 0.5f);
+		}
 	}
 }
 
diff --git a/Engine/Source/Editor/SourceControlWindowExtender/Private/SourceControlWindowExtenderModule.cpp b/Engine/Source/Editor/SourceControlWindowExtender/Private/SourceControlWindowExtenderModule.cpp
--- a/Engine/Source/Editor/SourceControlWindowExtender/Private/SourceControlWindowExtenderModule.cpp
+++ b/Engine/Source/Editor/SourceControlWindowExtender/Private/SourceControlWindowExtenderModule.cpp
@@ -259,7 +259,11 @@ void FSourceControlWindowExtenderModule::FocusActors(const TArray<FAssetData>& A
 	{
 		if (TUniquePtr<FWorldPartitionActorDesc> ActorDesc = FWorldPartitionActorDescUtils::GetActorDescriptorFromAssetData(ActorToFocus))
 		{
-			FocusBounds += ActorDesc->GetEditorBounds();
+			const FBox EditorBounds = ActorDesc->GetEditorBounds();
+			if (EditorBounds.IsValid)
+			{
+				FocusBounds += EditorBounds;
+			}
 		}
 	}
 
diff --git a/Engine/Source/Runtime/Core/Public/UObject/FortniteMainBranchObjectVersion.h b/Engine/Source/Runtime/Core/Public/UObject/FortniteMainBranchObjectVersion.h
--- a/Engine/Source/Runtime/Core/Public/UObject/FortniteMainBranchObjectVersion.h
+++ b/Engine/Source/Runtime/Core/Public/UObject/FortniteMainBranchObjectVersion.h
@@ -301,6 +301,9 @@ struct CORE_API FFortniteMainBranchObjectVersion
 		// Moved the water info mesh data and static water body meshes into new static mesh components for water bodies.
 		WaterBodyStaticMeshComponents,
 
+		// Serialize invalid bounds in world partition actor descriptors
+		WorldPartitionActorDescSerializeInvalidBounds,
+
 		// -----<new versions can be added above this line>-------------------------------------------------
 		VersionPlusOne,
 		LatestVersion = VersionPlusOne - 1
diff --git a/Engine/Source/Runtime/Engine/Private/WorldPartition/RuntimeSpatialHash/RuntimeSpatialHashHLOD.cpp b/Engine/Source/Runtime/Engine/Private/WorldPartition/RuntimeSpatialHash/RuntimeSpatialHashHLOD.cpp
--- a/Engine/Source/Runtime/Engine/Private/WorldPartition/RuntimeSpatialHash/RuntimeSpatialHashHLOD.cpp
+++ b/Engine/Source/Runtime/Engine/Private/WorldPartition/RuntimeSpatialHash/RuntimeSpatialHashHLOD.cpp
@@ -196,12 +196,10 @@ static TArray<FGuid> GenerateHLODsForGrid(UWorldPartition* WorldPartition, const
 					for (const FGuid& ActorGuid : ActorSetInstance->ActorSet->Actors)
 					{
 						const IStreamingGenerationContext::FActorInstance& ActorInstance = ActorInstances.Emplace_GetRef(ActorGuid, ActorSetInstance);
-
-						const double ActorMinZ = ActorInstance.GetActorDescView().GetRuntimeBounds().Min.Z;
-						const double ActorMaxZ = ActorInstance.GetActorDescView().GetRuntimeBounds().Max.Z;
-
-						CellBounds.Min.Z = FMath::Min(CellBounds.Min.Z, ActorMinZ);
-						CellBounds.Max.Z = FMath::Max(CellBounds.Max.Z, ActorMaxZ);
+						const FBox RuntimeBounds = ActorInstance.GetActorDescView().GetRuntimeBounds();
+						check(RuntimeBounds.IsValid);
+						CellBounds.Min.Z = FMath::Min(CellBounds.Min.Z, RuntimeBounds.Min.Z);
+						CellBounds.Max.Z = FMath::Max(CellBounds.Max.Z, RuntimeBounds.Max.Z);
 					}
 				}
 
diff --git a/Engine/Source/Runtime/Engine/Private/WorldPartition/WorldPartition.cpp b/Engine/Source/Runtime/Engine/Private/WorldPartition/WorldPartition.cpp
--- a/Engine/Source/Runtime/Engine/Private/WorldPartition/WorldPartition.cpp
+++ b/Engine/Source/Runtime/Engine/Private/WorldPartition/WorldPartition.cpp
@@ -1434,13 +1434,17 @@ void UWorldPartition::Tick(float DeltaSeconds)
 			{
 				if (ActorDescIterator->GetIsSpatiallyLoadedRaw() || ActorDescIterator->GetActorNativeClass()->IsChildOf<ALandscapeProxy>())
 				{
-					AllActorsBounds += ActorDescIterator->GetEditorBounds();
-
-					// Warn the user if the world becomes larger that 4km in any axis
-					if (AllActorsBounds.GetSize().GetMax() >= UWorldPartition::WorldExtentToEnableStreaming)
+					const FBox EditorBounds = ActorDescIterator->GetEditorBounds();
+					if (EditorBounds.IsValid)
 					{
-						bEnablingStreamingJustified = true;
-						break;
+						AllActorsBounds += EditorBounds;
+
+						// Warn the user if the world becomes larger that 4km in any axis
+						if (AllActorsBounds.GetSize().GetMax() >= UWorldPartition::WorldExtentToEnableStreaming)
+						{
+							bEnablingStreamingJustified = true;
+							break;
+						}
 					}
 				}
 			}
diff --git a/Engine/Source/Runtime/Engine/Private/WorldPartition/WorldPartitionActorDesc.cpp b/Engine/Source/Runtime/Engine/Private/WorldPartition/WorldPartitionActorDesc.cpp
--- a/Engine/Source/Runtime/Engine/Private/WorldPartition/WorldPartitionActorDesc.cpp
+++ b/Engine/Source/Runtime/Engine/Private/WorldPartition/WorldPartitionActorDesc.cpp
@@ -34,6 +34,7 @@ TMap<TSubclassOf<AActor>, FWorldPartitionActorDesc::FActorDescDeprecator> FWorld
 
 FWorldPartitionActorDesc::FWorldPartitionActorDesc()
 	: bIsUsingDataLayerAsset(false)
+	, bIsBoundsValid(false)
 	, SoftRefCount(0)
 	, HardRefCount(0)
 	, Container(nullptr)
@@ -83,8 +84,9 @@ void FWorldPartitionActorDesc::Init(const AActor* InActor)
 
 	check(Guid.IsValid());
 
-	const FBox StreamingBounds = InActor->GetStreamingBounds();
+	const FBox StreamingBounds = !bIsDefaultActorDesc ? InActor->GetStreamingBounds() : FBox(ForceInit);
 	StreamingBounds.GetCenterAndExtents(BoundsLocation, BoundsExtent);
+	bIsBoundsValid = StreamingBounds.IsValid == 1;
 
 	RuntimeGrid = InActor->GetRuntimeGrid();
 	bIsSpatiallyLoaded = InActor->GetIsSpatiallyLoaded();
@@ -241,6 +243,7 @@ bool FWorldPartitionActorDesc::Equals(const FWorldPartitionActorDesc* Other) con
 		ActorPackage == Other->ActorPackage &&
 		ActorPath == Other->ActorPath &&
 		ActorLabel == Other->ActorLabel &&
+		bIsBoundsValid == Other->bIsBoundsValid &&
 		BoundsLocation.Equals(Other->BoundsLocation, 0.1f) &&
 		BoundsExtent.Equals(Other->BoundsExtent, 0.1f) &&
 		RuntimeGrid == Other->RuntimeGrid &&
@@ -266,21 +269,11 @@ bool FWorldPartitionActorDesc::ShouldResave(const FWorldPartitionActorDesc* Othe
 	check(ActorPackage == Other->ActorPackage);
 	check(ActorPath == Other->ActorPath);
 
-	// Tolerate up to 5% for bounds change
-	const float BoundsChangeTolerance = 0.05f;
-	const FBox ThisBounds = GetRuntimeBounds();
-	const FBox OtherBounds = Other->GetRuntimeBounds();
-	const FVector BoundsMinChangeRatio = ThisBounds.Min / OtherBounds.Min;
-	const FVector BoundsMaxChangeRatio = ThisBounds.Max / OtherBounds.Max;
-	if (!BoundsMinChangeRatio.Equals(FVector::OneVector, BoundsChangeTolerance) || !BoundsMaxChangeRatio.Equals(FVector::OneVector, BoundsChangeTolerance))
-	{
-		return true;
-	}
-
 	if (RuntimeGrid != Other->RuntimeGrid ||
 		bIsSpatiallyLoaded != Other->bIsSpatiallyLoaded ||
 		bActorIsEditorOnly != Other->bActorIsEditorOnly ||
 		bActorIsRuntimeOnly != Other->bActorIsRuntimeOnly ||
+		bIsBoundsValid != Other->bIsBoundsValid ||
 		HLODLayer != Other->HLODLayer ||
 		ParentActor != Other->ParentActor ||
 		ContentBundleGuid != Other->ContentBundleGuid ||
@@ -290,6 +283,22 @@ bool FWorldPartitionActorDesc::ShouldResave(const FWorldPartitionActorDesc* Othe
 		return true;
 	}
 
+	// Tolerate up to 5% for bounds change
+	if (bIsBoundsValid)
+	{
+		const FBox ThisBounds = GetRuntimeBounds();
+		const FBox OtherBounds = Other->GetRuntimeBounds();
+		const FVector BoundsChangeTolerance = ThisBounds.GetSize() * 0.05f;
+		const FVector MinDiff = FVector(OtherBounds.Min - ThisBounds.Min).GetAbs();
+		const FVector MaxDiff = FVector(OtherBounds.Max - ThisBounds.Max).GetAbs();
+		if ((MinDiff.X > BoundsChangeTolerance.X) || (MaxDiff.X > BoundsChangeTolerance.X) ||
+			(MinDiff.Y > BoundsChangeTolerance.Y) || (MaxDiff.Y > BoundsChangeTolerance.Y) ||
+			(MinDiff.Z > BoundsChangeTolerance.Z) || (MaxDiff.Z > BoundsChangeTolerance.Z))
+		{
+			return true;
+		}
+	}
+
 	// If the actor descriptor says the actor is HLOD relebant but in reality it's not, this will incur a loading time penalty during HLOD generation
 	// but will not affect the final result, as the value from the loaded actor will be used instead, so don't consider this as affecting streaming generation.
 	return (bActorIsHLODRelevant == Other->bActorIsHLODRelevant) || (bActorIsHLODRelevant && !Other->bActorIsHLODRelevant);
@@ -357,15 +366,23 @@ FString FWorldPartitionActorDesc::ToString(EToStringMode Mode) const
 	if (Mode >= EToStringMode::Compact)
 	{
 		FString BoundsStr;
-		const FBox EditorBounds = GetEditorBounds();
-		const FBox RuntimeBounds = GetRuntimeBounds();
-		if (EditorBounds.Equals(RuntimeBounds))
+
+		if (bIsBoundsValid)
 		{
-			BoundsStr = RuntimeBounds.ToString();
+			const FBox EditorBounds = GetEditorBounds();
+			const FBox RuntimeBounds = GetRuntimeBounds();
+			if (EditorBounds.Equals(RuntimeBounds))
+			{
+				BoundsStr = RuntimeBounds.ToString();
+			}
+			else
+			{
+				BoundsStr = *FString::Printf(TEXT("(Editor:%s Runtime:%s)"), *EditorBounds.ToString(), *RuntimeBounds.ToString());
+			}
 		}
 		else
 		{
-			BoundsStr = *FString::Printf(TEXT("(Editor:%s Runtime:%s)"), *EditorBounds.ToString(), *RuntimeBounds.ToString());
+			BoundsStr = TEXT("Invalid");
 		}
 
 		Result.Appendf(
@@ -446,10 +463,23 @@ void FWorldPartitionActorDesc::Serialize(FArchive& Ar)
 		Ar << BoundsLocationFlt << BoundsExtentFlt;
 		BoundsLocation = FVector(BoundsLocationFlt);
 		BoundsExtent = FVector(BoundsExtentFlt);
+		bIsBoundsValid = true;
 	}
 	else if (!bIsDefaultActorDesc)
 	{
-		Ar << BoundsLocation << BoundsExtent;
+		if (Ar.CustomVer(FFortniteMainBranchObjectVersion::GUID) < FFortniteMainBranchObjectVersion::WorldPartitionActorDescSerializeInvalidBounds)
+		{
+			bIsBoundsValid = true;
+		}
+		else
+		{
+			Ar << bIsBoundsValid;
+		}
+
+		if (bIsBoundsValid)
+		{
+			Ar << BoundsLocation << BoundsExtent;
+		}
 	}
 	
 	if (Ar.CustomVer(FUE5ReleaseStreamObjectVersion::GUID) < FUE5ReleaseStreamObjectVersion::ConvertedActorGridPlacementToSpatiallyLoadedFlag)
@@ -595,12 +625,12 @@ PRAGMA_ENABLE_DEPRECATION_WARNINGS
 
 FBox FWorldPartitionActorDesc::GetEditorBounds() const
 {
-	return FBox(BoundsLocation - BoundsExtent, BoundsLocation + BoundsExtent);
+	return bIsBoundsValid ? FBox(BoundsLocation - BoundsExtent, BoundsLocation + BoundsExtent) : FBox(ForceInit);
 }
 
 FBox FWorldPartitionActorDesc::GetRuntimeBounds() const
 {
-	return FBox(BoundsLocation - BoundsExtent, BoundsLocation + BoundsExtent);
+	return bIsBoundsValid ? FBox(BoundsLocation - BoundsExtent, BoundsLocation + BoundsExtent) : FBox(ForceInit);
 }
 
 FName FWorldPartitionActorDesc::GetActorName() const
diff --git a/Engine/Source/Runtime/Engine/Private/WorldPartition/WorldPartitionRuntimeHash.cpp b/Engine/Source/Runtime/Engine/Private/WorldPartition/WorldPartitionRuntimeHash.cpp
--- a/Engine/Source/Runtime/Engine/Private/WorldPartition/WorldPartitionRuntimeHash.cpp
+++ b/Engine/Source/Runtime/Engine/Private/WorldPartition/WorldPartitionRuntimeHash.cpp
@@ -216,7 +216,11 @@ void UWorldPartitionRuntimeHash::PopulateRuntimeCell(UWorldPartitionRuntimeCell*
 	{
 		const FWorldPartitionActorDescView& ActorDescView = ActorInstance.GetActorDescView();
 		RuntimeCell->AddActorToCell(ActorDescView, ActorInstance.GetContainerID(), ActorInstance.GetTransform(), ActorInstance.GetActorDescContainer());
-		CellContentBounds += ActorDescView.GetRuntimeBounds().TransformBy(ActorInstance.GetTransform());
+		const FBox RuntimeBounds = ActorDescView.GetRuntimeBounds();
+		if (RuntimeBounds.IsValid)
+		{
+			CellContentBounds += RuntimeBounds.TransformBy(ActorInstance.GetTransform());
+		}
 					
 		if (ActorInstance.GetContainerID().IsMainContainer() && RuntimeCell->UnsavedActorsContainer)
 		{
diff --git a/Engine/Source/Runtime/Engine/Private/WorldPartition/WorldPartitionStreamingGeneration.cpp b/Engine/Source/Runtime/Engine/Private/WorldPartition/WorldPartitionStreamingGeneration.cpp
--- a/Engine/Source/Runtime/Engine/Private/WorldPartition/WorldPartitionStreamingGeneration.cpp
+++ b/Engine/Source/Runtime/Engine/Private/WorldPartition/WorldPartitionStreamingGeneration.cpp
@@ -191,7 +191,11 @@ class FWorldPartitionStreamingGenerator
 					for (const FGuid& ActorGuid : ActorSetInstance.ActorSet->Actors)
 					{
 						const FWorldPartitionActorDescView& ActorDescView = ActorSetContainer.ActorDescViewMap->FindByGuidChecked(ActorGuid);
-						ActorSetInstance.Bounds += ActorDescView.GetRuntimeBounds().TransformBy(ContainerInstanceDescriptor.Transform);
+						const FBox RuntimeBounds = ActorDescView.GetRuntimeBounds();
+						if (RuntimeBounds.IsValid)
+						{
+							ActorSetInstance.Bounds += RuntimeBounds.TransformBy(ContainerInstanceDescriptor.Transform);
+						}
 					}
 				}
 			}
@@ -452,7 +456,10 @@ class FWorldPartitionStreamingGenerator
 			{
 				if (ActorDescView.GetIsSpatiallyLoaded())
 				{
-					ContainerInstanceDescriptor.Bounds += ActorDescView.GetRuntimeBounds().TransformBy(ContainerInstanceDescriptor.Transform);
+					const FBox RuntimeBounds = ActorDescView.GetRuntimeBounds();
+					check(RuntimeBounds.IsValid);
+
+					ContainerInstanceDescriptor.Bounds += RuntimeBounds.TransformBy(ContainerInstanceDescriptor.Transform);
 				}
 			});
 		}
diff --git a/Engine/Source/Runtime/Engine/Private/WorldPartition/WorldPartitionSubsystem.cpp b/Engine/Source/Runtime/Engine/Private/WorldPartition/WorldPartitionSubsystem.cpp
--- a/Engine/Source/Runtime/Engine/Private/WorldPartition/WorldPartitionSubsystem.cpp
+++ b/Engine/Source/Runtime/Engine/Private/WorldPartition/WorldPartitionSubsystem.cpp
@@ -336,7 +336,12 @@ void UWorldPartitionSubsystem::FActorDescContainerInstanceManager::FActorDescCon
 		{
 			continue;
 		}
-		Bounds += ActorDescIt->GetRuntimeBounds();
+
+		const FBox RuntimeBounds = ActorDescIt->GetRuntimeBounds();
+		if (RuntimeBounds.IsValid)
+		{
+			Bounds += RuntimeBounds;
+		}
 	}
 }
 
diff --git a/Engine/Source/Runtime/Engine/Public/WorldPartition/WorldPartitionActorDesc.h b/Engine/Source/Runtime/Engine/Public/WorldPartition/WorldPartitionActorDesc.h
--- a/Engine/Source/Runtime/Engine/Public/WorldPartition/WorldPartitionActorDesc.h
+++ b/Engine/Source/Runtime/Engine/Public/WorldPartition/WorldPartitionActorDesc.h
@@ -103,7 +103,7 @@ PRAGMA_DISABLE_DEPRECATION_WARNINGS
 PRAGMA_ENABLE_DEPRECATION_WARNINGS
 
 	inline FName GetRuntimeGrid() const { return RuntimeGrid; }
-	inline bool GetIsSpatiallyLoaded() const { return bIsForcedNonSpatiallyLoaded ? false : bIsSpatiallyLoaded; }
+	inline bool GetIsSpatiallyLoaded() const { return (bIsForcedNonSpatiallyLoaded || !bIsBoundsValid) ? false : bIsSpatiallyLoaded; }
 	inline bool GetIsSpatiallyLoadedRaw() const { return bIsSpatiallyLoaded; }
 	inline bool GetActorIsEditorOnly() const { return bActorIsEditorOnly; }
 	inline bool GetActorIsRuntimeOnly() const { return bActorIsRuntimeOnly; }
@@ -291,6 +291,7 @@ protected:
 	bool							bActorIsRuntimeOnly;
 	bool							bActorIsHLODRelevant;
 	bool							bIsUsingDataLayerAsset; // Used to know if DataLayers array represents DataLayers Asset paths or the FNames of the deprecated version of Data Layers
+	bool							bIsBoundsValid;
 	FName							HLODLayer;
 	TArray<FName>					DataLayers;
 	TArray<FGuid>					References;
