Modified check in FUniqueNetIdRepl::ImportTextItem to prevent cooking errors in some projects.

diff --git a/Engine/Source/Runtime/Engine/Private/OnlineReplStructs.cpp b/Engine/Source/Runtime/Engine/Private/OnlineReplStructs.cpp
--- a/Engine/Source/Runtime/Engine/Private/OnlineReplStructs.cpp
+++ b/Engine/Source/Runtime/Engine/Private/OnlineReplStructs.cpp
@@ -548,7 +548,7 @@ bool FUniqueNetIdRepl::ImportTextItem(const TCHAR*& Buffer, int32 PortFlags, UOb
 		}
 		else
 		{
-			checkf(UOnlineEngineInterface::Get() && UOnlineEngineInterface::Get()->IsLoaded(), TEXT("Attempted to ImportText to FUniqueNetIdRepl while OSS is not loaded. Parent:%s"), *GetPathNameSafe(Parent));
+			check(UOnlineEngineInterface::Get());
 
 			FString Token;
 			if (const TCHAR* NewBuffer1 = FPropertyHelpers::ReadToken(Buffer, Token))
