Deduplicated primitive flags to a shared header & also deduplicated  shader paramter construction to shared function in the proxy.
- Also moved FPrimitiveUniformShaderParametersBuilder to its own header to cut down rebuilding when changing scene & nanite definitions.

diff --git a/Engine/Plugins/Experimental/GPULightmass/Source/GPULightmass/Private/LightmapRenderer.cpp b/Engine/Plugins/Experimental/GPULightmass/Source/GPULightmass/Private/LightmapRenderer.cpp
--- a/Engine/Plugins/Experimental/GPULightmass/Source/GPULightmass/Private/LightmapRenderer.cpp
+++ b/Engine/Plugins/Experimental/GPULightmass/Source/GPULightmass/Private/LightmapRenderer.cpp
@@ -33,6 +33,7 @@
 #include "PathTracing.h"
 #include "PostProcess/DrawRectangle.h"
 #include "SceneUniformBuffer.h"
+#include "PrimitiveUniformShaderParametersBuilder.h"
 
 RENDERER_API uint8 BlendModeToRayTracingInstanceMask(const EBlendMode BlendMode, ERayTracingViewMaskMode MaskMode);
 RENDERER_API uint8 ComputeRayTracingInstanceShadowMask(ERayTracingViewMaskMode MaskMode);
diff --git a/Engine/Plugins/FX/Niagara/Source/Niagara/Private/NiagaraComponent.cpp b/Engine/Plugins/FX/Niagara/Source/Niagara/Private/NiagaraComponent.cpp
--- a/Engine/Plugins/FX/Niagara/Source/Niagara/Private/NiagaraComponent.cpp
+++ b/Engine/Plugins/FX/Niagara/Source/Niagara/Private/NiagaraComponent.cpp
@@ -28,6 +28,7 @@
 #include "NiagaraCullProxyComponent.h"
 #include "SceneInterface.h"
 #include "UObject/UObjectIterator.h"
+#include "PrimitiveUniformShaderParametersBuilder.h"
 
 #include UE_INLINE_GENERATED_CPP_BY_NAME(NiagaraComponent)
 
diff --git a/Engine/Shaders/Private/SceneData.ush b/Engine/Shaders/Private/SceneData.ush
--- a/Engine/Shaders/Private/SceneData.ush
+++ b/Engine/Shaders/Private/SceneData.ush
@@ -5,6 +5,7 @@
 #include "LargeWorldCoordinates.ush"
 #include "OctahedralCommon.ush"
 #include "/Engine/Shared/LightSceneData.h"
+#include "/Engine/Shared/SceneDefinitions.h"
 
 #ifndef USE_GLOBAL_GPU_SCENE_DATA
 	#define USE_GLOBAL_GPU_SCENE_DATA 0
@@ -59,40 +60,6 @@
 	#endif
 #endif
 
-// Must match PrimitiveUniformShaderParameters.h
-#define PRIMITIVE_SCENE_DATA_FLAG_CAST_SHADOWS							0x1
-#define PRIMITIVE_SCENE_DATA_FLAG_USE_SINGLE_SAMPLE_SHADOW_SL			0x2
-#define PRIMITIVE_SCENE_DATA_FLAG_USE_VOLUMETRIC_LM_SHADOW_SL			0x4
-#define PRIMITIVE_SCENE_DATA_FLAG_DECAL_RECEIVER						0x8
-#define PRIMITIVE_SCENE_DATA_FLAG_CACHE_SHADOW_AS_STATIC				0x10
-#define PRIMITIVE_SCENE_DATA_FLAG_OUTPUT_VELOCITY						0x20
-#define PRIMITIVE_SCENE_DATA_FLAG_DETERMINANT_SIGN						0x40
-#define PRIMITIVE_SCENE_DATA_FLAG_HAS_DISTANCE_FIELD_REPRESENTATION		0x80
-#define PRIMITIVE_SCENE_DATA_FLAG_HAS_CAPSULE_REPRESENTATION			0x100
-#define PRIMITIVE_SCENE_DATA_FLAG_HAS_CAST_CONTACT_SHADOW				0x200
-#define PRIMITIVE_SCENE_DATA_FLAG_HAS_PRIMITIVE_CUSTOM_DATA				0x400
-#define PRIMITIVE_SCENE_DATA_FLAG_LIGHTING_CHANNEL_0					0x800
-#define PRIMITIVE_SCENE_DATA_FLAG_LIGHTING_CHANNEL_1					0x1000
-#define PRIMITIVE_SCENE_DATA_FLAG_LIGHTING_CHANNEL_2					0x2000
-#define PRIMITIVE_SCENE_DATA_FLAG_HAS_INSTANCE_LOCAL_BOUNDS				0x4000
-#define PRIMITIVE_SCENE_DATA_FLAG_EVALUATE_WORLD_POSITION_OFFSET		0x8000
-#define PRIMITIVE_SCENE_DATA_FLAG_VISIBLE_IN_GAME						0x10000
-#define PRIMITIVE_SCENE_DATA_FLAG_VISIBLE_IN_EDITOR						0x20000
-#define PRIMITIVE_SCENE_DATA_FLAG_VISIBLE_IN_REFLECTION_CAPTURES		0x40000
-#define PRIMITIVE_SCENE_DATA_FLAG_VISIBLE_IN_REAL_TIME_SKY_CAPTURES		0x80000
-#define PRIMITIVE_SCENE_DATA_FLAG_VISIBLE_IN_RAY_TRACING				0x100000
-#define PRIMITIVE_SCENE_DATA_FLAG_VISIBLE_IN_SCENE_CAPTURE_ONLY			0x200000
-#define PRIMITIVE_SCENE_DATA_FLAG_HIDDEN_IN_SCENE_CAPTURE				0x400000
-#define PRIMITIVE_SCENE_DATA_FLAG_FORCE_HIDDEN							0x800000 // Mobile allocates only 24 bits for primitive flags
-#define PRIMITIVE_SCENE_DATA_FLAG_CAST_HIDDEN_SHADOW					0x1000000
-#define PRIMITIVE_SCENE_DATA_FLAG_HAS_NANITE_IMPOSTER					0x2000000
-#define PRIMITIVE_SCENE_DATA_FLAG_INSTANCE_DRAW_DISTANCE_CULL			0x4000000
-#define PRIMITIVE_SCENE_DATA_FLAG_WPO_DISABLE_DISTANCE					0x8000000
-#define PRIMITIVE_SCENE_DATA_FLAG_WRITES_CUSTOM_DEPTH_STENCIL			0x10000000
-#define PRIMITIVE_SCENE_DATA_FLAG_REVERSE_CULLING						0x20000000
-#define PRIMITIVE_SCENE_DATA_FLAG_HOLDOUT								0x40000000
-#define PRIMITIVE_SCENE_DATA_FLAG_VISIBLE_IN_LUMEN_SCENE				0x80000000
-
 uint GetDefaultPrimitiveSceneDataFlags()
 {
 	// Should match FPrimitiveUniformShaderParametersBuilder::Default()
diff --git a/Engine/Shaders/Shared/SceneDefinitions.h b/Engine/Shaders/Shared/SceneDefinitions.h
new file mode 100644
--- /dev/null
+++ b/Engine/Shaders/Shared/SceneDefinitions.h
@@ -0,0 +1,40 @@
+// Copyright Epic Games, Inc. All Rights Reserved.
+
+/*================================================================================================
+!!!! Changing this file requires recompilation of the engine !!!
+=================================================================================================*/
+
+#pragma once
+
+#define PRIMITIVE_SCENE_DATA_FLAG_CAST_SHADOWS							0x1
+#define PRIMITIVE_SCENE_DATA_FLAG_USE_SINGLE_SAMPLE_SHADOW_SL			0x2
+#define PRIMITIVE_SCENE_DATA_FLAG_USE_VOLUMETRIC_LM_SHADOW_SL			0x4
+#define PRIMITIVE_SCENE_DATA_FLAG_DECAL_RECEIVER						0x8
+#define PRIMITIVE_SCENE_DATA_FLAG_CACHE_SHADOW_AS_STATIC				0x10
+#define PRIMITIVE_SCENE_DATA_FLAG_OUTPUT_VELOCITY						0x20
+#define PRIMITIVE_SCENE_DATA_FLAG_DETERMINANT_SIGN						0x40
+#define PRIMITIVE_SCENE_DATA_FLAG_HAS_DISTANCE_FIELD_REPRESENTATION		0x80
+#define PRIMITIVE_SCENE_DATA_FLAG_HAS_CAPSULE_REPRESENTATION			0x100
+#define PRIMITIVE_SCENE_DATA_FLAG_HAS_CAST_CONTACT_SHADOW				0x200
+#define PRIMITIVE_SCENE_DATA_FLAG_HAS_PRIMITIVE_CUSTOM_DATA				0x400
+#define PRIMITIVE_SCENE_DATA_FLAG_LIGHTING_CHANNEL_0					0x800
+#define PRIMITIVE_SCENE_DATA_FLAG_LIGHTING_CHANNEL_1					0x1000
+#define PRIMITIVE_SCENE_DATA_FLAG_LIGHTING_CHANNEL_2					0x2000
+#define PRIMITIVE_SCENE_DATA_FLAG_HAS_INSTANCE_LOCAL_BOUNDS				0x4000
+#define PRIMITIVE_SCENE_DATA_FLAG_EVALUATE_WORLD_POSITION_OFFSET		0x8000
+#define PRIMITIVE_SCENE_DATA_FLAG_VISIBLE_IN_GAME						0x10000
+#define PRIMITIVE_SCENE_DATA_FLAG_VISIBLE_IN_EDITOR						0x20000
+#define PRIMITIVE_SCENE_DATA_FLAG_VISIBLE_IN_REFLECTION_CAPTURES		0x40000
+#define PRIMITIVE_SCENE_DATA_FLAG_VISIBLE_IN_REAL_TIME_SKY_CAPTURES		0x80000
+#define PRIMITIVE_SCENE_DATA_FLAG_VISIBLE_IN_RAY_TRACING				0x100000
+#define PRIMITIVE_SCENE_DATA_FLAG_VISIBLE_IN_SCENE_CAPTURE_ONLY			0x200000
+#define PRIMITIVE_SCENE_DATA_FLAG_HIDDEN_IN_SCENE_CAPTURE				0x400000
+#define PRIMITIVE_SCENE_DATA_FLAG_FORCE_HIDDEN							0x800000 // Mobile allocates only 24 bits for primitive flags
+#define PRIMITIVE_SCENE_DATA_FLAG_CAST_HIDDEN_SHADOW					0x1000000
+#define PRIMITIVE_SCENE_DATA_FLAG_HAS_NANITE_IMPOSTER					0x2000000
+#define PRIMITIVE_SCENE_DATA_FLAG_INSTANCE_DRAW_DISTANCE_CULL			0x4000000
+#define PRIMITIVE_SCENE_DATA_FLAG_WPO_DISABLE_DISTANCE					0x8000000
+#define PRIMITIVE_SCENE_DATA_FLAG_WRITES_CUSTOM_DEPTH_STENCIL			0x10000000
+#define PRIMITIVE_SCENE_DATA_FLAG_REVERSE_CULLING						0x20000000
+#define PRIMITIVE_SCENE_DATA_FLAG_HOLDOUT								0x40000000
+#define PRIMITIVE_SCENE_DATA_FLAG_VISIBLE_IN_LUMEN_SCENE				0x80000000
diff --git a/Engine/Source/Developer/MaterialBaking/Private/MaterialRenderItem.cpp b/Engine/Source/Developer/MaterialBaking/Private/MaterialRenderItem.cpp
--- a/Engine/Source/Developer/MaterialBaking/Private/MaterialRenderItem.cpp
+++ b/Engine/Source/Developer/MaterialBaking/Private/MaterialRenderItem.cpp
@@ -10,6 +10,7 @@
 #include "CanvasRender.h"
 #include "RHIStaticStates.h"
 #include "UnrealClient.h"
+#include "PrimitiveUniformShaderParametersBuilder.h"
 
 #define SHOW_WIREFRAME_MESH 0
 
diff --git a/Engine/Source/Runtime/Engine/Private/DynamicMeshBuilder.cpp b/Engine/Source/Runtime/Engine/Private/DynamicMeshBuilder.cpp
--- a/Engine/Source/Runtime/Engine/Private/DynamicMeshBuilder.cpp
+++ b/Engine/Source/Runtime/Engine/Private/DynamicMeshBuilder.cpp
@@ -10,6 +10,7 @@
 #include "MeshBuilderOneFrameResources.h"
 #include "Math/Vector2DHalf.h"
 #include "ResourcePool.h"
+#include "PrimitiveUniformShaderParametersBuilder.h"
 
 class FGlobalDynamicMeshPoolPolicy
 {
diff --git a/Engine/Source/Runtime/Engine/Private/Particles/ParticleSystemRender.cpp b/Engine/Source/Runtime/Engine/Private/Particles/ParticleSystemRender.cpp
--- a/Engine/Source/Runtime/Engine/Private/Particles/ParticleSystemRender.cpp
+++ b/Engine/Source/Runtime/Engine/Private/Particles/ParticleSystemRender.cpp
@@ -34,6 +34,7 @@
 #include "UnrealEngine.h"
 #include "RenderCore.h"
 #include "DataDrivenShaderPlatformInfo.h"
+#include "PrimitiveUniformShaderParametersBuilder.h"
 
 DECLARE_CYCLE_STAT(TEXT("ParticleSystemSceneProxy Create GT"), STAT_FParticleSystemSceneProxy_Create, STATGROUP_Particles);
 DECLARE_CYCLE_STAT(TEXT("ParticleSystemSceneProxy GetMeshElements RT"), STAT_FParticleSystemSceneProxy_GetMeshElements, STATGROUP_Particles);
diff --git a/Engine/Source/Runtime/Engine/Private/PrimitiveSceneProxy.cpp b/Engine/Source/Runtime/Engine/Private/PrimitiveSceneProxy.cpp
--- a/Engine/Source/Runtime/Engine/Private/PrimitiveSceneProxy.cpp
+++ b/Engine/Source/Runtime/Engine/Private/PrimitiveSceneProxy.cpp
@@ -22,6 +22,7 @@
 #include "ComponentRecreateRenderStateContext.h"
 #include "DataDrivenShaderPlatformInfo.h"
 #include "SceneInterface.h"
+#include "PrimitiveUniformShaderParametersBuilder.h"
 
 #if WITH_EDITOR
 #include "FoliageHelper.h"
@@ -561,79 +562,128 @@ void FPrimitiveSceneProxy::UpdateUniformBuffer(FRHICommandList& RHICmdList)
 	// Skip expensive primitive uniform buffer creation for proxies whose vertex factories only use GPUScene for primitive data
 	if (DoesVFRequirePrimitiveUniformBuffer())
 	{
-		bool bHasPrecomputedVolumetricLightmap;
-		FMatrix PreviousLocalToWorld;
-		int32 SingleCaptureIndex;
-		bool bOutputVelocity;
+		// Update the uniform shader parameters.
+		FPrimitiveUniformShaderParametersBuilder Builder = FPrimitiveUniformShaderParametersBuilder{};
+		BuildUniformShaderParameters(Builder);
 
-		Scene->GetPrimitiveUniformShaderParameters_RenderThread(
-			PrimitiveSceneInfo,
-			bHasPrecomputedVolumetricLightmap,
-			PreviousLocalToWorld,
-			SingleCaptureIndex,
-			bOutputVelocity
-		);
+		FPrimitiveUniformShaderParameters PrimitiveParams = Builder.Build();
 
-		bOutputVelocity |= AlwaysHasVelocity();
+		check(UniformBuffer);
+		UniformBuffer.UpdateUniformBufferImmediate(RHICmdList, PrimitiveParams);
+	}
+}
 
-		FBoxSphereBounds PreSkinnedLocalBounds;
-		GetPreSkinnedLocalBounds(PreSkinnedLocalBounds);
+void FPrimitiveSceneProxy::BuildUniformShaderParameters(FPrimitiveUniformShaderParametersBuilder &Builder) const
+{
+	bool bHasPrecomputedVolumetricLightmap;
+	FMatrix PreviousLocalToWorld;
+	int32 SingleCaptureIndex;
+	bool bOutputVelocity;
+
+	Scene->GetPrimitiveUniformShaderParameters_RenderThread(
+		PrimitiveSceneInfo,
+		bHasPrecomputedVolumetricLightmap,
+		PreviousLocalToWorld,
+		SingleCaptureIndex,
+		bOutputVelocity
+	);
 
-		// Update the uniform shader parameters.
-		FPrimitiveUniformShaderParametersBuilder Builder = FPrimitiveUniformShaderParametersBuilder{}
-			.Defaults()
-				.LocalToWorld(LocalToWorld)
-				.PreviousLocalToWorld(PreviousLocalToWorld)
-				.ActorWorldPosition(ActorPosition)
-				.WorldBounds(Bounds)
-				.LocalBounds(LocalBounds)
-				.PreSkinnedLocalBounds(PreSkinnedLocalBounds)
-				.ReceivesDecals(bReceivesDecals)
-				.CacheShadowAsStatic(PrimitiveSceneInfo ? PrimitiveSceneInfo->ShouldCacheShadowAsStatic() : false)
-				.OutputVelocity(bOutputVelocity)
-				.EvaluateWorldPositionOffset(EvaluateWorldPositionOffset() && AnyMaterialHasWorldPositionOffset())
-				.MaxWorldPositionOffsetExtent(GetMaxWorldPositionOffsetExtent())
-				.MinMaxMaterialDisplacement(GetMinMaxMaterialDisplacement())
-				.LightingChannelMask(GetLightingChannelMask())
-				.LightmapDataIndex(PrimitiveSceneInfo ? PrimitiveSceneInfo->GetLightmapDataOffset() : 0)
-				.LightmapUVIndex(GetLightMapCoordinateIndex())
-				.SingleCaptureIndex(SingleCaptureIndex)
-				.CustomPrimitiveData(GetCustomPrimitiveData())
-				.HasDistanceFieldRepresentation(HasDistanceFieldRepresentation())
-				.HasCapsuleRepresentation(HasDynamicIndirectShadowCasterRepresentation())
-				.UseSingleSampleShadowFromStationaryLights(UseSingleSampleShadowFromStationaryLights())
-				.UseVolumetricLightmap(bHasPrecomputedVolumetricLightmap)
-				.CastContactShadow(CastsContactShadow())
-				.CastHiddenShadow(CastsHiddenShadow())
-				.CastShadow(CastsDynamicShadow())
-				.Holdout(Holdout())
-				.InstanceSceneDataOffset(PrimitiveSceneInfo ? PrimitiveSceneInfo->GetInstanceSceneDataOffset() : INDEX_NONE)
-				.NumInstanceSceneDataEntries(PrimitiveSceneInfo ? PrimitiveSceneInfo->GetNumInstanceSceneDataEntries() : 0)
-				.InstancePayloadDataOffset(PrimitiveSceneInfo ? PrimitiveSceneInfo->GetInstancePayloadDataOffset() : INDEX_NONE)
-				.InstancePayloadDataStride(PrimitiveSceneInfo ? PrimitiveSceneInfo->GetInstancePayloadDataStride() : 0);				
-
-		FVector2f InstanceDrawDistanceMinMax;
-		if (GetInstanceDrawDistanceMinMax(InstanceDrawDistanceMinMax))
+	bOutputVelocity |= AlwaysHasVelocity();
+
+	FBoxSphereBounds PreSkinnedLocalBounds;
+	GetPreSkinnedLocalBounds(PreSkinnedLocalBounds);
+
+	// Update the uniform shader parameters.
+	Builder = FPrimitiveUniformShaderParametersBuilder{}
+		.Defaults()
+			.LocalToWorld(LocalToWorld)
+			.PreviousLocalToWorld(PreviousLocalToWorld)
+			.ActorWorldPosition(ActorPosition)
+			.WorldBounds(Bounds)
+			.LocalBounds(LocalBounds)
+			.PreSkinnedLocalBounds(PreSkinnedLocalBounds)
+			.ReceivesDecals(bReceivesDecals)
+			.OutputVelocity(bOutputVelocity)
+			.EvaluateWorldPositionOffset(EvaluateWorldPositionOffset() && AnyMaterialHasWorldPositionOffset())
+			.MaxWorldPositionOffsetExtent(GetMaxWorldPositionOffsetExtent())
+			.MinMaxMaterialDisplacement(GetMinMaxMaterialDisplacement())
+			.LightingChannelMask(GetLightingChannelMask())
+			.LightmapUVIndex(GetLightMapCoordinateIndex())
+			.SingleCaptureIndex(SingleCaptureIndex)
+			.CustomPrimitiveData(GetCustomPrimitiveData())
+			.HasDistanceFieldRepresentation(HasDistanceFieldRepresentation())
+			.HasCapsuleRepresentation(HasDynamicIndirectShadowCasterRepresentation())
+			.UseSingleSampleShadowFromStationaryLights(UseSingleSampleShadowFromStationaryLights())
+			.UseVolumetricLightmap(bHasPrecomputedVolumetricLightmap)
+			.CastContactShadow(CastsContactShadow())
+			.CastHiddenShadow(CastsHiddenShadow())
+			.CastShadow(CastsDynamicShadow())
+			.Holdout(Holdout())
+			.VisibleInGame(IsDrawnInGame())
+			.VisibleInEditor(IsDrawnInEditor())
+			.VisibleInReflectionCaptures(IsVisibleInReflectionCaptures())
+			.VisibleInRealTimeSkyCaptures(IsVisibleInRealTimeSkyCaptures())
+			.VisibleInRayTracing(IsVisibleInRayTracing())
+			.VisibleInLumenScene(IsVisibleInLumenScene())
+			.VisibleInSceneCaptureOnly(IsVisibleInSceneCaptureOnly())
+			.HiddenInSceneCapture(IsHiddenInSceneCapture())
+			.ForceHidden(IsForceHidden())
+			.PrimitiveComponentId(GetPrimitiveComponentId().PrimIDValue)
+			.EditorColors(GetWireframeColor(), GetLevelColor());
+
+		if (PrimitiveSceneInfo != nullptr)
 		{
-			Builder.InstanceDrawDistance(InstanceDrawDistanceMinMax);
+			Builder.LightmapDataIndex(PrimitiveSceneInfo->GetLightmapDataOffset())
+			.CacheShadowAsStatic(PrimitiveSceneInfo->ShouldCacheShadowAsStatic())
+			.InstanceSceneDataOffset(PrimitiveSceneInfo->GetInstanceSceneDataOffset())
+			.NumInstanceSceneDataEntries(PrimitiveSceneInfo->GetNumInstanceSceneDataEntries())
+			.InstancePayloadDataOffset(PrimitiveSceneInfo->GetInstancePayloadDataOffset())
+			.InstancePayloadDataStride(PrimitiveSceneInfo->GetInstancePayloadDataStride())
+			.PersistentPrimitiveIndex(PrimitiveSceneInfo->GetPersistentIndex().Index);
 		}
 
-		float WPODisableDistance;
-		if (GetInstanceWorldPositionOffsetDisableDistance(WPODisableDistance))
-		{
-			Builder.InstanceWorldPositionOffsetDisableDistance(WPODisableDistance);
-		}
+	if (IsNaniteMesh())
+	{
+		const Nanite::FSceneProxyBase* NaniteProxy = static_cast<const Nanite::FSceneProxyBase*>(this);
 
-		const TConstArrayView<FRenderBounds> InstanceBounds = GetInstanceLocalBounds();
-		if (InstanceBounds.Num() > 0)
-		{
-			Builder.InstanceLocalBounds(InstanceBounds[0]);
-		}
+		uint32 NaniteResourceID = INDEX_NONE;
+		uint32 NaniteHierarchyOffset = INDEX_NONE;
+		uint32 NaniteImposterIndex = INDEX_NONE;
+		NaniteProxy->GetNaniteResourceInfo(NaniteResourceID, NaniteHierarchyOffset, NaniteImposterIndex);
+		uint32 NaniteFilterFlags = uint32(NaniteProxy->GetFilterFlags());
+		uint32 NaniteRayTracingDataOffset = NaniteProxy->GetRayTracingDataOffset();
+		bool bReverseCulling = NaniteProxy->IsCullingReversedByComponent(); // needed because Nanite doesn't use raster state
+		
+		Builder.NaniteResourceID(NaniteResourceID)
+			.NaniteHierarchyOffset(NaniteHierarchyOffset)
+			.NaniteImposterIndex(NaniteImposterIndex)
+			.NaniteFilterFlags(NaniteFilterFlags)
+			.NaniteRayTracingDataOffset(NaniteRayTracingDataOffset)
+			.ReverseCulling(bReverseCulling);
+	}
 
-		FPrimitiveUniformShaderParameters PrimitiveParams = Builder.Build();
 
-		check(UniformBuffer);
-		UniformBuffer.UpdateUniformBufferImmediate(RHICmdList, PrimitiveParams);
+	FVector2f InstanceDrawDistanceMinMax;
+	if (GetInstanceDrawDistanceMinMax(InstanceDrawDistanceMinMax))
+	{
+		Builder.InstanceDrawDistance(InstanceDrawDistanceMinMax);
+	}
+
+	float WPODisableDistance;
+	if (GetInstanceWorldPositionOffsetDisableDistance(WPODisableDistance))
+	{
+		Builder.InstanceWorldPositionOffsetDisableDistance(WPODisableDistance);
+	}
+
+	const TConstArrayView<FRenderBounds> InstanceBounds = GetInstanceLocalBounds();
+	if (InstanceBounds.Num() > 0)
+	{
+		Builder.InstanceLocalBounds(InstanceBounds[0]);
+	}
+
+	if (ShouldRenderCustomDepth())
+	{
+		Builder.CustomDepthStencil(GetCustomDepthStencilValue(), GetStencilWriteMask());
 	}
 }
 
diff --git a/Engine/Source/Runtime/Engine/Private/PrimitiveUniformShaderParameters.cpp b/Engine/Source/Runtime/Engine/Private/PrimitiveUniformShaderParameters.cpp
--- a/Engine/Source/Runtime/Engine/Private/PrimitiveUniformShaderParameters.cpp
+++ b/Engine/Source/Runtime/Engine/Private/PrimitiveUniformShaderParameters.cpp
@@ -7,6 +7,7 @@
 #include "SceneInterface.h"
 #include "UnrealEngine.h"
 #include "DataDrivenShaderPlatformInfo.h"
+#include "PrimitiveUniformShaderParametersBuilder.h"
 
 FPrimitiveUniformShaderParametersBuilder& FPrimitiveUniformShaderParametersBuilder::InstanceDrawDistance(FVector2f DistanceMinMax)
 {
@@ -133,115 +134,8 @@ PRAGMA_ENABLE_DEPRECATION_WARNINGS
 FPrimitiveSceneShaderData::FPrimitiveSceneShaderData(const FPrimitiveSceneProxy* RESTRICT Proxy)
 	: Data(InPlace, NoInit)
 {
-	bool bHasPrecomputedVolumetricLightmap;
-	FMatrix PreviousLocalToWorld;
-	int32 SingleCaptureIndex;
-	bool bOutputVelocity;
-
-	Proxy->GetScene().GetPrimitiveUniformShaderParameters_RenderThread(
-		Proxy->GetPrimitiveSceneInfo(),
-		bHasPrecomputedVolumetricLightmap,
-		PreviousLocalToWorld,
-		SingleCaptureIndex,
-		bOutputVelocity
-	);
-
-	bOutputVelocity |= Proxy->AlwaysHasVelocity();
-
-	FBoxSphereBounds PreSkinnedLocalBounds;
-	Proxy->GetPreSkinnedLocalBounds(PreSkinnedLocalBounds);
-
-	FPrimitiveSceneInfo* PrimitiveSceneInfo = Proxy->GetPrimitiveSceneInfo();
-
-	uint32 NaniteResourceID = INDEX_NONE;
-	uint32 NaniteHierarchyOffset = INDEX_NONE;
-	uint32 NaniteImposterIndex = INDEX_NONE;
-	uint32 NaniteFilterFlags = 0u;
-	uint32 NaniteRayTracingDataOffset = INDEX_NONE;
-	bool bReverseCulling = false;
-
-	if (Proxy->IsNaniteMesh())
-	{
-		auto NaniteProxy = static_cast<const Nanite::FSceneProxyBase*>(Proxy);
-		NaniteProxy->GetNaniteResourceInfo(NaniteResourceID, NaniteHierarchyOffset, NaniteImposterIndex);
-		NaniteFilterFlags = uint32(NaniteProxy->GetFilterFlags());
-		NaniteRayTracingDataOffset = NaniteProxy->GetRayTracingDataOffset();
-		bReverseCulling = NaniteProxy->IsCullingReversedByComponent(); // needed because Nanite doesn't use raster state
-	}
-
-	FPrimitiveUniformShaderParametersBuilder Builder = FPrimitiveUniformShaderParametersBuilder{}
-		.Defaults()
-		.LocalToWorld(Proxy->GetLocalToWorld())
-		.PreviousLocalToWorld(PreviousLocalToWorld)
-		.ActorWorldPosition(Proxy->GetActorPosition())
-		.WorldBounds(Proxy->GetBounds())
-		.LocalBounds(Proxy->GetLocalBounds())
-		.PreSkinnedLocalBounds(PreSkinnedLocalBounds)
-		.CustomPrimitiveData(Proxy->GetCustomPrimitiveData())
-		.LightingChannelMask(Proxy->GetLightingChannelMask())
-		.LightmapDataIndex(Proxy->GetPrimitiveSceneInfo()->GetLightmapDataOffset())
-		.LightmapUVIndex(Proxy->GetLightMapCoordinateIndex())
-		.SingleCaptureIndex(SingleCaptureIndex)
-		.PersistentPrimitiveIndex(Proxy->GetPrimitiveSceneInfo()->GetPersistentIndex().Index)
-		.InstanceSceneDataOffset(Proxy->GetPrimitiveSceneInfo()->GetInstanceSceneDataOffset())
-		.NumInstanceSceneDataEntries(Proxy->GetPrimitiveSceneInfo()->GetNumInstanceSceneDataEntries())
-		.InstancePayloadDataOffset(Proxy->GetPrimitiveSceneInfo()->GetInstancePayloadDataOffset())
-		.InstancePayloadDataStride(Proxy->GetPrimitiveSceneInfo()->GetInstancePayloadDataStride())
-		.HasDistanceFieldRepresentation(Proxy->HasDistanceFieldRepresentation())
-		.HasCapsuleRepresentation(Proxy->HasDynamicIndirectShadowCasterRepresentation())
-		.UseSingleSampleShadowFromStationaryLights(Proxy->UseSingleSampleShadowFromStationaryLights())
-		.ReceivesDecals(Proxy->ReceivesDecals())
-		.CacheShadowAsStatic(PrimitiveSceneInfo->ShouldCacheShadowAsStatic())
-		.OutputVelocity(bOutputVelocity)
-		.EvaluateWorldPositionOffset(Proxy->EvaluateWorldPositionOffset() && Proxy->AnyMaterialHasWorldPositionOffset())
-		.MaxWorldPositionOffsetExtent(Proxy->GetMaxWorldPositionOffsetExtent())
-		.MinMaxMaterialDisplacement(Proxy->GetMinMaxMaterialDisplacement())
-		.CastContactShadow(Proxy->CastsContactShadow())
-		.CastShadow(Proxy->CastsDynamicShadow())
-		.CastHiddenShadow(Proxy->CastsHiddenShadow())
-		.VisibleInGame(Proxy->IsDrawnInGame())
-		.VisibleInEditor(Proxy->IsDrawnInEditor())
-		.VisibleInReflectionCaptures(Proxy->IsVisibleInReflectionCaptures())
-		.VisibleInRealTimeSkyCaptures(Proxy->IsVisibleInRealTimeSkyCaptures())
-		.VisibleInRayTracing(Proxy->IsVisibleInRayTracing())
-		.VisibleInLumenScene(Proxy->IsVisibleInLumenScene())
-		.VisibleInSceneCaptureOnly(Proxy->IsVisibleInSceneCaptureOnly())
-		.HiddenInSceneCapture(Proxy->IsHiddenInSceneCapture())
-		.ForceHidden(Proxy->IsForceHidden())
-		.Holdout(Proxy->Holdout())
-		.UseVolumetricLightmap(bHasPrecomputedVolumetricLightmap)
-		.NaniteResourceID(NaniteResourceID)
-		.NaniteHierarchyOffset(NaniteHierarchyOffset)
-		.NaniteImposterIndex(NaniteImposterIndex)
-		.NaniteFilterFlags(NaniteFilterFlags)
-		.NaniteRayTracingDataOffset(NaniteRayTracingDataOffset)
-		.PrimitiveComponentId(Proxy->GetPrimitiveComponentId().PrimIDValue)
-		.EditorColors(Proxy->GetWireframeColor(), Proxy->GetLevelColor())
-		.ReverseCulling(bReverseCulling);
-
-	FVector2f InstanceDrawDistanceMinMax;
-	if (Proxy->GetInstanceDrawDistanceMinMax(InstanceDrawDistanceMinMax))
-	{
-		Builder.InstanceDrawDistance(InstanceDrawDistanceMinMax);
-	}
-
-	float WPODisableDistance;
-	if (Proxy->GetInstanceWorldPositionOffsetDisableDistance(WPODisableDistance))
-	{
-		Builder.InstanceWorldPositionOffsetDisableDistance(WPODisableDistance);
-	}
-
-	const TConstArrayView<FRenderBounds> InstanceBounds = Proxy->GetInstanceLocalBounds();
-	if (InstanceBounds.Num() > 0)
-	{
-		Builder.InstanceLocalBounds(InstanceBounds[0]);
-	}
-
-	if (Proxy->ShouldRenderCustomDepth())
-	{
-		Builder.CustomDepthStencil(Proxy->GetCustomDepthStencilValue(), Proxy->GetStencilWriteMask());
-	}
-
+	FPrimitiveUniformShaderParametersBuilder Builder = FPrimitiveUniformShaderParametersBuilder{};
+	Proxy->BuildUniformShaderParameters(Builder);
 	Setup(Builder.Build());
 }
 
@@ -347,3 +241,41 @@ void FPrimitiveSceneShaderData::Setup(const FPrimitiveUniformShaderParameters& P
 		Data[CustomPrimitiveDataStartIndex + DataIndex] = PrimitiveUniformShaderParameters.CustomPrimitiveData[DataIndex];
 	}
 }
+
+TUniformBufferRef<FPrimitiveUniformShaderParameters> CreatePrimitiveUniformBufferImmediate(
+	const FMatrix& LocalToWorld,
+	const FBoxSphereBounds& WorldBounds,
+	const FBoxSphereBounds& LocalBounds,
+	const FBoxSphereBounds& PreSkinnedLocalBounds,
+	bool bReceivesDecals,
+	bool bOutputVelocity
+)
+{
+	check(IsInRenderingThread());
+	return TUniformBufferRef<FPrimitiveUniformShaderParameters>::CreateUniformBufferImmediate(
+		FPrimitiveUniformShaderParametersBuilder{}
+		.Defaults()
+			.LocalToWorld(LocalToWorld)
+			.ActorWorldPosition(WorldBounds.Origin)
+			.WorldBounds(WorldBounds)
+			.LocalBounds(LocalBounds)
+			.PreSkinnedLocalBounds(PreSkinnedLocalBounds)
+			.ReceivesDecals(bReceivesDecals)
+			.OutputVelocity(bOutputVelocity)
+		.Build(),
+		UniformBuffer_MultiFrame
+	);
+}
+
+FPrimitiveUniformShaderParameters GetIdentityPrimitiveParameters()
+{
+	// Don't use FMatrix44f::Identity here as GetIdentityPrimitiveParameters is used by TGlobalResource<FIdentityPrimitiveUniformBuffer> and because
+	// static initialization order is undefined, FMatrix44f::Identity might be all 0's or random data the first time this is called.
+	return FPrimitiveUniformShaderParametersBuilder{}
+		.Defaults()
+			.LocalToWorld(FMatrix(FPlane(1, 0, 0, 0), FPlane(0, 1, 0, 0), FPlane(0, 0, 1, 0), FPlane(0, 0, 0, 1)))
+			.ActorWorldPosition(FVector(0.0, 0.0, 0.0))
+			.WorldBounds(FBoxSphereBounds(EForceInit::ForceInit))
+			.LocalBounds(FBoxSphereBounds(EForceInit::ForceInit))
+		.Build();
+}
\ No newline at end of file
diff --git a/Engine/Source/Runtime/Engine/Private/SceneManagement.cpp b/Engine/Source/Runtime/Engine/Private/SceneManagement.cpp
--- a/Engine/Source/Runtime/Engine/Private/SceneManagement.cpp
+++ b/Engine/Source/Runtime/Engine/Private/SceneManagement.cpp
@@ -20,6 +20,7 @@
 #include "ColorSpace.h"
 #include "DataDrivenShaderPlatformInfo.h"
 #include "StaticMeshBatch.h"
+#include "PrimitiveUniformShaderParametersBuilder.h"
 
 static TAutoConsoleVariable<float> CVarLODTemporalLag(
 	TEXT("lod.TemporalLag"),
diff --git a/Engine/Source/Runtime/Engine/Public/PrimitiveSceneProxy.h b/Engine/Source/Runtime/Engine/Public/PrimitiveSceneProxy.h
--- a/Engine/Source/Runtime/Engine/Public/PrimitiveSceneProxy.h
+++ b/Engine/Source/Runtime/Engine/Public/PrimitiveSceneProxy.h
@@ -34,6 +34,7 @@ struct FInstanceUpdateCmdBuffer;
 class FRayTracingGeometry;
 class FVertexFactory;
 class IHeterogeneousVolumeInterface;
+struct FPrimitiveUniformShaderParametersBuilder;
 
 namespace Nanite
 {
@@ -1009,6 +1010,11 @@ public:
 		UpdateUniformBuffer(FRHICommandListExecutor::GetImmediateCommandList());
 	}
 
+	/**
+	 * Apply the unform shader parameter settings for the proxy to the builder.
+	 */
+	ENGINE_API void BuildUniformShaderParameters(FPrimitiveUniformShaderParametersBuilder &Builder) const;
+
 #if ENABLE_DRAW_DEBUG
 
 	struct ENGINE_API FDebugMassData
diff --git a/Engine/Source/Runtime/Engine/Public/PrimitiveUniformShaderParameters.h b/Engine/Source/Runtime/Engine/Public/PrimitiveUniformShaderParameters.h
--- a/Engine/Source/Runtime/Engine/Public/PrimitiveUniformShaderParameters.h
+++ b/Engine/Source/Runtime/Engine/Public/PrimitiveUniformShaderParameters.h
@@ -13,11 +13,12 @@
 #include "LightmapUniformShaderParameters.h"
 #include "UnifiedBuffer.h"
 #include "Containers/StaticArray.h"
-#include "NaniteDefinitions.h"
 #if UE_ENABLE_INCLUDE_ORDER_DEPRECATED_IN_5_2
 #include "UnrealEngine.h"
 #endif
 
+struct FPrimitiveUniformShaderParametersBuilder;
+
 /** 
  * The uniform shader parameters associated with a primitive. 
  * Note: Must match FPrimitiveSceneData in shaders.
@@ -69,498 +70,16 @@ BEGIN_GLOBAL_SHADER_PARAMETER_STRUCT(FPrimitiveUniformShaderParameters,ENGINE_AP
 	SHADER_PARAMETER_ARRAY(FVector4f, CustomPrimitiveData, [FCustomPrimitiveData::NumCustomPrimitiveDataFloat4s]) // Custom data per primitive that can be accessed through material expression parameters and modified through UStaticMeshComponent
 END_GLOBAL_SHADER_PARAMETER_STRUCT()
 
-// Must match SceneData.ush
-#define PRIMITIVE_SCENE_DATA_FLAG_CAST_SHADOWS							0x1
-#define PRIMITIVE_SCENE_DATA_FLAG_USE_SINGLE_SAMPLE_SHADOW_SL			0x2
-#define PRIMITIVE_SCENE_DATA_FLAG_USE_VOLUMETRIC_LM_SHADOW_SL			0x4
-#define PRIMITIVE_SCENE_DATA_FLAG_DECAL_RECEIVER						0x8
-#define PRIMITIVE_SCENE_DATA_FLAG_CACHE_SHADOW_AS_STATIC				0x10
-#define PRIMITIVE_SCENE_DATA_FLAG_OUTPUT_VELOCITY						0x20
-#define PRIMITIVE_SCENE_DATA_FLAG_DETERMINANT_SIGN						0x40
-#define PRIMITIVE_SCENE_DATA_FLAG_HAS_DISTANCE_FIELD_REPRESENTATION		0x80
-#define PRIMITIVE_SCENE_DATA_FLAG_HAS_CAPSULE_REPRESENTATION			0x100
-#define PRIMITIVE_SCENE_DATA_FLAG_HAS_CAST_CONTACT_SHADOW				0x200
-#define PRIMITIVE_SCENE_DATA_FLAG_HAS_PRIMITIVE_CUSTOM_DATA				0x400
-#define PRIMITIVE_SCENE_DATA_FLAG_LIGHTING_CHANNEL_0					0x800
-#define PRIMITIVE_SCENE_DATA_FLAG_LIGHTING_CHANNEL_1					0x1000
-#define PRIMITIVE_SCENE_DATA_FLAG_LIGHTING_CHANNEL_2					0x2000
-#define PRIMITIVE_SCENE_DATA_FLAG_HAS_INSTANCE_LOCAL_BOUNDS				0x4000
-#define PRIMITIVE_SCENE_DATA_FLAG_EVALUATE_WORLD_POSITION_OFFSET		0x8000
-#define PRIMITIVE_SCENE_DATA_FLAG_VISIBLE_IN_GAME						0x10000
-#define PRIMITIVE_SCENE_DATA_FLAG_VISIBLE_IN_EDITOR						0x20000
-#define PRIMITIVE_SCENE_DATA_FLAG_VISIBLE_IN_REFLECTION_CAPTURES		0x40000
-#define PRIMITIVE_SCENE_DATA_FLAG_VISIBLE_IN_REAL_TIME_SKY_CAPTURES		0x80000
-#define PRIMITIVE_SCENE_DATA_FLAG_VISIBLE_IN_RAY_TRACING				0x100000
-#define PRIMITIVE_SCENE_DATA_FLAG_VISIBLE_IN_SCENE_CAPTURE_ONLY			0x200000
-#define PRIMITIVE_SCENE_DATA_FLAG_HIDDEN_IN_SCENE_CAPTURE				0x400000
-#define PRIMITIVE_SCENE_DATA_FLAG_FORCE_HIDDEN							0x800000 // Mobile allocates only 24 bits for primitive flags
-#define PRIMITIVE_SCENE_DATA_FLAG_CAST_HIDDEN_SHADOW					0x1000000
-#define PRIMITIVE_SCENE_DATA_FLAG_HAS_NANITE_IMPOSTER					0x2000000
-#define PRIMITIVE_SCENE_DATA_FLAG_INSTANCE_DRAW_DISTANCE_CULL			0x4000000
-#define PRIMITIVE_SCENE_DATA_FLAG_WPO_DISABLE_DISTANCE					0x8000000
-#define PRIMITIVE_SCENE_DATA_FLAG_WRITES_CUSTOM_DEPTH_STENCIL			0x10000000
-#define PRIMITIVE_SCENE_DATA_FLAG_REVERSE_CULLING						0x20000000
-#define PRIMITIVE_SCENE_DATA_FLAG_HOLDOUT								0x40000000
-#define PRIMITIVE_SCENE_DATA_FLAG_VISIBLE_IN_LUMEN_SCENE				0x80000000
-
-struct FPrimitiveUniformShaderParametersBuilder
-{
-public:
-	inline FPrimitiveUniformShaderParametersBuilder& Defaults()
-	{
-		ObjectRadius = 0.0f;
-
-		// Flags defaulted on
-		bCastShadow									= true;
-		bCastContactShadow							= true;
-		bEvaluateWorldPositionOffset				= true;
-		bVisibleInGame								= true;
-		bVisibleInEditor							= true;
-		bVisibleInReflectionCaptures				= true;
-		bVisibleInRealTimeSkyCaptures				= true;
-		bVisibleInRayTracing						= true;
-		bVisibleInLumenScene						= true;
-
-		// Flags defaulted off
-		bReceivesDecals								= false;
-		bUseSingleSampleShadowFromStationaryLights	= false;
-		bUseVolumetricLightmap						= false;
-		bCacheShadowAsStatic						= false;
-		bOutputVelocity								= false;
-		bHasDistanceFieldRepresentation				= false;
-		bHasCapsuleRepresentation					= false;
-		bHasPreSkinnedLocalBounds					= false;
-		bHasPreviousLocalToWorld					= false;
-		bHasInstanceLocalBounds						= false;
-		bCastHiddenShadow							= false;
-		bVisibleInSceneCaptureOnly					= false;
-		bHiddenInSceneCapture						= false;
-		bForceHidden								= false;
-		bHasNaniteImposter							= false;
-		bHasInstanceDrawDistanceCull				= false;
-		bHasWPODisableDistance						= false;
-		bWritesCustomDepthStencil					= false;
-		bReverseCulling								= false;
-		bHoldout									= false;
-
-		Parameters.MaxWPOExtent						= 0.0f;
-		Parameters.MinMaterialDisplacement			= 0.0f;
-		Parameters.MaxMaterialDisplacement			= 0.0f;
-
-		// Default colors
-		Parameters.WireframeColor					= FVector3f(1.0f, 1.0f, 1.0f);
-		Parameters.LevelColor						= FVector3f(1.0f, 1.0f, 1.0f);
-
-		// Invalid indices
-		Parameters.LightmapDataIndex				= INDEX_NONE;
-		Parameters.LightmapUVIndex					= INDEX_NONE;
-		Parameters.SingleCaptureIndex				= INDEX_NONE;
-		Parameters.PersistentPrimitiveIndex			= INDEX_NONE;
-		Parameters.PrimitiveComponentId				= ~uint32(0u);
-
-		// Nanite
-		Parameters.NaniteResourceID						= INDEX_NONE;
-		Parameters.NaniteHierarchyOffset				= INDEX_NONE;
-		Parameters.PackedNaniteFlags					= NANITE_IMPOSTER_INDEX_MASK;
-		Parameters.NaniteRayTracingDataOffset			= INDEX_NONE;
-
-		// Instance data
-		Parameters.InstanceSceneDataOffset			= INDEX_NONE;
-		Parameters.NumInstanceSceneDataEntries		= 0;
-		Parameters.InstancePayloadDataOffset		= INDEX_NONE;
-		Parameters.InstancePayloadDataStride		= 0;
-
-		LightingChannels = GetDefaultLightingChannelMask();
-
-		return CustomPrimitiveData(nullptr);
-	}
-
-#define PRIMITIVE_UNIFORM_BUILDER_METHOD(INPUT_TYPE, VARIABLE_NAME) \
-	inline FPrimitiveUniformShaderParametersBuilder& VARIABLE_NAME(INPUT_TYPE In##VARIABLE_NAME) { Parameters.VARIABLE_NAME = In##VARIABLE_NAME; return *this; }
-
-#define PRIMITIVE_UNIFORM_BUILDER_FLAG_METHOD(INPUT_TYPE, VARIABLE_NAME) \
-	inline FPrimitiveUniformShaderParametersBuilder& VARIABLE_NAME(INPUT_TYPE In##VARIABLE_NAME) { b##VARIABLE_NAME = In##VARIABLE_NAME; return *this; }
-
-	PRIMITIVE_UNIFORM_BUILDER_FLAG_METHOD(bool,			ReceivesDecals);
-	PRIMITIVE_UNIFORM_BUILDER_FLAG_METHOD(bool,			HasDistanceFieldRepresentation);
-	PRIMITIVE_UNIFORM_BUILDER_FLAG_METHOD(bool,			HasCapsuleRepresentation);
-	PRIMITIVE_UNIFORM_BUILDER_FLAG_METHOD(bool,			HasInstanceLocalBounds);
-	PRIMITIVE_UNIFORM_BUILDER_FLAG_METHOD(bool,			CastContactShadow);
-	PRIMITIVE_UNIFORM_BUILDER_FLAG_METHOD(bool,			CastHiddenShadow);
-	PRIMITIVE_UNIFORM_BUILDER_FLAG_METHOD(bool,			CastShadow);
-	PRIMITIVE_UNIFORM_BUILDER_FLAG_METHOD(bool,			UseSingleSampleShadowFromStationaryLights);
-	PRIMITIVE_UNIFORM_BUILDER_FLAG_METHOD(bool,			UseVolumetricLightmap);
-	PRIMITIVE_UNIFORM_BUILDER_FLAG_METHOD(bool,			CacheShadowAsStatic);
-	PRIMITIVE_UNIFORM_BUILDER_FLAG_METHOD(bool,			OutputVelocity);
-	PRIMITIVE_UNIFORM_BUILDER_FLAG_METHOD(bool,			EvaluateWorldPositionOffset);
-	PRIMITIVE_UNIFORM_BUILDER_FLAG_METHOD(bool,			VisibleInGame);
-	PRIMITIVE_UNIFORM_BUILDER_FLAG_METHOD(bool,			VisibleInEditor);
-	PRIMITIVE_UNIFORM_BUILDER_FLAG_METHOD(bool,			ReverseCulling);
-	PRIMITIVE_UNIFORM_BUILDER_FLAG_METHOD(bool,			VisibleInReflectionCaptures);
-	PRIMITIVE_UNIFORM_BUILDER_FLAG_METHOD(bool,			VisibleInRealTimeSkyCaptures);
-	PRIMITIVE_UNIFORM_BUILDER_FLAG_METHOD(bool,			VisibleInRayTracing);
-	PRIMITIVE_UNIFORM_BUILDER_FLAG_METHOD(bool,			VisibleInLumenScene);
-	PRIMITIVE_UNIFORM_BUILDER_FLAG_METHOD(bool,			VisibleInSceneCaptureOnly);
-	PRIMITIVE_UNIFORM_BUILDER_FLAG_METHOD(bool,			HiddenInSceneCapture);
-	PRIMITIVE_UNIFORM_BUILDER_FLAG_METHOD(bool,			ForceHidden);
-	PRIMITIVE_UNIFORM_BUILDER_FLAG_METHOD(bool,         Holdout);
-
-	PRIMITIVE_UNIFORM_BUILDER_METHOD(uint32,			InstanceSceneDataOffset);
-	PRIMITIVE_UNIFORM_BUILDER_METHOD(uint32,			NumInstanceSceneDataEntries);
-	PRIMITIVE_UNIFORM_BUILDER_METHOD(uint32,			InstancePayloadDataOffset);
-	PRIMITIVE_UNIFORM_BUILDER_METHOD(uint32,			InstancePayloadDataStride);
-	PRIMITIVE_UNIFORM_BUILDER_METHOD(int32,				SingleCaptureIndex);
-	PRIMITIVE_UNIFORM_BUILDER_METHOD(int32,				PersistentPrimitiveIndex);
-	PRIMITIVE_UNIFORM_BUILDER_METHOD(uint32,			PrimitiveComponentId);
-	PRIMITIVE_UNIFORM_BUILDER_METHOD(uint32,			NaniteResourceID);
-	PRIMITIVE_UNIFORM_BUILDER_METHOD(uint32,			NaniteHierarchyOffset);
-	PRIMITIVE_UNIFORM_BUILDER_METHOD(uint32,			NaniteRayTracingDataOffset);
-	PRIMITIVE_UNIFORM_BUILDER_METHOD(uint32,			LightmapUVIndex);
-	PRIMITIVE_UNIFORM_BUILDER_METHOD(uint32,			LightmapDataIndex);
-
-#undef PRIMITIVE_UNIFORM_BUILDER_FLAG_METHOD
-#undef PRIMITIVE_UNIFORM_BUILDER_METHOD
-
-	inline FPrimitiveUniformShaderParametersBuilder& LightingChannelMask(uint32 InLightingChannelMask)
-	{
-		LightingChannels = InLightingChannelMask;
-		return *this;
-	}
-
-	inline FPrimitiveUniformShaderParametersBuilder& MaxWorldPositionOffsetExtent(float InMaxExtent)
-	{
-		Parameters.MaxWPOExtent = InMaxExtent;
-		return *this;
-	}
-
-	inline FPrimitiveUniformShaderParametersBuilder& MinMaxMaterialDisplacement(const FVector2f& InMinMaxDisplacement)
-	{
-		Parameters.MinMaterialDisplacement = InMinMaxDisplacement.X;
-		Parameters.MaxMaterialDisplacement = InMinMaxDisplacement.Y;
-		return *this;
-	}
-
-	inline FPrimitiveUniformShaderParametersBuilder& ObjectBounds(const FVector3f& InObjectBounds)
-	{
-		Parameters.ObjectBoundsX = InObjectBounds.X;
-		Parameters.ObjectBoundsY = InObjectBounds.Y;
-		Parameters.ObjectBoundsZ = InObjectBounds.Z;
-		return *this;
-	}
-
-	inline FPrimitiveUniformShaderParametersBuilder& WorldBounds(const FBoxSphereBounds& InWorldBounds)
-	{
-		AbsoluteObjectWorldPosition = InWorldBounds.Origin;
-		ObjectRadius = static_cast<float>(InWorldBounds.SphereRadius);											//LWC_TODO: Precision loss
-		Parameters.ObjectBoundsX = static_cast<float>(InWorldBounds.BoxExtent.X);
-		Parameters.ObjectBoundsY = static_cast<float>(InWorldBounds.BoxExtent.Y);
-		Parameters.ObjectBoundsZ = static_cast<float>(InWorldBounds.BoxExtent.Z);
-		return *this;
-	}
-
-	inline FPrimitiveUniformShaderParametersBuilder& LocalBounds(const FBoxSphereBounds& InLocalBounds)
-	{
-		Parameters.LocalObjectBoundsMin = FVector3f(InLocalBounds.GetBoxExtrema(0)); // 0 == minimum		//LWC_TODO: Precision loss
-		Parameters.LocalObjectBoundsMax = FVector3f(InLocalBounds.GetBoxExtrema(1)); // 1 == maximum		//LWC_TODO: Precision loss
-		return *this;
-	}
-
-	inline FPrimitiveUniformShaderParametersBuilder& PreSkinnedLocalBounds(const FBoxSphereBounds& InPreSkinnedLocalBounds)
-	{
-		bHasPreSkinnedLocalBounds = true;
-		Parameters.PreSkinnedLocalBoundsMin = FVector3f(InPreSkinnedLocalBounds.GetBoxExtrema(0)); // 0 == minimum		//LWC_TODO: Precision loss
-		Parameters.PreSkinnedLocalBoundsMax = FVector3f(InPreSkinnedLocalBounds.GetBoxExtrema(1)); // 1 == maximum		//LWC_TODO: Precision loss
-		return *this;
-	}
-
-	inline FPrimitiveUniformShaderParametersBuilder& ActorWorldPosition(const FVector& InActorWorldPosition)
-	{
-		AbsoluteActorWorldPosition = InActorWorldPosition;
-		return *this;
-	}
-
-	inline FPrimitiveUniformShaderParametersBuilder& LocalToWorld(const FMatrix& InLocalToWorld)
-	{
-		AbsoluteLocalToWorld = InLocalToWorld;
-		return *this;
-	}
-
-	inline FPrimitiveUniformShaderParametersBuilder& InstanceLocalBounds(const FRenderBounds& InInstanceLocalBounds)
-	{
-		bHasInstanceLocalBounds = true;
-		Parameters.InstanceLocalBoundsCenter = InInstanceLocalBounds.GetCenter();
-		Parameters.InstanceLocalBoundsExtent = InInstanceLocalBounds.GetExtent();
-		return *this;
-	}
-
-	inline FPrimitiveUniformShaderParametersBuilder& PreviousLocalToWorld(const FMatrix& InPreviousLocalToWorld)
-	{
-		bHasPreviousLocalToWorld = true;
-		AbsolutePreviousLocalToWorld = InPreviousLocalToWorld;
-		return *this;
-	}
-
-	inline FPrimitiveUniformShaderParametersBuilder& EditorColors(const FLinearColor& InWireframeColor, const FLinearColor& InLevelColor)
-	{
-		Parameters.WireframeColor = FVector3f(InWireframeColor.R, InWireframeColor.G, InWireframeColor.B);
-		Parameters.LevelColor = FVector3f(InLevelColor.R, InLevelColor.G, InLevelColor.B);
-		return *this;
-	}
-
-	inline FPrimitiveUniformShaderParametersBuilder& CustomPrimitiveData(const FCustomPrimitiveData* InCustomPrimitiveData)
-	{
-		// If this primitive has custom primitive data, set it
-		if (InCustomPrimitiveData)
-		{
-			// Copy at most up to the max supported number of floats for safety
-			FMemory::Memcpy(
-				&Parameters.CustomPrimitiveData,
-				InCustomPrimitiveData->Data.GetData(),
-				InCustomPrimitiveData->Data.GetTypeSize() * FMath::Min(InCustomPrimitiveData->Data.Num(),
-				FCustomPrimitiveData::NumCustomPrimitiveDataFloats)
-			);
-
-			bHasCustomData = true;
-		}
-		else
-		{
-			// Clear to 0
-			FMemory::Memzero(Parameters.CustomPrimitiveData);
-			bHasCustomData = false;
-		}
-
-		return *this;
-	}
-
-	inline FPrimitiveUniformShaderParametersBuilder& NaniteImposterIndex(uint32 ImposterIndex)
-	{
-		bHasNaniteImposter = ImposterIndex != INDEX_NONE;
-
-		check(!bHasNaniteImposter || ImposterIndex < NANITE_IMPOSTER_INDEX_MASK);
-		Parameters.PackedNaniteFlags = (Parameters.PackedNaniteFlags & NANITE_FILTER_FLAGS_MASK) | (ImposterIndex & NANITE_IMPOSTER_INDEX_MASK);
-
-		return *this;
-	}
-
-	inline FPrimitiveUniformShaderParametersBuilder& NaniteFilterFlags(uint32 FilterFlags)
-	{
-		check(FilterFlags < (1u << NANITE_FILTER_FLAGS_NUM_BITS));
-		Parameters.PackedNaniteFlags = (FilterFlags << NANITE_IMPOSTER_INDEX_NUM_BITS) | (Parameters.PackedNaniteFlags & NANITE_IMPOSTER_INDEX_MASK);
-
-		return *this;
-	}
-
-	inline FPrimitiveUniformShaderParametersBuilder& CustomDepthStencil(uint8 StencilValue, EStencilMask StencilWriteMask)
-	{
-		// Translate the enum to a mask that can be consumed by the GPU with fewer operations
-		uint32 GPUStencilMask;
-		switch (StencilWriteMask)
-		{		
-		case SM_Default:
-			GPUStencilMask = 0u; // Use zero to mean replace
-			break;
-		case SM_255:
-			GPUStencilMask = 0xFFu;
-			break;
-		default:
-			GPUStencilMask = (1u << uint32(StencilWriteMask - SM_1)) & 0xFFu;
-			break;
-		}
-		Parameters.CustomStencilValueAndMask = (GPUStencilMask << 8u) | StencilValue;
-		bWritesCustomDepthStencil = true;
-
-		return *this;
-	}
-
-	ENGINE_API FPrimitiveUniformShaderParametersBuilder& InstanceDrawDistance(FVector2f DistanceMinMax);
-
-	ENGINE_API FPrimitiveUniformShaderParametersBuilder& InstanceWorldPositionOffsetDisableDistance(float WPODisableDistance);
-
-	inline const FPrimitiveUniformShaderParameters& Build()
-	{
-		const FLargeWorldRenderPosition AbsoluteWorldPosition(AbsoluteLocalToWorld.GetOrigin());
-		const FVector TilePositionOffset = AbsoluteWorldPosition.GetTileOffset();
-
-		Parameters.TilePosition = AbsoluteWorldPosition.GetTile();
-
-		{
-			// Inverse on FMatrix44f can generate NaNs if the source matrix contains large scaling, so do it in double precision.
-			// Also use double precision to calculate WorldToPreviousWorld to prevent precision issues at far distances
-			FMatrix LocalToRelativeWorld = FLargeWorldRenderScalar::MakeToRelativeWorldMatrixDouble(TilePositionOffset, AbsoluteLocalToWorld);
-			FMatrix PrevLocalToRelativeWorld = FLargeWorldRenderScalar::MakeClampedToRelativeWorldMatrixDouble(TilePositionOffset, AbsolutePreviousLocalToWorld);
-			FMatrix RelativeWorldToLocal = LocalToRelativeWorld.Inverse();
-
-			Parameters.LocalToRelativeWorld = FMatrix44f(LocalToRelativeWorld);
-			Parameters.RelativeWorldToLocal = FMatrix44f(RelativeWorldToLocal);
-			Parameters.WorldToPreviousWorld = FMatrix44f(RelativeWorldToLocal * PrevLocalToRelativeWorld);
-
-			if (bHasPreviousLocalToWorld)
-			{
-				Parameters.PreviousLocalToRelativeWorld = FMatrix44f(PrevLocalToRelativeWorld);
-				Parameters.PreviousRelativeWorldToLocal = FMatrix44f(PrevLocalToRelativeWorld.Inverse());
-			}
-			else
-			{
-				Parameters.PreviousLocalToRelativeWorld = Parameters.LocalToRelativeWorld;
-				Parameters.PreviousRelativeWorldToLocal = Parameters.RelativeWorldToLocal;
-			}
-		}
-
-		Parameters.ActorRelativeWorldPosition = FVector3f(AbsoluteActorWorldPosition - TilePositionOffset);	//LWC_TODO: Precision loss
-		const FVector3f ObjectRelativeWorldPositionAsFloat = FVector3f(AbsoluteObjectWorldPosition - TilePositionOffset);
-		Parameters.ObjectRelativeWorldPositionAndRadius = FVector4f(ObjectRelativeWorldPositionAsFloat, ObjectRadius);
-
-		if (!bHasInstanceLocalBounds)
-		{
-			FRenderBounds InstanceLocalBounds(Parameters.LocalObjectBoundsMin, Parameters.LocalObjectBoundsMax);
-			Parameters.InstanceLocalBoundsCenter = InstanceLocalBounds.GetCenter();
-			Parameters.InstanceLocalBoundsExtent = InstanceLocalBounds.GetExtent();
-		}
-
-		if (!bHasPreSkinnedLocalBounds)
-		{
-			Parameters.PreSkinnedLocalBoundsMin = Parameters.LocalObjectBoundsMin;
-			Parameters.PreSkinnedLocalBoundsMax = Parameters.LocalObjectBoundsMax;
-		}
-
-		Parameters.ObjectOrientation = Parameters.LocalToRelativeWorld.GetUnitAxis(EAxis::Z);
-
-		{
-			// Extract per axis scales from LocalToWorld transform
-			FVector4f WorldX = FVector4f(Parameters.LocalToRelativeWorld.M[0][0], Parameters.LocalToRelativeWorld.M[0][1], Parameters.LocalToRelativeWorld.M[0][2], 0);
-			FVector4f WorldY = FVector4f(Parameters.LocalToRelativeWorld.M[1][0], Parameters.LocalToRelativeWorld.M[1][1], Parameters.LocalToRelativeWorld.M[1][2], 0);
-			FVector4f WorldZ = FVector4f(Parameters.LocalToRelativeWorld.M[2][0], Parameters.LocalToRelativeWorld.M[2][1], Parameters.LocalToRelativeWorld.M[2][2], 0);
-			float ScaleX = FVector3f(WorldX).Size();
-			float ScaleY = FVector3f(WorldY).Size();
-			float ScaleZ = FVector3f(WorldZ).Size();
-			Parameters.NonUniformScale = FVector4f(ScaleX, ScaleY, ScaleZ, FMath::Max3(FMath::Abs(ScaleX), FMath::Abs(ScaleY), FMath::Abs(ScaleZ)));
-			Parameters.InvNonUniformScale = FVector3f(
-				ScaleX > UE_KINDA_SMALL_NUMBER ? 1.0f / ScaleX : 0.0f,
-				ScaleY > UE_KINDA_SMALL_NUMBER ? 1.0f / ScaleY : 0.0f,
-				ScaleZ > UE_KINDA_SMALL_NUMBER ? 1.0f / ScaleZ : 0.0f);
-		}
-
-		// If SingleCaptureIndex is invalid, set it to 0 since there will be a default cubemap at that slot
-		Parameters.SingleCaptureIndex = FMath::Max(Parameters.SingleCaptureIndex, 0);
-
-		Parameters.Flags = 0;
-		Parameters.Flags |= bReceivesDecals ? PRIMITIVE_SCENE_DATA_FLAG_DECAL_RECEIVER : 0u;
-		Parameters.Flags |= bHasDistanceFieldRepresentation ? PRIMITIVE_SCENE_DATA_FLAG_HAS_DISTANCE_FIELD_REPRESENTATION : 0u;
-		Parameters.Flags |= bHasCapsuleRepresentation ? PRIMITIVE_SCENE_DATA_FLAG_HAS_CAPSULE_REPRESENTATION : 0u;
-		Parameters.Flags |= bUseSingleSampleShadowFromStationaryLights ? PRIMITIVE_SCENE_DATA_FLAG_USE_SINGLE_SAMPLE_SHADOW_SL : 0u;
-		Parameters.Flags |= (bUseVolumetricLightmap && bUseSingleSampleShadowFromStationaryLights) ? PRIMITIVE_SCENE_DATA_FLAG_USE_VOLUMETRIC_LM_SHADOW_SL : 0u;
-		Parameters.Flags |= bCacheShadowAsStatic ? PRIMITIVE_SCENE_DATA_FLAG_CACHE_SHADOW_AS_STATIC : 0u;
-		Parameters.Flags |= bOutputVelocity ? PRIMITIVE_SCENE_DATA_FLAG_OUTPUT_VELOCITY : 0u;
-		Parameters.Flags |= bEvaluateWorldPositionOffset ? PRIMITIVE_SCENE_DATA_FLAG_EVALUATE_WORLD_POSITION_OFFSET : 0u;
-		Parameters.Flags |= (Parameters.LocalToRelativeWorld.RotDeterminant() < 0.0f) ? PRIMITIVE_SCENE_DATA_FLAG_DETERMINANT_SIGN : 0u;
-		Parameters.Flags |= bHasCustomData ? PRIMITIVE_SCENE_DATA_FLAG_HAS_PRIMITIVE_CUSTOM_DATA : 0u;
-		Parameters.Flags |= ((LightingChannels & 0x1) != 0) ? PRIMITIVE_SCENE_DATA_FLAG_LIGHTING_CHANNEL_0 : 0u;
-		Parameters.Flags |= ((LightingChannels & 0x2) != 0) ? PRIMITIVE_SCENE_DATA_FLAG_LIGHTING_CHANNEL_1 : 0u;
-		Parameters.Flags |= ((LightingChannels & 0x4) != 0) ? PRIMITIVE_SCENE_DATA_FLAG_LIGHTING_CHANNEL_2 : 0u;
-		Parameters.Flags |= bHasNaniteImposter ? PRIMITIVE_SCENE_DATA_FLAG_HAS_NANITE_IMPOSTER : 0u;
-		Parameters.Flags |= bHasInstanceLocalBounds ? PRIMITIVE_SCENE_DATA_FLAG_HAS_INSTANCE_LOCAL_BOUNDS : 0u;
-		Parameters.Flags |= bCastShadow ? PRIMITIVE_SCENE_DATA_FLAG_CAST_SHADOWS : 0u;
-		Parameters.Flags |= bCastContactShadow ? PRIMITIVE_SCENE_DATA_FLAG_HAS_CAST_CONTACT_SHADOW : 0u;
-		Parameters.Flags |= bCastHiddenShadow ? PRIMITIVE_SCENE_DATA_FLAG_CAST_HIDDEN_SHADOW : 0u;
-		Parameters.Flags |= bVisibleInGame ? PRIMITIVE_SCENE_DATA_FLAG_VISIBLE_IN_GAME : 0u;
-	#if WITH_EDITOR
-		Parameters.Flags |= bVisibleInEditor ? PRIMITIVE_SCENE_DATA_FLAG_VISIBLE_IN_EDITOR : 0u;
-	#endif
-		Parameters.Flags |= bVisibleInReflectionCaptures ? PRIMITIVE_SCENE_DATA_FLAG_VISIBLE_IN_REFLECTION_CAPTURES : 0u;
-		Parameters.Flags |= bVisibleInRealTimeSkyCaptures ? PRIMITIVE_SCENE_DATA_FLAG_VISIBLE_IN_REAL_TIME_SKY_CAPTURES : 0u;
-		Parameters.Flags |= bVisibleInRayTracing ? PRIMITIVE_SCENE_DATA_FLAG_VISIBLE_IN_RAY_TRACING : 0u;
-		Parameters.Flags |= bVisibleInLumenScene ? PRIMITIVE_SCENE_DATA_FLAG_VISIBLE_IN_LUMEN_SCENE : 0u;
-		Parameters.Flags |= bVisibleInSceneCaptureOnly ? PRIMITIVE_SCENE_DATA_FLAG_VISIBLE_IN_SCENE_CAPTURE_ONLY : 0u;
-		Parameters.Flags |= bHiddenInSceneCapture ? PRIMITIVE_SCENE_DATA_FLAG_HIDDEN_IN_SCENE_CAPTURE : 0u;
-		Parameters.Flags |= bForceHidden ? PRIMITIVE_SCENE_DATA_FLAG_FORCE_HIDDEN : 0u;
-		Parameters.Flags |= bHasInstanceDrawDistanceCull ? PRIMITIVE_SCENE_DATA_FLAG_INSTANCE_DRAW_DISTANCE_CULL : 0u;
-		Parameters.Flags |= bHasWPODisableDistance ? PRIMITIVE_SCENE_DATA_FLAG_WPO_DISABLE_DISTANCE : 0u;
-		Parameters.Flags |= bWritesCustomDepthStencil ? PRIMITIVE_SCENE_DATA_FLAG_WRITES_CUSTOM_DEPTH_STENCIL : 0u;
-		Parameters.Flags |= bReverseCulling ? PRIMITIVE_SCENE_DATA_FLAG_REVERSE_CULLING : 0u;
-		Parameters.Flags |= bHoldout ? PRIMITIVE_SCENE_DATA_FLAG_HOLDOUT : 0u;
-		return Parameters;
-	}
-
-private:
-	FPrimitiveUniformShaderParameters Parameters;
-	
-	FMatrix AbsoluteLocalToWorld;
-	FMatrix AbsolutePreviousLocalToWorld;
-	FVector AbsoluteObjectWorldPosition;
-	FVector AbsoluteActorWorldPosition;
-
-	float ObjectRadius;
-
-	uint32 LightingChannels : 3;
-	uint32 bReceivesDecals : 1;
-	uint32 bUseSingleSampleShadowFromStationaryLights : 1;
-	uint32 bUseVolumetricLightmap : 1;
-	uint32 bCacheShadowAsStatic : 1;
-	uint32 bOutputVelocity : 1;
-	uint32 bEvaluateWorldPositionOffset : 1;
-	uint32 bCastShadow : 1;
-	uint32 bCastContactShadow : 1;
-	uint32 bCastHiddenShadow : 1;
-	uint32 bHasDistanceFieldRepresentation : 1;
-	uint32 bHasCapsuleRepresentation : 1;
-	uint32 bHasPreSkinnedLocalBounds : 1;
-	uint32 bHasInstanceLocalBounds : 1;
-	uint32 bHasCustomData : 1;
-	uint32 bHasPreviousLocalToWorld : 1;
-	uint32 bVisibleInGame : 1;
-	uint32 bVisibleInEditor : 1;
-	uint32 bVisibleInReflectionCaptures : 1;
-	uint32 bVisibleInRealTimeSkyCaptures : 1;
-	uint32 bVisibleInRayTracing : 1;
-	uint32 bVisibleInLumenScene : 1;
-	uint32 bVisibleInSceneCaptureOnly : 1;
-	uint32 bHiddenInSceneCapture : 1;
-	uint32 bForceHidden : 1;
-	uint32 bHasNaniteImposter : 1;
-	uint32 bHasInstanceDrawDistanceCull : 1;
-	uint32 bHasWPODisableDistance : 1;
-	uint32 bWritesCustomDepthStencil : 1;
-	uint32 bReverseCulling: 1;
-	uint32 bHoldout : 1;
-};
-
-inline TUniformBufferRef<FPrimitiveUniformShaderParameters> CreatePrimitiveUniformBufferImmediate(
+ENGINE_API TUniformBufferRef<FPrimitiveUniformShaderParameters> CreatePrimitiveUniformBufferImmediate(
 	const FMatrix& LocalToWorld,
 	const FBoxSphereBounds& WorldBounds,
 	const FBoxSphereBounds& LocalBounds,
 	const FBoxSphereBounds& PreSkinnedLocalBounds,
 	bool bReceivesDecals,
 	bool bOutputVelocity
-)
-{
-	check(IsInRenderingThread());
-	return TUniformBufferRef<FPrimitiveUniformShaderParameters>::CreateUniformBufferImmediate(
-		FPrimitiveUniformShaderParametersBuilder{}
-		.Defaults()
-			.LocalToWorld(LocalToWorld)
-			.ActorWorldPosition(WorldBounds.Origin)
-			.WorldBounds(WorldBounds)
-			.LocalBounds(LocalBounds)
-			.PreSkinnedLocalBounds(PreSkinnedLocalBounds)
-			.ReceivesDecals(bReceivesDecals)
-			.OutputVelocity(bOutputVelocity)
-		.Build(),
-		UniformBuffer_MultiFrame
-	);
-}
-
-inline FPrimitiveUniformShaderParameters GetIdentityPrimitiveParameters()
-{
-	// Don't use FMatrix44f::Identity here as GetIdentityPrimitiveParameters is used by TGlobalResource<FIdentityPrimitiveUniformBuffer> and because
-	// static initialization order is undefined, FMatrix44f::Identity might be all 0's or random data the first time this is called.
-	return FPrimitiveUniformShaderParametersBuilder{}
-		.Defaults()
-			.LocalToWorld(FMatrix(FPlane(1, 0, 0, 0), FPlane(0, 1, 0, 0), FPlane(0, 0, 1, 0), FPlane(0, 0, 0, 1)))
-			.ActorWorldPosition(FVector(0.0, 0.0, 0.0))
-			.WorldBounds(FBoxSphereBounds(EForceInit::ForceInit))
-			.LocalBounds(FBoxSphereBounds(EForceInit::ForceInit))
-		.Build();
-}
+);
+
+ENGINE_API FPrimitiveUniformShaderParameters GetIdentityPrimitiveParameters();
 
 /**
  * Primitive uniform buffer containing only identity transforms.
@@ -690,3 +209,7 @@ public:
 extern ENGINE_API TGlobalResource<FSinglePrimitiveStructured> GIdentityPrimitiveBuffer;
 UE_DEPRECATED(5.3, "Use FRendererModule::CreateSinglePrimitiveSceneUniformBuffer instead")
 extern ENGINE_API TGlobalResource<FSinglePrimitiveStructured> GTilePrimitiveBuffer;
+
+#if UE_ENABLE_INCLUDE_ORDER_DEPRECATED_IN_5_3
+#include "PrimitiveUniformShaderParametersBuilder.h"
+#endif
diff --git a/Engine/Source/Runtime/Engine/Public/PrimitiveUniformShaderParametersBuilder.h b/Engine/Source/Runtime/Engine/Public/PrimitiveUniformShaderParametersBuilder.h
new file mode 100644
--- /dev/null
+++ b/Engine/Source/Runtime/Engine/Public/PrimitiveUniformShaderParametersBuilder.h
@@ -0,0 +1,439 @@
+// Copyright Epic Games, Inc. All Rights Reserved.
+
+#pragma once
+
+#include "CoreMinimal.h"
+#include "Misc/LargeWorldRenderPosition.h"
+#include "Engine/EngineTypes.h"
+#include "SceneTypes.h"
+#include "ShaderParameters.h"
+#include "UniformBuffer.h"
+#include "LightmapUniformShaderParameters.h"
+#include "NaniteDefinitions.h"
+#include "RenderTransform.h"
+#include "SceneDefinitions.h"
+
+#if !UE_ENABLE_INCLUDE_ORDER_DEPRECATED_IN_5_3
+#include "PrimitiveUniformShaderParameters.h"
+#endif
+
+struct FPrimitiveUniformShaderParametersBuilder
+{
+public:
+	inline FPrimitiveUniformShaderParametersBuilder& Defaults()
+	{
+		ObjectRadius = 0.0f;
+
+		// Flags defaulted on
+		bCastShadow									= true;
+		bCastContactShadow							= true;
+		bEvaluateWorldPositionOffset				= true;
+		bVisibleInGame								= true;
+		bVisibleInEditor							= true;
+		bVisibleInReflectionCaptures				= true;
+		bVisibleInRealTimeSkyCaptures				= true;
+		bVisibleInRayTracing						= true;
+		bVisibleInLumenScene						= true;
+
+		// Flags defaulted off
+		bReceivesDecals								= false;
+		bUseSingleSampleShadowFromStationaryLights	= false;
+		bUseVolumetricLightmap						= false;
+		bCacheShadowAsStatic						= false;
+		bOutputVelocity								= false;
+		bHasDistanceFieldRepresentation				= false;
+		bHasCapsuleRepresentation					= false;
+		bHasPreSkinnedLocalBounds					= false;
+		bHasPreviousLocalToWorld					= false;
+		bHasInstanceLocalBounds						= false;
+		bCastHiddenShadow							= false;
+		bVisibleInSceneCaptureOnly					= false;
+		bHiddenInSceneCapture						= false;
+		bForceHidden								= false;
+		bHasNaniteImposter							= false;
+		bHasInstanceDrawDistanceCull				= false;
+		bHasWPODisableDistance						= false;
+		bWritesCustomDepthStencil					= false;
+		bReverseCulling								= false;
+		bHoldout									= false;
+
+		Parameters.MaxWPOExtent						= 0.0f;
+		Parameters.MinMaterialDisplacement			= 0.0f;
+		Parameters.MaxMaterialDisplacement			= 0.0f;
+
+		// Default colors
+		Parameters.WireframeColor					= FVector3f(1.0f, 1.0f, 1.0f);
+		Parameters.LevelColor						= FVector3f(1.0f, 1.0f, 1.0f);
+
+		// Invalid indices
+		Parameters.LightmapDataIndex				= 0;
+		Parameters.LightmapUVIndex					= INDEX_NONE;
+		Parameters.SingleCaptureIndex				= INDEX_NONE;
+		Parameters.PersistentPrimitiveIndex			= INDEX_NONE;
+		Parameters.PrimitiveComponentId				= ~uint32(0u);
+
+		// Nanite
+		Parameters.NaniteResourceID						= INDEX_NONE;
+		Parameters.NaniteHierarchyOffset				= INDEX_NONE;
+		Parameters.PackedNaniteFlags					= NANITE_IMPOSTER_INDEX_MASK;
+		Parameters.NaniteRayTracingDataOffset			= INDEX_NONE;
+
+		// Instance data
+		Parameters.InstanceSceneDataOffset			= INDEX_NONE;
+		Parameters.NumInstanceSceneDataEntries		= 0;
+		Parameters.InstancePayloadDataOffset		= INDEX_NONE;
+		Parameters.InstancePayloadDataStride		= 0;
+
+		LightingChannels = GetDefaultLightingChannelMask();
+
+		return CustomPrimitiveData(nullptr);
+	}
+
+#define PRIMITIVE_UNIFORM_BUILDER_METHOD(INPUT_TYPE, VARIABLE_NAME) \
+	inline FPrimitiveUniformShaderParametersBuilder& VARIABLE_NAME(INPUT_TYPE In##VARIABLE_NAME) { Parameters.VARIABLE_NAME = In##VARIABLE_NAME; return *this; }
+
+#define PRIMITIVE_UNIFORM_BUILDER_FLAG_METHOD(INPUT_TYPE, VARIABLE_NAME) \
+	inline FPrimitiveUniformShaderParametersBuilder& VARIABLE_NAME(INPUT_TYPE In##VARIABLE_NAME) { b##VARIABLE_NAME = In##VARIABLE_NAME; return *this; }
+
+	PRIMITIVE_UNIFORM_BUILDER_FLAG_METHOD(bool,			ReceivesDecals);
+	PRIMITIVE_UNIFORM_BUILDER_FLAG_METHOD(bool,			HasDistanceFieldRepresentation);
+	PRIMITIVE_UNIFORM_BUILDER_FLAG_METHOD(bool,			HasCapsuleRepresentation);
+	PRIMITIVE_UNIFORM_BUILDER_FLAG_METHOD(bool,			HasInstanceLocalBounds);
+	PRIMITIVE_UNIFORM_BUILDER_FLAG_METHOD(bool,			CastContactShadow);
+	PRIMITIVE_UNIFORM_BUILDER_FLAG_METHOD(bool,			CastHiddenShadow);
+	PRIMITIVE_UNIFORM_BUILDER_FLAG_METHOD(bool,			CastShadow);
+	PRIMITIVE_UNIFORM_BUILDER_FLAG_METHOD(bool,			UseSingleSampleShadowFromStationaryLights);
+	PRIMITIVE_UNIFORM_BUILDER_FLAG_METHOD(bool,			UseVolumetricLightmap);
+	PRIMITIVE_UNIFORM_BUILDER_FLAG_METHOD(bool,			CacheShadowAsStatic);
+	PRIMITIVE_UNIFORM_BUILDER_FLAG_METHOD(bool,			OutputVelocity);
+	PRIMITIVE_UNIFORM_BUILDER_FLAG_METHOD(bool,			EvaluateWorldPositionOffset);
+	PRIMITIVE_UNIFORM_BUILDER_FLAG_METHOD(bool,			VisibleInGame);
+	PRIMITIVE_UNIFORM_BUILDER_FLAG_METHOD(bool,			VisibleInEditor);
+	PRIMITIVE_UNIFORM_BUILDER_FLAG_METHOD(bool,			ReverseCulling);
+	PRIMITIVE_UNIFORM_BUILDER_FLAG_METHOD(bool,			VisibleInReflectionCaptures);
+	PRIMITIVE_UNIFORM_BUILDER_FLAG_METHOD(bool,			VisibleInRealTimeSkyCaptures);
+	PRIMITIVE_UNIFORM_BUILDER_FLAG_METHOD(bool,			VisibleInRayTracing);
+	PRIMITIVE_UNIFORM_BUILDER_FLAG_METHOD(bool,			VisibleInLumenScene);
+	PRIMITIVE_UNIFORM_BUILDER_FLAG_METHOD(bool,			VisibleInSceneCaptureOnly);
+	PRIMITIVE_UNIFORM_BUILDER_FLAG_METHOD(bool,			HiddenInSceneCapture);
+	PRIMITIVE_UNIFORM_BUILDER_FLAG_METHOD(bool,			ForceHidden);
+	PRIMITIVE_UNIFORM_BUILDER_FLAG_METHOD(bool,         Holdout);
+
+	PRIMITIVE_UNIFORM_BUILDER_METHOD(uint32,			InstanceSceneDataOffset);
+	PRIMITIVE_UNIFORM_BUILDER_METHOD(uint32,			NumInstanceSceneDataEntries);
+	PRIMITIVE_UNIFORM_BUILDER_METHOD(uint32,			InstancePayloadDataOffset);
+	PRIMITIVE_UNIFORM_BUILDER_METHOD(uint32,			InstancePayloadDataStride);
+	PRIMITIVE_UNIFORM_BUILDER_METHOD(int32,				SingleCaptureIndex);
+	PRIMITIVE_UNIFORM_BUILDER_METHOD(int32,				PersistentPrimitiveIndex);
+	PRIMITIVE_UNIFORM_BUILDER_METHOD(uint32,			PrimitiveComponentId);
+	PRIMITIVE_UNIFORM_BUILDER_METHOD(uint32,			NaniteResourceID);
+	PRIMITIVE_UNIFORM_BUILDER_METHOD(uint32,			NaniteHierarchyOffset);
+	PRIMITIVE_UNIFORM_BUILDER_METHOD(uint32,			NaniteRayTracingDataOffset);
+	PRIMITIVE_UNIFORM_BUILDER_METHOD(uint32,			LightmapUVIndex);
+	PRIMITIVE_UNIFORM_BUILDER_METHOD(uint32,			LightmapDataIndex);
+
+#undef PRIMITIVE_UNIFORM_BUILDER_FLAG_METHOD
+#undef PRIMITIVE_UNIFORM_BUILDER_METHOD
+
+	inline FPrimitiveUniformShaderParametersBuilder& LightingChannelMask(uint32 InLightingChannelMask)
+	{
+		LightingChannels = InLightingChannelMask;
+		return *this;
+	}
+
+	inline FPrimitiveUniformShaderParametersBuilder& MaxWorldPositionOffsetExtent(float InMaxExtent)
+	{
+		Parameters.MaxWPOExtent = InMaxExtent;
+		return *this;
+	}
+
+	inline FPrimitiveUniformShaderParametersBuilder& MinMaxMaterialDisplacement(const FVector2f& InMinMaxDisplacement)
+	{
+		Parameters.MinMaterialDisplacement = InMinMaxDisplacement.X;
+		Parameters.MaxMaterialDisplacement = InMinMaxDisplacement.Y;
+		return *this;
+	}
+
+	inline FPrimitiveUniformShaderParametersBuilder& ObjectBounds(const FVector3f& InObjectBounds)
+	{
+		Parameters.ObjectBoundsX = InObjectBounds.X;
+		Parameters.ObjectBoundsY = InObjectBounds.Y;
+		Parameters.ObjectBoundsZ = InObjectBounds.Z;
+		return *this;
+	}
+
+	inline FPrimitiveUniformShaderParametersBuilder& WorldBounds(const FBoxSphereBounds& InWorldBounds)
+	{
+		AbsoluteObjectWorldPosition = InWorldBounds.Origin;
+		ObjectRadius = static_cast<float>(InWorldBounds.SphereRadius);											//LWC_TODO: Precision loss
+		Parameters.ObjectBoundsX = static_cast<float>(InWorldBounds.BoxExtent.X);
+		Parameters.ObjectBoundsY = static_cast<float>(InWorldBounds.BoxExtent.Y);
+		Parameters.ObjectBoundsZ = static_cast<float>(InWorldBounds.BoxExtent.Z);
+		return *this;
+	}
+
+	inline FPrimitiveUniformShaderParametersBuilder& LocalBounds(const FBoxSphereBounds& InLocalBounds)
+	{
+		Parameters.LocalObjectBoundsMin = FVector3f(InLocalBounds.GetBoxExtrema(0)); // 0 == minimum		//LWC_TODO: Precision loss
+		Parameters.LocalObjectBoundsMax = FVector3f(InLocalBounds.GetBoxExtrema(1)); // 1 == maximum		//LWC_TODO: Precision loss
+		return *this;
+	}
+
+	inline FPrimitiveUniformShaderParametersBuilder& PreSkinnedLocalBounds(const FBoxSphereBounds& InPreSkinnedLocalBounds)
+	{
+		bHasPreSkinnedLocalBounds = true;
+		Parameters.PreSkinnedLocalBoundsMin = FVector3f(InPreSkinnedLocalBounds.GetBoxExtrema(0)); // 0 == minimum		//LWC_TODO: Precision loss
+		Parameters.PreSkinnedLocalBoundsMax = FVector3f(InPreSkinnedLocalBounds.GetBoxExtrema(1)); // 1 == maximum		//LWC_TODO: Precision loss
+		return *this;
+	}
+
+	inline FPrimitiveUniformShaderParametersBuilder& ActorWorldPosition(const FVector& InActorWorldPosition)
+	{
+		AbsoluteActorWorldPosition = InActorWorldPosition;
+		return *this;
+	}
+
+	inline FPrimitiveUniformShaderParametersBuilder& LocalToWorld(const FMatrix& InLocalToWorld)
+	{
+		AbsoluteLocalToWorld = InLocalToWorld;
+		return *this;
+	}
+
+	inline FPrimitiveUniformShaderParametersBuilder& InstanceLocalBounds(const FRenderBounds& InInstanceLocalBounds)
+	{
+		bHasInstanceLocalBounds = true;
+		Parameters.InstanceLocalBoundsCenter = InInstanceLocalBounds.GetCenter();
+		Parameters.InstanceLocalBoundsExtent = InInstanceLocalBounds.GetExtent();
+		return *this;
+	}
+
+	inline FPrimitiveUniformShaderParametersBuilder& PreviousLocalToWorld(const FMatrix& InPreviousLocalToWorld)
+	{
+		bHasPreviousLocalToWorld = true;
+		AbsolutePreviousLocalToWorld = InPreviousLocalToWorld;
+		return *this;
+	}
+
+	inline FPrimitiveUniformShaderParametersBuilder& EditorColors(const FLinearColor& InWireframeColor, const FLinearColor& InLevelColor)
+	{
+		Parameters.WireframeColor = FVector3f(InWireframeColor.R, InWireframeColor.G, InWireframeColor.B);
+		Parameters.LevelColor = FVector3f(InLevelColor.R, InLevelColor.G, InLevelColor.B);
+		return *this;
+	}
+
+	inline FPrimitiveUniformShaderParametersBuilder& CustomPrimitiveData(const FCustomPrimitiveData* InCustomPrimitiveData)
+	{
+		// If this primitive has custom primitive data, set it
+		if (InCustomPrimitiveData)
+		{
+			// Copy at most up to the max supported number of floats for safety
+			FMemory::Memcpy(
+				&Parameters.CustomPrimitiveData,
+				InCustomPrimitiveData->Data.GetData(),
+				InCustomPrimitiveData->Data.GetTypeSize() * FMath::Min(InCustomPrimitiveData->Data.Num(),
+				FCustomPrimitiveData::NumCustomPrimitiveDataFloats)
+			);
+
+			bHasCustomData = true;
+		}
+		else
+		{
+			// Clear to 0
+			FMemory::Memzero(Parameters.CustomPrimitiveData);
+			bHasCustomData = false;
+		}
+
+		return *this;
+	}
+
+	inline FPrimitiveUniformShaderParametersBuilder& NaniteImposterIndex(uint32 ImposterIndex)
+	{
+		bHasNaniteImposter = ImposterIndex != INDEX_NONE;
+
+		check(!bHasNaniteImposter || ImposterIndex < NANITE_IMPOSTER_INDEX_MASK);
+		Parameters.PackedNaniteFlags = (Parameters.PackedNaniteFlags & NANITE_FILTER_FLAGS_MASK) | (ImposterIndex & NANITE_IMPOSTER_INDEX_MASK);
+
+		return *this;
+	}
+
+	inline FPrimitiveUniformShaderParametersBuilder& NaniteFilterFlags(uint32 FilterFlags)
+	{
+		check(FilterFlags < (1u << NANITE_FILTER_FLAGS_NUM_BITS));
+		Parameters.PackedNaniteFlags = (FilterFlags << NANITE_IMPOSTER_INDEX_NUM_BITS) | (Parameters.PackedNaniteFlags & NANITE_IMPOSTER_INDEX_MASK);
+
+		return *this;
+	}
+
+	inline FPrimitiveUniformShaderParametersBuilder& CustomDepthStencil(uint8 StencilValue, EStencilMask StencilWriteMask)
+	{
+		// Translate the enum to a mask that can be consumed by the GPU with fewer operations
+		uint32 GPUStencilMask;
+		switch (StencilWriteMask)
+		{		
+		case SM_Default:
+			GPUStencilMask = 0u; // Use zero to mean replace
+			break;
+		case SM_255:
+			GPUStencilMask = 0xFFu;
+			break;
+		default:
+			GPUStencilMask = (1u << uint32(StencilWriteMask - SM_1)) & 0xFFu;
+			break;
+		}
+		Parameters.CustomStencilValueAndMask = (GPUStencilMask << 8u) | StencilValue;
+		bWritesCustomDepthStencil = true;
+
+		return *this;
+	}
+
+	ENGINE_API FPrimitiveUniformShaderParametersBuilder& InstanceDrawDistance(FVector2f DistanceMinMax);
+
+	ENGINE_API FPrimitiveUniformShaderParametersBuilder& InstanceWorldPositionOffsetDisableDistance(float WPODisableDistance);
+
+	inline const FPrimitiveUniformShaderParameters& Build()
+	{
+		const FLargeWorldRenderPosition AbsoluteWorldPosition(AbsoluteLocalToWorld.GetOrigin());
+		const FVector TilePositionOffset = AbsoluteWorldPosition.GetTileOffset();
+
+		Parameters.TilePosition = AbsoluteWorldPosition.GetTile();
+
+		{
+			// Inverse on FMatrix44f can generate NaNs if the source matrix contains large scaling, so do it in double precision.
+			// Also use double precision to calculate WorldToPreviousWorld to prevent precision issues at far distances
+			FMatrix LocalToRelativeWorld = FLargeWorldRenderScalar::MakeToRelativeWorldMatrixDouble(TilePositionOffset, AbsoluteLocalToWorld);
+			FMatrix PrevLocalToRelativeWorld = FLargeWorldRenderScalar::MakeClampedToRelativeWorldMatrixDouble(TilePositionOffset, AbsolutePreviousLocalToWorld);
+			FMatrix RelativeWorldToLocal = LocalToRelativeWorld.Inverse();
+
+			Parameters.LocalToRelativeWorld = FMatrix44f(LocalToRelativeWorld);
+			Parameters.RelativeWorldToLocal = FMatrix44f(RelativeWorldToLocal);
+			Parameters.WorldToPreviousWorld = FMatrix44f(RelativeWorldToLocal * PrevLocalToRelativeWorld);
+
+			if (bHasPreviousLocalToWorld)
+			{
+				Parameters.PreviousLocalToRelativeWorld = FMatrix44f(PrevLocalToRelativeWorld);
+				Parameters.PreviousRelativeWorldToLocal = FMatrix44f(PrevLocalToRelativeWorld.Inverse());
+			}
+			else
+			{
+				Parameters.PreviousLocalToRelativeWorld = Parameters.LocalToRelativeWorld;
+				Parameters.PreviousRelativeWorldToLocal = Parameters.RelativeWorldToLocal;
+			}
+		}
+
+		Parameters.ActorRelativeWorldPosition = FVector3f(AbsoluteActorWorldPosition - TilePositionOffset);	//LWC_TODO: Precision loss
+		const FVector3f ObjectRelativeWorldPositionAsFloat = FVector3f(AbsoluteObjectWorldPosition - TilePositionOffset);
+		Parameters.ObjectRelativeWorldPositionAndRadius = FVector4f(ObjectRelativeWorldPositionAsFloat, ObjectRadius);
+
+		if (!bHasInstanceLocalBounds)
+		{
+			FRenderBounds InstanceLocalBounds(Parameters.LocalObjectBoundsMin, Parameters.LocalObjectBoundsMax);
+			Parameters.InstanceLocalBoundsCenter = InstanceLocalBounds.GetCenter();
+			Parameters.InstanceLocalBoundsExtent = InstanceLocalBounds.GetExtent();
+		}
+
+		if (!bHasPreSkinnedLocalBounds)
+		{
+			Parameters.PreSkinnedLocalBoundsMin = Parameters.LocalObjectBoundsMin;
+			Parameters.PreSkinnedLocalBoundsMax = Parameters.LocalObjectBoundsMax;
+		}
+
+		Parameters.ObjectOrientation = Parameters.LocalToRelativeWorld.GetUnitAxis(EAxis::Z);
+
+		{
+			// Extract per axis scales from LocalToWorld transform
+			FVector4f WorldX = FVector4f(Parameters.LocalToRelativeWorld.M[0][0], Parameters.LocalToRelativeWorld.M[0][1], Parameters.LocalToRelativeWorld.M[0][2], 0);
+			FVector4f WorldY = FVector4f(Parameters.LocalToRelativeWorld.M[1][0], Parameters.LocalToRelativeWorld.M[1][1], Parameters.LocalToRelativeWorld.M[1][2], 0);
+			FVector4f WorldZ = FVector4f(Parameters.LocalToRelativeWorld.M[2][0], Parameters.LocalToRelativeWorld.M[2][1], Parameters.LocalToRelativeWorld.M[2][2], 0);
+			float ScaleX = FVector3f(WorldX).Size();
+			float ScaleY = FVector3f(WorldY).Size();
+			float ScaleZ = FVector3f(WorldZ).Size();
+			Parameters.NonUniformScale = FVector4f(ScaleX, ScaleY, ScaleZ, FMath::Max3(FMath::Abs(ScaleX), FMath::Abs(ScaleY), FMath::Abs(ScaleZ)));
+			Parameters.InvNonUniformScale = FVector3f(
+				ScaleX > UE_KINDA_SMALL_NUMBER ? 1.0f / ScaleX : 0.0f,
+				ScaleY > UE_KINDA_SMALL_NUMBER ? 1.0f / ScaleY : 0.0f,
+				ScaleZ > UE_KINDA_SMALL_NUMBER ? 1.0f / ScaleZ : 0.0f);
+		}
+
+		// If SingleCaptureIndex is invalid, set it to 0 since there will be a default cubemap at that slot
+		Parameters.SingleCaptureIndex = FMath::Max(Parameters.SingleCaptureIndex, 0);
+
+		Parameters.Flags = 0;
+		Parameters.Flags |= bReceivesDecals ? PRIMITIVE_SCENE_DATA_FLAG_DECAL_RECEIVER : 0u;
+		Parameters.Flags |= bHasDistanceFieldRepresentation ? PRIMITIVE_SCENE_DATA_FLAG_HAS_DISTANCE_FIELD_REPRESENTATION : 0u;
+		Parameters.Flags |= bHasCapsuleRepresentation ? PRIMITIVE_SCENE_DATA_FLAG_HAS_CAPSULE_REPRESENTATION : 0u;
+		Parameters.Flags |= bUseSingleSampleShadowFromStationaryLights ? PRIMITIVE_SCENE_DATA_FLAG_USE_SINGLE_SAMPLE_SHADOW_SL : 0u;
+		Parameters.Flags |= (bUseVolumetricLightmap && bUseSingleSampleShadowFromStationaryLights) ? PRIMITIVE_SCENE_DATA_FLAG_USE_VOLUMETRIC_LM_SHADOW_SL : 0u;
+		Parameters.Flags |= bCacheShadowAsStatic ? PRIMITIVE_SCENE_DATA_FLAG_CACHE_SHADOW_AS_STATIC : 0u;
+		Parameters.Flags |= bOutputVelocity ? PRIMITIVE_SCENE_DATA_FLAG_OUTPUT_VELOCITY : 0u;
+		Parameters.Flags |= bEvaluateWorldPositionOffset ? PRIMITIVE_SCENE_DATA_FLAG_EVALUATE_WORLD_POSITION_OFFSET : 0u;
+		Parameters.Flags |= (Parameters.LocalToRelativeWorld.RotDeterminant() < 0.0f) ? PRIMITIVE_SCENE_DATA_FLAG_DETERMINANT_SIGN : 0u;
+		Parameters.Flags |= bHasCustomData ? PRIMITIVE_SCENE_DATA_FLAG_HAS_PRIMITIVE_CUSTOM_DATA : 0u;
+		Parameters.Flags |= ((LightingChannels & 0x1) != 0) ? PRIMITIVE_SCENE_DATA_FLAG_LIGHTING_CHANNEL_0 : 0u;
+		Parameters.Flags |= ((LightingChannels & 0x2) != 0) ? PRIMITIVE_SCENE_DATA_FLAG_LIGHTING_CHANNEL_1 : 0u;
+		Parameters.Flags |= ((LightingChannels & 0x4) != 0) ? PRIMITIVE_SCENE_DATA_FLAG_LIGHTING_CHANNEL_2 : 0u;
+		Parameters.Flags |= bHasNaniteImposter ? PRIMITIVE_SCENE_DATA_FLAG_HAS_NANITE_IMPOSTER : 0u;
+		Parameters.Flags |= bHasInstanceLocalBounds ? PRIMITIVE_SCENE_DATA_FLAG_HAS_INSTANCE_LOCAL_BOUNDS : 0u;
+		Parameters.Flags |= bCastShadow ? PRIMITIVE_SCENE_DATA_FLAG_CAST_SHADOWS : 0u;
+		Parameters.Flags |= bCastContactShadow ? PRIMITIVE_SCENE_DATA_FLAG_HAS_CAST_CONTACT_SHADOW : 0u;
+		Parameters.Flags |= bCastHiddenShadow ? PRIMITIVE_SCENE_DATA_FLAG_CAST_HIDDEN_SHADOW : 0u;
+		Parameters.Flags |= bVisibleInGame ? PRIMITIVE_SCENE_DATA_FLAG_VISIBLE_IN_GAME : 0u;
+	#if WITH_EDITOR
+		Parameters.Flags |= bVisibleInEditor ? PRIMITIVE_SCENE_DATA_FLAG_VISIBLE_IN_EDITOR : 0u;
+	#endif
+		Parameters.Flags |= bVisibleInReflectionCaptures ? PRIMITIVE_SCENE_DATA_FLAG_VISIBLE_IN_REFLECTION_CAPTURES : 0u;
+		Parameters.Flags |= bVisibleInRealTimeSkyCaptures ? PRIMITIVE_SCENE_DATA_FLAG_VISIBLE_IN_REAL_TIME_SKY_CAPTURES : 0u;
+		Parameters.Flags |= bVisibleInRayTracing ? PRIMITIVE_SCENE_DATA_FLAG_VISIBLE_IN_RAY_TRACING : 0u;
+		Parameters.Flags |= bVisibleInLumenScene ? PRIMITIVE_SCENE_DATA_FLAG_VISIBLE_IN_LUMEN_SCENE : 0u;
+		Parameters.Flags |= bVisibleInSceneCaptureOnly ? PRIMITIVE_SCENE_DATA_FLAG_VISIBLE_IN_SCENE_CAPTURE_ONLY : 0u;
+		Parameters.Flags |= bHiddenInSceneCapture ? PRIMITIVE_SCENE_DATA_FLAG_HIDDEN_IN_SCENE_CAPTURE : 0u;
+		Parameters.Flags |= bForceHidden ? PRIMITIVE_SCENE_DATA_FLAG_FORCE_HIDDEN : 0u;
+		Parameters.Flags |= bHasInstanceDrawDistanceCull ? PRIMITIVE_SCENE_DATA_FLAG_INSTANCE_DRAW_DISTANCE_CULL : 0u;
+		Parameters.Flags |= bHasWPODisableDistance ? PRIMITIVE_SCENE_DATA_FLAG_WPO_DISABLE_DISTANCE : 0u;
+		Parameters.Flags |= bWritesCustomDepthStencil ? PRIMITIVE_SCENE_DATA_FLAG_WRITES_CUSTOM_DEPTH_STENCIL : 0u;
+		Parameters.Flags |= bReverseCulling ? PRIMITIVE_SCENE_DATA_FLAG_REVERSE_CULLING : 0u;
+		Parameters.Flags |= bHoldout ? PRIMITIVE_SCENE_DATA_FLAG_HOLDOUT : 0u;
+		return Parameters;
+	}
+
+private:
+	FPrimitiveUniformShaderParameters Parameters;
+	
+	FMatrix AbsoluteLocalToWorld;
+	FMatrix AbsolutePreviousLocalToWorld;
+	FVector AbsoluteObjectWorldPosition;
+	FVector AbsoluteActorWorldPosition;
+
+	float ObjectRadius;
+
+	uint32 LightingChannels : 3;
+	uint32 bReceivesDecals : 1;
+	uint32 bUseSingleSampleShadowFromStationaryLights : 1;
+	uint32 bUseVolumetricLightmap : 1;
+	uint32 bCacheShadowAsStatic : 1;
+	uint32 bOutputVelocity : 1;
+	uint32 bEvaluateWorldPositionOffset : 1;
+	uint32 bCastShadow : 1;
+	uint32 bCastContactShadow : 1;
+	uint32 bCastHiddenShadow : 1;
+	uint32 bHasDistanceFieldRepresentation : 1;
+	uint32 bHasCapsuleRepresentation : 1;
+	uint32 bHasPreSkinnedLocalBounds : 1;
+	uint32 bHasInstanceLocalBounds : 1;
+	uint32 bHasCustomData : 1;
+	uint32 bHasPreviousLocalToWorld : 1;
+	uint32 bVisibleInGame : 1;
+	uint32 bVisibleInEditor : 1;
+	uint32 bVisibleInReflectionCaptures : 1;
+	uint32 bVisibleInRealTimeSkyCaptures : 1;
+	uint32 bVisibleInRayTracing : 1;
+	uint32 bVisibleInLumenScene : 1;
+	uint32 bVisibleInSceneCaptureOnly : 1;
+	uint32 bHiddenInSceneCapture : 1;
+	uint32 bForceHidden : 1;
+	uint32 bHasNaniteImposter : 1;
+	uint32 bHasInstanceDrawDistanceCull : 1;
+	uint32 bHasWPODisableDistance : 1;
+	uint32 bWritesCustomDepthStencil : 1;
+	uint32 bReverseCulling: 1;
+	uint32 bHoldout : 1;
+};
diff --git a/Engine/Source/Runtime/Renderer/Private/GPUScene.cpp b/Engine/Source/Runtime/Renderer/Private/GPUScene.cpp
--- a/Engine/Source/Runtime/Renderer/Private/GPUScene.cpp
+++ b/Engine/Source/Runtime/Renderer/Private/GPUScene.cpp
@@ -27,6 +27,7 @@
 #include "LightSceneData.h"
 #include "LightSceneProxy.h"
 #include "SystemTextures.h"
+#include "SceneDefinitions.h"
 
 #define LOG_INSTANCE_ALLOCATIONS 0
 
