Replaced FSparseVolumeTextureFrame and USparseVolumeTextureFrame (which was just a view) with a new USparseVolumeTextureFrame class that is owned by UStreamableSparseVolumeTexture. These USVTFrames share their lifetime with the parent SVT which solves the problem of USVTFrames (with view behavior) acting as dangling pointers if they are kept around for longer than a single frame.

#rnx

diff --git a/Engine/Plugins/FX/Niagara/Source/Niagara/Private/NiagaraDataInterfaceRenderTargetVolume.cpp b/Engine/Plugins/FX/Niagara/Source/Niagara/Private/NiagaraDataInterfaceRenderTargetVolume.cpp
--- a/Engine/Plugins/FX/Niagara/Source/Niagara/Private/NiagaraDataInterfaceRenderTargetVolume.cpp
+++ b/Engine/Plugins/FX/Niagara/Source/Niagara/Private/NiagaraDataInterfaceRenderTargetVolume.cpp
@@ -826,7 +826,7 @@ bool UNiagaraDataInterfaceRenderTargetVolume::SimCacheReadFrame(UObject* Storage
 	{			
 		UAnimatedSparseVolumeTexture* SVT = Cast<UAnimatedSparseVolumeTexture>(StorageObject);
 		
-		USparseVolumeTextureFrame *SVTFrame = USparseVolumeTextureFrame::CreateFrame(SVT, FrameA, 0);
+		USparseVolumeTextureFrame *SVTFrame = USparseVolumeTextureFrame::GetFrame(SVT, FrameA);
 		
 		FUintVector4 CurrentPackedUniforms0 = FUintVector4();
 		FUintVector4 CurrentPackedUniforms1 = FUintVector4();
diff --git a/Engine/Source/Runtime/Engine/Classes/SparseVolumeTexture/SparseVolumeTexture.h b/Engine/Source/Runtime/Engine/Classes/SparseVolumeTexture/SparseVolumeTexture.h
--- a/Engine/Source/Runtime/Engine/Classes/SparseVolumeTexture/SparseVolumeTexture.h
+++ b/Engine/Source/Runtime/Engine/Classes/SparseVolumeTexture/SparseVolumeTexture.h
@@ -44,26 +44,6 @@ struct ENGINE_API FSparseVolumeTextureHeader
 	void Serialize(FArchive& Ar);
 };
 
-struct ENGINE_API FSparseVolumeTextureFrame
-{
-	// The frame data that can be streamed in when in game.
-	FByteBulkData						RuntimeStreamedInData;
-
-	// The render side proxy for the sparse volume texture asset.
-	FSparseVolumeTextureSceneProxy*		SparseVolumeTextureSceneProxy;
-
-#if WITH_EDITORONLY_DATA
-	/** The raw data that can be loaded when we want to update cook the data with different settings or updated code without re importing. */
-	UE::Serialization::FEditorBulkData	RawData;
-#endif
-
-	FSparseVolumeTextureFrame();
-	virtual ~FSparseVolumeTextureFrame();
-	bool BuildDerivedData(const FIntVector3& VolumeResolution, TextureAddress AddressX, TextureAddress AddressY, TextureAddress AddressZ, FSparseVolumeTextureData* OutMippedTextureData);
-	void Serialize(FArchive& Ar, class UStreamableSparseVolumeTexture* Owner, int32 FrameIndex);
-};
-
-
 enum ESparseVolumeTextureShaderUniform
 {
 	ESparseVolumeTexture_TileSize,
@@ -119,7 +99,6 @@ public:
 	/** Getter for the shader uniform parameter type with index as ESparseVolumeTextureShaderUniform. */
 	static UE::Shader::EValueType GetUniformParameterType(int32 Index);
 
-protected:
 #if WITH_EDITOR
 	enum class ENotifyMaterialsEffectOnShaders
 	{
@@ -132,6 +111,58 @@ protected:
 #endif // WITH_EDITOR
 };
 
+UCLASS(ClassGroup = Rendering, BlueprintType)
+class ENGINE_API USparseVolumeTextureFrame : public USparseVolumeTexture
+{
+	GENERATED_UCLASS_BODY()
+
+	// SVT_TODO: Remove this once the new streaming manager has been implemented. For now we need this to keep the old system running.
+	friend struct FStreamingSparseVolumeTextureData;
+
+public:
+
+	USparseVolumeTextureFrame();
+	virtual ~USparseVolumeTextureFrame() = default;
+
+	static USparseVolumeTextureFrame* GetFrame(USparseVolumeTexture* SparseVolumeTexture, int32 FrameIndex);
+
+	bool Initialize(USparseVolumeTexture* InOwner, int32 InFrameIndex, FSparseVolumeTextureData& UncookedFrame);
+	bool BuildDerivedData(FSparseVolumeTextureData* OutMippedTextureData);
+	int32 GetFrameIndex() const { return FrameIndex; }
+	void GenerateOrLoadDDCRuntimeData(UE::DerivedData::FRequestOwner& DDCRequestOwner);
+
+	//~ Begin UObject Interface.
+	virtual void PostLoad() override;
+	virtual void FinishDestroy() override;
+	virtual void BeginDestroy() override;
+	virtual void Serialize(FArchive& Ar) override;
+	virtual void GetResourceSizeEx(FResourceSizeEx& CumulativeResourceSize) override;
+	//~ End UObject Interface.
+
+	//~ Begin USparseVolumeTexture Interface.
+	virtual int32 GetNumFrames() const override { return 1; }
+	virtual int32 GetNumMipLevels() const { return Owner->GetNumMipLevels(); }
+	virtual FIntVector GetVolumeResolution() const { return Owner->GetVolumeResolution(); }
+	virtual TextureAddress GetTextureAddressX() const { return Owner->GetTextureAddressX(); }
+	virtual TextureAddress GetTextureAddressY() const { return Owner->GetTextureAddressY(); }
+	virtual TextureAddress GetTextureAddressZ() const { return Owner->GetTextureAddressZ(); }
+	virtual const FSparseVolumeTextureSceneProxy* GetSparseVolumeTextureSceneProxy() const { return SceneProxy; }
+	//~ End USparseVolumeTexture Interface.
+
+private:
+	UPROPERTY()
+	TObjectPtr<USparseVolumeTexture> Owner;
+	int32 FrameIndex;
+
+	FSparseVolumeTextureSceneProxy* SceneProxy;
+	FByteBulkData StreamingData;
+
+#if WITH_EDITORONLY_DATA
+	/** The raw data that can be loaded when we want to update cook the data with different settings or updated code without re importing. */
+	UE::Serialization::FEditorBulkData	RawData;
+#endif
+};
+
 UCLASS(ClassGroup = Rendering, BlueprintType)//, hidecategories = (Object))
 class UStreamableSparseVolumeTexture : public USparseVolumeTexture
 {
@@ -179,7 +210,7 @@ public:
 	virtual bool Initialize(const TArrayView<FSparseVolumeTextureData>& UncookedData, int32 NumMipLevels = INDEX_NONE /*Create entire mip chain by default*/);
 
 	const FSparseVolumeTextureSceneProxy* GetStreamedFrameProxyOrFallback(int32 FrameIndex, int32 MipLevel) const;
-	TArrayView<const FSparseVolumeTextureFrame> GetFrames() const { return Frames; };
+	USparseVolumeTextureFrame* GetFrame(int32 FrameIndex) const { return Frames.IsValidIndex(FrameIndex) ? Frames[FrameIndex] : nullptr; }
 
 	//~ Begin UObject Interface.
 	virtual void PostLoad() override;
@@ -204,7 +235,8 @@ public:
 
 protected:
 
-	TArray<FSparseVolumeTextureFrame> Frames;
+	UPROPERTY()
+	TArray<TObjectPtr<USparseVolumeTextureFrame>> Frames;
 #if WITH_EDITORONLY_DATA
 	FIntVector VolumeBoundsMin;
 	FIntVector VolumeBoundsMax;
@@ -212,7 +244,6 @@ protected:
 #endif // WITH_EDITORONLY_DATA
 
 	void GenerateOrLoadDDCRuntimeDataAndCreateSceneProxy();
-	void GenerateOrLoadDDCRuntimeDataForFrame(FSparseVolumeTextureFrame& Frame, UE::DerivedData::FRequestOwner& DDCRequestOwner);
 };
 
 UCLASS(ClassGroup = Rendering, BlueprintType)//, hidecategories = (Object))
@@ -255,38 +286,6 @@ private:
 	int32 PreviewMipLevel;
 };
 
-// USparseVolumeTextureFrame inherits from USparseVolumeTexture to be viewed using any given frame of a UAnimatedSparseVolumeTexture (or UStaticSparseVolumeTexture)
-UCLASS(ClassGroup = Rendering, BlueprintType)//, hidecategories = (Object))
-class ENGINE_API USparseVolumeTextureFrame : public USparseVolumeTexture
-{
-	GENERATED_UCLASS_BODY()
-
-public:
-	USparseVolumeTextureFrame();
-	virtual ~USparseVolumeTextureFrame() = default;
-
-	static USparseVolumeTextureFrame* CreateFrame(USparseVolumeTexture* Texture, int32 FrameIndex, int32 MipLevel);
-
-	void Initialize(const FSparseVolumeTextureSceneProxy* InSceneProxy, const FIntVector& InVolumeResolution, TextureAddress InAddressX, TextureAddress InAddressY, TextureAddress InAddressZ);
-
-	//~ Begin USparseVolumeTexture Interface.
-	virtual int32 GetNumFrames() const override { return 1; }
-	virtual int32 GetNumMipLevels() const override { return 1; }
-	virtual FIntVector GetVolumeResolution() const override { return VolumeResolution; };
-	virtual TextureAddress GetTextureAddressX() const override { return AddressX; }
-	virtual TextureAddress GetTextureAddressY() const override { return AddressY; }
-	virtual TextureAddress GetTextureAddressZ() const override { return AddressZ; }
-	virtual const FSparseVolumeTextureSceneProxy* GetSparseVolumeTextureSceneProxy() const override { return SceneProxy; };
-	//~ End USparseVolumeTexture Interface.
-
-private:
-	FIntVector3 VolumeResolution;
-	TEnumAsByte<enum TextureAddress> AddressX;
-	TEnumAsByte<enum TextureAddress> AddressY;
-	TEnumAsByte<enum TextureAddress> AddressZ;
-	const FSparseVolumeTextureSceneProxy* SceneProxy;
-};
-
 UCLASS(ClassGroup = Rendering, BlueprintType)//, hidecategories = (Object))
 class ENGINE_API UAnimatedSparseVolumeTextureController : public UObject
 {
diff --git a/Engine/Source/Runtime/Engine/Private/Components/HeterogeneousVolumeComponent.cpp b/Engine/Source/Runtime/Engine/Private/Components/HeterogeneousVolumeComponent.cpp
--- a/Engine/Source/Runtime/Engine/Private/Components/HeterogeneousVolumeComponent.cpp
+++ b/Engine/Source/Runtime/Engine/Private/Components/HeterogeneousVolumeComponent.cpp
@@ -291,7 +291,7 @@ void UHeterogeneousVolumeComponent::TickComponent(float DeltaTime, ELevelTick Ti
 				{
 					int32 FrameIndex = int32(Time * Framerate) % DefaultSVT->GetNumFrames();
 					int32 MipLevel = 0;
-					USparseVolumeTextureFrame* SVTFrame = USparseVolumeTextureFrame::CreateFrame(DefaultSVT, FrameIndex, MipLevel);
+					USparseVolumeTextureFrame* SVTFrame = USparseVolumeTextureFrame::GetFrame(DefaultSVT, FrameIndex);
 					MaterialInstanceDynamic->SetSparseVolumeTextureParameterValue(SVTParameterInfo[MaterialIndex].Name, SVTFrame);
 
 					VolumeResolution = SVTFrame ? SVTFrame->GetVolumeResolution() : FIntVector(1);
diff --git a/Engine/Source/Runtime/Engine/Private/SparseVolumeTexture/SparseVolumeTexture.cpp b/Engine/Source/Runtime/Engine/Private/SparseVolumeTexture/SparseVolumeTexture.cpp
--- a/Engine/Source/Runtime/Engine/Private/SparseVolumeTexture/SparseVolumeTexture.cpp
+++ b/Engine/Source/Runtime/Engine/Private/SparseVolumeTexture/SparseVolumeTexture.cpp
@@ -106,117 +106,6 @@ void FSparseVolumeTextureHeader::Serialize(FArchive& Ar)
 
 ////////////////////////////////////////////////////////////////////////////////////////////////
 
-FSparseVolumeTextureFrame::FSparseVolumeTextureFrame()
-	: RuntimeStreamedInData()
-	, SparseVolumeTextureSceneProxy()
-#if WITH_EDITORONLY_DATA
-	, RawData()
-#endif
-{
-}
-
-FSparseVolumeTextureFrame::~FSparseVolumeTextureFrame()
-{
-}
-
-bool FSparseVolumeTextureFrame::BuildDerivedData(const FIntVector3& VolumeResolution, TextureAddress AddressX, TextureAddress AddressY, TextureAddress AddressZ, FSparseVolumeTextureData* OutMippedTextureData)
-{
-#if WITH_EDITORONLY_DATA
-	// Check if the virtualized bulk data payload is available now
-	if (RawData.HasPayloadData())
-	{
-		// First, read the source data in from the raw data stored as bulk data
-		UE::Serialization::FEditorBulkDataReader RawDataArchiveReader(RawData);
-		FSparseVolumeTextureData TextureData;
-		TextureData.Serialize(RawDataArchiveReader);
-
-		FSparseVolumeTextureDataAddressingInfo AddressingInfo{};
-		AddressingInfo.VolumeResolution = VolumeResolution;
-		AddressingInfo.AddressX = AddressX;
-		AddressingInfo.AddressY = AddressY;
-		AddressingInfo.AddressZ = AddressZ;
-
-		const int32 NumMipLevels = -1; // generate entire mip chain
-		const bool bMoveMip0FromSource = true; // we have no need to keep TextureData around
-		if (!TextureData.BuildDerivedData(AddressingInfo, NumMipLevels, bMoveMip0FromSource, *OutMippedTextureData))
-		{
-			return false;
-		}
-
-		// Now unload the raw data
-		RawData.UnloadData();
-		
-		return true;
-	}
-#endif
-	return false;
-}
-
-void FSparseVolumeTextureFrame::Serialize(FArchive& Ar, UStreamableSparseVolumeTexture* Owner, int32 FrameIndex)
-{
-	FStripDataFlags StripFlags(Ar);
-
-	const bool bInlinePayload = (FrameIndex == 0);
-	RuntimeStreamedInData.SetBulkDataFlags(bInlinePayload ? BULKDATA_ForceInlinePayload : BULKDATA_Force_NOT_InlinePayload);
-
-	if (StripFlags.IsEditorDataStripped() && Ar.IsLoadingFromCookedPackage())
-	{
-		// In this case we are loading in game with a cooked build so we only need to load the runtime data.
-
-		// Read cooked bulk data from archive
-		RuntimeStreamedInData.Serialize(Ar, Owner);
-
-		if (bInlinePayload)
-		{
-			SparseVolumeTextureSceneProxy = new FSparseVolumeTextureSceneProxy();
-
-			// Create runtime data from cooked bulk data
-			{
-				FBulkDataReader BulkDataReader(RuntimeStreamedInData);
-				FSparseVolumeTextureData TextureData;
-				TextureData.Serialize(BulkDataReader);
-				bool bSuccess = SparseVolumeTextureSceneProxy->GetRuntimeData().Create(TextureData);
-				check(bSuccess); // SVT_TODO
-			}
-
-			// The bulk data is no longer needed
-			RuntimeStreamedInData.RemoveBulkData();
-
-			// Runtime data is now valid, initialize the render thread proxy
-			BeginInitResource(SparseVolumeTextureSceneProxy);
-		}
-	}
-	else if (Ar.IsCooking())
-	{
-		// We are cooking the game, serialize the asset out.
-
-		FSparseVolumeTextureData DerivedData;
-		const bool bBuiltDerivedData = BuildDerivedData(Owner->VolumeResolution, Owner->AddressX, Owner->AddressY, Owner->AddressZ, &DerivedData);
-		check(bBuiltDerivedData); // SVT_TODO: actual error handling
-
-		// Write derived data into RuntimeStreamedInData
-		{
-			FBulkDataWriter BulkDataWriter(RuntimeStreamedInData);
-			DerivedData.Serialize(BulkDataWriter);
-		}
-
-		// And now write the cooked bulk data to the archive
-		RuntimeStreamedInData.Serialize(Ar, Owner);
-	}
-	else if (!Ar.IsObjectReferenceCollector())
-	{
-#if WITH_EDITORONLY_DATA
-		// When in EDITOR:
-		//  - We only serialize raw data 
-		//  - The runtime data is fetched/put from/to DDC
-		//  - This EditorBulk data do not load the full and huge OpenVDB data. That is only done explicitly later.
-		RawData.Serialize(Ar, Owner);
-#endif
-	}
-}
-
-////////////////////////////////////////////////////////////////////////////////////////////////
-
 USparseVolumeTexture::USparseVolumeTexture(const FObjectInitializer& ObjectInitializer)
 	: Super(ObjectInitializer)
 {
@@ -379,6 +268,259 @@ void USparseVolumeTexture::NotifyMaterials(const ENotifyMaterialsEffectOnShaders
 
 ////////////////////////////////////////////////////////////////////////////////////////////////
 
+USparseVolumeTextureFrame::USparseVolumeTextureFrame(const FObjectInitializer& ObjectInitializer)
+	: Super(ObjectInitializer)
+{
+}
+
+USparseVolumeTextureFrame* USparseVolumeTextureFrame::GetFrame(USparseVolumeTexture* SparseVolumeTexture, int32 FrameIndex)
+{
+	if (UStreamableSparseVolumeTexture* StreamableSVT = Cast<UStreamableSparseVolumeTexture>(SparseVolumeTexture))
+	{
+		return StreamableSVT->GetFrame(FrameIndex);
+	}
+	return nullptr;
+}
+
+bool USparseVolumeTextureFrame::Initialize(USparseVolumeTexture* InOwner, int32 InFrameIndex, FSparseVolumeTextureData& UncookedFrame)
+{
+#if WITH_EDITORONLY_DATA
+	Owner = InOwner;
+	FrameIndex = InFrameIndex;
+	{
+		UE::Serialization::FEditorBulkDataWriter RawDataArchiveWriter(RawData);
+		UncookedFrame.Serialize(RawDataArchiveWriter);
+	}
+
+	return true;
+#else
+	return false;
+#endif
+}
+
+bool USparseVolumeTextureFrame::BuildDerivedData(FSparseVolumeTextureData* OutMippedTextureData)
+{
+#if WITH_EDITORONLY_DATA
+	// Check if the virtualized bulk data payload is available now
+	if (RawData.HasPayloadData())
+	{
+		// First, read the source data in from the raw data stored as bulk data
+		UE::Serialization::FEditorBulkDataReader RawDataArchiveReader(RawData);
+		FSparseVolumeTextureData TextureData;
+		TextureData.Serialize(RawDataArchiveReader);
+
+		FSparseVolumeTextureDataAddressingInfo AddressingInfo{};
+		AddressingInfo.VolumeResolution = GetVolumeResolution();
+		AddressingInfo.AddressX = GetTextureAddressX();
+		AddressingInfo.AddressY = GetTextureAddressY();
+		AddressingInfo.AddressZ = GetTextureAddressZ();
+
+		const int32 NumMipLevels = -1; // generate entire mip chain
+		const bool bMoveMip0FromSource = true; // we have no need to keep TextureData around
+		if (!TextureData.BuildDerivedData(AddressingInfo, NumMipLevels, bMoveMip0FromSource, *OutMippedTextureData))
+		{
+			return false;
+		}
+
+		// Now unload the raw data
+		RawData.UnloadData();
+
+		return true;
+	}
+#endif
+	return false;
+}
+
+void USparseVolumeTextureFrame::PostLoad()
+{
+	Super::PostLoad();
+}
+
+void USparseVolumeTextureFrame::FinishDestroy()
+{
+	Super::FinishDestroy();
+}
+
+void USparseVolumeTextureFrame::BeginDestroy()
+{
+	if (SceneProxy)
+	{
+		ENQUEUE_RENDER_COMMAND(USparseVolumeTextureFrame_DeleteSVTProxy)(
+			[Proxy = SceneProxy](FRHICommandListImmediate& RHICmdList)
+			{
+				Proxy->ReleaseResource();
+				delete Proxy;
+			});
+		SceneProxy = nullptr;
+	}
+
+	Super::BeginDestroy();
+}
+
+void USparseVolumeTextureFrame::Serialize(FArchive& Ar)
+{
+	Super::Serialize(Ar);
+
+	FStripDataFlags StripFlags(Ar);
+
+	const bool bInlinePayload = (FrameIndex == 0);
+	StreamingData.SetBulkDataFlags(bInlinePayload ? BULKDATA_ForceInlinePayload : BULKDATA_Force_NOT_InlinePayload);
+
+	if (StripFlags.IsEditorDataStripped() && Ar.IsLoadingFromCookedPackage())
+	{
+		// In this case we are loading in game with a cooked build so we only need to load the runtime data.
+
+		// Read cooked bulk data from archive
+		StreamingData.Serialize(Ar, Owner);
+
+		if (bInlinePayload)
+		{
+			SceneProxy = new FSparseVolumeTextureSceneProxy();
+
+			// Create runtime data from cooked bulk data
+			{
+				FBulkDataReader BulkDataReader(StreamingData);
+				FSparseVolumeTextureData TextureData;
+				TextureData.Serialize(BulkDataReader);
+				bool bSuccess = SceneProxy->GetRuntimeData().Create(TextureData);
+				check(bSuccess); // SVT_TODO
+			}
+
+			// The bulk data is no longer needed
+			StreamingData.RemoveBulkData();
+
+			// Runtime data is now valid, initialize the render thread proxy
+			BeginInitResource(SceneProxy);
+		}
+	}
+	else if (Ar.IsCooking())
+	{
+		// We are cooking the game, serialize the asset out.
+
+		FSparseVolumeTextureData DerivedData;
+		const bool bBuiltDerivedData = BuildDerivedData(&DerivedData);
+		check(bBuiltDerivedData); // SVT_TODO: actual error handling
+
+		// Write derived data into StreamingData
+		{
+			FBulkDataWriter BulkDataWriter(StreamingData);
+			DerivedData.Serialize(BulkDataWriter);
+		}
+
+		// And now write the cooked bulk data to the archive
+		StreamingData.Serialize(Ar, Owner);
+	}
+	else if (!Ar.IsObjectReferenceCollector())
+	{
+#if WITH_EDITORONLY_DATA
+		// When in EDITOR:
+		//  - We only serialize raw data 
+		//  - The runtime data is fetched/put from/to DDC
+		//  - This EditorBulk data do not load the full and huge OpenVDB data. That is only done explicitly later.
+		RawData.Serialize(Ar, Owner);
+#endif
+	}
+}
+
+void USparseVolumeTextureFrame::GetResourceSizeEx(FResourceSizeEx& CumulativeResourceSize)
+{
+	Super::GetResourceSizeEx(CumulativeResourceSize);
+}
+
+void USparseVolumeTextureFrame::GenerateOrLoadDDCRuntimeData(UE::DerivedData::FRequestOwner& DDCRequestOwner)
+{
+#if WITH_EDITORONLY_DATA
+	using namespace UE::DerivedData;
+
+	// Release any previously allocated render thread proxy
+	if (SceneProxy)
+	{
+		BeginReleaseResource(SceneProxy);
+	}
+	else
+	{
+		SceneProxy = new FSparseVolumeTextureSceneProxy();
+	}
+
+	static const FString SparseVolumeTextureDDCVersion = TEXT("381AE2A9-A903-4C8F-8486-891E24D6EC70");	// Bump this if you want to ignore all cached data so far.
+	const FString DerivedDataKey = RawData.GetIdentifier().ToString()
+		+ FString::Format(TEXT("{0},{1},{2}"), { GetTextureAddressX(), GetTextureAddressY(), GetTextureAddressZ() })
+		+ SparseVolumeTextureDDCVersion;
+
+	const FCacheKey Key = ConvertLegacyCacheKey(DerivedDataKey);
+	const FSharedString Name = MakeStringView(GetPathName());
+
+	UE::DerivedData::GetCache().GetValue({ {Name, Key} }, DDCRequestOwner,
+		[this, &DDCRequestOwner](FCacheGetValueResponse&& Response)
+		{
+			if (Response.Status == EStatus::Ok)
+			{
+				DDCRequestOwner.LaunchTask(TEXT("USparseVolumeTextureFrame_DerivedDataLoad"),
+					[this, Value = MoveTemp(Response.Value)]()
+					{
+						FSharedBuffer Data = Value.GetData().Decompress();
+						FMemoryReaderView Ar(Data, true /*bIsPersistent*/);
+						FSparseVolumeTextureData TextureData;
+						TextureData.Serialize(Ar);
+						bool bSuccess = SceneProxy->GetRuntimeData().Create(TextureData);
+						check(bSuccess); // SVT_TODO
+
+						// Runtime data is now valid, initialize the render thread proxy
+						BeginInitResource(SceneProxy);
+					});
+			}
+			else if (Response.Status == EStatus::Error)
+			{
+				DDCRequestOwner.LaunchTask(TEXT("USparseVolumeTextureFrame_DerivedDataBuild"),
+					[this, &DDCRequestOwner, Name = Response.Name, Key = Response.Key]()
+					{
+						FSparseVolumeTextureRuntime& RuntimeData = SceneProxy->GetRuntimeData();
+
+						// Check if the virtualized bulk data payload is available now
+						if (RawData.HasPayloadData())
+						{
+							FSparseVolumeTextureData TextureData;
+							bool bSuccess = BuildDerivedData(&TextureData);
+							ensure(bSuccess);
+
+							bSuccess = RuntimeData.Create(TextureData);
+							ensure(bSuccess);
+
+							// Using a LargeMemoryWriter for serialization since the data can be bigger than 2 GB
+							FLargeMemoryWriter LargeMemWriter(0, /*bIsPersistent=*/ true);
+							TextureData.Serialize(LargeMemWriter);
+
+							const int64 UncompressedSize = LargeMemWriter.TotalSize();
+
+							// Since the DDC doesn't support data bigger than 2 GB, we only cache for such uncompressed size.
+							constexpr int64 SizeThreshold = 2147483648LL;	// 2GB
+							const bool bIsCacheable = UncompressedSize < SizeThreshold;
+							if (bIsCacheable)
+							{
+								FValue Value = FValue::Compress(FSharedBuffer::MakeView(LargeMemWriter.GetView()));
+								UE::DerivedData::GetCache().PutValue({ {Name, Key, Value} }, DDCRequestOwner);
+							}
+							else
+							{
+								UE_LOG(LogSparseVolumeTexture, Error, TEXT("SparseVolumeTexture - the asset is too large to fit in Derived Data Cache %s"), *GetName());
+							}
+						}
+						else
+						{
+							UE_LOG(LogSparseVolumeTexture, Error, TEXT("SparseVolumeTexture - Raw source data is not available for %s. Using default data."), *GetName());
+							RuntimeData.SetAsDefaultTexture();
+						}
+
+						// Runtime data is now valid, initialize the render thread proxy
+						BeginInitResource(SceneProxy);
+					});
+			}
+		});
+#endif // WITH_EDITORONLY_DATA
+}
+
+////////////////////////////////////////////////////////////////////////////////////////////////
+
 UStreamableSparseVolumeTexture::UStreamableSparseVolumeTexture(const FObjectInitializer& ObjectInitializer)
 	: Super(ObjectInitializer)
 {
@@ -433,13 +575,15 @@ bool UStreamableSparseVolumeTexture::AppendFrame(FSparseVolumeTextureData& Uncoo
 	VolumeBoundsMax.Z = FMath::Max(VolumeBoundsMax.Z, UncookedFrame.Header.VirtualVolumeAABBMax.Z);
 
 	VolumeResolution = VolumeBoundsMax;
-	Frames.AddDefaulted();
+
+	USparseVolumeTextureFrame* Frame = NewObject<USparseVolumeTextureFrame>(this);
+	if (Frame->Initialize(this, Frames.Num(), UncookedFrame))
 	{
-		UE::Serialization::FEditorBulkDataWriter RawDataArchiveWriter(Frames.Last().RawData);
-		UncookedFrame.Serialize(RawDataArchiveWriter);
+		Frames.Add(Frame);
+		return true;
 	}
-
-	return true;
+	return false;
+	
 #else
 	return false;
 #endif
@@ -529,7 +673,7 @@ const FSparseVolumeTextureSceneProxy* UStreamableSparseVolumeTexture::GetStreame
 	FrameIndex = FMath::Clamp(FrameIndex, 0, Frames.Num() - 1);
 
 #if WITH_EDITORONLY_DATA
-	return Frames[FrameIndex].SparseVolumeTextureSceneProxy;
+	return Frames[FrameIndex]->GetSparseVolumeTextureSceneProxy();
 #else
 	ISparseVolumeTextureStreamingManager& StreamingManager = IStreamingManager::Get().GetSparseVolumeTextureStreamingManager();
 	const FSparseVolumeTextureSceneProxy* Proxy = StreamingManager.GetSparseVolumeTextureSceneProxy(this, FrameIndex, MipLevel, true);
@@ -571,19 +715,6 @@ void UStreamableSparseVolumeTexture::FinishDestroy()
 void UStreamableSparseVolumeTexture::BeginDestroy()
 {
 	Super::BeginDestroy();
-	for (FSparseVolumeTextureFrame& Frame : Frames)
-	{
-		if (Frame.SparseVolumeTextureSceneProxy)
-		{
-			ENQUEUE_RENDER_COMMAND(UStreamableSparseVolumeTexture_DeleteSVTProxy)(
-				[Proxy = Frame.SparseVolumeTextureSceneProxy](FRHICommandListImmediate& RHICmdList)
-				{
-					Proxy->ReleaseResource();
-					delete Proxy;
-				});
-			Frame.SparseVolumeTextureSceneProxy = nullptr;
-		}
-	}
 }
 
 void UStreamableSparseVolumeTexture::Serialize(FArchive& Ar)
@@ -615,19 +746,6 @@ void UStreamableSparseVolumeTexture::Serialize(FArchive& Ar)
 	UE::SVT::Private::SerializeEnumAs<uint8>(Ar, *InitStatePtr);
 	Ar << *VolumeBoundsMinPtr;
 	Ar << *VolumeBoundsMaxPtr;
-
-	int32 NumFrames = Frames.Num();
-	Ar << NumFrames;
-
-	if (Ar.IsLoading())
-	{
-		Frames.Reset(NumFrames);
-		Frames.AddDefaulted(NumFrames);
-	}
-	for (int32 FrameIndex = 0; FrameIndex < NumFrames; ++FrameIndex)
-	{
-		Frames[FrameIndex].Serialize(Ar, this, FrameIndex);
-	}
 }
 
 #if WITH_EDITOR
@@ -639,6 +757,10 @@ void UStreamableSparseVolumeTexture::PostEditChangeProperty(FPropertyChangedEven
 	{
 		// SVT need to recompile shaders when address mode changes
 		NotifyMaterials();
+		for (USparseVolumeTextureFrame* Frame : Frames)
+		{
+			Frame->NotifyMaterials();
+		}
 	}
 
 	Super::PostEditChangeProperty(PropertyChangedEvent);
@@ -653,12 +775,9 @@ void UStreamableSparseVolumeTexture::GetResourceSizeEx(FResourceSizeEx& Cumulati
 	SIZE_T SizeCPU = sizeof(*this);
 	SIZE_T SizeGPU = 0;
 	SizeCPU += Frames.GetAllocatedSize();
-	for (const FSparseVolumeTextureFrame& Frame : Frames)
+	for (USparseVolumeTextureFrame* Frame : Frames)
 	{
-		if (Frame.SparseVolumeTextureSceneProxy)
-		{
-			Frame.SparseVolumeTextureSceneProxy->GetMemorySize(&SizeCPU, &SizeGPU);
-		}
+		Frame->GetResourceSizeEx(CumulativeResourceSize);
 	}
 	ISparseVolumeTextureStreamingManager& StreamingManager = IStreamingManager::Get().GetSparseVolumeTextureStreamingManager();
 	StreamingManager.GetMemorySizeForSparseVolumeTexture(this, &SizeCPU, &SizeGPU);
@@ -679,117 +798,19 @@ void UStreamableSparseVolumeTexture::GenerateOrLoadDDCRuntimeDataAndCreateSceneP
 	UE::DerivedData::FRequestOwner DDCRequestOwner(UE::DerivedData::EPriority::Normal);
 	{
 		UE::DerivedData::FRequestBarrier DDCRequestBarrier(DDCRequestOwner);
-		for (FSparseVolumeTextureFrame& Frame : Frames)
+		for (USparseVolumeTextureFrame* Frame : Frames)
 		{
-			// Release any previously allocated render thread proxy
-			if (Frame.SparseVolumeTextureSceneProxy)
-			{
-				BeginReleaseResource(Frame.SparseVolumeTextureSceneProxy);
-			}
-			else
-			{
-				Frame.SparseVolumeTextureSceneProxy = new FSparseVolumeTextureSceneProxy();
-			}
-
-			GenerateOrLoadDDCRuntimeDataForFrame(Frame, DDCRequestOwner);
+			Frame->GenerateOrLoadDDCRuntimeData(DDCRequestOwner);
 		}
 	}
 
-	// Wait for all DDC requests to complete before creating the proxies
+	// Wait for all DDC requests to complete
 	DDCRequestOwner.Wait();
 
-	for (FSparseVolumeTextureFrame& Frame : Frames)
-	{
-		// Runtime data is now valid, initialize the render thread proxy
-		BeginInitResource(Frame.SparseVolumeTextureSceneProxy);
-	}
-
 	IStreamingManager::Get().GetSparseVolumeTextureStreamingManager().AddSparseVolumeTexture(this);
 #endif
 }
 
-void UStreamableSparseVolumeTexture::GenerateOrLoadDDCRuntimeDataForFrame(FSparseVolumeTextureFrame& Frame, UE::DerivedData::FRequestOwner& DDCRequestOwner)
-{
-#if WITH_EDITORONLY_DATA
-	using namespace UE::DerivedData;
-
-	if (InitState != EInitState::Done)
-	{
-		UE_LOG(LogSparseVolumeTexture, Warning, TEXT("Tried to cache derived data of an uninitialized SVT: %s"), *GetName());
-		return;
-	}
-
-	static const FString SparseVolumeTextureDDCVersion = TEXT("381AE2A9-A903-4C8F-8486-891E24D6EC70");	// Bump this if you want to ignore all cached data so far.
-	const FString DerivedDataKey = Frame.RawData.GetIdentifier().ToString() 
-		+ FString::Format(TEXT("{0},{1},{2}"), {AddressX.GetIntValue(), AddressY.GetIntValue(), AddressZ.GetIntValue()}) 
-		+ SparseVolumeTextureDDCVersion;
-
-	const FCacheKey Key = ConvertLegacyCacheKey(DerivedDataKey);
-	const FSharedString Name = MakeStringView(GetPathName());
-
-	UE::DerivedData::GetCache().GetValue({ {Name, Key} }, DDCRequestOwner,
-		[this, &Frame, &DDCRequestOwner](FCacheGetValueResponse&& Response)
-		{
-			if (Response.Status == EStatus::Ok)
-			{
-				DDCRequestOwner.LaunchTask(TEXT("UStreamableSparseVolumeTexture_DerivedDataLoad"),
-					[this, &Frame, Value = MoveTemp(Response.Value)]()
-					{
-						FSharedBuffer Data = Value.GetData().Decompress();
-						FMemoryReaderView Ar(Data, true /*bIsPersistent*/);
-						FSparseVolumeTextureData TextureData;
-						TextureData.Serialize(Ar);
-						bool bSuccess = Frame.SparseVolumeTextureSceneProxy->GetRuntimeData().Create(TextureData);
-						check(bSuccess); // SVT_TODO
-					});
-			}
-			else if (Response.Status == EStatus::Error)
-			{
-				DDCRequestOwner.LaunchTask(TEXT("UStreamableSparseVolumeTexture_DerivedDataBuild"),
-					[this, &Frame, &DDCRequestOwner, Name = Response.Name, Key = Response.Key]()
-					{
-						FSparseVolumeTextureRuntime& RuntimeData = Frame.SparseVolumeTextureSceneProxy->GetRuntimeData();
-
-						// Check if the virtualized bulk data payload is available now
-						if (Frame.RawData.HasPayloadData())
-						{
-							FSparseVolumeTextureData TextureData;
-							bool bSuccess = Frame.BuildDerivedData(VolumeResolution, AddressX, AddressY, AddressZ, &TextureData);
-							ensure(bSuccess);
-
-							bSuccess = RuntimeData.Create(TextureData);
-							ensure(bSuccess);
-
-							// Using a LargeMemoryWriter for serialization since the data can be bigger than 2 GB
-							FLargeMemoryWriter LargeMemWriter(0, /*bIsPersistent=*/ true);
-							TextureData.Serialize(LargeMemWriter);
-
-							const int64 UncompressedSize = LargeMemWriter.TotalSize();
-
-							// Since the DDC doesn't support data bigger than 2 GB, we only cache for such uncompressed size.
-							constexpr int64 SizeThreshold = 2147483648LL;	// 2GB
-							const bool bIsCacheable = UncompressedSize < SizeThreshold;
-							if (bIsCacheable)
-							{
-								FValue Value = FValue::Compress(FSharedBuffer::MakeView(LargeMemWriter.GetView()));
-								UE::DerivedData::GetCache().PutValue({ {Name, Key, Value} }, DDCRequestOwner);
-							}
-							else
-							{
-								UE_LOG(LogSparseVolumeTexture, Error, TEXT("SparseVolumeTexture - the asset is too large to fit in Derived Data Cache %s"), *GetName());
-							}
-						}
-						else
-						{
-							UE_LOG(LogSparseVolumeTexture, Error, TEXT("SparseVolumeTexture - Raw source data is not available for %s. Using default data."), *GetName());
-							RuntimeData.SetAsDefaultTexture();
-						}
-					});
-			}
-		});
-#endif // WITH_EDITORONLY_DATA
-}
-
 ////////////////////////////////////////////////////////////////////////////////////////////////
 
 UStaticSparseVolumeTexture::UStaticSparseVolumeTexture(const FObjectInitializer& ObjectInitializer)
@@ -825,51 +846,6 @@ const FSparseVolumeTextureSceneProxy* UAnimatedSparseVolumeTexture::GetSparseVol
 
 ////////////////////////////////////////////////////////////////////////////////////////////////
 
-USparseVolumeTextureFrame::USparseVolumeTextureFrame(const FObjectInitializer& ObjectInitializer)
-	: Super(ObjectInitializer)
-{
-}
-
-USparseVolumeTextureFrame* USparseVolumeTextureFrame::CreateFrame(USparseVolumeTexture* Texture, int32 FrameIndex, int32 MipLevel)
-{
-	if (!Texture)
-	{
-		return nullptr;
-	}
-	
-	const FSparseVolumeTextureSceneProxy* Proxy = nullptr;
-	if (Texture->IsA<UStreamableSparseVolumeTexture>())
-	{
-		UStreamableSparseVolumeTexture* StreamableSVT = CastChecked<UStreamableSparseVolumeTexture>(Texture);
-		check(StreamableSVT);
-		Proxy = StreamableSVT->GetStreamedFrameProxyOrFallback(FrameIndex, MipLevel);
-	}
-	else
-	{
-		Proxy = Texture->GetSparseVolumeTextureSceneProxy();
-	}
-
-	if (Proxy)
-	{
-		USparseVolumeTextureFrame* Frame = NewObject<USparseVolumeTextureFrame>();
-		Frame->Initialize(Proxy, Texture->GetVolumeResolution(), Texture->GetTextureAddressX(), Texture->GetTextureAddressY(), Texture->GetTextureAddressZ());
-		return Frame;
-	}
-	
-	return nullptr;
-}
-
-void USparseVolumeTextureFrame::Initialize(const FSparseVolumeTextureSceneProxy* InSceneProxy, const FIntVector& InVolumeResolution, TextureAddress InAddressX, TextureAddress InAddressY, TextureAddress InAddressZ)
-{
-	SceneProxy = InSceneProxy;
-	VolumeResolution = InVolumeResolution;
-	AddressX = InAddressX;
-	AddressY = InAddressY;
-	AddressZ = InAddressZ;
-}
-
-////////////////////////////////////////////////////////////////////////////////////////////////
-
 UAnimatedSparseVolumeTextureController::UAnimatedSparseVolumeTextureController(const FObjectInitializer& ObjectInitializer)
 	: Super(ObjectInitializer)
 {
@@ -925,10 +901,7 @@ USparseVolumeTextureFrame* UAnimatedSparseVolumeTextureController::GetFrameByInd
 		return nullptr;
 	}
 
-	// Create and initialize a USparseVolumeTextureFrame which holds the frame to sample and can be bound to shaders
-	USparseVolumeTextureFrame* Frame = USparseVolumeTextureFrame::CreateFrame(SparseVolumeTexture, FrameIndex, MipLevel);
-
-	return Frame;
+	return USparseVolumeTextureFrame::GetFrame(SparseVolumeTexture, FrameIndex);
 }
 
 USparseVolumeTextureFrame* UAnimatedSparseVolumeTextureController::GetCurrentFrame()
@@ -942,10 +915,7 @@ USparseVolumeTextureFrame* UAnimatedSparseVolumeTextureController::GetCurrentFra
 	const float FrameIndexF = GetFractionalFrameIndex();
 	const int32 FrameIndex = (int32)FrameIndexF;
 
-	// Create and initialize a USparseVolumeTextureFrame which holds the frame to sample and can be bound to shaders
-	USparseVolumeTextureFrame* Frame = USparseVolumeTextureFrame::CreateFrame(SparseVolumeTexture, FrameIndex, MipLevel);
-
-	return Frame;
+	return USparseVolumeTextureFrame::GetFrame(SparseVolumeTexture, FrameIndex);
 }
 
 void UAnimatedSparseVolumeTextureController::GetCurrentFramesForInterpolation(USparseVolumeTextureFrame*& Frame0, USparseVolumeTextureFrame*& Frame1, float& LerpAlpha)
@@ -960,9 +930,8 @@ void UAnimatedSparseVolumeTextureController::GetCurrentFramesForInterpolation(US
 	const int32 FrameIndex = (int32)FrameIndexF;
 	LerpAlpha = FMath::Frac(FrameIndexF);
 
-	// Create and initialize a USparseVolumeTextureFrame which holds the frame to sample and can be bound to shaders
-	Frame0 = USparseVolumeTextureFrame::CreateFrame(SparseVolumeTexture, FrameIndex, MipLevel);
-	Frame1 = USparseVolumeTextureFrame::CreateFrame(SparseVolumeTexture, FrameIndex + 1, MipLevel);
+	Frame0 = USparseVolumeTextureFrame::GetFrame(SparseVolumeTexture, FrameIndex);
+	Frame1 = USparseVolumeTextureFrame::GetFrame(SparseVolumeTexture, (FrameIndex + 1) % SparseVolumeTexture->GetNumFrames());
 }
 
 float UAnimatedSparseVolumeTextureController::GetDuration()
diff --git a/Engine/Source/Runtime/Engine/Private/SparseVolumeTexture/SparseVolumeTextureStreaming.cpp b/Engine/Source/Runtime/Engine/Private/SparseVolumeTexture/SparseVolumeTextureStreaming.cpp
--- a/Engine/Source/Runtime/Engine/Private/SparseVolumeTexture/SparseVolumeTextureStreaming.cpp
+++ b/Engine/Source/Runtime/Engine/Private/SparseVolumeTexture/SparseVolumeTextureStreaming.cpp
@@ -66,7 +66,7 @@ bool FStreamingSparseVolumeTextureData::Initialize(UStreamableSparseVolumeTextur
 {
 	check(InSparseVolumeTexture);
 
-	if (InSparseVolumeTexture->GetFrames().IsEmpty())
+	if (InSparseVolumeTexture->GetNumFrames() == 0)
 	{
 		UE_LOG(LogSparseVolumeTextureStreaming, Error, TEXT("Failed to initialize streaming SparseVolumeTexture due to lack of SVT or serialized stream frames. '%s'"), *InSparseVolumeTexture->GetFullName());
 		return false;
@@ -78,9 +78,8 @@ bool FStreamingSparseVolumeTextureData::Initialize(UStreamableSparseVolumeTextur
 	// Always get the first frame of data so we can play immediately
 	check(LoadedChunks.IsEmpty());
 	check(LoadedChunkIndices.IsEmpty());
-	check(!InSparseVolumeTexture->GetFrames().IsEmpty()); // Must hold at least the first frame
 
-	AddNewLoadedChunk(0, InSparseVolumeTexture->GetFrames()[0].SparseVolumeTextureSceneProxy);
+	AddNewLoadedChunk(0, InSparseVolumeTexture->GetFrame(0)->SceneProxy);
 	LoadedChunkIndices.Add(0);
 
 	return true;
@@ -242,12 +241,11 @@ void FStreamingSparseVolumeTextureData::BeginPendingRequests(const TArray<int32>
 	const int32 NumFrames = SparseVolumeTexture->GetNumFrames();
 	check(NumFrames > 0);
 	const EAsyncIOPriorityAndFlags AsyncIOPriority = AIOP_CriticalPath; //Set to Crit temporarily as emergency speculative fix for streaming issue
-	TArrayView<const FSparseVolumeTextureFrame> SVTFrames = SparseVolumeTexture->GetFrames();
 	for (int32 IndexToLoad : IndicesToLoad)
 	{
 		const int32 FrameToLoad = SVTChunkIndexToFrame(IndexToLoad, NumFrames) % NumFrames;
-		const FSparseVolumeTextureFrame& Frame = SVTFrames[FrameToLoad];
-		FSparseVolumeTextureSceneProxy* ExistingProxy = Frame.SparseVolumeTextureSceneProxy;
+		USparseVolumeTextureFrame* Frame = SparseVolumeTexture->GetFrame(FrameToLoad);
+		FSparseVolumeTextureSceneProxy* ExistingProxy = Frame->SceneProxy;
 		FLoadedSparseVolumeTextureChunk& ChunkStorage = AddNewLoadedChunk(IndexToLoad, ExistingProxy);
 
 		if (!ExistingProxy)
@@ -255,7 +253,7 @@ void FStreamingSparseVolumeTextureData::BeginPendingRequests(const TArray<int32>
 			UE_CLOG(ChunkStorage.Proxy.load() != nullptr, LogSparseVolumeTextureStreaming, Fatal, TEXT("Existing render proxy for streaming SparseVolumeTexture frame."));
 			UE_CLOG(ChunkStorage.IORequest, LogSparseVolumeTextureStreaming, Fatal, TEXT("Streaming SparseVolumeTexture frame already has IORequest."));
 
-			const int64 ChunkSize = Frame.RuntimeStreamedInData.GetBulkDataSize();
+			const int64 ChunkSize = Frame->StreamingData.GetBulkDataSize();
 			ChunkStorage.RequestStart = FPlatformTime::Seconds();
 			UE_LOG(LogSparseVolumeTextureStreaming, Warning, TEXT("SparseVolumeTexture streaming request started %s Frame:%i At:%.3f\n"), *SparseVolumeTexture->GetName(), IndexToLoad, ChunkStorage.RequestStart);
 			FBulkDataIORequestCallBack AsyncFileCallBack = [this, IndexToLoad, ChunkSize](bool bWasCancelled, IBulkDataIORequest* Req)
@@ -264,8 +262,8 @@ void FStreamingSparseVolumeTextureData::BeginPendingRequests(const TArray<int32>
 			};
 
 			UE_LOG(LogSparseVolumeTextureStreaming, Warning, TEXT("Loading streaming SparseVolumeTexture %s Frame:%i Offset:%i Size:%i File:%s\n"),
-				*SparseVolumeTexture->GetName(), IndexToLoad, Frame.RuntimeStreamedInData.GetBulkDataOffsetInFile(), Frame.RuntimeStreamedInData.GetBulkDataSize(), *Frame.RuntimeStreamedInData.GetDebugName());
-			ChunkStorage.IORequest = Frame.RuntimeStreamedInData.CreateStreamingRequest(AsyncIOPriority, &AsyncFileCallBack, nullptr);
+				*SparseVolumeTexture->GetName(), IndexToLoad, Frame->StreamingData.GetBulkDataOffsetInFile(), Frame->StreamingData.GetBulkDataSize(), *Frame->StreamingData.GetDebugName());
+			ChunkStorage.IORequest = Frame->StreamingData.CreateStreamingRequest(AsyncIOPriority, &AsyncFileCallBack, nullptr);
 			if (!ChunkStorage.IORequest)
 			{
 				UE_LOG(LogSparseVolumeTextureStreaming, Error, TEXT("SparseVolumeTexture streaming read request failed."));
diff --git a/Engine/Source/Runtime/Renderer/Private/SparseVolumeTexture/SparseVolumeTextureViewerComponent.cpp b/Engine/Source/Runtime/Renderer/Private/SparseVolumeTexture/SparseVolumeTextureViewerComponent.cpp
--- a/Engine/Source/Runtime/Renderer/Private/SparseVolumeTexture/SparseVolumeTextureViewerComponent.cpp
+++ b/Engine/Source/Runtime/Renderer/Private/SparseVolumeTexture/SparseVolumeTextureViewerComponent.cpp
@@ -209,7 +209,7 @@ void USparseVolumeTextureViewerComponent::TickComponent(float DeltaTime, enum EL
 			FrameIndex = int32(AnimationFrame * float(SparseVolumeTexturePreview->GetNumFrames()));
 		}
 
-		SparseVolumeTextureFrame = USparseVolumeTextureFrame::CreateFrame(SparseVolumeTexturePreview, FrameIndex, MipLevel);
+		SparseVolumeTextureFrame = USparseVolumeTextureFrame::GetFrame(SparseVolumeTexturePreview, FrameIndex);
 	}
 	else
 	{
