remove the need to manually sub-query cluster unions to determine which component in the cluster union we hit

diff --git a/Engine/Source/Runtime/Engine/Classes/Components/PrimitiveComponent.h b/Engine/Source/Runtime/Engine/Classes/Components/PrimitiveComponent.h
--- a/Engine/Source/Runtime/Engine/Classes/Components/PrimitiveComponent.h
+++ b/Engine/Source/Runtime/Engine/Classes/Components/PrimitiveComponent.h
@@ -2750,6 +2750,19 @@ public:
 	 *  @return true if a hit is found
 	 */
 	virtual bool LineTraceComponent( FHitResult& OutHit, const FVector Start, const FVector End, const FCollisionQueryParams& Params );
+
+	/**
+	 *  Trace a ray against just this component.
+	 *  @param  OutHit          Information about hit against this component, if true is returned
+	 *  @param  Start           Start location of the ray
+	 *  @param  End             End location of the ray
+	 *  @param  TraceChannel    The 'channel' that this query is in, used to determine which components to hit
+	 *  @param  Params          Additional parameters used for the trace
+	 * 	@param 	ResponseParam	ResponseContainer to be used for this trace
+	 *	@param	ObjectQueryParams	List of object types it's looking for
+	 *  @return true if a hit is found
+	 */
+	virtual bool LineTraceComponent(FHitResult& OutHit, const FVector Start, const FVector End, ECollisionChannel TraceChannel, const struct FCollisionQueryParams& Params, const struct FCollisionResponseParams& ResponseParams, const struct FCollisionObjectQueryParams& ObjectParams);
 	
 	/** 
 	 *  Trace a shape against just this component.
@@ -2763,6 +2776,21 @@ public:
 	 */
 	virtual bool SweepComponent(FHitResult& OutHit, const FVector Start, const FVector End, const FQuat& ShapeWorldRotation, const FCollisionShape &CollisionShape, bool bTraceComplex=false);
 
+	/**
+	 *  Trace a shape against just this component.
+	 *  @param  OutHit          	Information about hit against this component, if true is returned
+	 *  @param  Start           	Start location of the box
+	 *  @param  End             	End location of the box
+	 *  @param  ShapeWorldRotation  The rotation applied to the collision shape in world space.
+	 *  @param  Geometry			Geometry to sweep with.
+	 *  @param  TraceChannel    The 'channel' that this query is in, used to determine which components to hit
+	 *  @param  Params          Additional parameters used for the trace
+	 * 	@param 	ResponseParam	ResponseContainer to be used for this trace
+	 *	@param	ObjectQueryParams	List of object types it's looking for
+	 *  @return true if a hit is found
+	 */
+	virtual bool SweepComponent(FHitResult& OutHit, const FVector Start, const FVector End, const FQuat& ShapeWorldRotation, const FPhysicsGeometry& Geometry, ECollisionChannel TraceChannel, const struct FCollisionQueryParams& Params, const struct FCollisionResponseParams& ResponseParams, const struct FCollisionObjectQueryParams& ObjectParams);
+
 	/** 
 	 *  Test the collision of the supplied component at the supplied location/rotation, and determine if it overlaps this component.
 	 *  @note This overload taking rotation as a FQuat is slightly faster than the version using FRotator.
@@ -2822,6 +2850,20 @@ public:
 	 */
 	virtual bool OverlapComponentWithResult(const FVector& Pos, const FQuat& Rot, const FCollisionShape& CollisionShape, TArray<FOverlapResult>& OutOverlap) const;
 
+	/**
+	 * Test the collision of the supplied shape at the supplied location, and determine if it overlaps this component.
+	 * Also will return information about the overlap.
+	 *  @param  Pos             Location to place PrimComp geometry at
+	 *	@param  Rot             Rotation of PrimComp geometry
+	 *  @param  Geometry		Geometry to use for the overlap check.
+	 *  @param  TraceChannel    The 'channel' that this query is in, used to determine which components to hit
+	 *  @param  Params          Additional parameters used for the trace
+	 * 	@param 	ResponseParam	ResponseContainer to be used for this trace
+	 *	@param	ObjectQueryParams	List of object types it's looking for
+	 *  @param  OutOverlap      Additional information about what exactly was overlapped.
+	 *  @return true if PrimComp overlaps this component at the specified location/rotation
+	 */
+	virtual bool OverlapComponentWithResult(const FVector& Pos, const FQuat& Rot, const FPhysicsGeometry& Geometry, ECollisionChannel TraceChannel, const struct FCollisionQueryParams& Params, const struct FCollisionResponseParams& ResponseParams, const struct FCollisionObjectQueryParams& ObjectParams, TArray<FOverlapResult>& OutOverlap) const;
 
 	/**
 	 * Computes the minimum translation direction (MTD) when an overlap exists between the component and the given shape.
diff --git a/Engine/Source/Runtime/Engine/Classes/PhysicsEngine/ClusterUnionComponent.h b/Engine/Source/Runtime/Engine/Classes/PhysicsEngine/ClusterUnionComponent.h
--- a/Engine/Source/Runtime/Engine/Classes/PhysicsEngine/ClusterUnionComponent.h
+++ b/Engine/Source/Runtime/Engine/Classes/PhysicsEngine/ClusterUnionComponent.h
@@ -230,6 +230,9 @@ private:
 	void HandleRemovedClusteredComponent(UPrimitiveComponent* ChangedComponent, bool bDestroyReplicatedProxy);
 
 	TArray<UPrimitiveComponent*> GetAllCurrentChildComponents() const;
+	// Lambda returns whether or not iteration should continue;
+	void VisitAllCurrentChildComponents(const TFunction<bool(UPrimitiveComponent*)>& Lambda) const;
+	void VisitAllCurrentChildComponentsForCollision(ECollisionChannel TraceChannel, const struct FCollisionQueryParams& Params, const struct FCollisionResponseParams& ResponseParams, const struct FCollisionObjectQueryParams& ObjectParams, const TFunction<bool(UPrimitiveComponent*)>& Lambda) const;
 
 	// Whether or not this code is running on the server.
 	UFUNCTION()
@@ -250,12 +253,16 @@ public:
 
 	//~ Begin UPrimitiveComponent Interface
 public:
+	using UPrimitiveComponent::LineTraceComponent;
+	using UPrimitiveComponent::SweepComponent;
+	using UPrimitiveComponent::OverlapComponentWithResult;
+
 	virtual FBodyInstance* GetBodyInstance(FName BoneName, bool bGetWelded, int32 Index) const override { return nullptr; }
 	virtual void SetSimulatePhysics(bool bSimulate) override;
 	virtual bool CanEditSimulatePhysics() override { return true; }
-	virtual bool LineTraceComponent(FHitResult& OutHit, const FVector Start, const FVector End, const FCollisionQueryParams& Params) override;
-	virtual bool SweepComponent(FHitResult& OutHit, const FVector Start, const FVector End, const FQuat& ShapeWorldRotation, const FCollisionShape& CollisionShape, bool bTraceComplex) override;
-	virtual bool OverlapComponentWithResult(const FVector& Pos, const FQuat& Rot, const FCollisionShape& CollisionShape, TArray<FOverlapResult>& OutOverlap) const override;
+	virtual bool LineTraceComponent(FHitResult& OutHit, const FVector Start, const FVector End, ECollisionChannel TraceChannel, const struct FCollisionQueryParams& Params, const struct FCollisionResponseParams& ResponseParams, const struct FCollisionObjectQueryParams& ObjectParams) override;
+	virtual bool SweepComponent(FHitResult& OutHit, const FVector Start, const FVector End, const FQuat& ShapeWorldRotation, const FPhysicsGeometry& Geometry, ECollisionChannel TraceChannel, const struct FCollisionQueryParams& Params, const struct FCollisionResponseParams& ResponseParams, const struct FCollisionObjectQueryParams& ObjectParams) override;
+	virtual bool OverlapComponentWithResult(const FVector& Pos, const FQuat& Rot, const FPhysicsGeometry& Geometry, ECollisionChannel TraceChannel, const struct FCollisionQueryParams& Params, const struct FCollisionResponseParams& ResponseParams, const struct FCollisionObjectQueryParams& ObjectParams, TArray<FOverlapResult>& OutOverlap) const override;
 	virtual bool ComponentOverlapComponentWithResultImpl(const class UPrimitiveComponent* const PrimComp, const FVector& Pos, const FQuat& Rot, const FCollisionQueryParams& Params, TArray<FOverlapResult>& OutOverlap) const override;
 	//~ End UPrimitiveComponent Interface
 
diff --git a/Engine/Source/Runtime/Engine/Private/Collision/SceneQuery.cpp b/Engine/Source/Runtime/Engine/Private/Collision/SceneQuery.cpp
--- a/Engine/Source/Runtime/Engine/Private/Collision/SceneQuery.cpp
+++ b/Engine/Source/Runtime/Engine/Private/Collision/SceneQuery.cpp
@@ -5,6 +5,7 @@
 #include "Physics/Experimental/PhysScene_Chaos.h"
 #include "Physics/PhysicsInterfaceUtils.h"
 #include "Collision/CollisionConversions.h"
+#include "PhysicsEngine/ClusterUnionComponent.h"
 #include "PhysicsEngine/CollisionQueryFilterCallback.h"
 #include "PhysicsEngine/ScopedSQHitchRepeater.h"
 
@@ -294,6 +295,10 @@ struct TSQTraits
 
 	CA_SUPPRESS(6326);
 	constexpr static bool IsSweep() { return GeometryQuery == ESweepOrRay::Sweep;  }
+
+	/** Easy way to query whether this SQ trait is for the GT or the PT based on the based in hit type. */
+	constexpr static bool IsExternalData() { return std::is_base_of_v<ChaosInterface::FActorShape, InHitType>; }
+
 };
 
 enum class EThreadQueryContext
@@ -408,6 +413,63 @@ bool TSceneCastCommonImp(const UWorld* World, typename Traits::TOutHits& OutHits
 				UE_LOG(LogCollision, Error, TEXT("--------%s"), *Params.ToString());
 #endif
 			}
+
+			// This block is only necessary on the game thread when dealing with cluster unions.
+			// TODO: Is there a way to get this to generalize better to the PT as well? Right now it depends on GT functions on the cluster union component.
+			if constexpr (Traits::IsExternalData())
+			{
+				auto DoClusterUnionTrace = [&GeomInputs, &Start, &End, TraceChannel, &Params, &ResponseParams, &ObjectParams](const FHitResult& OriginalHit, FHitResult& NewHit)
+				{
+					if (UClusterUnionComponent* ClusterUnion = Cast<UClusterUnionComponent>(OriginalHit.GetComponent()))
+					{
+						if constexpr (Traits::IsRay())
+						{
+							return ClusterUnion->LineTraceComponent(NewHit, Start, End, TraceChannel, Params, ResponseParams, ObjectParams);
+						}
+						else
+						{
+							return ClusterUnion->SweepComponent(NewHit, Start, End, *GeomInputs.GetGeometryOrientation(), *GeomInputs.GetGeometry(), TraceChannel, Params, ResponseParams, ObjectParams);
+						}
+					}
+
+					return false;
+				};
+
+				if (bSuccess && Params.bTraceIntoSubComponents)
+				{
+					if constexpr (Traits::IsMulti())
+					{
+						TArray<FHitResult> AllNewHits;
+						TArray<int32> ClusterUnionIndices;
+
+						for (int32 Index = 0; Index < OutHits.Num(); ++Index)
+						{
+							FHitResult NewHit;
+							if (DoClusterUnionTrace(OutHits[Index], NewHit))
+							{
+								AllNewHits.Add(NewHit);
+								ClusterUnionIndices.Add(Index);
+							}
+						}
+
+						for (int32 Index = ClusterUnionIndices.Num() - 1; Index >= 0; --Index)
+						{
+							// No shrinking since we're going to be adding more elements shortly.
+							OutHits.RemoveAtSwap(ClusterUnionIndices[Index], 1, false);
+						}
+
+						OutHits.Append(AllNewHits);
+					}
+					else
+					{
+						FHitResult NewHit;
+						if (DoClusterUnionTrace(OutHits, NewHit))
+						{
+							OutHits = NewHit;
+						}
+					}
+				}
+			}
 		}
 
 		bHaveBlockingHit = bBlockingHit;
@@ -630,6 +692,45 @@ bool GeomOverlapMultiImp(const UWorld* World, const FPhysicsGeometry& Geom, cons
 				if(NumHits > 0)
 				{
 					bHaveBlockingHit = ConvertOverlapResults(NumHits, OverlapBuffer.GetHits(), Filter, OutOverlaps);
+
+					// This block is only necessary on the game thread when dealing with cluster unions.
+					// TODO: Is there a way to get this to generalize better to the PT as well? Right now it depends on GT functions on the cluster union component.
+					if constexpr (std::is_base_of_v<ChaosInterface::FActorShape, TOverlapHit>)
+					{
+						auto DoClusterUnionOverlap = [&GeomPose, &Geom, TraceChannel, &Params, &ResponseParams, &ObjectParams](const FOverlapResult& OriginalOverlap, TArray<FOverlapResult>& NewOverlaps)
+						{
+							if (UClusterUnionComponent* ClusterUnion = Cast<UClusterUnionComponent>(OriginalOverlap.GetComponent()))
+							{
+								return ClusterUnion->OverlapComponentWithResult(GeomPose.GetTranslation(), GeomPose.GetRotation(), Geom, TraceChannel, Params, ResponseParams, ObjectParams, NewOverlaps);
+							}
+
+							return false;
+						};
+
+						if (bHaveBlockingHit && Params.bTraceIntoSubComponents)
+						{
+							TArray<FOverlapResult> AllNewOverlaps;
+							TArray<int32> ClusterUnionIndices;
+
+							for (int32 Index = 0; Index < OutOverlaps.Num(); ++Index)
+							{
+								TArray<FOverlapResult> NewOverlaps;
+								if (DoClusterUnionOverlap(OutOverlaps[Index], NewOverlaps))
+								{
+									AllNewOverlaps.Append(NewOverlaps);
+									ClusterUnionIndices.Add(Index);
+								}
+							}
+
+							for (int32 Index = ClusterUnionIndices.Num() - 1; Index >= 0; --Index)
+							{
+								// No shrinking since we're going to be adding more elements shortly.
+								OutOverlaps.RemoveAtSwap(ClusterUnionIndices[Index], 1, false);
+							}
+
+							OutOverlaps.Append(AllNewOverlaps);
+						}
+					}
 				}
 			}
 
diff --git a/Engine/Source/Runtime/Engine/Private/Components/PrimitiveComponent.cpp b/Engine/Source/Runtime/Engine/Private/Components/PrimitiveComponent.cpp
--- a/Engine/Source/Runtime/Engine/Private/Components/PrimitiveComponent.cpp
+++ b/Engine/Source/Runtime/Engine/Private/Components/PrimitiveComponent.cpp
@@ -2919,9 +2919,14 @@ extern float DebugLineLifetime;
 
 
 bool UPrimitiveComponent::LineTraceComponent(struct FHitResult& OutHit, const FVector Start, const FVector End, const struct FCollisionQueryParams& Params)
+{
+	return LineTraceComponent(OutHit, Start, End, DefaultCollisionChannel, Params, FCollisionResponseParams::DefaultResponseParam, FCollisionObjectQueryParams::DefaultObjectQueryParam);
+}
+
+bool UPrimitiveComponent::LineTraceComponent(FHitResult& OutHit, const FVector Start, const FVector End, ECollisionChannel TraceChannel, const struct FCollisionQueryParams& Params, const struct FCollisionResponseParams& ResponseParams, const struct FCollisionObjectQueryParams& ObjectParams)
 {
 	bool bHaveHit = false;
-	
+
 	if (FBodyInstance* ThisBodyInstance = GetBodyInstance())
 	{
 		bHaveHit = ThisBodyInstance->LineTrace(OutHit, Start, End, Params.bTraceComplex, Params.bReturnPhysicalMaterial);
@@ -2973,14 +2978,20 @@ bool UPrimitiveComponent::SweepComponent(struct FHitResult& OutHit, const FVecto
 		return ThisBodyInstance->Sweep(OutHit, Start, End, ShapeWorldRotation, CollisionShape, bTraceComplex);
 	}
 
+	FCollisionQueryParams Params = FCollisionQueryParams::DefaultQueryParam;
+	Params.bTraceComplex = bTraceComplex;
+
 	if (CollisionShape.IsNearlyZero())
 	{
-		FCollisionQueryParams Params;
-		Params.bTraceComplex = bTraceComplex;
 		return LineTraceComponent(OutHit, Start, End, Params);
 	}
 
 	FPhysicsShapeAdapter_Chaos ShapeAdapter(ShapeWorldRotation, CollisionShape);
+	return SweepComponent(OutHit, Start, End, ShapeWorldRotation, ShapeAdapter.GetGeometry(), DefaultCollisionChannel, Params, FCollisionResponseParams::DefaultResponseParam, FCollisionObjectQueryParams::DefaultObjectQueryParam);
+}
+
+bool UPrimitiveComponent::SweepComponent(FHitResult& OutHit, const FVector Start, const FVector End, const FQuat& ShapeWorldRotation, const FPhysicsGeometry& Geometry, ECollisionChannel TraceChannel, const struct FCollisionQueryParams& Params, const struct FCollisionResponseParams& ResponseParams, const struct FCollisionObjectQueryParams& ObjectParams)
+{
 	TArray<Chaos::FPhysicsObjectHandle> Objects = GetAllPhysicsObjects();
 	FLockedReadPhysicsObjectExternalInterface Interface = FPhysicsObjectExternalInterface::LockRead(Objects);
 	Objects = Objects.FilterByPredicate(
@@ -2993,14 +3004,14 @@ bool UPrimitiveComponent::SweepComponent(struct FHitResult& OutHit, const FVecto
 	Chaos::FPhysicsObjectCollisionInterface_External CollisionInterface{ Interface.GetInterface() };
 	ChaosInterface::FSweepHit BestHit;
 
-	Chaos::FSweepParameters Params;
-	Params.bSweepComplex = bTraceComplex;
+	Chaos::FSweepParameters SweepParams;
+	SweepParams.bSweepComplex = Params.bTraceComplex;
 
 	// TODO: Expose this even further via parameters in the primitive component.
 	// For now, having this be always true guarantees us identical behavior to tracing via the Chaos SQ
 	// since TSQTraits::GetHitFlags() will always have the MTD flag on.
-	Params.bComputeMTD = true;
-	if (CollisionInterface.ShapeSweep(Objects, ShapeAdapter.GetGeometry(), ShapeAdapter.GetGeomPose(Start), End, Params, BestHit))
+	SweepParams.bComputeMTD = true;
+	if (CollisionInterface.ShapeSweep(Objects, Geometry, FTransform{ ShapeWorldRotation, Start }, End, SweepParams, BestHit))
 	{
 		FCollisionFilterData QueryFilter;
 		QueryFilter.Word1 = 0xFFFFF;
@@ -3136,6 +3147,11 @@ bool UPrimitiveComponent::OverlapComponent(const FVector& Pos, const FQuat& Rot,
 bool UPrimitiveComponent::OverlapComponentWithResult(const FVector& Pos, const FQuat& Rot, const FCollisionShape& CollisionShape, TArray<FOverlapResult>& OutOverlap) const
 {
 	FPhysicsShapeAdapter_Chaos ShapeAdapter(Rot, CollisionShape);
+	return OverlapComponentWithResult(Pos, Rot, ShapeAdapter.GetGeometry(), DefaultCollisionChannel, FCollisionQueryParams::DefaultQueryParam, FCollisionResponseParams::DefaultResponseParam, FCollisionObjectQueryParams::DefaultObjectQueryParam, OutOverlap);
+}
+
+bool UPrimitiveComponent::OverlapComponentWithResult(const FVector& Pos, const FQuat& Rot, const FPhysicsGeometry& Geometry, ECollisionChannel TraceChannel, const struct FCollisionQueryParams& Params, const struct FCollisionResponseParams& ResponseParams, const struct FCollisionObjectQueryParams& ObjectParams, TArray<FOverlapResult>& OutOverlap) const
+{
 	TArray<Chaos::FPhysicsObjectHandle> Objects = GetAllPhysicsObjects();
 	FLockedReadPhysicsObjectExternalInterface Interface = FPhysicsObjectExternalInterface::LockRead(Objects);
 	Objects = Objects.FilterByPredicate(
@@ -3147,7 +3163,7 @@ bool UPrimitiveComponent::OverlapComponentWithResult(const FVector& Pos, const F
 
 	Chaos::FPhysicsObjectCollisionInterface_External CollisionInterface{ Interface.GetInterface() };
 	TArray<ChaosInterface::FOverlapHit> OverlapHits;
-	if (CollisionInterface.ShapeOverlap(Objects, ShapeAdapter.GetGeometry(), ShapeAdapter.GetGeomPose(Pos), OverlapHits))
+	if (CollisionInterface.ShapeOverlap(Objects, Geometry, FTransform{ Rot, Pos }, OverlapHits))
 	{
 		TArray<FOverlapResult> Overlaps;
 
diff --git a/Engine/Source/Runtime/Engine/Private/PhysicsEngine/ClusterUnionComponent.cpp b/Engine/Source/Runtime/Engine/Private/PhysicsEngine/ClusterUnionComponent.cpp
--- a/Engine/Source/Runtime/Engine/Private/PhysicsEngine/ClusterUnionComponent.cpp
+++ b/Engine/Source/Runtime/Engine/Private/PhysicsEngine/ClusterUnionComponent.cpp
@@ -683,64 +683,75 @@ void UClusterUnionComponent::SetSimulatePhysics(bool bSimulate)
 	PhysicsProxy->SetObjectState_External(bSimulate ? Chaos::EObjectStateType::Dynamic : Chaos::EObjectStateType::Kinematic);
 }
 
-bool UClusterUnionComponent::LineTraceComponent(FHitResult& OutHit, const FVector Start, const FVector End, const FCollisionQueryParams& Params)
+bool UClusterUnionComponent::LineTraceComponent(FHitResult& OutHit, const FVector Start, const FVector End, ECollisionChannel TraceChannel, const struct FCollisionQueryParams& Params, const struct FCollisionResponseParams& ResponseParams, const struct FCollisionObjectQueryParams& ObjectParams)
 {
 	bool bHasHit = false;
 	OutHit.Distance = TNumericLimits<float>::Max();
-	for (UPrimitiveComponent* Component : GetAllCurrentChildComponents())
-	{
-		if (Component)
+
+	VisitAllCurrentChildComponentsForCollision(TraceChannel, Params, ResponseParams, ObjectParams,
+		[&Start, &End, TraceChannel, &Params, &ResponseParams, &ObjectParams, &bHasHit, &OutHit](UPrimitiveComponent* Component)
 		{
 			FHitResult ComponentHit;
-			if (Component->LineTraceComponent(ComponentHit, Start, End, Params) && ComponentHit.Distance < OutHit.Distance)
+			if (Component->LineTraceComponent(ComponentHit, Start, End, TraceChannel, Params, ResponseParams, ObjectParams) && ComponentHit.Distance < OutHit.Distance)
 			{
 				bHasHit = true;
 				OutHit = ComponentHit;
 			}
+
+			return true;
 		}
-	}
+	);
+
 	return bHasHit;
 }
 
-bool UClusterUnionComponent::SweepComponent(FHitResult& OutHit, const FVector Start, const FVector End, const FQuat& ShapeWorldRotation, const FCollisionShape& CollisionShape, bool bTraceComplex)
+bool UClusterUnionComponent::SweepComponent(FHitResult& OutHit, const FVector Start, const FVector End, const FQuat& ShapeWorldRotation, const FPhysicsGeometry& Geometry, ECollisionChannel TraceChannel, const struct FCollisionQueryParams& Params, const struct FCollisionResponseParams& ResponseParams, const struct FCollisionObjectQueryParams& ObjectParams)
 {
 	bool bHasHit = false;
 	OutHit.Distance = TNumericLimits<float>::Max();
-	for (UPrimitiveComponent* Component : GetAllCurrentChildComponents())
-	{
-		if (Component)
+
+	VisitAllCurrentChildComponentsForCollision(TraceChannel, Params, ResponseParams, ObjectParams,
+		[&Geometry, &Start, &End, &ShapeWorldRotation, TraceChannel, &Params, &ResponseParams, &ObjectParams, &bHasHit, &OutHit](UPrimitiveComponent* Component)
 		{
 			FHitResult ComponentHit;
-			if (Component->SweepComponent(ComponentHit, Start, End, ShapeWorldRotation, CollisionShape, bTraceComplex) && ComponentHit.Distance < OutHit.Distance)
+			if (Component->SweepComponent(ComponentHit, Start, End, ShapeWorldRotation, Geometry, TraceChannel, Params, ResponseParams, ObjectParams) && ComponentHit.Distance < OutHit.Distance)
 			{
 				bHasHit = true;
 				OutHit = ComponentHit;
 			}
+
+			return true;
 		}
-	}
+	);
+
 	return bHasHit;
 }
 
-bool UClusterUnionComponent::OverlapComponentWithResult(const FVector& Pos, const FQuat& Rot, const FCollisionShape& CollisionShape, TArray<FOverlapResult>& OutOverlap) const
+bool UClusterUnionComponent::OverlapComponentWithResult(const FVector& Pos, const FQuat& Rot, const FPhysicsGeometry& Geometry, ECollisionChannel TraceChannel, const struct FCollisionQueryParams& Params, const struct FCollisionResponseParams& ResponseParams, const struct FCollisionObjectQueryParams& ObjectParams, TArray<FOverlapResult>& OutOverlap) const
 {
 	bool bHasOverlap = false;
 
-	FVector QueryHalfExtent = CollisionShape.GetExtent();
+	FVector QueryHalfExtent = Geometry.BoundingBox().Extents();
 	FBoxSphereBounds QueryBounds(FBox(-QueryHalfExtent, QueryHalfExtent));
 	QueryBounds = QueryBounds.TransformBy(FTransform(Rot, Pos));
 
-	for (UPrimitiveComponent* Component : GetAllCurrentChildComponents())
-	{
-		if(Component && FBoxSphereBounds::SpheresIntersect(QueryBounds, Component->Bounds))
+	VisitAllCurrentChildComponentsForCollision(TraceChannel, Params, ResponseParams, ObjectParams,
+		[&QueryBounds, &Geometry, &Pos, &Rot, TraceChannel, &Params, &ResponseParams, &ObjectParams, &bHasOverlap, &OutOverlap](UPrimitiveComponent* Component)
 		{
-			TArray<FOverlapResult> SubOverlaps;
-			if(Component->OverlapComponentWithResult(Pos, Rot, CollisionShape, SubOverlaps))
+			if (FBoxSphereBounds::SpheresIntersect(QueryBounds, Component->Bounds))
 			{
-				bHasOverlap = true;
-				OutOverlap.Append(SubOverlaps);
+				TArray<FOverlapResult> SubOverlaps;
+				if (Component->OverlapComponentWithResult(Pos, Rot, Geometry, TraceChannel, Params, ResponseParams, ObjectParams, SubOverlaps))
+				{
+					bHasOverlap = true;
+					OutOverlap.Append(SubOverlaps);
+				}
 			}
+
+			return true;
 		}
-	}
+	);
+
 	return bHasOverlap;
 }
 
@@ -753,47 +764,25 @@ bool UClusterUnionComponent::ComponentOverlapComponentWithResultImpl(const class
 
 	bool bHasOverlap = false;
 
-	TSet<uint32> IgnoredActors;
-	IgnoredActors.Reserve(Params.GetIgnoredActors().Num());
-	for (uint32 Id : Params.GetIgnoredActors())
-	{
-		IgnoredActors.Add(Id);
-	}
-
-	TSet<uint32> IgnoredComponents;
-	IgnoredComponents.Reserve(Params.GetIgnoredComponents().Num());
-	for (uint32 Id : Params.GetIgnoredActors())
-	{
-		IgnoredComponents.Add(Id);
-	}
-
 	FBoxSphereBounds QueryBounds = PrimComp->Bounds.TransformBy(FTransform(Rot, Pos));
 
-	for (UPrimitiveComponent* Component : GetAllCurrentChildComponents())
-	{
-		if(Component && FBoxSphereBounds::SpheresIntersect(QueryBounds, Component->Bounds))
+	VisitAllCurrentChildComponentsForCollision(DefaultCollisionChannel, Params, FCollisionResponseParams::DefaultResponseParam, FCollisionObjectQueryParams::DefaultObjectQueryParam,
+		[&QueryBounds, PrimComp, &Pos, &Rot, &Params, &bHasOverlap, &OutOverlap](UPrimitiveComponent* Component)
 		{
-			if(IgnoredComponents.Contains(Component->GetUniqueID()))
-			{
-				continue;
-			}
-
-			if(AActor* Owner = Component->GetOwner())
+			if (FBoxSphereBounds::SpheresIntersect(QueryBounds, Component->Bounds))
 			{
-				if(IgnoredActors.Contains(Owner->GetUniqueID()))
+				TArray<FOverlapResult> SubOverlaps;
+				if (Component->ComponentOverlapComponentWithResult(PrimComp, Pos, Rot, Params, SubOverlaps))
 				{
-					continue;
+					bHasOverlap = true;
+					OutOverlap.Append(SubOverlaps);
 				}
 			}
 
-			TArray<FOverlapResult> SubOverlaps;
-			if(Component->ComponentOverlapComponentWithResult(PrimComp, Pos, Rot, Params, SubOverlaps))
-			{
-				bHasOverlap = true;
-				OutOverlap.Append(SubOverlaps);
-			}
+			return true;
 		}
-	}
+	);
+
 	return bHasOverlap;
 }
 
@@ -802,11 +791,27 @@ TArray<UPrimitiveComponent*> UClusterUnionComponent::GetAllCurrentChildComponent
 	TArray<UPrimitiveComponent*> Components;
 	Components.Reserve(ComponentToPhysicsObjects.Num() + PendingComponentSync.Num());
 
+	VisitAllCurrentChildComponents(
+		[&Components](UPrimitiveComponent* Component)
+		{
+			Components.Add(Component);
+			return true;
+		}
+	);
+
+	return Components;
+}
+
+void UClusterUnionComponent::VisitAllCurrentChildComponents(const TFunction<bool(UPrimitiveComponent*)>& Lambda) const
+{
 	for (const TPair<TObjectKey<UPrimitiveComponent>, FClusterUnionPendingAddData>& Kvp : PendingComponentSync)
 	{
 		if (UPrimitiveComponent* Component = Kvp.Key.ResolveObjectPtr(); Component && Component->HasValidPhysicsState())
 		{
-			Components.Add(Component);
+			if (!Lambda(Component))
+			{
+				return;
+			}
 		}
 	}
 
@@ -814,11 +819,59 @@ TArray<UPrimitiveComponent*> UClusterUnionComponent::GetAllCurrentChildComponent
 	{
 		if (UPrimitiveComponent* Component = Kvp.Key.ResolveObjectPtr(); Component && Component->HasValidPhysicsState() && !Kvp.Value.bPendingDeletion)
 		{
-			Components.Add(Component);
+			if (!Lambda(Component))
+			{
+				return;
+			}
 		}
 	}
+}
 
-	return Components;
+void UClusterUnionComponent::VisitAllCurrentChildComponentsForCollision(ECollisionChannel TraceChannel, const struct FCollisionQueryParams& Params, const struct FCollisionResponseParams& ResponseParams, const struct FCollisionObjectQueryParams& ObjectParams, const TFunction<bool(UPrimitiveComponent*)>& Lambda) const
+{
+	TSet<uint32> IgnoredActors;
+	IgnoredActors.Reserve(Params.GetIgnoredActors().Num());
+	for (uint32 Id : Params.GetIgnoredActors())
+	{
+		IgnoredActors.Add(Id);
+	}
+
+	TSet<uint32> IgnoredComponents;
+	IgnoredComponents.Reserve(Params.GetIgnoredComponents().Num());
+	for (uint32 Id : Params.GetIgnoredActors())
+	{
+		IgnoredComponents.Add(Id);
+	}
+
+	VisitAllCurrentChildComponents(
+		[TraceChannel, &IgnoredActors, &IgnoredComponents, &Lambda](UPrimitiveComponent* Component)
+		{
+			if (Component)
+			{
+				// TODO: Support all the other filters that could exist here.
+				if (IgnoredComponents.Contains(Component->GetUniqueID()))
+				{
+					return true;
+				}
+
+				if (AActor* Owner = Component->GetOwner())
+				{
+					if (IgnoredActors.Contains(Owner->GetUniqueID()))
+					{
+						return true;
+					}
+
+					if (Owner->GetComponentsCollisionResponseToChannel(TraceChannel) == ECollisionResponse::ECR_Ignore)
+					{
+						return true;
+					}
+				}
+
+				Lambda(Component);
+			}
+			return true;
+		}
+	);
 }
 
 TArray<int32> UClusterUnionComponent::GetAddedBoneIdsForComponent(UPrimitiveComponent* Component) const
diff --git a/Engine/Source/Runtime/Engine/Public/CollisionQueryParams.h b/Engine/Source/Runtime/Engine/Public/CollisionQueryParams.h
--- a/Engine/Source/Runtime/Engine/Public/CollisionQueryParams.h
+++ b/Engine/Source/Runtime/Engine/Public/CollisionQueryParams.h
@@ -67,6 +67,9 @@ struct ENGINE_API FCollisionQueryParams
 	/** Whether to skip narrow phase checks (only for overlaps). */
 	bool bSkipNarrowPhase;
 
+	/** Whether to ignore traces to the cluster union and trace against its children instead. */
+	bool bTraceIntoSubComponents;
+
 	/** Filters query by mobility types (static vs stationary/movable)*/
 	EQueryMobilityType MobilityType;
 
@@ -154,6 +157,7 @@ public:
 #if !(UE_BUILD_TEST || UE_BUILD_SHIPPING)
 		bDebugQuery = false;
 #endif
+		bTraceIntoSubComponents = true;
 	}
 
 	FCollisionQueryParams(FName InTraceTag, bool bInTraceComplex=false, const AActor* InIgnoreActor=NULL)
