Fixed invalid ensure in WorldPartitionSubsystem and adapted WorldPartition::Canstream to also consider the LevelStreaming flag bIsRequestingUnloadAndRemoval.

diff --git a/Engine/Source/Runtime/Engine/Private/WorldPartition/WorldPartition.cpp b/Engine/Source/Runtime/Engine/Private/WorldPartition/WorldPartition.cpp
--- a/Engine/Source/Runtime/Engine/Private/WorldPartition/WorldPartition.cpp
+++ b/Engine/Source/Runtime/Engine/Private/WorldPartition/WorldPartition.cpp
@@ -783,7 +783,7 @@ bool UWorldPartition::CanStream() const
 	// Is it part of a Sub-level that should be visible.
 	if (ULevelStreaming* LevelStreaming = ULevelStreaming::FindStreamingLevel(PersistentLevel))
 	{
-		return LevelStreaming->ShouldBeVisible();
+		return !LevelStreaming->GetIsRequestingUnloadAndRemoval() && LevelStreaming->ShouldBeVisible();
 	}
 
 	return true;
diff --git a/Engine/Source/Runtime/Engine/Private/WorldPartition/WorldPartitionSubsystem.cpp b/Engine/Source/Runtime/Engine/Private/WorldPartition/WorldPartitionSubsystem.cpp
--- a/Engine/Source/Runtime/Engine/Private/WorldPartition/WorldPartitionSubsystem.cpp
+++ b/Engine/Source/Runtime/Engine/Private/WorldPartition/WorldPartitionSubsystem.cpp
@@ -508,12 +508,11 @@ void UWorldPartitionSubsystem::OnLevelStreamingTargetStateChanged(UWorld* World,
 		return;
 	}
 
-	// Make sure when a WorldPartiton is LevelStreamed that changing its state to remove it from world will update the target states of its Cells right away.
+	// Make sure when a WorldPartition is LevelStreamed that changing its state to remove it from world will update the target states of its Cells right away.
 	if(LevelIfLoaded && NewTarget != ELevelStreamingTargetState::LoadedVisible)
 	{
-		// At this point the StreamingLevel should not be visible or the global flag on the world should have been set to unload all streaming levels
-		ensure(!StreamingLevel->ShouldBeVisible() || World->GetShouldForceUnloadStreamingLevels());
-		if (UWorldPartition* WorldPartition = LevelIfLoaded->GetTypedOuter<UWorld>()->GetWorldPartition(); WorldPartition && WorldPartition->IsInitialized())
+		UWorldPartition* WorldPartition = LevelIfLoaded->GetTypedOuter<UWorld>()->GetWorldPartition();
+		if (WorldPartition && WorldPartition->IsInitialized() && !WorldPartition->CanStream())
 		{
 			WorldPartition->UpdateStreamingState();
 		}
