Actor Filter: Fixup redirected assets

#rnx

diff --git a/Engine/Source/Runtime/Engine/Private/WorldPartition/DataLayer/WorldDataLayersActorDesc.cpp b/Engine/Source/Runtime/Engine/Private/WorldPartition/DataLayer/WorldDataLayersActorDesc.cpp
--- a/Engine/Source/Runtime/Engine/Private/WorldPartition/DataLayer/WorldDataLayersActorDesc.cpp
+++ b/Engine/Source/Runtime/Engine/Private/WorldPartition/DataLayer/WorldDataLayersActorDesc.cpp
@@ -4,7 +4,7 @@
 
 #if WITH_EDITOR
 
-#include "WorldPartition/WorldPartitionActorDescUtils.h"
+#include "WorldPartition/WorldPartitionHelpers.h"
 #include "WorldPartition/DataLayer/WorldDataLayers.h"
 #include "WorldPartition/DataLayer/DataLayerInstanceWithAsset.h"
 #include "WorldPartition/WorldPartitionActorContainerID.h"
@@ -38,7 +38,7 @@ FArchive& operator<<(FArchive& Ar, FDataLayerInstanceDesc& Desc)
 	// Fixup redirected data layer asset path
 	if (Ar.IsLoading() && Desc.bIsUsingAsset)
 	{
-		FWorldPartitionActorDescUtils::FixupRedirectedAssetPath(Desc.AssetPath);
+		FWorldPartitionHelpers::FixupRedirectedAssetPath(Desc.AssetPath);
 	}
 
 	if (Ar.CustomVer(FFortniteMainBranchObjectVersion::GUID) >= FFortniteMainBranchObjectVersion::WorldPartitionActorFilter)
diff --git a/Engine/Source/Runtime/Engine/Private/WorldPartition/Filter/WorldPartitionActorFilter.cpp b/Engine/Source/Runtime/Engine/Private/WorldPartition/Filter/WorldPartitionActorFilter.cpp
--- a/Engine/Source/Runtime/Engine/Private/WorldPartition/Filter/WorldPartitionActorFilter.cpp
+++ b/Engine/Source/Runtime/Engine/Private/WorldPartition/Filter/WorldPartitionActorFilter.cpp
@@ -2,6 +2,7 @@
 
 #include "WorldPartition/Filter/WorldPartitionActorFilter.h"
 #include "WorldPartition/WorldPartitionActorLoaderInterface.h"
+#include "WorldPartition/WorldPartitionHelpers.h"
 
 #if WITH_EDITORONLY_DATA
 FWorldPartitionActorFilter::FOnWorldPartitionActorFilterChanged FWorldPartitionActorFilter::OnWorldPartitionActorFilterChanged;
@@ -159,7 +160,9 @@ bool FWorldPartitionActorFilter::Serialize(FArchive& Ar)
 			bool bIncluded;
 			Ar << bIncluded;
 
-			DataLayerFilters.Add(FSoftObjectPath(AssetPath)).bIncluded = bIncluded;
+			FSoftObjectPath SoftObjectPath(AssetPath);
+			FWorldPartitionHelpers::FixupRedirectedAssetPath(SoftObjectPath);
+			DataLayerFilters.Add(SoftObjectPath).bIncluded = bIncluded;
 		}
 	}
 	else
diff --git a/Engine/Source/Runtime/Engine/Private/WorldPartition/WorldPartitionActorDesc.cpp b/Engine/Source/Runtime/Engine/Private/WorldPartition/WorldPartitionActorDesc.cpp
--- a/Engine/Source/Runtime/Engine/Private/WorldPartition/WorldPartitionActorDesc.cpp
+++ b/Engine/Source/Runtime/Engine/Private/WorldPartition/WorldPartitionActorDesc.cpp
@@ -3,7 +3,6 @@
 #if WITH_EDITOR
 #include "WorldPartition/WorldPartitionActorDesc.h"
 #include "WorldPartition/WorldPartitionActorDescArchive.h"
-#include "WorldPartition/WorldPartitionActorDescUtils.h"
 
 #include "Misc/Paths.h"
 #include "Misc/ArchiveMD5.h"
@@ -618,7 +617,7 @@ PRAGMA_ENABLE_DEPRECATION_WARNINGS
 	{
 		for (FName& DataLayer : DataLayers)
 		{
-			FWorldPartitionActorDescUtils::FixupRedirectedAssetPath(DataLayer);
+			FWorldPartitionHelpers::FixupRedirectedAssetPath(DataLayer);
 		}
 	}
 }
diff --git a/Engine/Source/Runtime/Engine/Private/WorldPartition/WorldPartitionActorDescUtils.cpp b/Engine/Source/Runtime/Engine/Private/WorldPartition/WorldPartitionActorDescUtils.cpp
--- a/Engine/Source/Runtime/Engine/Private/WorldPartition/WorldPartitionActorDescUtils.cpp
+++ b/Engine/Source/Runtime/Engine/Private/WorldPartition/WorldPartitionActorDescUtils.cpp
@@ -9,10 +9,10 @@
 #include "WorldPartition/WorldPartitionLog.h"
 #include "Engine/Level.h"
 #include "GameFramework/Actor.h"
-#include "AssetRegistry/AssetRegistryModule.h"
 #include "UObject/CoreRedirects.h"
 #include "Misc/Base64.h"
 #include "WorldPartition/WorldPartitionActorDesc.h"
+#include "WorldPartition/WorldPartitionHelpers.h"
 
 static FName NAME_ActorMetaDataClass(TEXT("ActorMetaDataClass"));
 static FName NAME_ActorMetaData(TEXT("ActorMetaData"));
@@ -174,53 +174,6 @@ void FWorldPartitionActorDescUtils::ReplaceActorDescriptorPointerFromActor(const
 
 bool FWorldPartitionActorDescUtils::FixupRedirectedAssetPath(FName& InOutAssetPath)
 {
-	const FAssetData* AssetData;
-	IAssetRegistry& AssetRegistry = FModuleManager::LoadModuleChecked<FAssetRegistryModule>(TEXT("AssetRegistry")).Get();
-
-	FSoftObjectPath SoftObjectPath(InOutAssetPath.ToString());
-	FTopLevelAssetPath AssetPath = SoftObjectPath.GetAssetPath();
-
-	for (;;)
-	{
-		TArray<FAssetData> Assets;							
-		AssetRegistry.ScanFilesSynchronous({ AssetPath.GetPackageName().ToString() }, /*bForceRescan*/false);
-		AssetRegistry.GetAssetsByPackageName(AssetPath.GetPackageName(), Assets, /*bIncludeOnlyOnDiskAssets*/true);
-
-		if (!Assets.Num())
-		{
-			UE_LOG(LogWorldPartition, Warning, TEXT("Failed to find assets for asset path '%s'"), *AssetPath.ToString());
-			return false;
-		}
-
-		AssetData = Assets.FindByPredicate([&AssetPath](const FAssetData& AssetData)
-		{ 
-			return (AssetData.ToSoftObjectPath().GetAssetPath() == AssetPath);
-		});
-
-		if (!AssetData)
-		{
-			UE_LOG(LogWorldPartition, Warning, TEXT("Failed to find asset for asset path '%s'"), *AssetPath.ToString());
-			return false;
-		}
-
-		if (!AssetData->IsRedirector())
-		{
-			break;
-		}
-
-		FString DestinationObjectPath;
-		if (!AssetData->GetTagValue(TEXT("DestinationObject"), DestinationObjectPath))
-		{
-			UE_LOG(LogWorldPartition, Warning, TEXT("Failed to follow redirector for '%s'"), *AssetPath.ToString());
-			return false;
-		}
-
-		// Update asset path
-		AssetPath = FTopLevelAssetPath(DestinationObjectPath);
-	}
-
-	SoftObjectPath.SetPath(AssetPath, SoftObjectPath.GetSubPathString());	
-	InOutAssetPath = FName(*SoftObjectPath.ToString());
-	return true;
+	return FWorldPartitionHelpers::FixupRedirectedAssetPath(InOutAssetPath);
 }
 #endif
diff --git a/Engine/Source/Runtime/Engine/Private/WorldPartition/WorldPartitionHelpers.cpp b/Engine/Source/Runtime/Engine/Private/WorldPartition/WorldPartitionHelpers.cpp
--- a/Engine/Source/Runtime/Engine/Private/WorldPartition/WorldPartitionHelpers.cpp
+++ b/Engine/Source/Runtime/Engine/Private/WorldPartition/WorldPartitionHelpers.cpp
@@ -11,7 +11,11 @@
 #include "Commandlets/Commandlet.h"
 #include "WorldPartition/WorldPartitionLog.h"
 
-DEFINE_LOG_CATEGORY_STATIC(LogWorldPartitionHelpers, Log, All);
+#if WITH_EDITOR
+#include "Misc/RedirectCollector.h"
+#include "AssetRegistry/AssetRegistryModule.h"
+#include "Modules/ModuleManager.h"
+#endif
 
 namespace FWorldPartitionHelpersPrivate
 {
@@ -58,7 +62,7 @@ bool FWorldPartitionHelpers::IsActorDescClassCompatibleWith(const FWorldPartitio
 
 			if (!ActorBaseClass)
 			{
-				UE_LOG(LogWorldPartitionHelpers, Warning, TEXT("Failed to find actor base class: %s."), *ActorBasePath.ToString());
+				UE_LOG(LogWorldPartition, Warning, TEXT("Failed to find actor base class: %s."), *ActorBasePath.ToString());
 				ActorBaseClass = ActorNativeClass;
 			}
 		}
@@ -300,6 +304,77 @@ bool FWorldPartitionHelpers::ConvertRuntimePathToEditorPath(const FSoftObjectPat
 	return false;
 }
 
+bool FWorldPartitionHelpers::FixupRedirectedAssetPath(FSoftObjectPath& InOutSoftObjectPath)
+{
+	// Check GRedirectCollector first for faster fixup
+	FSoftObjectPath FoundRedirection = GRedirectCollector.GetAssetPathRedirection(InOutSoftObjectPath.GetWithoutSubPath());
+	if (!FoundRedirection.IsNull())
+	{
+		InOutSoftObjectPath.SetPath(FoundRedirection.GetAssetPath(), InOutSoftObjectPath.GetSubPathString());
+		return true;
+	}
+
+	const FAssetData* AssetData;
+	IAssetRegistry& AssetRegistry = FModuleManager::LoadModuleChecked<FAssetRegistryModule>(TEXT("AssetRegistry")).Get();
+
+	FTopLevelAssetPath AssetPath = InOutSoftObjectPath.GetAssetPath();
+
+	for (;;)
+	{
+		TArray<FAssetData> Assets;
+		AssetRegistry.ScanFilesSynchronous({ AssetPath.GetPackageName().ToString() }, /*bForceRescan*/false);
+		AssetRegistry.GetAssetsByPackageName(AssetPath.GetPackageName(), Assets, /*bIncludeOnlyOnDiskAssets*/true);
+
+		if (!Assets.Num())
+		{
+			UE_LOG(LogWorldPartition, Warning, TEXT("Failed to find assets for asset path '%s'"), *AssetPath.ToString());
+			return false;
+		}
+
+		AssetData = Assets.FindByPredicate([&AssetPath](const FAssetData& AssetData)
+		{
+			return (AssetData.ToSoftObjectPath().GetAssetPath() == AssetPath);
+		});
+
+		if (!AssetData)
+		{
+			UE_LOG(LogWorldPartition, Warning, TEXT("Failed to find asset for asset path '%s'"), *AssetPath.ToString());
+			return false;
+		}
+
+		if (!AssetData->IsRedirector())
+		{
+			break;
+		}
+
+		FString DestinationObjectPath;
+		if (!AssetData->GetTagValue(TEXT("DestinationObject"), DestinationObjectPath))
+		{
+			UE_LOG(LogWorldPartition, Warning, TEXT("Failed to follow redirector for '%s'"), *AssetPath.ToString());
+			return false;
+		}
+
+		// Update asset path
+		AssetPath = FTopLevelAssetPath(DestinationObjectPath);
+	}
+
+	InOutSoftObjectPath.SetPath(AssetPath, InOutSoftObjectPath.GetSubPathString());
+
+	return true;
+}
+
+bool FWorldPartitionHelpers::FixupRedirectedAssetPath(FName& InOutAssetPath)
+{
+	FSoftObjectPath SoftObjectPath(InOutAssetPath.ToString());
+	if (FixupRedirectedAssetPath(SoftObjectPath))
+	{
+		InOutAssetPath = FName(*SoftObjectPath.ToString());
+		return true;
+	}
+
+	return false;
+}
+
 #endif // #if WITH_EDITOR
 
 bool FWorldPartitionHelpers::ConvertEditorPathToRuntimePath(const FSoftObjectPath& InPath, FSoftObjectPath& OutPath)
diff --git a/Engine/Source/Runtime/Engine/Public/WorldPartition/WorldPartitionActorDescUtils.h b/Engine/Source/Runtime/Engine/Public/WorldPartition/WorldPartitionActorDescUtils.h
--- a/Engine/Source/Runtime/Engine/Public/WorldPartition/WorldPartitionActorDescUtils.h
+++ b/Engine/Source/Runtime/Engine/Public/WorldPartition/WorldPartitionActorDescUtils.h
@@ -73,11 +73,7 @@ struct ENGINE_API FWorldPartitionActorDescUtils
 	 */
 	static void ReplaceActorDescriptorPointerFromActor(const AActor* InOldActor, AActor* InNewActor, FWorldPartitionActorDesc* InActorDesc);
 
-	/**
-	 * Resolves the provided asset path using asset redirectors.
-	 * @param	InOutAssetPath	Asset path to resolve [In/Out].
-	 * @return					False if an error occured while trying to resolve, else true.
-	 */
+	UE_DEPRECATED(5.3, "Use FWorldPartitionHelpers::FixupRedirectedAssetPath instead")
 	static bool FixupRedirectedAssetPath(FName& InOutAssetPath);
 };
 #endif
\ No newline at end of file
diff --git a/Engine/Source/Runtime/Engine/Public/WorldPartition/WorldPartitionHelpers.h b/Engine/Source/Runtime/Engine/Public/WorldPartition/WorldPartitionHelpers.h
--- a/Engine/Source/Runtime/Engine/Public/WorldPartition/WorldPartitionHelpers.h
+++ b/Engine/Source/Runtime/Engine/Public/WorldPartition/WorldPartitionHelpers.h
@@ -138,6 +138,20 @@ public:
 
 	// Runtime/Editor conversions
 	static bool ConvertRuntimePathToEditorPath(const FSoftObjectPath& InPath, FSoftObjectPath& OutPath);
+
+	/**
+	 * Resolves the provided asset path using asset redirectors.
+	 * @param	InOutAssetPath	Asset path to resolve [In/Out].
+	 * @return	False if an error occured while trying to resolve, else true.
+	 */
+	static bool FixupRedirectedAssetPath(FName& InOutAssetPath);
+	
+	/**
+	 * Resolves the provided soft object asset path using asset redirectors.
+	 * @param	InOutSoftObjectPath	Asset path to resolve [In/Out].
+	 * @return	False if an error occured while trying to resolve, else true.
+	 */
+	static bool FixupRedirectedAssetPath(FSoftObjectPath& InOutSoftObjectPath);
 #endif // WITH_EDITOR
 
 	// Editor/Runtime conversions
