World Partition - HLOD: Better fix for cached bounds
* Only recompute bounds if HLODs belongs to a WP that isn't for the main world (which can have no offset)

diff --git a/Engine/Source/Runtime/Engine/Private/WorldPartition/HLOD/HLODActor.cpp b/Engine/Source/Runtime/Engine/Private/WorldPartition/HLOD/HLODActor.cpp
--- a/Engine/Source/Runtime/Engine/Private/WorldPartition/HLOD/HLODActor.cpp
+++ b/Engine/Source/Runtime/Engine/Private/WorldPartition/HLOD/HLODActor.cpp
@@ -11,6 +11,9 @@
 
 #include UE_INLINE_GENERATED_CPP_BY_NAME(HLODActor)
 
+#include "WorldPartition/WorldPartition.h"
+#include "WorldPartition/WorldPartitionHelpers.h"
+
 #if WITH_EDITOR
 #include "Misc/ArchiveMD5.h"
 #include "UObject/FortniteMainBranchObjectVersion.h"
@@ -132,14 +135,6 @@ void AWorldPartitionHLOD::PostLoad()
 {
 	Super::PostLoad();
 
-	if (GWorldPartitionHLODForceDisableShadows && GetWorld() && GetWorld()->IsGameWorld())
-	{
-		ForEachComponent<UPrimitiveComponent>(false, [](UPrimitiveComponent* PrimitiveComponent)
-		{
-			PrimitiveComponent->SetCastShadow(false);
-		});
-	}
-
 	// If world is instanced, we need to recompute our bounds since they are in the instanced-world space
 	ForEachComponent<USceneComponent>(false, [](USceneComponent* SceneComponent)
 	{
@@ -223,6 +218,32 @@ void AWorldPartitionHLOD::PostLoad()
 #endif
 }
 
+void AWorldPartitionHLOD::PreRegisterAllComponents()
+{
+	Super::PreRegisterAllComponents();
+
+	if (GWorldPartitionHLODForceDisableShadows && GetWorld() && GetWorld()->IsGameWorld())
+	{
+		ForEachComponent<UPrimitiveComponent>(false, [](UPrimitiveComponent* PrimitiveComponent)
+		{
+			PrimitiveComponent->SetCastShadow(false);
+		});
+	}
+
+	// If world is instanced, we need to recompute our bounds since they are in the instanced-world space
+	if (UWorldPartition* WorldPartition = FWorldPartitionHelpers::GetWorldPartition(this))
+	{
+		if (!WorldPartition->IsMainWorldPartition())
+		{
+			ForEachComponent<USceneComponent>(false, [](USceneComponent* SceneComponent)
+			{
+				// Clear bComputedBoundsOnceForGame so that the bounds are recomputed once
+				SceneComponent->bComputedBoundsOnceForGame = false;
+			});
+		}
+	}
+}
+
 #if WITH_EDITOR
 
 void AWorldPartitionHLOD::RerunConstructionScripts()
diff --git a/Engine/Source/Runtime/Engine/Public/WorldPartition/HLOD/HLODActor.h b/Engine/Source/Runtime/Engine/Public/WorldPartition/HLOD/HLODActor.h
--- a/Engine/Source/Runtime/Engine/Public/WorldPartition/HLOD/HLODActor.h
+++ b/Engine/Source/Runtime/Engine/Public/WorldPartition/HLOD/HLODActor.h
@@ -72,6 +72,7 @@ protected:
 	//~ End UObject Interface.
 
 	//~ Begin AActor Interface.
+	virtual void PreRegisterAllComponents() override;
 	virtual void BeginPlay() override;
 	virtual void EndPlay(const EEndPlayReason::Type EndPlayReason) override;
 	virtual bool IsRuntimeOnly() const override { return true; }
