Fix for default gravity direction initialization in CharacterMovementComponent
#rnx

diff --git a/Engine/Source/Runtime/Engine/Private/Components/CharacterMovementComponent.cpp b/Engine/Source/Runtime/Engine/Private/Components/CharacterMovementComponent.cpp
--- a/Engine/Source/Runtime/Engine/Private/Components/CharacterMovementComponent.cpp
+++ b/Engine/Source/Runtime/Engine/Private/Components/CharacterMovementComponent.cpp
@@ -83,7 +83,7 @@ namespace CharacterMovementConstants
 	const float VERTICAL_SLOPE_NORMAL_Z = 0.001f; // Slope is vertical if Abs(Normal.Z) <= this threshold. Accounts for precision problems that sometimes angle normals slightly off horizontal for vertical surface.
 }
 
-const FVector UCharacterMovementComponent::DefaultGravityDirection = FVector::DownVector;
+const FVector UCharacterMovementComponent::DefaultGravityDirection = FVector(0.0, 0.0, -1.0);
 const float UCharacterMovementComponent::MIN_TICK_TIME = 1e-6f;
 const float UCharacterMovementComponent::MIN_FLOOR_DIST = 1.9f;
 const float UCharacterMovementComponent::MAX_FLOOR_DIST = 2.4f;
