Deprecate ray traced skylight and remove it from project settings.

diff --git a/Engine/Source/Runtime/Engine/Classes/Engine/RendererSettings.h b/Engine/Source/Runtime/Engine/Classes/Engine/RendererSettings.h
--- a/Engine/Source/Runtime/Engine/Classes/Engine/RendererSettings.h
+++ b/Engine/Source/Runtime/Engine/Classes/Engine/RendererSettings.h
@@ -528,10 +528,8 @@ class ENGINE_API URendererSettings : public UDeveloperSettings
 		ToolTip = "Controls whether Ray Traced Shadows are used by default. Lights can still override and force Ray Traced shadows on or off. Requires Hardware Ray Tracing to be enabled."))
 		uint32 bEnableRayTracingShadows : 1;
 
-	UPROPERTY(config, EditAnywhere, Category = HardwareRayTracing, meta = (
-		ConsoleVariable = "r.RayTracing.Skylight", DisplayName = "Ray Traced Skylight",
-		ToolTip = "Controls whether Ray Traced Skylight is used by default. Skylights can still override and force Ray Traced Skylight on or off. Requires Hardware Ray Tracing to be enabled.  Has no effect when Dynamic Global Illumination Method is set to Lumen."))
-		uint32 bEnableRayTracingSkylight : 1;
+	UPROPERTY()
+		uint32 bEnableRayTracingSkylight_DEPRECATED : 1;
 
 	UPROPERTY(config, EditAnywhere, Category = HardwareRayTracing, meta = (
 		ConsoleVariable = "r.RayTracing.UseTextureLod", DisplayName = "Texture LOD",
diff --git a/Engine/Source/Runtime/Engine/Private/RendererSettings.cpp b/Engine/Source/Runtime/Engine/Private/RendererSettings.cpp
--- a/Engine/Source/Runtime/Engine/Private/RendererSettings.cpp
+++ b/Engine/Source/Runtime/Engine/Private/RendererSettings.cpp
@@ -87,7 +87,6 @@ URendererSettings::URendererSettings(const FObjectInitializer& ObjectInitializer
 	bEnableRayTracing = 0;
 	bUseHardwareRayTracingForLumen = 0;
 	bEnableRayTracingShadows = 0;
-	bEnableRayTracingSkylight = 0;
 	bEnablePathTracing = 0;
 	bEnableRayTracingTextureLOD = 0;
 	DefaultBoneInfluenceLimit = 0;
