Texture.h MaxTextureSize change to BlueprintReadWrite, not sure why it was ReadOnly

#rnx

diff --git a/Engine/Source/Runtime/Engine/Classes/Engine/Texture.h b/Engine/Source/Runtime/Engine/Classes/Engine/Texture.h
--- a/Engine/Source/Runtime/Engine/Classes/Engine/Texture.h
+++ b/Engine/Source/Runtime/Engine/Classes/Engine/Texture.h
@@ -1173,7 +1173,7 @@ public:
 	FName OodleTextureSdkVersion;
 
 	/** The maximum resolution for generated textures. A value of 0 means the maximum size for the format on each platform. */
-	UPROPERTY(EditAnywhere, BlueprintReadOnly, Category=Compression, meta=(DisplayName="Maximum Texture Size", ClampMin = "0.0"), AdvancedDisplay)
+	UPROPERTY(EditAnywhere, BlueprintReadWrite, Category=Compression, meta=(DisplayName="Maximum Texture Size", ClampMin = "0.0"), AdvancedDisplay)
 	int32 MaxTextureSize;
 
 	/** The compression quality for generated ASTC textures (i.e. mobile platform textures). */
