Substrate - Properly initialise compilation result structures to fix DDC nondeterminism.

diff --git a/Engine/Source/Runtime/Engine/Public/Rendering/StrataMaterialShared.cpp b/Engine/Source/Runtime/Engine/Public/Rendering/StrataMaterialShared.cpp
--- a/Engine/Source/Runtime/Engine/Public/Rendering/StrataMaterialShared.cpp
+++ b/Engine/Source/Runtime/Engine/Public/Rendering/StrataMaterialShared.cpp
@@ -15,8 +15,15 @@ FStrataMaterialCompilationOutput::FStrataMaterialCompilationOutput()
 	, RequestedBytePixePixel(0)
 	, PlatformBytePixePixel(0)
 	, bMaterialOutOfBudgetHasBeenSimplified(0)
+	, RootOperatorIndex(0)
 #endif
 {
+#if WITH_EDITOR
+	for (uint32 i = 0; i < STRATA_COMPILATION_OUTPUT_MAX_OPERATOR; ++i)
+	{
+		Operators[i] = FStrataOperator();
+	}
+#endif
 }
 
 
@@ -35,6 +42,7 @@ FStrataOperator::FStrataOperator()
 
 	BSDFIndex = INDEX_NONE;
 	BSDFType = 0;
+	BSDFRegisteredSharedLocalBasis = FStrataRegisteredSharedLocalBasis();
 	bBSDFHasSSS = false;
 	bBSDFHasMFPPluggedIn = false;
 	bBSDFHasEdgeColor = false;
diff --git a/Engine/Source/Runtime/Engine/Public/Rendering/StrataMaterialShared.h b/Engine/Source/Runtime/Engine/Public/Rendering/StrataMaterialShared.h
--- a/Engine/Source/Runtime/Engine/Public/Rendering/StrataMaterialShared.h
+++ b/Engine/Source/Runtime/Engine/Public/Rendering/StrataMaterialShared.h
@@ -91,8 +91,7 @@ public:
 	////
 	//// The following data is only needed when compiling with the editor.
 	////
-
-	// STRATA_TODO Strata Tree topology
+	
 	// STRATA_TODO corresponding uint8 ShadingModel? IsThin?
 
 	/** The Strata verbose description */
