[LWM] debug display the spawn points in cooked

diff --git a/Engine/Source/Runtime/Engine/Public/DebugRenderSceneProxy.h b/Engine/Source/Runtime/Engine/Public/DebugRenderSceneProxy.h
--- a/Engine/Source/Runtime/Engine/Public/DebugRenderSceneProxy.h
+++ b/Engine/Source/Runtime/Engine/Public/DebugRenderSceneProxy.h
@@ -312,6 +312,8 @@ protected:
 
 	ENGINE_API virtual void DrawDebugLabels(UCanvas* Canvas, APlayerController*);
 	ENGINE_API void ResetTexts() { Texts.Reset(); }
+	ENGINE_API const TextArray& GetTexts() const { return Texts; }
+	ENGINE_API float GetTextWithoutShadowDistance() const {return TextWithoutShadowDistance; }
 
 protected:
 	FDebugDrawDelegate DebugTextDrawingDelegate;
