WorldPartition: Fixed PartitionActorDesc losing grid indices when resaved after a fixup redirector.

diff --git a/Engine/Source/Runtime/Engine/Private/WorldPartition/ActorPartition/PartitionActorDesc.cpp b/Engine/Source/Runtime/Engine/Private/WorldPartition/ActorPartition/PartitionActorDesc.cpp
--- a/Engine/Source/Runtime/Engine/Private/WorldPartition/ActorPartition/PartitionActorDesc.cpp
+++ b/Engine/Source/Runtime/Engine/Private/WorldPartition/ActorPartition/PartitionActorDesc.cpp
@@ -64,4 +64,17 @@ FBox FPartitionActorDesc::GetEditorBounds() const
 	const UActorPartitionSubsystem::FCellCoord CellCoord(GridIndexX, GridIndexY, GridIndexZ, 0);
 	return UActorPartitionSubsystem::FCellCoord::GetCellBounds(CellCoord, GridSize);
 }
+
+void FPartitionActorDesc::TransferWorldData(const FWorldPartitionActorDesc* From)
+{
+	FWorldPartitionActorDesc::TransferWorldData(From);
+
+	// TransferWorldData is called for actors that are not added to the world (not registered)
+	// Transfer properties that depend on actor being added to the world (components being registered)
+	const FPartitionActorDesc* FromPartitionActorDesc = (FPartitionActorDesc*)From;
+	GridIndexX = FromPartitionActorDesc->GridIndexX;
+	GridIndexY = FromPartitionActorDesc->GridIndexY;
+	GridIndexZ = FromPartitionActorDesc->GridIndexZ;
+}
+
 #endif
diff --git a/Engine/Source/Runtime/Engine/Public/WorldPartition/ActorPartition/PartitionActorDesc.h b/Engine/Source/Runtime/Engine/Public/WorldPartition/ActorPartition/PartitionActorDesc.h
--- a/Engine/Source/Runtime/Engine/Public/WorldPartition/ActorPartition/PartitionActorDesc.h
+++ b/Engine/Source/Runtime/Engine/Public/WorldPartition/ActorPartition/PartitionActorDesc.h
@@ -24,5 +24,6 @@ protected:
 	virtual bool Equals(const FWorldPartitionActorDesc* Other) const override;
 	virtual void Serialize(FArchive& Ar) override;
 	virtual FBox GetEditorBounds() const override;
+	virtual void TransferWorldData(const FWorldPartitionActorDesc* From) override;
 };
 #endif
