Events optimization:  Remove duplicated search for component from proxy

diff --git a/Engine/Source/Runtime/Engine/Private/PhysicsEngine/Experimental/PhysScene_Chaos.cpp b/Engine/Source/Runtime/Engine/Private/PhysicsEngine/Experimental/PhysScene_Chaos.cpp
--- a/Engine/Source/Runtime/Engine/Private/PhysicsEngine/Experimental/PhysScene_Chaos.cpp
+++ b/Engine/Source/Runtime/Engine/Private/PhysicsEngine/Experimental/PhysScene_Chaos.cpp
@@ -807,7 +807,7 @@ FCollisionNotifyInfo& FPhysScene_Chaos::GetPendingCollisionForContactPair(const
 	return PendingCollisionNotifies[NewIdx];
 }
 
-FORCEINLINE void FPhysScene_Chaos::HandleEachCollisionEvent(const TArray<int32>& CollisionIndices, IPhysicsProxyBase* PhysicsProxy0, Chaos::FCollisionDataArray const& CollisionData, Chaos::FReal MinDeltaVelocityThreshold)
+FORCEINLINE void FPhysScene_Chaos::HandleEachCollisionEvent(const TArray<int32>& CollisionIndices, IPhysicsProxyBase* PhysicsProxy0, UPrimitiveComponent* const Comp0, Chaos::FCollisionDataArray const& CollisionData, Chaos::FReal MinDeltaVelocityThreshold)
 {
 	for (int32 EncodedCollisionIdx : CollisionIndices)
 	{
@@ -818,7 +818,8 @@ FORCEINLINE void FPhysScene_Chaos::HandleEachCollisionEvent(const TArray<int32>&
 		IPhysicsProxyBase* const PhysicsProxy1 = bSwapOrder ? CollisionDataItem.Proxy1 : CollisionDataItem.Proxy2;
 
 		// Are the proxies pending destruction? If they are no longer tracked by the PhysScene, the proxy is deleted or pending deletion.
-		if (GetOwningComponent<UPrimitiveComponent>(PhysicsProxy0) == nullptr || GetOwningComponent<UPrimitiveComponent>(PhysicsProxy1) == nullptr)
+		UPrimitiveComponent* const Comp1 = GetOwningComponent<UPrimitiveComponent>(PhysicsProxy1);
+		if (Comp0 == nullptr || Comp1 == nullptr)
 		{
 			return;
 		}
@@ -875,12 +876,12 @@ FORCEINLINE void FPhysScene_Chaos::HandleEachCollisionEvent(const TArray<int32>&
 			// in some case ( like with geometry collections ) we don't have a body instance so the component part will null, we need to handle that 
 			if (NotifyInfo.Info0.Component == nullptr)
 			{
-				NotifyInfo.Info0.Component = GetOwningComponent<UPrimitiveComponent>(PhysicsProxy0);
+				NotifyInfo.Info0.Component = Comp0;
 				NotifyInfo.Info0.Actor = (NotifyInfo.Info0.Component != nullptr) ? NotifyInfo.Info0.Component->GetOwner() : nullptr;
 			}
 			if (NotifyInfo.Info1.Component == nullptr)
 			{
-				NotifyInfo.Info1.Component = GetOwningComponent<UPrimitiveComponent>(PhysicsProxy1);
+				NotifyInfo.Info1.Component = Comp1;
 				NotifyInfo.Info1.Actor = (NotifyInfo.Info1.Component != nullptr) ? NotifyInfo.Info1.Component->GetOwner() : nullptr;
 			}
 		}
@@ -913,7 +914,7 @@ void FPhysScene_Chaos::HandleCollisionEvents(const Chaos::FCollisionEventData& E
 
 				if (Comp0 != nullptr && CollisionEventRegistrations.Contains(Comp0))
 				{
-					HandleEachCollisionEvent(CollisionIndices, PhysicsProxy0, CollisionData, MinDeltaVelocityThreshold);
+					HandleEachCollisionEvent(CollisionIndices, PhysicsProxy0, Comp0, CollisionData, MinDeltaVelocityThreshold);
 				}
 			}
 		}
@@ -933,7 +934,7 @@ void FPhysScene_Chaos::HandleCollisionEvents(const Chaos::FCollisionEventData& E
 						TArray<int32> const* const CollisionIndices = PhysicsProxyToCollisionIndicesMap.Find(PhysicsProxy0);
 						if (CollisionIndices)
 						{
-							HandleEachCollisionEvent(*CollisionIndices, PhysicsProxy0, CollisionData, MinDeltaVelocityThreshold);
+							HandleEachCollisionEvent(*CollisionIndices, PhysicsProxy0, Comp0, CollisionData, MinDeltaVelocityThreshold);
 						}
 					}
 				}
diff --git a/Engine/Source/Runtime/Engine/Public/Physics/Experimental/PhysScene_Chaos.h b/Engine/Source/Runtime/Engine/Public/Physics/Experimental/PhysScene_Chaos.h
--- a/Engine/Source/Runtime/Engine/Public/Physics/Experimental/PhysScene_Chaos.h
+++ b/Engine/Source/Runtime/Engine/Public/Physics/Experimental/PhysScene_Chaos.h
@@ -275,7 +275,7 @@ private:
 	TArray<FCollisionNotifyInfo> PendingCollisionNotifies;
 
 	// Chaos Event Handlers
-	void HandleEachCollisionEvent(const TArray<int32>& CollisionIndices, IPhysicsProxyBase* PhysicsProxy0, Chaos::FCollisionDataArray const& CollisionData, Chaos::FReal MinDeltaVelocityThreshold);
+	void HandleEachCollisionEvent(const TArray<int32>& CollisionIndices, IPhysicsProxyBase* PhysicsProxy0, UPrimitiveComponent* const Comp0, Chaos::FCollisionDataArray const& CollisionData, Chaos::FReal MinDeltaVelocityThreshold);
 	void HandleCollisionEvents(const Chaos::FCollisionEventData& CollisionData);
 
 	void DispatchPendingCollisionNotifies();
