Make sure the export name is updated when renaming a bone

diff --git a/Engine/Source/Runtime/Engine/Private/Animation/ReferenceSkeleton.cpp b/Engine/Source/Runtime/Engine/Private/Animation/ReferenceSkeleton.cpp
--- a/Engine/Source/Runtime/Engine/Private/Animation/ReferenceSkeleton.cpp
+++ b/Engine/Source/Runtime/Engine/Private/Animation/ReferenceSkeleton.cpp
@@ -203,6 +203,9 @@ void FReferenceSkeleton::Rename(const FName InBoneName, const FName InNewName)
 
 	FMeshBoneInfo& BoneInfo = RawRefBoneInfo[RawBoneIndex];
 	BoneInfo.Name = InNewName;
+#if WITH_EDITORONLY_DATA
+	BoneInfo.ExportName = BoneInfo.Name.ToString();
+#endif
 	RawNameToIndexMap.Remove(InBoneName);
 	RawNameToIndexMap.Add(InNewName, RawBoneIndex);
 }
