Level Instance: Fix recursive editing

- Avoid using Package/ActorName as part of FLevelInstanceID outside of GameWorlds

#rnx

diff --git a/Engine/Source/Runtime/Engine/Private/LevelInstance/LevelInstanceSubsystem.cpp b/Engine/Source/Runtime/Engine/Private/LevelInstance/LevelInstanceSubsystem.cpp
--- a/Engine/Source/Runtime/Engine/Private/LevelInstance/LevelInstanceSubsystem.cpp
+++ b/Engine/Source/Runtime/Engine/Private/LevelInstance/LevelInstanceSubsystem.cpp
@@ -126,7 +126,7 @@ FLevelInstanceID::FLevelInstanceID(ULevelInstanceSubsystem* LevelInstanceSubsyst
 	
 	uint64 TmpHash = 0;
 		
-	if (LevelInstanceActor->IsNameStableForNetworking())
+	if (LevelInstanceActor->IsNameStableForNetworking() && LevelInstanceActor->GetWorld()->IsGameWorld())
 	{
 		ActorName = LevelInstanceActor->GetFName();
 		FString NameStr = ActorName.ToString();
