fix possibility where using cluster union tracing with ignored components would result in potential infinite loops

- fixed issue where the ignored component would not actually be respected properly by UClusterUnionComponent::VisitAllCurrentChildComponentsForCollision
- fixed issue where TSceneCastCommonImp and GeomOverlapMultiImp would make an assumption that hitting a cluster union equates to hitting a child component (not accounting for a possibility of a sub-ignored actor/component). this results in cases where we can hit the cluster union and not hit a child component and then not remove the cluster union from the hit result (as you would expect).

diff --git a/Engine/Source/Runtime/Engine/Private/Collision/SceneQuery.cpp b/Engine/Source/Runtime/Engine/Private/Collision/SceneQuery.cpp
--- a/Engine/Source/Runtime/Engine/Private/Collision/SceneQuery.cpp
+++ b/Engine/Source/Runtime/Engine/Private/Collision/SceneQuery.cpp
@@ -332,6 +332,12 @@ EThreadQueryContext GetThreadQueryContext(const Chaos::FPhysicsSolver& Solver)
 	}
 }
 
+struct FClusterUnionHit
+{
+	bool bIsClusterUnion = false;
+	bool bHit = false;
+};
+
 template <typename Traits, typename TGeomInputs>
 bool TSceneCastCommonImp(const UWorld* World, typename Traits::TOutHits& OutHits, const TGeomInputs& GeomInputs, const FVector Start, const FVector End, ECollisionChannel TraceChannel, const struct FCollisionQueryParams& Params, const struct FCollisionResponseParams& ResponseParams, const struct FCollisionObjectQueryParams& ObjectParams)
 {
@@ -420,19 +426,21 @@ bool TSceneCastCommonImp(const UWorld* World, typename Traits::TOutHits& OutHits
 			{
 				auto DoClusterUnionTrace = [&GeomInputs, &Start, &End, TraceChannel, &Params, &ResponseParams, &ObjectParams](const FHitResult& OriginalHit, FHitResult& NewHit)
 				{
+					FClusterUnionHit Result;
 					if (UClusterUnionComponent* ClusterUnion = Cast<UClusterUnionComponent>(OriginalHit.GetComponent()))
 					{
+						Result.bIsClusterUnion = true;
 						if constexpr (Traits::IsRay())
 						{
-							return ClusterUnion->LineTraceComponent(NewHit, Start, End, TraceChannel, Params, ResponseParams, ObjectParams);
+							Result.bHit = ClusterUnion->LineTraceComponent(NewHit, Start, End, TraceChannel, Params, ResponseParams, ObjectParams);
 						}
 						else
 						{
-							return ClusterUnion->SweepComponent(NewHit, Start, End, *GeomInputs.GetGeometryOrientation(), *GeomInputs.GetGeometry(), TraceChannel, Params, ResponseParams, ObjectParams);
+							Result.bHit = ClusterUnion->SweepComponent(NewHit, Start, End, *GeomInputs.GetGeometryOrientation(), *GeomInputs.GetGeometry(), TraceChannel, Params, ResponseParams, ObjectParams);
 						}
 					}
 
-					return false;
+					return Result;
 				};
 
 				if (bSuccess && Params.bTraceIntoSubComponents)
@@ -445,10 +453,15 @@ bool TSceneCastCommonImp(const UWorld* World, typename Traits::TOutHits& OutHits
 						for (int32 Index = 0; Index < OutHits.Num(); ++Index)
 						{
 							FHitResult NewHit;
-							if (DoClusterUnionTrace(OutHits[Index], NewHit))
+							FClusterUnionHit ClusterUnionHit = DoClusterUnionTrace(OutHits[Index], NewHit);
+							if (ClusterUnionHit.bIsClusterUnion)
 							{
-								AllNewHits.Add(NewHit);
 								ClusterUnionIndices.Add(Index);
+
+								if (ClusterUnionHit.bHit)
+								{
+									AllNewHits.Add(NewHit);
+								}
 							}
 						}
 
@@ -463,9 +476,14 @@ bool TSceneCastCommonImp(const UWorld* World, typename Traits::TOutHits& OutHits
 					else
 					{
 						FHitResult NewHit;
-						if (DoClusterUnionTrace(OutHits, NewHit))
+						FClusterUnionHit ClusterUnionHit = DoClusterUnionTrace(OutHits, NewHit);
+						if (ClusterUnionHit.bIsClusterUnion)
 						{
-							OutHits = NewHit;
+							Traits::ResetOutHits(OutHits, Start, End);
+							if (ClusterUnionHit.bHit)
+							{
+								OutHits = NewHit;
+							}
 						}
 					}
 				}
@@ -699,12 +717,14 @@ bool GeomOverlapMultiImp(const UWorld* World, const FPhysicsGeometry& Geom, cons
 					{
 						auto DoClusterUnionOverlap = [&GeomPose, &Geom, TraceChannel, &Params, &ResponseParams, &ObjectParams](const FOverlapResult& OriginalOverlap, TArray<FOverlapResult>& NewOverlaps)
 						{
+							FClusterUnionHit Result;
 							if (UClusterUnionComponent* ClusterUnion = Cast<UClusterUnionComponent>(OriginalOverlap.GetComponent()))
 							{
-								return ClusterUnion->OverlapComponentWithResult(GeomPose.GetTranslation(), GeomPose.GetRotation(), Geom, TraceChannel, Params, ResponseParams, ObjectParams, NewOverlaps);
+								Result.bIsClusterUnion = true;
+								Result.bHit = ClusterUnion->OverlapComponentWithResult(GeomPose.GetTranslation(), GeomPose.GetRotation(), Geom, TraceChannel, Params, ResponseParams, ObjectParams, NewOverlaps);
 							}
 
-							return false;
+							return Result;
 						};
 
 						if (bHaveBlockingHit && Params.bTraceIntoSubComponents)
@@ -715,10 +735,15 @@ bool GeomOverlapMultiImp(const UWorld* World, const FPhysicsGeometry& Geom, cons
 							for (int32 Index = 0; Index < OutOverlaps.Num(); ++Index)
 							{
 								TArray<FOverlapResult> NewOverlaps;
-								if (DoClusterUnionOverlap(OutOverlaps[Index], NewOverlaps))
+								FClusterUnionHit ClusterUnionHit = DoClusterUnionOverlap(OutOverlaps[Index], NewOverlaps);
+								if (ClusterUnionHit.bIsClusterUnion)
 								{
-									AllNewOverlaps.Append(NewOverlaps);
 									ClusterUnionIndices.Add(Index);
+
+									if (ClusterUnionHit.bHit)
+									{
+										AllNewOverlaps.Append(NewOverlaps);
+									}
 								}
 							}
 
diff --git a/Engine/Source/Runtime/Engine/Private/PhysicsEngine/ClusterUnionComponent.cpp b/Engine/Source/Runtime/Engine/Private/PhysicsEngine/ClusterUnionComponent.cpp
--- a/Engine/Source/Runtime/Engine/Private/PhysicsEngine/ClusterUnionComponent.cpp
+++ b/Engine/Source/Runtime/Engine/Private/PhysicsEngine/ClusterUnionComponent.cpp
@@ -838,7 +838,7 @@ void UClusterUnionComponent::VisitAllCurrentChildComponentsForCollision(ECollisi
 
 	TSet<uint32> IgnoredComponents;
 	IgnoredComponents.Reserve(Params.GetIgnoredComponents().Num());
-	for (uint32 Id : Params.GetIgnoredActors())
+	for (uint32 Id : Params.GetIgnoredComponents())
 	{
 		IgnoredComponents.Add(Id);
 	}
