World Partition Streaming Generation Logs: only output cook generation logs on build machines.

diff --git a/Engine/Source/Runtime/Engine/Private/WorldPartition/WorldPartitionStreamingGeneration.cpp b/Engine/Source/Runtime/Engine/Private/WorldPartition/WorldPartitionStreamingGeneration.cpp
--- a/Engine/Source/Runtime/Engine/Private/WorldPartition/WorldPartitionStreamingGeneration.cpp
+++ b/Engine/Source/Runtime/Engine/Private/WorldPartition/WorldPartitionStreamingGeneration.cpp
@@ -1140,7 +1140,7 @@ bool UWorldPartition::GenerateContainerStreaming(const FGenerateStreamingParams&
 	// Dump state log
 	TUniquePtr<FArchive> LogFileAr;
 
-	if (IsMainWorldPartition())
+	if (IsMainWorldPartition() && (!GIsBuildMachine || IsRunningCookCommandlet()))
 	{
 		TStringBuilder<256> StateLogSuffix;
 		StateLogSuffix += bIsPIE ? TEXT("PIE") : (IsRunningGame() ? TEXT("Game") : (IsRunningCookCommandlet() ? TEXT("Cook") : (GIsAutomationTesting ? TEXT("UnitTest") : TEXT("Manual"))));
