cluster union components should properly compute and cache its local bounds

diff --git a/Engine/Source/Runtime/Engine/Classes/PhysicsEngine/ClusterUnionComponent.h b/Engine/Source/Runtime/Engine/Classes/PhysicsEngine/ClusterUnionComponent.h
--- a/Engine/Source/Runtime/Engine/Classes/PhysicsEngine/ClusterUnionComponent.h
+++ b/Engine/Source/Runtime/Engine/Classes/PhysicsEngine/ClusterUnionComponent.h
@@ -205,6 +205,10 @@ private:
 	UPROPERTY(ReplicatedUsing=OnRep_RigidState)
 	FClusterUnionReplicatedData ReplicatedRigidState;
 
+	// Cached local bounds from the physics particle.
+	mutable bool bHasCachedLocalBounds;
+	mutable FBoxSphereBounds CachedLocalBounds;
+
 	// Handles changes to ReplicatedRigidState. Note that this function does not handle replication of X/R since we make use
 	// of the scene component's default replication for that.
 	UFUNCTION()
@@ -269,6 +273,7 @@ public:
 	//~ Begin USceneComponent Interface
 public:
 	virtual void OnUpdateTransform(EUpdateTransformFlags UpdateTransformFlags, ETeleportType Teleport) override;
+	virtual FBoxSphereBounds CalcBounds(const FTransform& LocalToWorld) const override;
 	//~ End USceneComponent Interface
 
 	//~ Begin IPhysicsComponent Interface.
diff --git a/Engine/Source/Runtime/Engine/Private/PhysicsEngine/ClusterUnionComponent.cpp b/Engine/Source/Runtime/Engine/Private/PhysicsEngine/ClusterUnionComponent.cpp
--- a/Engine/Source/Runtime/Engine/Private/PhysicsEngine/ClusterUnionComponent.cpp
+++ b/Engine/Source/Runtime/Engine/Private/PhysicsEngine/ClusterUnionComponent.cpp
@@ -24,7 +24,9 @@ UClusterUnionComponent::UClusterUnionComponent(const FObjectInitializer& ObjectI
 {
 	PhysicsProxy = nullptr;
 	SetIsReplicatedByDefault(true);
+	bComputeBoundsOnceForGame = false;
 	bHasReceivedTransform = false;
+	bHasCachedLocalBounds = false;
 #if WITH_EDITORONLY_DATA
 	bVisualizeComponent = true;
 #endif
@@ -361,6 +363,27 @@ void UClusterUnionComponent::OnUpdateTransform(EUpdateTransformFlags UpdateTrans
 	}
 }
 
+FBoxSphereBounds UClusterUnionComponent::CalcBounds(const FTransform& LocalToWorld) const
+{
+	if (!bHasCachedLocalBounds)
+	{
+		if (!PhysicsProxy)
+		{
+			return Super::CalcBounds(LocalToWorld);
+		}
+		else
+		{
+			Chaos::FPhysicsObjectHandle Handle = PhysicsProxy->GetPhysicsObjectHandle();
+			FLockedReadPhysicsObjectExternalInterface Interface = FPhysicsObjectExternalInterface::LockRead({ &Handle, 1 });
+			FBoxSphereBounds NewBounds = Interface->GetBounds({ &Handle, 1 });
+			bHasCachedLocalBounds = true;
+			CachedLocalBounds = NewBounds;
+		}
+	}
+
+	return CachedLocalBounds.TransformBy(LocalToWorld);
+}
+
 bool UClusterUnionComponent::ShouldCreatePhysicsState() const
 {
 	return true;
@@ -474,6 +497,9 @@ void UClusterUnionComponent::SyncClusterUnionFromProxy()
 	{
 		HandleRemovedClusteredComponent(Component, true);
 	}
+
+	bHasCachedLocalBounds = false;
+	UpdateBounds();
 }
 
 void UClusterUnionComponent::HandleAddOrModifiedClusteredComponent(UPrimitiveComponent* ChangedComponent, const TMap<int32, FTransform>& PerBoneChildToParent)
