Adding cvar "net.DiscardTornOffActorRPCs" to skip sending RPCs for torn off actors, and trigger an actor channel failure if a new channel is opened for a torn off actor on clients

diff --git a/Engine/Source/Runtime/Engine/Private/DataChannel.cpp b/Engine/Source/Runtime/Engine/Private/DataChannel.cpp
--- a/Engine/Source/Runtime/Engine/Private/DataChannel.cpp
+++ b/Engine/Source/Runtime/Engine/Private/DataChannel.cpp
@@ -111,6 +111,7 @@ static FAutoConsoleVariableRef CVarDormancyHysteresis(
 namespace UE::Net
 {
 	extern int32 FilterGuidRemapping;
+	extern bool bDiscardTornOffActorRPCs;
 
 	static float QueuedBunchTimeoutSeconds = 30.0f;
 	static FAutoConsoleVariableRef CVarQueuedBunchTimeoutSeconds(
@@ -3040,6 +3041,24 @@ void UActorChannel::ProcessBunch( FInBunch & Bunch )
 			}
 			return;
 		}
+		else
+		{
+			if (UE::Net::bDiscardTornOffActorRPCs && NewChannelActor->GetTearOff())
+			{
+				UE_LOG(LogNet, Warning, TEXT("UActorChannel::ProcessBunch: SerializeNewActor received an open bunch for a torn off actor. Actor: %s, Channel: %i"), *GetFullNameSafe(NewChannelActor), ChIndex);
+				Broken = 1;
+
+				if (!Connection->IsInternalAck()
+#if !UE_BUILD_SHIPPING
+					&& !bBlockChannelFailure
+#endif
+					)
+				{
+					FNetControlMessage<NMT_ActorChannelFailure>::Send(Connection, ChIndex);
+				}
+				return;
+			}
+		}
 
 		ESetChannelActorFlags Flags = ESetChannelActorFlags::None;
 		if (GSkipReplicatorForDestructionInfos != 0 && Bunch.bClose && Bunch.AtEnd())
diff --git a/Engine/Source/Runtime/Engine/Private/NetDriver.cpp b/Engine/Source/Runtime/Engine/Private/NetDriver.cpp
--- a/Engine/Source/Runtime/Engine/Private/NetDriver.cpp
+++ b/Engine/Source/Runtime/Engine/Private/NetDriver.cpp
@@ -327,6 +327,13 @@ namespace UE::Net
 	{
 		return UE::Net::Private::bIgnoreStaticActorDestruction;
 	}
+
+	bool bDiscardTornOffActorRPCs = true;
+	FAutoConsoleVariableRef CVarNetDiscardTornOffActorRPCs(
+		TEXT("net.DiscardTornOffActorRPCs"),
+		bDiscardTornOffActorRPCs,
+		TEXT("If enabled, discard RPCs if the actor has been torn off."),
+		ECVF_Default);
 }
 
 #if UE_BUILD_SHIPPING
@@ -6660,6 +6667,12 @@ void UNetDriver::ProcessRemoteFunction(
 		return;
 	}
 
+	if (UE::Net::bDiscardTornOffActorRPCs && Actor->GetTearOff())
+	{
+		UE_LOG(LogNet, Warning, TEXT("UNetDriver::ProcessRemoteFunction: Remote function %s called from actor %s while actor is torn off. Function will not be processed."), *Function->GetName(), *Actor->GetName());
+		return;
+	}
+
 #if !UE_BUILD_SHIPPING
 	SCOPE_CYCLE_COUNTER(STAT_NetProcessRemoteFunc);
 	SCOPE_CYCLE_UOBJECT(Function, Function);
