bTraceIntoSubComponents needs a default value in every FCollisionQueryParams constructor

diff --git a/Engine/Source/Runtime/Engine/Private/Collision/Collision.cpp b/Engine/Source/Runtime/Engine/Private/Collision/Collision.cpp
--- a/Engine/Source/Runtime/Engine/Private/Collision/Collision.cpp
+++ b/Engine/Source/Runtime/Engine/Private/Collision/Collision.cpp
@@ -55,6 +55,7 @@ FCollisionQueryParams::FCollisionQueryParams(FName InTraceTag, const TStatId& In
 #if !(UE_BUILD_TEST || UE_BUILD_SHIPPING)
 		bDebugQuery = false;
 #endif
+	bTraceIntoSubComponents = true;
 }
 
 
