Added BP exposed setter function for NumCustomDataFloats on InstanceStaticMeshComponent

diff --git a/Engine/Source/Runtime/Engine/Classes/Components/InstancedStaticMeshComponent.h b/Engine/Source/Runtime/Engine/Classes/Components/InstancedStaticMeshComponent.h
--- a/Engine/Source/Runtime/Engine/Classes/Components/InstancedStaticMeshComponent.h
+++ b/Engine/Source/Runtime/Engine/Classes/Components/InstancedStaticMeshComponent.h
@@ -162,6 +162,10 @@ class ENGINE_API UInstancedStaticMeshComponent : public UStaticMeshComponent, pu
 	UFUNCTION(BlueprintCallable, Category = "Components|InstancedStaticMesh")
 	virtual bool SetCustomDataValue(int32 InstanceIndex, int32 CustomDataIndex, float CustomDataValue, bool bMarkRenderStateDirty = false);
 
+	/** Update number of custom data entries per instance. This applies to all instances and will reallocate the full custom data buffer and reset all values to 0 */
+	UFUNCTION(BlueprintCallable, Category = "Components|InstancedStaticMesh")
+	virtual void SetNumCustomDataFloats(int32 InNumCustomDataFloats);
+
 	/** Per Instance Custom Data */
 	virtual bool SetCustomData(int32 InstanceIndex, TArrayView<const float> CustomDataFloats, bool bMarkRenderStateDirty = false); 
 
diff --git a/Engine/Source/Runtime/Engine/Private/InstancedStaticMesh.cpp b/Engine/Source/Runtime/Engine/Private/InstancedStaticMesh.cpp
--- a/Engine/Source/Runtime/Engine/Private/InstancedStaticMesh.cpp
+++ b/Engine/Source/Runtime/Engine/Private/InstancedStaticMesh.cpp
@@ -3418,6 +3418,18 @@ bool UInstancedStaticMeshComponent::SetCustomData(int32 InstanceIndex, TArrayVie
 	return true;
 }
 
+void UInstancedStaticMeshComponent::SetNumCustomDataFloats(int32 InNumCustomDataFloats)
+{
+	NumCustomDataFloats = FMath::Max(InNumCustomDataFloats, 0);
+
+	// Clear out and reinit to 0
+	PerInstanceSMCustomData.Empty(PerInstanceSMData.Num() * NumCustomDataFloats);
+	PerInstanceSMCustomData.SetNumZeroed(PerInstanceSMData.Num() * NumCustomDataFloats);
+
+	InstanceUpdateCmdBuffer.Edit();
+	MarkRenderStateDirty();
+}
+
 bool UInstancedStaticMeshComponent::SupportsRemoveSwap() const
 {
 	return bSupportRemoveAtSwap || CVarISMForceRemoveAtSwap.GetValueOnGameThread() != 0;
@@ -5074,14 +5086,7 @@ void UInstancedStaticMeshComponent::PostEditChangeChainProperty(FPropertyChanged
 		}
 		else if (PropertyChangedEvent.Property->GetFName() == "NumCustomDataFloats")
 		{
-			NumCustomDataFloats = FMath::Max(NumCustomDataFloats, 0);
-
-			// Clear out and reinit to 0
-			PerInstanceSMCustomData.Empty(PerInstanceSMData.Num() * NumCustomDataFloats);
-			PerInstanceSMCustomData.SetNumZeroed(PerInstanceSMData.Num() * NumCustomDataFloats);
-
-			InstanceUpdateCmdBuffer.Edit();
-			MarkRenderStateDirty();
+			SetNumCustomDataFloats(NumCustomDataFloats);
 		}
 		else if (PropertyChangedEvent.PropertyChain.GetActiveMemberNode()->GetValue()->GetFName() == "PerInstanceSMCustomData")
 		{
