UPrimitiveComponent::OverlapComponentWithResult should return blocking hits

diff --git a/Engine/Source/Runtime/Engine/Private/Components/PrimitiveComponent.cpp b/Engine/Source/Runtime/Engine/Private/Components/PrimitiveComponent.cpp
--- a/Engine/Source/Runtime/Engine/Private/Components/PrimitiveComponent.cpp
+++ b/Engine/Source/Runtime/Engine/Private/Components/PrimitiveComponent.cpp
@@ -3168,7 +3168,7 @@ bool UPrimitiveComponent::OverlapComponentWithResult(const FVector& Pos, const F
 		TArray<FOverlapResult> Overlaps;
 
 		FCollisionFilterData QueryFilter;
-		QueryFilter.Word2 = 0xFFFFF;
+		QueryFilter.Word1 = 0xFFFFF;
 		ConvertOverlapResults(OverlapHits.Num(), OverlapHits.GetData(), QueryFilter, Overlaps);
 
 		if (!Overlaps.IsEmpty())
