PSO precaching: Added common implementation of CollectPSOPrecacheData for static mesh component and its derived types, that can correctly handle non-MVF case
Added precaching for overlay materials
Added vertex strides for some common VFs as VK PSO requires correct strides

diff --git a/Engine/Source/Runtime/Engine/Classes/Components/StaticMeshComponent.h b/Engine/Source/Runtime/Engine/Classes/Components/StaticMeshComponent.h
--- a/Engine/Source/Runtime/Engine/Classes/Components/StaticMeshComponent.h
+++ b/Engine/Source/Runtime/Engine/Classes/Components/StaticMeshComponent.h
@@ -35,6 +35,7 @@ struct FNavigableGeometryExport;
 struct FNavigationRelevantData;
 struct FStaticMeshComponentLODInfo;
 struct FStaticLightingPrimitiveInfo;
+struct FStaticMeshLODResources;
 
 /** Whether FStaticMeshSceneProxy should to store data and enable codepaths needed for debug rendering */
 #define STATICMESH_ENABLE_DEBUG_RENDERING			ENABLE_DRAW_DEBUG
@@ -690,10 +691,12 @@ private:
 	UMaterialInterface* GetMaterial(int32 MaterialIndex, bool bDoingNaniteMaterialAudit) const;
 
 protected:
-
 	/** Collect all the PSO precache data used by the static mesh component */
 	virtual void CollectPSOPrecacheData(const FPSOPrecacheParams& BasePrecachePSOParams, FComponentPSOPrecacheParamsList& OutParams) override;
-
+	/** Shared implementation for all StaticMesh derived components */
+	using GetPSOVertexElementsFn = TFunctionRef<void(const FStaticMeshLODResources& LODRenderData, bool bSupportsManualVertexFetch, FVertexDeclarationElementList& Elements)>;
+	void CollectPSOPrecacheDataImpl(const FVertexFactoryType* VFType, const FPSOPrecacheParams& BasePrecachePSOParams, GetPSOVertexElementsFn GetVertexElements, FComponentPSOPrecacheParamsList& OutParams) const;
+		
 	/** Whether the component type supports static lighting. */
 	virtual bool SupportsStaticLighting() const override
 	{
diff --git a/Engine/Source/Runtime/Engine/Classes/Engine/InstancedStaticMesh.h b/Engine/Source/Runtime/Engine/Classes/Engine/InstancedStaticMesh.h
--- a/Engine/Source/Runtime/Engine/Classes/Engine/InstancedStaticMesh.h
+++ b/Engine/Source/Runtime/Engine/Classes/Engine/InstancedStaticMesh.h
@@ -286,6 +286,7 @@ public:
 	 * Get vertex elements used when during PSO precaching materials using this vertex factory type
 	 */
 	static void GetPSOPrecacheVertexFetchElements(EVertexInputStreamType VertexInputStreamType, FVertexDeclarationElementList& Elements);
+	static void GetVertexElements(ERHIFeatureLevel::Type FeatureLevel, EVertexInputStreamType InputStreamType, bool bSupportsManualVertexFetch, FDataType& Data, FVertexDeclarationElementList& Elements);
 
 	/**
 	 * An implementation of the interface used by TSynchronizedResource to update the resource with new data from the game thread.
@@ -336,6 +337,8 @@ public:
 	{
 		return UniformBuffer.GetReference();
 	}
+protected:
+	static void GetVertexElements(ERHIFeatureLevel::Type FeatureLevel, EVertexInputStreamType InputStreamType, bool bSupportsManualVertexFetch, FDataType& Data, FVertexDeclarationElementList& Elements, FVertexStreamList& Streams);
 
 private:
 	FDataType Data;
diff --git a/Engine/Source/Runtime/Engine/Private/Components/MeshComponent.cpp b/Engine/Source/Runtime/Engine/Private/Components/MeshComponent.cpp
--- a/Engine/Source/Runtime/Engine/Private/Components/MeshComponent.cpp
+++ b/Engine/Source/Runtime/Engine/Private/Components/MeshComponent.cpp
@@ -229,6 +229,8 @@ void UMeshComponent::SetOverlayMaterial(UMaterialInterface* NewOverlayMaterial)
 	if (OverlayMaterial != NewOverlayMaterial)
 	{
 		OverlayMaterial = NewOverlayMaterial;
+		// Precache PSOs again
+		PrecachePSOs();
 		MarkRenderStateDirty();
 	}
 }
diff --git a/Engine/Source/Runtime/Engine/Private/Components/SkinnedMeshComponent.cpp b/Engine/Source/Runtime/Engine/Private/Components/SkinnedMeshComponent.cpp
--- a/Engine/Source/Runtime/Engine/Private/Components/SkinnedMeshComponent.cpp
+++ b/Engine/Source/Runtime/Engine/Private/Components/SkinnedMeshComponent.cpp
@@ -596,6 +596,17 @@ void USkinnedMeshComponent::CollectPSOPrecacheData(const FPSOPrecacheParams& Bas
 		ComponentParams.VertexFactoryDataList = VFsPerMaterial.VertexFactoryDataList;
 		ComponentParams.PSOPrecacheParams = PrecachePSOParams;
 	}
+
+	if (OverlayMaterial && VFsPerMaterials.Num() != 0)
+	{
+		// Overlay is rendered with the same set of VFs
+		FComponentPSOPrecacheParams& ComponentParams = OutParams[OutParams.AddDefaulted()];
+
+		ComponentParams.MaterialInterface = OverlayMaterial;
+		ComponentParams.VertexFactoryDataList = VFsPerMaterials[0].VertexFactoryDataList;
+		ComponentParams.PSOPrecacheParams = PrecachePSOParams;
+		ComponentParams.PSOPrecacheParams.bCastShadow = false;
+	}
 }
 
 void USkinnedMeshComponent::OnRegister()
diff --git a/Engine/Source/Runtime/Engine/Private/Components/SplineMeshComponent.cpp b/Engine/Source/Runtime/Engine/Private/Components/SplineMeshComponent.cpp
--- a/Engine/Source/Runtime/Engine/Private/Components/SplineMeshComponent.cpp
+++ b/Engine/Source/Runtime/Engine/Private/Components/SplineMeshComponent.cpp
@@ -20,7 +20,6 @@
 #include "RenderUtils.h"
 #include "SceneInterface.h"
 #include "UObject/UnrealType.h"
-#include "MaterialDomain.h"
 #include "Materials/Material.h"
 
 #if WITH_EDITOR
@@ -131,8 +130,8 @@ IMPLEMENT_VERTEX_FACTORY_TYPE(FSplineMeshVertexFactory, "/Engine/Private/LocalVe
 );
 
 void InitSplineMeshVertexFactoryComponents(
-	FStaticMeshVertexBuffers& VertexBuffers, 
-	FSplineMeshVertexFactory* VertexFactory, 
+	const FStaticMeshVertexBuffers& VertexBuffers, 
+	const FSplineMeshVertexFactory* VertexFactory, 
 	int32 LightMapCoordinateIndex, 
 	bool bOverrideColorVertexBuffer,
 	FLocalVertexFactory::FDataType& OutData)
@@ -641,65 +640,18 @@ void USplineMeshComponent::CollectPSOPrecacheData(const FPSOPrecacheParams& Base
 	}
 
 	FVertexFactoryType* VertexFactoryType = &FSplineMeshVertexFactory::StaticType;
-	bool bSupportsManualVertexFetch = VertexFactoryType->SupportsManualVertexFetch(GMaxRHIFeatureLevel);
-
-	// Override color buffer is using same vertex element data as default one
-	bool bOverrideColorVertexBuffer = false;
 	int32 LightMapCoordinateIndex = GetStaticMesh()->GetLightMapCoordinateIndex();
 
-	bool bAnySectionCastsShadows = false;
-	FPSOPrecacheVertexFactoryDataPerMaterialIndexList VFTypesPerMaterialIndex;
-	for (FStaticMeshLODResources& LODRenderData : GetStaticMesh()->GetRenderData()->LODResources)
-	{
-		FVertexDeclarationElementList VertexElements;
-		if (!bSupportsManualVertexFetch)
-		{
-			FLocalVertexFactory::FDataType Data;
-			InitSplineMeshVertexFactoryComponents(LODRenderData.VertexBuffers, nullptr /*VertexFactory*/, LightMapCoordinateIndex, bOverrideColorVertexBuffer, Data);
-			FSplineMeshVertexFactory::GetVertexElements(GMaxRHIFeatureLevel, EVertexInputStreamType::Default, bSupportsManualVertexFetch, Data, VertexElements);						
-		}
-
-		for (FStaticMeshSection& RenderSection : LODRenderData.Sections)
-		{
-			bAnySectionCastsShadows |= RenderSection.bCastShadow;
-
-			int16 MaterialIndex = RenderSection.MaterialIndex;
-			FPSOPrecacheVertexFactoryDataPerMaterialIndex* VFsPerMaterial = VFTypesPerMaterialIndex.FindByPredicate(
-				[MaterialIndex](const FPSOPrecacheVertexFactoryDataPerMaterialIndex& Other) { return Other.MaterialIndex == MaterialIndex; });
-			if (VFsPerMaterial == nullptr)
-			{
-				VFsPerMaterial = &VFTypesPerMaterialIndex.AddDefaulted_GetRef();
-				VFsPerMaterial->MaterialIndex = RenderSection.MaterialIndex;
-			}
-
-			if (bSupportsManualVertexFetch)
-			{
-				VFsPerMaterial->VertexFactoryDataList.AddUnique(FPSOPrecacheVertexFactoryData(VertexFactoryType));
-			}
-			else
-			{	
-				VFsPerMaterial->VertexFactoryDataList.AddUnique(FPSOPrecacheVertexFactoryData(VertexFactoryType, VertexElements));
-			}			
-		}
-	}
-
-	FPSOPrecacheParams PrecachePSOParams = BasePrecachePSOParams;
-	PrecachePSOParams.bCastShadow = bAnySectionCastsShadows;
-	PrecachePSOParams.bReverseCulling = bReverseCulling;
-
-	for (FPSOPrecacheVertexFactoryDataPerMaterialIndex& VFsPerMaterial : VFTypesPerMaterialIndex)
+	auto SMC_GetElements = [LightMapCoordinateIndex](const FStaticMeshLODResources& LODRenderData, bool bSupportsManualVertexFetch, FVertexDeclarationElementList& Elements)
 	{
-		UMaterialInterface* MaterialInterface = GetMaterial(VFsPerMaterial.MaterialIndex);
-		if (MaterialInterface == nullptr)
-		{
-			MaterialInterface = UMaterial::GetDefaultMaterial(MD_Surface);
-		}
+		// FIXME: This will miss when SM component overrides vertex colors and source StaticMesh does not have vertex colors
+		bool bOverrideColorVertexBuffer = false;
+		FLocalVertexFactory::FDataType Data;
+		InitSplineMeshVertexFactoryComponents(LODRenderData.VertexBuffers, nullptr /*VertexFactory*/, LightMapCoordinateIndex, bOverrideColorVertexBuffer, Data);
+		FLocalVertexFactory::GetVertexElements(GMaxRHIFeatureLevel, EVertexInputStreamType::Default, bSupportsManualVertexFetch, Data, Elements);
+	};
 
-		FComponentPSOPrecacheParams& ComponentParams = OutParams[OutParams.AddDefaulted()];
-		ComponentParams.MaterialInterface = MaterialInterface;
-		ComponentParams.VertexFactoryDataList = VFsPerMaterial.VertexFactoryDataList;
-		ComponentParams.PSOPrecacheParams = PrecachePSOParams;
-	}
+	CollectPSOPrecacheDataImpl(VertexFactoryType, BasePrecachePSOParams, SMC_GetElements, OutParams);
 }
 
 FPrimitiveSceneProxy* USplineMeshComponent::CreateSceneProxy()
diff --git a/Engine/Source/Runtime/Engine/Private/Components/StaticMeshComponent.cpp b/Engine/Source/Runtime/Engine/Private/Components/StaticMeshComponent.cpp
--- a/Engine/Source/Runtime/Engine/Private/Components/StaticMeshComponent.cpp
+++ b/Engine/Source/Runtime/Engine/Private/Components/StaticMeshComponent.cpp
@@ -43,6 +43,7 @@
 #include "PhysicsEngine/BodySetup.h"
 #include "ComponentRecreateRenderStateContext.h"
 #include "Engine/StaticMesh.h"
+#include "MaterialDomain.h"
 
 #include UE_INLINE_GENERATED_CPP_BY_NAME(StaticMeshComponent)
 
@@ -1520,24 +1521,71 @@ void UStaticMeshComponent::InitResources()
 	}
 }
 
-void UStaticMeshComponent::CollectPSOPrecacheData(const FPSOPrecacheParams& BasePrecachePSOParams, FComponentPSOPrecacheParamsList& OutParams)
+void InitStaticMeshVertexFactoryComponents(
+	const FStaticMeshVertexBuffers& VertexBuffers,
+	FLocalVertexFactory* VertexFactory,
+	int32 LightMapCoordinateIndex,
+	bool bOverrideColorVertexBuffer,
+	FLocalVertexFactory::FDataType& OutData)
 {
-	if (StaticMesh == nullptr || StaticMesh->GetRenderData() == nullptr)
+	VertexBuffers.PositionVertexBuffer.BindPositionVertexBuffer(VertexFactory, OutData);
+	VertexBuffers.StaticMeshVertexBuffer.BindTangentVertexBuffer(VertexFactory, OutData);
+	VertexBuffers.StaticMeshVertexBuffer.BindPackedTexCoordVertexBuffer(VertexFactory, OutData);
+	VertexBuffers.StaticMeshVertexBuffer.BindLightMapVertexBuffer(VertexFactory, OutData, LightMapCoordinateIndex);
+	if (bOverrideColorVertexBuffer)
 	{
-		return;
+		FColorVertexBuffer::BindDefaultColorVertexBuffer(VertexFactory, OutData, FColorVertexBuffer::NullBindStride::FColorSizeForComponentOverride);
+	}
+	else
+	{
+		VertexBuffers.ColorVertexBuffer.BindColorVertexBuffer(VertexFactory, OutData);
 	}
+}
+
+void UStaticMeshComponent::CollectPSOPrecacheDataImpl(
+	const FVertexFactoryType* VFType, 
+	const FPSOPrecacheParams& BasePrecachePSOParams, 
+	GetPSOVertexElementsFn GetVertexElements,
+	FComponentPSOPrecacheParamsList& OutParams) const
+{
+	check(StaticMesh != nullptr && StaticMesh->GetRenderData() != nullptr);
 
 	UWorld* World = GetWorld();
 	ERHIFeatureLevel::Type FeatureLevel = World ? World->GetFeatureLevel() : GMaxRHIFeatureLevel;
 
+	bool bSupportsManualVertexFetch = VFType->SupportsManualVertexFetch(GMaxRHIFeatureLevel);
 	bool bAnySectionCastsShadows = false;
-	TArray<int16, TInlineAllocator<2>> UsedMaterialIndices;
-	for (FStaticMeshLODResources& LODRenderData : StaticMesh->GetRenderData()->LODResources)
+
+	FPSOPrecacheVertexFactoryDataPerMaterialIndexList VFTypesPerMaterialIndex;
+	for (FStaticMeshLODResources& LODRenderData : GetStaticMesh()->GetRenderData()->LODResources)
 	{
+		FVertexDeclarationElementList VertexElements;
+		if (!bSupportsManualVertexFetch)
+		{
+			GetVertexElements(LODRenderData, bSupportsManualVertexFetch, VertexElements);
+		}
+
 		for (FStaticMeshSection& RenderSection : LODRenderData.Sections)
 		{
-			UsedMaterialIndices.AddUnique(RenderSection.MaterialIndex);
 			bAnySectionCastsShadows |= RenderSection.bCastShadow;
+
+			int16 MaterialIndex = RenderSection.MaterialIndex;
+			FPSOPrecacheVertexFactoryDataPerMaterialIndex* VFsPerMaterial = VFTypesPerMaterialIndex.FindByPredicate(
+				[MaterialIndex](const FPSOPrecacheVertexFactoryDataPerMaterialIndex& Other) { return Other.MaterialIndex == MaterialIndex; });
+			if (VFsPerMaterial == nullptr)
+			{
+				VFsPerMaterial = &VFTypesPerMaterialIndex.AddDefaulted_GetRef();
+				VFsPerMaterial->MaterialIndex = RenderSection.MaterialIndex;
+			}
+
+			if (bSupportsManualVertexFetch)
+			{
+				VFsPerMaterial->VertexFactoryDataList.AddUnique(FPSOPrecacheVertexFactoryData(VFType));
+			}
+			else
+			{	
+				VFsPerMaterial->VertexFactoryDataList.AddUnique(FPSOPrecacheVertexFactoryData(VFType, VertexElements));
+			}			
 		}
 	}
 
@@ -1548,22 +1596,54 @@ void UStaticMeshComponent::CollectPSOPrecacheData(const FPSOPrecacheParams& Base
 	PrecachePSOParams.bReverseCulling = bReverseCulling != bIsLocalToWorldDeterminantNegative;
 	PrecachePSOParams.bForceLODModel = ForcedLodModel > 0;
 
-	const FVertexFactoryType* VFType = ShouldCreateNaniteProxy() ? &Nanite::FVertexFactory::StaticType :  &FLocalVertexFactory::StaticType;
-
-	for (uint16 MaterialIndex : UsedMaterialIndices)
+	for (FPSOPrecacheVertexFactoryDataPerMaterialIndex& VFsPerMaterial : VFTypesPerMaterialIndex)
 	{
-		UMaterialInterface* MaterialInterface = GetMaterial(MaterialIndex);
-		if (MaterialInterface)
+		UMaterialInterface* MaterialInterface = GetMaterial(VFsPerMaterial.MaterialIndex);
+		if (MaterialInterface == nullptr)
 		{
-			PrecachePSOParams.bHasWorldPositionOffsetVelocity = SupportsWorldPositionOffsetVelocity()
-				&& MaterialInterface->GetRelevance_Concurrent(FeatureLevel).bUsesWorldPositionOffset;
-
-			FComponentPSOPrecacheParams& ComponentParams = OutParams[OutParams.AddDefaulted()];
-			ComponentParams.MaterialInterface = MaterialInterface;
-			ComponentParams.VertexFactoryDataList.Add(FPSOPrecacheVertexFactoryData(VFType));
-			ComponentParams.PSOPrecacheParams = PrecachePSOParams;
+			MaterialInterface = UMaterial::GetDefaultMaterial(MD_Surface);
 		}
+
+		FComponentPSOPrecacheParams& ComponentParams = OutParams[OutParams.AddDefaulted()];
+		ComponentParams.MaterialInterface = MaterialInterface;
+		ComponentParams.VertexFactoryDataList = VFsPerMaterial.VertexFactoryDataList;
+		ComponentParams.PSOPrecacheParams = PrecachePSOParams;
+	}
+
+	if (OverlayMaterial && VFTypesPerMaterialIndex.Num() != 0)
+	{
+		// Overlay is rendered with the same set of VFs
+		FComponentPSOPrecacheParams& ComponentParams = OutParams[OutParams.AddDefaulted()];
+		
+		ComponentParams.MaterialInterface = OverlayMaterial;
+		ComponentParams.VertexFactoryDataList = VFTypesPerMaterialIndex[0].VertexFactoryDataList;
+		ComponentParams.PSOPrecacheParams = PrecachePSOParams;
+		ComponentParams.PSOPrecacheParams.bCastShadow = false;
+	}
+}
+
+void UStaticMeshComponent::CollectPSOPrecacheData(const FPSOPrecacheParams& BasePrecachePSOParams, FComponentPSOPrecacheParamsList& OutParams)
+{
+	if (StaticMesh == nullptr || StaticMesh->GetRenderData() == nullptr)
+	{
+		return;
 	}
+
+	const FVertexFactoryType* VFType = ShouldCreateNaniteProxy() ? &Nanite::FVertexFactory::StaticType : &FLocalVertexFactory::StaticType;
+	int32 LightMapCoordinateIndex = StaticMesh->GetLightMapCoordinateIndex();
+	// FIXME: Need a precise per-LOD test
+	bool bOverrideColorVertexBuffer = LODData.Num() != 0 && LODData[0].OverrideVertexColors != nullptr;
+	
+	auto SMC_GetElements = [LightMapCoordinateIndex, bOverrideColorVertexBuffer](const FStaticMeshLODResources& LODRenderData, bool bSupportsManualVertexFetch, FVertexDeclarationElementList& Elements)
+	{
+		int32 NumTexCoords = (int32)LODRenderData.VertexBuffers.StaticMeshVertexBuffer.GetNumTexCoords();
+		int32 LODLightMapCoordinateIndex = LightMapCoordinateIndex < NumTexCoords ? LightMapCoordinateIndex : NumTexCoords - 1;
+		FLocalVertexFactory::FDataType Data;
+		InitStaticMeshVertexFactoryComponents(LODRenderData.VertexBuffers, nullptr /*VertexFactory*/, LODLightMapCoordinateIndex, bOverrideColorVertexBuffer, Data);
+		FLocalVertexFactory::GetVertexElements(GMaxRHIFeatureLevel, EVertexInputStreamType::Default, bSupportsManualVertexFetch, Data, Elements);
+	};
+	
+	CollectPSOPrecacheDataImpl(VFType, BasePrecachePSOParams, SMC_GetElements, OutParams);
 }
 
 #if WITH_EDITOR
diff --git a/Engine/Source/Runtime/Engine/Private/GPUSkinVertexFactory.cpp b/Engine/Source/Runtime/Engine/Private/GPUSkinVertexFactory.cpp
--- a/Engine/Source/Runtime/Engine/Private/GPUSkinVertexFactory.cpp
+++ b/Engine/Source/Runtime/Engine/Private/GPUSkinVertexFactory.cpp
@@ -745,9 +745,12 @@ void TGPUSkinVertexFactory<BoneInfluenceType>::GetVertexElements(ERHIFeatureLeve
 	int32 MorphDeltaStreamIndex;
 	GetVertexElements(FeatureLevel, InputStreamType, GPUSkinData, OutElements, VertexStreams, MorphDeltaStreamIndex);
 
-	// For ES3.1 attribute ID needs to be done differently
-	check(FeatureLevel > ERHIFeatureLevel::ES3_1);
-	OutElements.Add(FVertexElement(VertexStreams.Num(), 0, VET_UInt, 16, 0, true));
+	if (UseGPUScene(GMaxRHIShaderPlatform, GMaxRHIFeatureLevel))
+	{
+		// For ES3.1 attribute ID needs to be done differently
+		check(FeatureLevel > ERHIFeatureLevel::ES3_1);
+		OutElements.Add(FVertexElement(VertexStreams.Num(), 0, VET_UInt, 16, 0, true));
+	}
 }
 
 /**
diff --git a/Engine/Source/Runtime/Engine/Private/InstancedStaticMesh.cpp b/Engine/Source/Runtime/Engine/Private/InstancedStaticMesh.cpp
--- a/Engine/Source/Runtime/Engine/Private/InstancedStaticMesh.cpp
+++ b/Engine/Source/Runtime/Engine/Private/InstancedStaticMesh.cpp
@@ -277,6 +277,19 @@ public:
 
 TGlobalResource<FDummyFloatBuffer> GDummyFloatBuffer;
 
+// Dummy instance buffer used for PSO pre-caching
+class FDummyStaticMeshInstanceBuffer : public FStaticMeshInstanceBuffer
+{
+public:
+	FDummyStaticMeshInstanceBuffer()
+	: FStaticMeshInstanceBuffer(GMaxRHIFeatureLevel, false /*InRequireCPUAccess*/, true /*bDeferGPUUploadIn*/)
+	{
+		InstanceData = MakeShared<FStaticMeshInstanceData, ESPMode::ThreadSafe>(GVertexElementTypeSupport.IsSupported(VET_Half2));
+	}
+};
+
+TGlobalResource<FDummyStaticMeshInstanceBuffer> GDummyStaticMeshInstanceBuffer;
+
 FInstancedStaticMeshDelegates::FOnInstanceIndexUpdated FInstancedStaticMeshDelegates::OnInstanceIndexUpdated;
 
 /** InstancedStaticMeshInstance hit proxy */
@@ -937,22 +950,67 @@ void FInstancedStaticMeshVertexFactory::InitRHI()
 
 	check(HasValidFeatureLevel());
 
+	const ERHIFeatureLevel::Type ThisFeatureLevel = GetFeatureLevel();
+	const bool bCanUseGPUScene = UseGPUScene(GMaxRHIShaderPlatform, ThisFeatureLevel);
+	const bool bUseManualVertexFetch = GetType()->SupportsManualVertexFetch(ThisFeatureLevel);
+
 	FVertexDeclarationElementList Elements;
-	if(Data.PositionComponent.VertexBuffer != NULL)
-	{
-		Elements.Add(AccessStreamComponent(Data.PositionComponent,0));
-	}
+	GetVertexElements(ThisFeatureLevel, EVertexInputStreamType::Default, bUseManualVertexFetch, Data, Elements, Streams);
 
 	// on mobile with GPUScene enabled instanced attributes[8-12] are used for a general auto-instancing
 	// so we add them only for desktop or if mobile has GPUScene disabled
 	// FIXME mobile: instanced attributes encode some editor related data as well (selection etc), need to split it into separate SRV as it's not supported with auto-instancing
+	// FIXME: Need to capture PrimitiveId elements for PSO precaching
 	uint8 AutoInstancingAttr_Mobile = 8;
 	const bool bMobileUsesGPUScene = MobileSupportsGPUScene();
+	if (ThisFeatureLevel > ERHIFeatureLevel::ES3_1 || !bMobileUsesGPUScene)
+	{
+		// Do not add general auto-instancing attributes for mobile
+		AutoInstancingAttr_Mobile = 0xff;
+	}
 
-	const bool bCanUseGPUScene = UseGPUScene(GMaxRHIShaderPlatform, GetFeatureLevel());
-	const bool bUseManualVertexFetch = GetType()->SupportsManualVertexFetch(GetFeatureLevel());
+	AddPrimitiveIdStreamElement(EVertexInputStreamType::Default, Elements, 13, AutoInstancingAttr_Mobile);
+
+	// we don't need per-vertex shadow or lightmap rendering
+	InitDeclaration(Elements);
 
-	if (!bUseManualVertexFetch)
+	if (!bCanUseGPUScene)
+	{
+		FInstancedStaticMeshVertexFactoryUniformShaderParameters UniformParameters;
+		UniformParameters.VertexFetch_InstanceOriginBuffer = GetInstanceOriginSRV();
+		UniformParameters.VertexFetch_InstanceTransformBuffer = GetInstanceTransformSRV();
+		UniformParameters.VertexFetch_InstanceLightmapBuffer = GetInstanceLightmapSRV();
+		UniformParameters.InstanceCustomDataBuffer = GetInstanceCustomDataSRV();
+		UniformParameters.NumCustomDataFloats = Data.NumCustomDataFloats;
+		UniformBuffer = TUniformBufferRef<FInstancedStaticMeshVertexFactoryUniformShaderParameters>::CreateUniformBufferImmediate(UniformParameters, UniformBuffer_MultiFrame, EUniformBufferValidation::None);
+	}
+}
+
+void FInstancedStaticMeshVertexFactory::GetVertexElements(
+	ERHIFeatureLevel::Type FeatureLevel,
+	EVertexInputStreamType InputStreamType,
+	bool bSupportsManualVertexFetch,
+	FDataType& Data,
+	FVertexDeclarationElementList& Elements)
+{
+	FVertexStreamList VertexStreams;
+	GetVertexElements(FeatureLevel, InputStreamType, bSupportsManualVertexFetch, Data, Elements, VertexStreams);
+}
+
+void FInstancedStaticMeshVertexFactory::GetVertexElements(
+	ERHIFeatureLevel::Type FeatureLevel, 
+	EVertexInputStreamType InputStreamType, 
+	bool bSupportsManualVertexFetch, 
+	FDataType& Data, 
+	FVertexDeclarationElementList& Elements, 
+	FVertexStreamList& Streams)
+{
+	if (Data.PositionComponent.VertexBuffer != NULL)
+	{
+		Elements.Add(AccessStreamComponent(Data.PositionComponent, 0, Streams));
+	}
+
+	if (!bSupportsManualVertexFetch)
 	{
 		// only tangent,normal are used by the stream. the binormal is derived in the shader
 		uint8 TangentBasisAttributes[2] = { 1, 2 };
@@ -960,7 +1018,7 @@ void FInstancedStaticMeshVertexFactory::InitRHI()
 		{
 			if (Data.TangentBasisComponents[AxisIndex].VertexBuffer != NULL)
 			{
-				Elements.Add(AccessStreamComponent(Data.TangentBasisComponents[AxisIndex], TangentBasisAttributes[AxisIndex]));
+				Elements.Add(AccessStreamComponent(Data.TangentBasisComponents[AxisIndex], TangentBasisAttributes[AxisIndex], Streams));
 			}
 		}
 
@@ -972,14 +1030,14 @@ void FInstancedStaticMeshVertexFactory::InitRHI()
 
 		if (Data.ColorComponent.VertexBuffer)
 		{
-			Elements.Add(AccessStreamComponent(Data.ColorComponent, 3));
+			Elements.Add(AccessStreamComponent(Data.ColorComponent, 3, Streams));
 		}
 		else
 		{
 			//If the mesh has no color component, set the null color buffer on a new stream with a stride of 0.
 			//This wastes 4 bytes of bandwidth per vertex, but prevents having to compile out twice the number of vertex factories.
 			FVertexStreamComponent NullColorComponent(&GNullColorVertexBuffer, 0, 0, VET_Color, EVertexStreamUsage::ManualFetch);
-			Elements.Add(AccessStreamComponent(NullColorComponent, 3));
+			Elements.Add(AccessStreamComponent(NullColorComponent, 3, Streams));
 		}
 
 		if (Data.TextureCoordinates.Num())
@@ -989,7 +1047,8 @@ void FInstancedStaticMeshVertexFactory::InitRHI()
 			{
 				Elements.Add(AccessStreamComponent(
 					Data.TextureCoordinates[CoordinateIndex],
-					BaseTexCoordAttribute + CoordinateIndex
+					BaseTexCoordAttribute + CoordinateIndex,
+					Streams
 				));
 			}
 
@@ -997,7 +1056,8 @@ void FInstancedStaticMeshVertexFactory::InitRHI()
 			{
 				Elements.Add(AccessStreamComponent(
 					Data.TextureCoordinates[Data.TextureCoordinates.Num() - 1],
-					BaseTexCoordAttribute + CoordinateIndex
+					BaseTexCoordAttribute + CoordinateIndex,
+					Streams
 				));
 			}
 		}
@@ -1007,59 +1067,43 @@ void FInstancedStaticMeshVertexFactory::InitRHI()
 		if (FLocalVertexFactory::IsGPUSkinPassThroughSupported(GMaxRHIShaderPlatform))
 		{
 			FVertexStreamComponent NullComponent(&GNullVertexBuffer, 0, 0, VET_Float4);
-			Elements.Add(AccessStreamComponent(NullComponent, 14));
+			Elements.Add(AccessStreamComponent(NullComponent, 14, Streams));
 		}
 
 		if (Data.LightMapCoordinateComponent.VertexBuffer)
 		{
-			Elements.Add(AccessStreamComponent(Data.LightMapCoordinateComponent, 15));
+			Elements.Add(AccessStreamComponent(Data.LightMapCoordinateComponent, 15, Streams));
 		}
 		else if (Data.TextureCoordinates.Num())
 		{
-			Elements.Add(AccessStreamComponent(Data.TextureCoordinates[0], 15));
+			Elements.Add(AccessStreamComponent(Data.TextureCoordinates[0], 15, Streams));
 		}
 	}
 
-	if (GetFeatureLevel() > ERHIFeatureLevel::ES3_1 || !bMobileUsesGPUScene)
+	const bool bCanUseGPUScene = UseGPUScene(GMaxRHIShaderPlatform, FeatureLevel);
+	const bool bMobileUsesGPUScene = MobileSupportsGPUScene();
+	
+	if (FeatureLevel > ERHIFeatureLevel::ES3_1 || !bMobileUsesGPUScene)
 	{
 		// toss in the instanced location stream
 		check(bCanUseGPUScene || Data.InstanceOriginComponent.VertexBuffer);
 		if (Data.InstanceOriginComponent.VertexBuffer)
 		{
-			Elements.Add(AccessStreamComponent(Data.InstanceOriginComponent, 8));
+			Elements.Add(AccessStreamComponent(Data.InstanceOriginComponent, 8, Streams));
 		}
 
 		check(bCanUseGPUScene || Data.InstanceTransformComponent[0].VertexBuffer);
 		if (Data.InstanceTransformComponent[0].VertexBuffer)
 		{
-			Elements.Add(AccessStreamComponent(Data.InstanceTransformComponent[0], 9));
-			Elements.Add(AccessStreamComponent(Data.InstanceTransformComponent[1], 10));
-			Elements.Add(AccessStreamComponent(Data.InstanceTransformComponent[2], 11));
+			Elements.Add(AccessStreamComponent(Data.InstanceTransformComponent[0], 9, Streams));
+			Elements.Add(AccessStreamComponent(Data.InstanceTransformComponent[1], 10, Streams));
+			Elements.Add(AccessStreamComponent(Data.InstanceTransformComponent[2], 11, Streams));
 		}
 
 		if (Data.InstanceLightmapAndShadowMapUVBiasComponent.VertexBuffer)
 		{
-			Elements.Add(AccessStreamComponent(Data.InstanceLightmapAndShadowMapUVBiasComponent, 12));
+			Elements.Add(AccessStreamComponent(Data.InstanceLightmapAndShadowMapUVBiasComponent, 12, Streams));
 		}
-
-		// Do not add general auto-instancing attributes for mobile
-		AutoInstancingAttr_Mobile = 0xff;
-	}
-
-	AddPrimitiveIdStreamElement(EVertexInputStreamType::Default, Elements, 13, AutoInstancingAttr_Mobile);
-
-	// we don't need per-vertex shadow or lightmap rendering
-	InitDeclaration(Elements);
-
-	if (!bCanUseGPUScene)
-	{
-		FInstancedStaticMeshVertexFactoryUniformShaderParameters UniformParameters;
-		UniformParameters.VertexFetch_InstanceOriginBuffer = GetInstanceOriginSRV();
-		UniformParameters.VertexFetch_InstanceTransformBuffer = GetInstanceTransformSRV();
-		UniformParameters.VertexFetch_InstanceLightmapBuffer = GetInstanceLightmapSRV();
-		UniformParameters.InstanceCustomDataBuffer = GetInstanceCustomDataSRV();
-		UniformParameters.NumCustomDataFloats = Data.NumCustomDataFloats;
-		UniformBuffer = TUniformBufferRef<FInstancedStaticMeshVertexFactoryUniformShaderParameters>::CreateUniformBufferImmediate(UniformParameters, UniformBuffer_MultiFrame, EUniformBufferValidation::None);
 	}
 }
 
@@ -1117,17 +1161,53 @@ void FInstancedStaticMeshRenderData::ReleaseResources(FSceneInterface* Scene, co
 	}
 }
 
+void InitInstancedStaticMeshVertexFactoryComponents(
+	const FStaticMeshVertexBuffers& VertexBuffers,
+	const FStaticMeshInstanceBuffer* InstanceBuffer,
+	const FInstancedStaticMeshVertexFactory* VertexFactory,
+	int32 LightMapCoordinateIndex,
+	bool bHasColorVertexData,
+	bool bRHISupportsManualVertexFetch,
+	FInstancedStaticMeshVertexFactory::FDataType& OutData)
+{
+	VertexBuffers.PositionVertexBuffer.BindPositionVertexBuffer(VertexFactory, OutData);
+	VertexBuffers.StaticMeshVertexBuffer.BindTangentVertexBuffer(VertexFactory, OutData);
+	VertexBuffers.StaticMeshVertexBuffer.BindPackedTexCoordVertexBuffer(VertexFactory, OutData);
+	
+	if (LightMapCoordinateIndex < (int32)VertexBuffers.StaticMeshVertexBuffer.GetNumTexCoords() && LightMapCoordinateIndex >= 0)
+	{
+		VertexBuffers.StaticMeshVertexBuffer.BindLightMapVertexBuffer(VertexFactory, OutData, LightMapCoordinateIndex);
+	}
+
+	if (bHasColorVertexData)
+	{
+		VertexBuffers.ColorVertexBuffer.BindColorVertexBuffer(VertexFactory, OutData);
+	}
+	else
+	{
+		// shouldn't this check if ISM component actually has a color data for override?
+		FColorVertexBuffer::BindDefaultColorVertexBuffer(VertexFactory, OutData, bRHISupportsManualVertexFetch ? FColorVertexBuffer::NullBindStride::FColorSizeForComponentOverride : FColorVertexBuffer::NullBindStride::ZeroForDefaultBufferBind);
+	}
+
+	if (InstanceBuffer)
+	{
+		InstanceBuffer->BindInstanceVertexBuffer(VertexFactory, OutData);
+	}
+}
+
 void FInstancedStaticMeshRenderData::BindBuffersToVertexFactories()
 {
 	TRACE_CPUPROFILER_EVENT_SCOPE_STR("FInstancedStaticMeshRenderData::BindBuffersToVertexFactories");
 
 	check(IsInRenderingThread());
-
+	
+	FStaticMeshInstanceBuffer* InstanceBuffer = nullptr;
 	const bool bRHISupportsManualVertexFetch = RHISupportsManualVertexFetch(GShaderPlatformForFeatureLevel[FeatureLevel]);
 	const bool bCanUseGPUScene = UseGPUScene(GMaxRHIShaderPlatform, FeatureLevel);
 	if (!bCanUseGPUScene)
 	{
 		PerInstanceRenderData->InstanceBuffer.FlushGPUUpload();
+		InstanceBuffer = &PerInstanceRenderData->InstanceBuffer;
 	}
 
 	for (int32 LODIndex = 0; LODIndex < VertexFactories.Num(); LODIndex++)
@@ -1137,30 +1217,7 @@ void FInstancedStaticMeshRenderData::BindBuffersToVertexFactories()
 		FInstancedStaticMeshVertexFactory::FDataType Data;
 		// Assign to the vertex factory for this LOD.
 		FInstancedStaticMeshVertexFactory& VertexFactory = VertexFactories[LODIndex];
-
-		RenderData->VertexBuffers.PositionVertexBuffer.BindPositionVertexBuffer(&VertexFactory, Data);
-		RenderData->VertexBuffers.StaticMeshVertexBuffer.BindTangentVertexBuffer(&VertexFactory, Data);
-		RenderData->VertexBuffers.StaticMeshVertexBuffer.BindPackedTexCoordVertexBuffer(&VertexFactory, Data);
-		if (LightMapCoordinateIndex < (int32)RenderData->VertexBuffers.StaticMeshVertexBuffer.GetNumTexCoords() && LightMapCoordinateIndex >= 0)
-		{
-			RenderData->VertexBuffers.StaticMeshVertexBuffer.BindLightMapVertexBuffer(&VertexFactory, Data, LightMapCoordinateIndex);
-		}
-
-		if (RenderData->bHasColorVertexData)
-		{
-			RenderData->VertexBuffers.ColorVertexBuffer.BindColorVertexBuffer(&VertexFactory, Data);
-		}
-		else
-		{
-			FColorVertexBuffer::BindDefaultColorVertexBuffer(&VertexFactory, Data, bRHISupportsManualVertexFetch ? FColorVertexBuffer::NullBindStride::FColorSizeForComponentOverride : FColorVertexBuffer::NullBindStride::ZeroForDefaultBufferBind);
-		}
-
-		check(PerInstanceRenderData);
-		if (!bCanUseGPUScene)
-		{
-			PerInstanceRenderData->InstanceBuffer.BindInstanceVertexBuffer(&VertexFactory, Data);
-		}
-
+		InitInstancedStaticMeshVertexFactoryComponents(RenderData->VertexBuffers, InstanceBuffer, &VertexFactory, LightMapCoordinateIndex, RenderData->bHasColorVertexData, bRHISupportsManualVertexFetch, Data);
 		VertexFactory.SetData(Data);
 		VertexFactory.InitResource();
 	}
@@ -4721,34 +4778,19 @@ void UInstancedStaticMeshComponent::CollectPSOPrecacheData(const FPSOPrecachePar
 		return;
 	}
 
-	bool bAnySectionCastsShadows = false;
-	TArray<int16, TInlineAllocator<2>> UsedMaterialIndices;
-	for (FStaticMeshLODResources& LODRenderData : GetStaticMesh()->GetRenderData()->LODResources)
-	{
-		for (FStaticMeshSection& RenderSection : LODRenderData.Sections)
-		{
-			UsedMaterialIndices.AddUnique(RenderSection.MaterialIndex);
-			bAnySectionCastsShadows |= RenderSection.bCastShadow;
-		}
-	}
-
-	FPSOPrecacheParams PrecachePSOParams = BasePrecachePSOParams;
-	PrecachePSOParams.bCastShadow = bAnySectionCastsShadows;
-	PrecachePSOParams.bReverseCulling = bReverseCulling;
-	
+	const bool bCanUseGPUScene = UseGPUScene(GMaxRHIShaderPlatform, GMaxRHIFeatureLevel);
+	FStaticMeshInstanceBuffer* InstanceBuffer = bCanUseGPUScene ? nullptr : &GDummyStaticMeshInstanceBuffer;
 	const FVertexFactoryType* VFType = ShouldCreateNaniteProxy() ? &Nanite::FVertexFactory::StaticType : &FInstancedStaticMeshVertexFactory::StaticType;
-
-	for (uint16 MaterialIndex : UsedMaterialIndices)
+	int32 LightMapCoordinateIndex = GetStaticMesh()->GetLightMapCoordinateIndex();
+	
+	auto ISMC_GetElements = [LightMapCoordinateIndex, InstanceBuffer](const FStaticMeshLODResources& LODRenderData, bool bSupportsManualVertexFetch, FVertexDeclarationElementList& Elements)
 	{
-		UMaterialInterface* MaterialInterface = GetMaterial(MaterialIndex);
-		if (MaterialInterface)
-		{
-			FComponentPSOPrecacheParams& ComponentParams = OutParams[OutParams.AddDefaulted()];
-			ComponentParams.MaterialInterface = MaterialInterface;
-			ComponentParams.VertexFactoryDataList.Add(FPSOPrecacheVertexFactoryData(VFType));
-			ComponentParams.PSOPrecacheParams = PrecachePSOParams;
-		}
-	}
+		FInstancedStaticMeshVertexFactory::FDataType Data;
+		InitInstancedStaticMeshVertexFactoryComponents(LODRenderData.VertexBuffers, InstanceBuffer, nullptr /*VertexFactory*/, LightMapCoordinateIndex, LODRenderData.bHasColorVertexData, bSupportsManualVertexFetch, Data);
+		FInstancedStaticMeshVertexFactory::GetVertexElements(GMaxRHIFeatureLevel, EVertexInputStreamType::Default, bSupportsManualVertexFetch, Data, Elements);
+	};
+
+	CollectPSOPrecacheDataImpl(VFType, BasePrecachePSOParams, ISMC_GetElements, OutParams);
 }
 
 void UInstancedStaticMeshComponent::OnPostLoadPerInstanceData()
diff --git a/Engine/Source/Runtime/Engine/Private/LocalVertexFactory.cpp b/Engine/Source/Runtime/Engine/Private/LocalVertexFactory.cpp
--- a/Engine/Source/Runtime/Engine/Private/LocalVertexFactory.cpp
+++ b/Engine/Source/Runtime/Engine/Private/LocalVertexFactory.cpp
@@ -329,28 +329,36 @@ void FLocalVertexFactory::ValidateCompiledResult(const FVertexFactoryType* Type,
 */
 void FLocalVertexFactory::GetPSOPrecacheVertexFetchElements(EVertexInputStreamType VertexInputStreamType, FVertexDeclarationElementList& Elements)
 {
-	Elements.Add(FVertexElement(0, 0, VET_Float3, 0, 12, false));
-
-	switch (VertexInputStreamType)
-	{
-	case EVertexInputStreamType::Default:
-	{
-		Elements.Add(FVertexElement(1, 0, VET_UInt, 13, 0, true));
-		break;
-	}
-	case EVertexInputStreamType::PositionOnly:
+	Elements.Add(FVertexElement(0, 0, VET_Float3, 0, sizeof(float)*3u, false));
+	
+	if (VertexInputStreamType == EVertexInputStreamType::PositionAndNormalOnly)
 	{
-		Elements.Add(FVertexElement(1, 0, VET_UInt, 1, 0, true));
-		break;
+		// 2-axis TangentBasis components in a single buffer, hence *2u
+		Elements.Add(FVertexElement(1, 4, VET_PackedNormal, 2, sizeof(FPackedNormal)*2u, false));
 	}
-	case EVertexInputStreamType::PositionAndNormalOnly:
+
+	if (UseGPUScene(GMaxRHIShaderPlatform, GMaxRHIFeatureLevel))
 	{
-		Elements.Add(FVertexElement(1, 4, VET_PackedNormal, 2, 0, false));
-		Elements.Add(FVertexElement(2, 0, VET_UInt, 1, 0, true));
-		break;
-	}
-	default:
-		checkNoEntry();
+		switch (VertexInputStreamType)
+		{
+		case EVertexInputStreamType::Default:
+		{
+			Elements.Add(FVertexElement(1, 0, VET_UInt, 13, sizeof(uint32), true));
+			break;
+		}
+		case EVertexInputStreamType::PositionOnly:
+		{
+			Elements.Add(FVertexElement(1, 0, VET_UInt, 1, sizeof(uint32), true));
+			break;
+		}
+		case EVertexInputStreamType::PositionAndNormalOnly:
+		{
+			Elements.Add(FVertexElement(2, 0, VET_UInt, 1, sizeof(uint32), true));
+			break;
+		}
+		default:
+			checkNoEntry();
+		}
 	}
 }
 
@@ -360,12 +368,14 @@ void FLocalVertexFactory::GetVertexElements(ERHIFeatureLevel::Type FeatureLevel,
 	int32 ColorStreamIndex;
 	GetVertexElements(FeatureLevel, InputStreamType, bSupportsManualVertexFetch, Data, Elements, VertexStreams, ColorStreamIndex);
 
-	// For ES3.1 attribute ID needs to be done differently
-	check(FeatureLevel > ERHIFeatureLevel::ES3_1);
-	Elements.Add(FVertexElement(VertexStreams.Num(), 0, VET_UInt, 13, 0, true));
+	if (UseGPUScene(GMaxRHIShaderPlatform, GMaxRHIFeatureLevel))
+	{
+		// For ES3.1 attribute ID needs to be done differently
+		check(FeatureLevel > ERHIFeatureLevel::ES3_1);
+		Elements.Add(FVertexElement(VertexStreams.Num(), 0, VET_UInt, 13, 0, true));
+	}
 }
 
-
 void FLocalVertexFactory::SetData(const FDataType& InData)
 {
 	check(IsInRenderingThread());
diff --git a/Engine/Source/Runtime/Engine/Private/SkeletalRenderStatic.cpp b/Engine/Source/Runtime/Engine/Private/SkeletalRenderStatic.cpp
--- a/Engine/Source/Runtime/Engine/Private/SkeletalRenderStatic.cpp
+++ b/Engine/Source/Runtime/Engine/Private/SkeletalRenderStatic.cpp
@@ -215,6 +215,7 @@ void FSkeletalMeshObjectStatic::FSkeletalMeshObjectLOD::InitResources(FSkelMeshC
 			PositionVertexBufferPtr->BindPositionVertexBuffer(VertexFactoryPtr, Data);
 			StaticMeshVertexBufferPtr->BindTangentVertexBuffer(VertexFactoryPtr, Data);
 			StaticMeshVertexBufferPtr->BindPackedTexCoordVertexBuffer(VertexFactoryPtr, Data);
+			StaticMeshVertexBufferPtr->BindLightMapVertexBuffer(VertexFactoryPtr, Data, 0);
 			ColorVertexBufferPtr->BindColorVertexBuffer(VertexFactoryPtr, Data);
 
 			VertexFactoryPtr->SetData(Data);
diff --git a/Engine/Source/Runtime/Engine/Private/SkinnedAsset.cpp b/Engine/Source/Runtime/Engine/Private/SkinnedAsset.cpp
--- a/Engine/Source/Runtime/Engine/Private/SkinnedAsset.cpp
+++ b/Engine/Source/Runtime/Engine/Private/SkinnedAsset.cpp
@@ -116,6 +116,9 @@ void USkinnedAsset::PostLoad()
 	}
 }
 
+// TODO: move this to a StaticMeshResources.h?
+extern void InitStaticMeshVertexFactoryComponents(const FStaticMeshVertexBuffers& VertexBuffers, FLocalVertexFactory* VertexFactory, int32 LightMapCoordinateIndex, bool bOverrideColorVertexBuffer, FLocalVertexFactory::FDataType& OutData);
+
 FPSOPrecacheVertexFactoryDataPerMaterialIndexList USkinnedAsset::GetVertexFactoryTypesPerMaterialIndex(
 	USkinnedMeshComponent* SkinnedMeshComponent, int32 MinLODIndex, bool bCPUSkin, ERHIFeatureLevel::Type FeatureLevel)
 {
@@ -154,7 +157,21 @@ FPSOPrecacheVertexFactoryDataPerMaterialIndexList USkinnedAsset::GetVertexFactor
 			if (bCPUSkin)
 			{
 				// Force static from GPU point of view
-				VFsPerMaterial->VertexFactoryDataList.AddUnique(FPSOPrecacheVertexFactoryData(&FLocalVertexFactory::StaticType));
+				const FVertexFactoryType* CPUSkinVFType = &FLocalVertexFactory::StaticType;
+				bool bSupportsManualVertexFetch = CPUSkinVFType->SupportsManualVertexFetch(GMaxRHIFeatureLevel);
+				if (!bSupportsManualVertexFetch)
+				{
+					FVertexDeclarationElementList VertexElements;
+					bool bOverrideColorVertexBuffer = false;
+					FLocalVertexFactory::FDataType Data;
+					InitStaticMeshVertexFactoryComponents(LODRenderData.StaticVertexBuffers, nullptr /*VertexFactory*/, 0, bOverrideColorVertexBuffer, Data);
+					FLocalVertexFactory::GetVertexElements(GMaxRHIFeatureLevel, EVertexInputStreamType::Default, bSupportsManualVertexFetch, Data, VertexElements);
+					VFsPerMaterial->VertexFactoryDataList.AddUnique(FPSOPrecacheVertexFactoryData(CPUSkinVFType, VertexElements));
+				}
+				else
+				{
+					VFsPerMaterial->VertexFactoryDataList.AddUnique(FPSOPrecacheVertexFactoryData(CPUSkinVFType));
+				}
 			}
 			else if (!SkelMeshRenderData->RequiresCPUSkinning(FeatureLevel, LODIndex))
 			{
diff --git a/Engine/Source/Runtime/Landscape/Private/LandscapeRender.cpp b/Engine/Source/Runtime/Landscape/Private/LandscapeRender.cpp
--- a/Engine/Source/Runtime/Landscape/Private/LandscapeRender.cpp
+++ b/Engine/Source/Runtime/Landscape/Private/LandscapeRender.cpp
@@ -244,19 +244,23 @@ void ULandscapeComponent::GetUsedMaterials(TArray<UMaterialInterface*>& OutMater
 
 	ALandscapeProxy* Actor = GetLandscapeProxy();
 
-	if (Actor != nullptr && Actor->bUseDynamicMaterialInstance)
+	if (GMaxRHIFeatureLevel > ERHIFeatureLevel::ES3_1)
 	{
-		OutMaterials.Append(MaterialInstancesDynamic.FilterByPredicate([](UMaterialInstanceDynamic* MaterialInstance) { return MaterialInstance != nullptr; }));
-	}
-	else
-	{
-		OutMaterials.Append(MaterialInstances.FilterByPredicate([](UMaterialInstanceConstant* MaterialInstance) { return MaterialInstance != nullptr; }));
+		if (Actor != nullptr && Actor->bUseDynamicMaterialInstance)
+		{
+			OutMaterials.Append(MaterialInstancesDynamic.FilterByPredicate([](UMaterialInstanceDynamic* MaterialInstance) { return MaterialInstance != nullptr; }));
+		}
+		else
+		{
+			OutMaterials.Append(MaterialInstances.FilterByPredicate([](UMaterialInstanceConstant* MaterialInstance) { return MaterialInstance != nullptr; }));
+		}
 	}
 
 	if (MobileMaterialInterfaces.Num())
 	{
 		OutMaterials.Append(MobileMaterialInterfaces.FilterByPredicate([](UMaterialInterface* MaterialInstance) { return MaterialInstance != nullptr; }));
 	}
+
 #if	WITH_EDITORONLY_DATA	
 	if (MobileCombinationMaterialInstances.Num())
 	{
@@ -3097,8 +3101,12 @@ void FLandscapeVertexFactory::ModifyCompilationEnvironment(const FVertexFactoryS
 
 void FLandscapeVertexFactory::GetPSOPrecacheVertexFetchElements(EVertexInputStreamType VertexInputStreamType, FVertexDeclarationElementList& Elements)
 {
-	Elements.Add(FVertexElement(0, 0, VET_UByte4, 0, 0, false));
-	Elements.Add(FVertexElement(1, 0, VET_UInt, 1, 0, true));
+	Elements.Add(FVertexElement(0, 0, VET_UByte4, 0, sizeof(FLandscapeVertex), false));
+	
+	if (UseGPUScene(GMaxRHIShaderPlatform, GMaxRHIFeatureLevel))
+	{
+		Elements.Add(FVertexElement(1, 0, VET_UInt, 1, sizeof(uint32), true));
+	}
 }
 
 IMPLEMENT_VERTEX_FACTORY_PARAMETER_TYPE(FLandscapeVertexFactory, SF_Vertex, FLandscapeVertexFactoryVertexShaderParameters);
