Remove Garbage collect lock from skeletalmesh build to avoid dead lock, log an error instead.

diff --git a/Engine/Source/Runtime/Engine/Private/MorphTools.cpp b/Engine/Source/Runtime/Engine/Private/MorphTools.cpp
--- a/Engine/Source/Runtime/Engine/Private/MorphTools.cpp
+++ b/Engine/Source/Runtime/Engine/Private/MorphTools.cpp
@@ -225,15 +225,15 @@ void FFinishBuildMorphTargetData::ApplyEditorData(USkeletalMesh * SkeletalMesh,
 		return;
 	}
 	
-	if (SkeletalMesh->HasAnyFlags(RF_BeginDestroyed))
+	if (SkeletalMesh->HasAnyFlags(RF_BeginDestroyed | RF_FinishDestroyed))
 	{
 		//No need to apply the morph targets if the asset is being deleted.
 		//Also acquiring the GCScopeGuard when we are in the destruction process will create a deadlock, if there is code
 		//in the BeginDestroy function of the skeletalmesh that touch any property that is lock by the async build.
 		return;
 	}
-	//Make sure we do not create new morph target during a gc
-	FGCScopeGuard GCScopeGuard;
+	//GC should not be active during the build since we force finish the skinned asset compilation during the pre garbage delegate
+	check(!IsGarbageCollectingAndLockingUObjectHashTables());
 
 	FSkeletalMeshModel * SkelMeshModel = SkeletalMesh->GetImportedModel();
 	check(SkelMeshModel);
@@ -263,11 +263,21 @@ void FFinishBuildMorphTargetData::ApplyEditorData(USkeletalMesh * SkeletalMesh,
 				//Skip this empty morphtarget
 				continue;
 			}
-			//When we save the cook result we should never have to build a new morph target
-			//When saving cook build, we call GetPlatformSkeletalMeshRenderData in USkeletalMesh::BeginCacheForCookedPlatformData
-			//which happen before the serialization of that cook skeletalmesh
-			if (!bIsSerializeSaving)
+
+			//Reuse the morph target if it already exist and was not garbage collect (it can happen if we play with the morph target threshold build options)
+			MorphTarget = FindObjectSafe<UMorphTarget>(SkeletalMesh, *MorphTargetName.ToString(), true);
+			if (MorphTarget && MorphTarget->HasAnyFlags(RF_BeginDestroyed | RF_FinishDestroyed))
+			{
+				MorphTarget = nullptr;
+			}
+
+			if (!MorphTarget)
 			{
+				//When we save the cook result we should never have to build a new morph target
+				//When saving cook build, we call GetPlatformSkeletalMeshRenderData in USkeletalMesh::BeginCacheForCookedPlatformData
+				//which happen before the serialization of that cook skeletalmesh
+				if (!bIsSerializeSaving)
+				{
 				//Avoid recycling morphtarget with NewObject it cannot be done asynchronously
 				//Find the UMorphTarget and simply clear the data if it exist.
 				TArray<UObject*> SubObjects;
@@ -290,12 +300,13 @@ void FFinishBuildMorphTargetData::ApplyEditorData(USkeletalMesh * SkeletalMesh,
 				{
 					MorphTarget = NewObject<UMorphTarget>(SkeletalMesh, MorphTargetName);
 				}
-				check(MorphTarget);
-			}
-			else
-			{
-				UE_ASSET_LOG(LogSkeletalMesh, Error, SkeletalMesh, TEXT("Cannot cache a skeletalmesh during a serialize if some morph targets need to be created. The solution is to Pre cache the skeletalmesh before the serialization so no morph target get created."));
-				continue;
+					check(MorphTarget);
+				}
+				else
+				{
+					UE_ASSET_LOG(LogSkeletalMesh, Error, SkeletalMesh, TEXT("Cannot cache a skeletalmesh during a serialize if some morph targets need to be created. The solution is to Pre cache the skeletalmesh before the serialization so no morph target get created."));
+					continue;
+				}
 			}
 		}
 		else
diff --git a/Engine/Source/Runtime/Engine/Private/SkinnedAssetCompiler.cpp b/Engine/Source/Runtime/Engine/Private/SkinnedAssetCompiler.cpp
--- a/Engine/Source/Runtime/Engine/Private/SkinnedAssetCompiler.cpp
+++ b/Engine/Source/Runtime/Engine/Private/SkinnedAssetCompiler.cpp
@@ -54,6 +54,7 @@ FSkinnedAssetCompilingManager::FSkinnedAssetCompilingManager()
 {
 	SkinnedAssetCompilingManagerImpl::EnsureInitializedCVars();
 	PostReachabilityAnalysisHandle = FCoreUObjectDelegates::PostReachabilityAnalysis.AddRaw(this, &FSkinnedAssetCompilingManager::OnPostReachabilityAnalysis);
+	PreGarbageCollectHandle = FCoreUObjectDelegates::GetPreGarbageCollectDelegate().AddRaw(this, &FSkinnedAssetCompilingManager::OnPreGarbageCollect);
 }
 
 FName FSkinnedAssetCompilingManager::GetAssetTypeName() const
@@ -149,6 +150,7 @@ void FSkinnedAssetCompilingManager::Shutdown()
 	}
 
 	FCoreUObjectDelegates::PostReachabilityAnalysis.Remove(PostReachabilityAnalysisHandle);
+	FCoreUObjectDelegates::GetPreGarbageCollectDelegate().Remove(PreGarbageCollectHandle);
 }
 
 bool FSkinnedAssetCompilingManager::IsAsyncCompilationEnabled() const
@@ -500,6 +502,11 @@ void FSkinnedAssetCompilingManager::OnPostReachabilityAnalysis()
 	}
 }
 
+void FSkinnedAssetCompilingManager::OnPreGarbageCollect()
+{
+	FinishAllCompilation();
+}
+
 #endif // #if WITH_EDITOR
 
 #undef LOCTEXT_NAMESPACE
diff --git a/Engine/Source/Runtime/Engine/Public/SkinnedAssetCompiler.h b/Engine/Source/Runtime/Engine/Public/SkinnedAssetCompiler.h
--- a/Engine/Source/Runtime/Engine/Public/SkinnedAssetCompiler.h
+++ b/Engine/Source/Runtime/Engine/Public/SkinnedAssetCompiler.h
@@ -99,6 +99,9 @@ private:
 
 	void OnPostReachabilityAnalysis();
 	FDelegateHandle PostReachabilityAnalysisHandle;
+	
+	void OnPreGarbageCollect();
+	FDelegateHandle PreGarbageCollectHandle;
 };
 
 #endif // #if WITH_EDITOR
\ No newline at end of file
