Metasound support for SoundNodeWaveParam

diff --git a/Engine/Source/Runtime/Engine/Private/SoundNodeWaveParam.cpp b/Engine/Source/Runtime/Engine/Private/SoundNodeWaveParam.cpp
--- a/Engine/Source/Runtime/Engine/Private/SoundNodeWaveParam.cpp
+++ b/Engine/Source/Runtime/Engine/Private/SoundNodeWaveParam.cpp
@@ -3,7 +3,9 @@
 
 #include "Sound/SoundNodeWaveParam.h"
 #include "ActiveSound.h"
+#include "AudioDevice.h"
 #include "IAudioParameterTransmitter.h"
+#include "Sound/SoundCue.h"
 #include "Sound/SoundWave.h"
 
 #include UE_INLINE_GENERATED_CPP_BY_NAME(SoundNodeWaveParam)
@@ -29,7 +31,37 @@ void USoundNodeWaveParam::ParseNodes(FAudioDevice* AudioDevice, const UPTRINT No
 
 	if (USoundWave* NewWave = Cast<USoundWave>(ParamValue.ObjectParam))
 	{
-		NewWave->Parse(AudioDevice, GetNodeWaveInstanceHash(NodeWaveInstanceHash, (UPTRINT)NewWave, 0), ActiveSound, ParseParams, WaveInstances);
+		RETRIEVE_SOUNDNODE_PAYLOAD(sizeof(USoundWave*));
+		DECLARE_SOUNDNODE_ELEMENT(USoundWave*, PrevWave);
+
+		const UPTRINT WaveHash = GetNodeWaveInstanceHash(NodeWaveInstanceHash, (UPTRINT)NewWave, 0);
+		
+		if (PrevWave != NewWave)
+		{
+			if (FSoundCueParameterTransmitter* SoundCueTransmitter = static_cast<FSoundCueParameterTransmitter*>(ActiveSound.GetTransmitter()))
+			{
+				// removing here prevents waste in the case that a metasound is replaced with a soundwave
+				SoundCueTransmitter->Transmitters.Remove(GetNodeWaveInstanceHash(NodeWaveInstanceHash, (UPTRINT)PrevWave, 0));
+				
+				Audio::FParameterTransmitterInitParams Params;
+				Params.DefaultParams = ActiveSound.GetTransmitter()->GetParameters();
+				Params.InstanceID = Audio::GetTransmitterID(ActiveSound.GetAudioComponentID(), WaveHash, ActiveSound.GetPlayOrder()); 
+				Params.SampleRate = AudioDevice->GetSampleRate();
+				
+				Params.AudioDeviceID = AudioDevice->DeviceID;
+				
+				const TSharedPtr<Audio::IParameterTransmitter> SoundWaveTransmitter = NewWave->CreateParameterTransmitter(MoveTemp(Params));
+			
+				if (SoundWaveTransmitter.IsValid())
+				{
+					SoundCueTransmitter->Transmitters.Add(WaveHash, SoundWaveTransmitter);
+				}
+			}
+
+			PrevWave = NewWave;
+		}
+		
+		NewWave->Parse(AudioDevice, WaveHash, ActiveSound, ParseParams, WaveInstances);
 	}
 	else
 	{
