Move the scope for STAT_CreatePhysicsMeshes down to where we're actually doing busy work (and removing the redundant trace scope).

diff --git a/Engine/Source/Runtime/Engine/Private/PhysicsEngine/BodySetup.cpp b/Engine/Source/Runtime/Engine/Private/PhysicsEngine/BodySetup.cpp
--- a/Engine/Source/Runtime/Engine/Private/PhysicsEngine/BodySetup.cpp
+++ b/Engine/Source/Runtime/Engine/Private/PhysicsEngine/BodySetup.cpp
@@ -401,12 +401,8 @@ FByteBulkData* UBodySetup::GetCookedFormatData()
 
 void UBodySetup::CreatePhysicsMeshes()
 {
-	TRACE_CPUPROFILER_EVENT_SCOPE(UBodySetup::CreatePhysicsMeshes);
-
-	SCOPE_CYCLE_COUNTER(STAT_CreatePhysicsMeshes);
-
 	// Create meshes from cooked data if not already done
-	if(bCreatedPhysicsMeshes)
+	if (bCreatedPhysicsMeshes)
 	{
 		return;
 	}
@@ -416,6 +412,8 @@ void UBodySetup::CreatePhysicsMeshes()
 	{
 		return;
 	}
+
+	SCOPE_CYCLE_COUNTER(STAT_CreatePhysicsMeshes);
 	
 	bool bClearMeshes = true;
 	bool bSkipProcessFormatData = false;
