Change the scope cycle counter in AnimInstanceProxy update to only output one for the first node.

diff --git a/Engine/Source/Runtime/Engine/Private/Animation/AnimInstanceProxy.cpp b/Engine/Source/Runtime/Engine/Private/Animation/AnimInstanceProxy.cpp
--- a/Engine/Source/Runtime/Engine/Private/Animation/AnimInstanceProxy.cpp
+++ b/Engine/Source/Runtime/Engine/Private/Animation/AnimInstanceProxy.cpp
@@ -1187,8 +1187,8 @@ void FAnimInstanceProxy::UpdateAnimation_WithRoot(const FAnimationUpdateContext&
 	DECLARE_SCOPE_HIERARCHICAL_COUNTER_FUNC()
 
 	ANIM_MT_SCOPE_CYCLE_COUNTER(ProxyUpdateAnimation, !IsInGameThread());
-	FScopeCycleCounterUObject AnimScope(GetAnimInstanceObject());
-	
+	FScopeCycleCounterUObject AnimScope(bUpdatingRoot ? nullptr : GetAnimInstanceObject());
+
 #if WITH_EDITORONLY_DATA
 	UpdatedNodesThisFrame.Reset();
 	NodeInputAttributesThisFrame.Reset();
