Speculative fix for a crash caused by updating a texture via a render target while it is being build/compiled.

diff --git a/Engine/Source/Runtime/Engine/Private/KismetRenderingLibrary.cpp b/Engine/Source/Runtime/Engine/Private/KismetRenderingLibrary.cpp
--- a/Engine/Source/Runtime/Engine/Private/KismetRenderingLibrary.cpp
+++ b/Engine/Source/Runtime/Engine/Private/KismetRenderingLibrary.cpp
@@ -598,7 +598,7 @@ UTexture2D* UKismetRenderingLibrary::RenderTargetCreateStaticTexture2DEditorOnly
 		UObject* NewObj = nullptr;
 
 		// create a static 2d texture
-		NewObj = RenderTarget->ConstructTexture2D(CreatePackage( *PackageName), Name, RenderTarget->GetMaskedFlags() | RF_Public | RF_Standalone, CTF_Default | CTF_AllowMips, nullptr);
+		NewObj = RenderTarget->ConstructTexture2D(CreatePackage( *PackageName), Name, RenderTarget->GetMaskedFlags() | RF_Public | RF_Standalone, CTF_Default | CTF_AllowMips | CTF_SkipPostEdit, nullptr);
 		UTexture2D* NewTex = Cast<UTexture2D>(NewObj);
 
 		if (NewTex != nullptr)
@@ -606,14 +606,14 @@ UTexture2D* UKismetRenderingLibrary::RenderTargetCreateStaticTexture2DEditorOnly
 			// package needs saving
 			NewObj->MarkPackageDirty();
 
-			// Notify the asset registry
-			FAssetRegistryModule::AssetCreated(NewObj);
-
 			// Update Compression and Mip settings
 			NewTex->CompressionSettings = CompressionSettings;
 			NewTex->MipGenSettings = MipSettings;
 			NewTex->PostEditChange();
 
+			// Notify the asset registry
+			FAssetRegistryModule::AssetCreated(NewObj);
+
 			return NewTex;
 		}
 		FMessageLog("Blueprint").Warning(LOCTEXT("RenderTargetCreateStaticTexture2D_FailedToCreateTexture", "RenderTargetCreateStaticTexture2DEditorOnly: Failed to create a new texture."));
diff --git a/Engine/Source/Runtime/Engine/Private/TextureRenderTarget2D.cpp b/Engine/Source/Runtime/Engine/Private/TextureRenderTarget2D.cpp
--- a/Engine/Source/Runtime/Engine/Private/TextureRenderTarget2D.cpp
+++ b/Engine/Source/Runtime/Engine/Private/TextureRenderTarget2D.cpp
@@ -18,6 +18,7 @@
 #include "UObject/UnrealType.h"
 #if WITH_EDITOR
 #include "Components/SceneCaptureComponent2D.h"
+#include "TextureCompiler.h"
 #include "UObject/UObjectIterator.h"
 #endif
 
@@ -370,7 +371,11 @@ UTexture2D* UTextureRenderTarget2D::ConstructTexture2D(UObject* Outer, const FSt
 		Result->CompressionNone = true;
 		Result->DeferCompression = false;
 	}
-	Result->PostEditChange();
+
+	if ((Flags & CTF_SkipPostEdit) != 0)
+	{
+		Result->PostEditChange();
+	}
 #endif
 
 	return Result;
@@ -540,6 +545,9 @@ void UTextureRenderTarget2D::UpdateTexture2D(UTexture2D* InTexture2D, ETextureSo
 	{
 		TextureChangingDelegate.ExecuteIfBound(InTexture2D);
 
+		// Ensure the texture is not being compiled
+		FTextureCompilingManager::Get().FinishCompilation({InTexture2D});
+
 		// init to the same size as the 2d texture
 		InTexture2D->Source.Init(SizeX, SizeY, NbSlices, NbMips, InTextureFormat, NewData);
 		InTexture2D->CompressionSettings = CompressionSettingsForTexture;
diff --git a/Engine/Source/Runtime/RenderCore/Public/RenderUtils.h b/Engine/Source/Runtime/RenderCore/Public/RenderUtils.h
--- a/Engine/Source/Runtime/RenderCore/Public/RenderUtils.h
+++ b/Engine/Source/Runtime/RenderCore/Public/RenderUtils.h
@@ -133,7 +133,7 @@ enum EConstructTextureFlags
 {
 	/** Compress RGBA8 to DXT */
 	CTF_Compress =				0x01,
-	/** Don't actually compress until the pacakge is saved */
+	/** Don't actually compress until the package is saved */
 	CTF_DeferCompression =		0x02,
 	/** Enable SRGB on the texture */
 	CTF_SRGB =					0x04,
@@ -145,6 +145,8 @@ enum EConstructTextureFlags
 	CTF_RemapAlphaAsMasked =	0x20,
 	/** Ensure the alpha channel of the texture is opaque white (255). */
 	CTF_ForceOpaque =			0x40,
+	/** Don't call a post edit change on the texture. */
+	CTF_SkipPostEdit =			0x80,
 
 	/** Default flags (maps to previous defaults to ConstructTexture2D) */
 	CTF_Default = CTF_Compress | CTF_SRGB,
