WorldPartitionObjectResolver Interface

#rnx

diff --git a/Engine/Source/Runtime/Engine/Private/WorldPartition/WorldPartitionLevelHelper.cpp b/Engine/Source/Runtime/Engine/Private/WorldPartition/WorldPartitionLevelHelper.cpp
--- a/Engine/Source/Runtime/Engine/Private/WorldPartition/WorldPartitionLevelHelper.cpp
+++ b/Engine/Source/Runtime/Engine/Private/WorldPartition/WorldPartitionLevelHelper.cpp
@@ -8,6 +8,9 @@
 #include "Misc/PackageName.h"
 #include "WorldPartition/WorldPartitionPackageHelper.h"
 #include "WorldPartition/WorldPartitionRuntimeCell.h"
+#include "Containers/StringFwd.h"
+#include "WorldPartition/WorldPartitionActorContainerID.h"
+#include "WorldPartition/IWorldPartitionObjectResolver.h"
 
 #if WITH_EDITOR
 
@@ -21,8 +24,50 @@
 #include "WorldPartition/ContentBundle/ContentBundleEditor.h"
 #include "LevelUtils.h"
 #include "ActorFolder.h"
-#include "Containers/StringFwd.h"
 
+#endif
+
+bool FWorldPartitionResolveData::ResolveObject(UWorld* InWorld, const FSoftObjectPath& InObjectPath, UObject*& OutObject)
+{
+	OutObject = nullptr;
+	if (InWorld)
+	{
+		if (IsValid() && SourceWorldAssetPath == InObjectPath.GetAssetPath().ToString())
+		{
+			const FString SubPathString = FWorldPartitionLevelHelper::AddActorContainerIDToSubPathString(ContainerID, InObjectPath.GetSubPathString());
+			// We don't read the return value as we always want to return true when using the resolve data.
+			InWorld->ResolveSubobject(*SubPathString, OutObject, /*bLoadIfExists*/false);
+			return true;
+		}
+	}
+
+	return false;
+}
+
+FString FWorldPartitionLevelHelper::AddActorContainerIDToSubPathString(const FActorContainerID& InContainerID, const FString& InSubPathString)
+{
+	if (!InContainerID.IsMainContainer())
+	{
+		constexpr const TCHAR PersistenLevelName[] = TEXT("PersistentLevel.");
+		constexpr const int32 DotPos = UE_ARRAY_COUNT(PersistenLevelName);
+		if (InSubPathString.StartsWith(PersistenLevelName))
+		{
+			const int32 SubObjectPos = InSubPathString.Find(TEXT("."), ESearchCase::IgnoreCase, ESearchDir::FromStart, DotPos);
+			if (SubObjectPos == INDEX_NONE)
+			{
+				return InSubPathString + TEXT("_") + InContainerID.ToShortString();
+			}
+			else
+			{
+				return InSubPathString.Mid(0, SubObjectPos) + TEXT("_") + InContainerID.ToShortString() + InSubPathString.Mid(SubObjectPos);
+			}
+		}
+	}
+
+	return InSubPathString;
+}
+
+#if WITH_EDITOR
 FWorldPartitionLevelHelper& FWorldPartitionLevelHelper::Get()
 {
 	static FWorldPartitionLevelHelper Instance;
@@ -183,29 +228,6 @@ void FWorldPartitionLevelHelper::RemapLevelSoftObjectPaths(ULevel* InLevel, UWor
 	FixupSerializer.Fixup(InLevel);
 }
 
-FString FWorldPartitionLevelHelper::AddActorContainerIDToSubPathString(const FActorContainerID& InContainerID, const FString& InSubPathString)
-{
-	if (!InContainerID.IsMainContainer())
-	{
-		constexpr const TCHAR PersistenLevelName[] = TEXT("PersistentLevel.");
-		constexpr const int32 DotPos = UE_ARRAY_COUNT(PersistenLevelName);
-		if (InSubPathString.StartsWith(PersistenLevelName))
-		{
-			const int32 SubObjectPos = InSubPathString.Find(TEXT("."), ESearchCase::IgnoreCase, ESearchDir::FromStart, DotPos);
-			if (SubObjectPos == INDEX_NONE)
-			{
-				return InSubPathString + TEXT("_") + InContainerID.ToShortString();
-			}
-			else
-			{
-				return InSubPathString.Mid(0, SubObjectPos) + TEXT("_") + InContainerID.ToShortString() + InSubPathString.Mid(SubObjectPos);
-			}
-		}
-	}
-
-	return InSubPathString;
-}
-
 FString FWorldPartitionLevelHelper::GetContainerPackage(const FActorContainerID& InContainerID, const FString& InPackageName, const FString& InDestLevelPackageName)
 {
 	// Generate a unique name to load a Level Instance embedded actor if there are multiple instances of this Level Instance and possibly across multiple instances of the World Partition world
@@ -475,6 +497,11 @@ bool FWorldPartitionLevelHelper::LoadActors(UWorld* InOuterWorld, ULevel* InDest
 						}
 					});
 					FixupArchive.Fixup(Actor);
+
+					if (IWorldPartitionObjectResolver* ObjectResolver = Cast<IWorldPartitionObjectResolver>(Actor))
+					{
+						ObjectResolver->SetWorldPartitionResolveData(FWorldPartitionResolveData(PackageObjectMapping->ContainerID, SourceWorldPath));
+					}
 				}
 
 				if (InDestLevel)
diff --git a/Engine/Source/Runtime/Engine/Public/WorldPartition/IWorldPartitionObjectResolver.h b/Engine/Source/Runtime/Engine/Public/WorldPartition/IWorldPartitionObjectResolver.h
new file mode 100644
--- /dev/null
+++ b/Engine/Source/Runtime/Engine/Public/WorldPartition/IWorldPartitionObjectResolver.h
@@ -0,0 +1,48 @@
+// Copyright Epic Games, Inc. All Rights Reserved.
+
+#pragma once
+
+#include "UObject/Interface.h"
+#include "UObject/ObjectMacros.h"
+#include "Engine/EngineTypes.h"
+#include "WorldPartition/WorldPartitionActorContainerID.h"
+#include "IWorldPartitionObjectResolver.generated.h"
+
+USTRUCT()
+struct FWorldPartitionResolveData
+{
+	GENERATED_USTRUCT_BODY()
+
+	FWorldPartitionResolveData() {}
+	FWorldPartitionResolveData(const FActorContainerID& InContainerID, const FString& InSourceWorldAssetPath) 
+		: ContainerID(InContainerID), SourceWorldAssetPath(InSourceWorldAssetPath) {}
+
+	bool IsValid() const { return !ContainerID.IsMainContainer(); }
+		
+	bool ResolveObject(UWorld* InWorld, const FSoftObjectPath& InObjectPath, UObject*& OutObject);
+
+	UPROPERTY()
+	FActorContainerID ContainerID;
+
+	UPROPERTY()
+	FString SourceWorldAssetPath;
+};
+
+UINTERFACE()
+class ENGINE_API UWorldPartitionObjectResolver
+	: public UInterface
+{
+public:
+	GENERATED_BODY()
+};
+
+class ENGINE_API IWorldPartitionObjectResolver
+{
+public:
+	GENERATED_BODY()
+
+#if WITH_EDITOR
+	virtual void SetWorldPartitionResolveData(const FWorldPartitionResolveData& InEmbeddingData) = 0;
+#endif
+	virtual const FWorldPartitionResolveData& GetWorldPartitionResolveData() const = 0;
+};
diff --git a/Engine/Source/Runtime/Engine/Public/WorldPartition/WorldPartitionLevelHelper.h b/Engine/Source/Runtime/Engine/Public/WorldPartition/WorldPartitionLevelHelper.h
--- a/Engine/Source/Runtime/Engine/Public/WorldPartition/WorldPartitionLevelHelper.h
+++ b/Engine/Source/Runtime/Engine/Public/WorldPartition/WorldPartitionLevelHelper.h
@@ -9,8 +9,6 @@
 
 #pragma once
 
-#if WITH_EDITOR
-
 #include "Engine/World.h"
 #include "WorldPartition/WorldPartitionRuntimeCell.h"
 
@@ -20,6 +18,9 @@ struct FActorContainerID;
 
 class FWorldPartitionLevelHelper
 {
+public:
+	static FString AddActorContainerIDToSubPathString(const FActorContainerID& InContainerID, const FString& InSubPathString);
+#if WITH_EDITOR
 public:
 	static FWorldPartitionLevelHelper& Get();
 
@@ -46,7 +47,6 @@ private:
 	void RemoveReferences(FPackageReferencer* InReferencer);
 	void PreGarbageCollect();
 
-	static FString AddActorContainerIDToSubPathString(const FActorContainerID& InContainerID, const FString& InSubPathString);
 	static FString GetContainerPackage(const FActorContainerID& InContainerID, const FString& InPackageName, const FString& InDestLevelPackageName = FString());
 	static UWorld::InitializationValues GetWorldInitializationValues();
 
@@ -65,6 +65,6 @@ private:
 	TMap<FName, FPackageReference> PackageReferences;
 
 	TSet<TWeakObjectPtr<UPackage>> PreGCPackagesToUnload;
+#endif
 };
 
-#endif
\ No newline at end of file
