Temporarily default LevelStreaming.DefaultAllowClientUseMakingInvisibleTransactionRequests to false even when compiling with Iris as it can cause an infinite loop in UWorld::BlockTillLevelStreamingCompleted

diff --git a/Engine/Source/Runtime/Engine/Private/LevelStreaming.cpp b/Engine/Source/Runtime/Engine/Private/LevelStreaming.cpp
--- a/Engine/Source/Runtime/Engine/Private/LevelStreaming.cpp
+++ b/Engine/Source/Runtime/Engine/Private/LevelStreaming.cpp
@@ -45,7 +45,7 @@ namespace LevelStreamingCVars
 	// starting to replicate data based on an older visibility/streamingstatus update which can lead to broken channels
 	// to mitigate this problem we assign a TransactionId to each request/update to make sure that we are acting on the correct data
 #if UE_WITH_IRIS
-	static bool bDefaultAllowClientUseMakingInvisibleTransactionRequests = true;
+	static bool bDefaultAllowClientUseMakingInvisibleTransactionRequests = false;
 #else
 	static bool bDefaultAllowClientUseMakingInvisibleTransactionRequests = false;
 #endif
