Adjust check for valid actor channel net guid for replays

diff --git a/Engine/Source/Runtime/Engine/Private/DataChannel.cpp b/Engine/Source/Runtime/Engine/Private/DataChannel.cpp
--- a/Engine/Source/Runtime/Engine/Private/DataChannel.cpp
+++ b/Engine/Source/Runtime/Engine/Private/DataChannel.cpp
@@ -2418,36 +2418,40 @@ bool UActorChannel::CleanUp(const bool bForDestroy, EChannelCloseReason CloseRea
 	{
 		LLM_SCOPE_BYTAG(NetConnection);
 
-		checkf(ActorNetGUID.IsValid(), TEXT("UActorChannel::Cleanup: ActorNetGUID is invalid! Channel: %i"), ChIndex);
-		
-		TArray<UActorChannel*>& ChannelsStillProcessing = Connection->KeepProcessingActorChannelBunchesMap.FindOrAdd(ActorNetGUID);
+		// Allow channels in replays to be cleaned up even if the guid was never set
+		checkf(ActorNetGUID.IsValid() || Connection->IsInternalAck(), TEXT("UActorChannel::Cleanup: ActorNetGUID is invalid! Channel: %i"), ChIndex);
 		
+		if (ActorNetGUID.IsValid())
+		{
+			TArray<UActorChannel*>& ChannelsStillProcessing = Connection->KeepProcessingActorChannelBunchesMap.FindOrAdd(ActorNetGUID);
+
 #if DO_CHECK
-		if (ensureMsgf(!ChannelsStillProcessing.Contains(this), TEXT("UActorChannel::CleanUp encountered a channel already within the KeepProcessingActorChannelBunchMap. Channel: %i"), ChIndex))
+			if (ensureMsgf(!ChannelsStillProcessing.Contains(this), TEXT("UActorChannel::CleanUp encountered a channel already within the KeepProcessingActorChannelBunchMap. Channel: %i"), ChIndex))
 #endif // #if DO_CHECK
-		{
-			UE_LOG(LogNet, VeryVerbose, TEXT("UActorChannel::CleanUp: Adding to KeepProcessingActorChannelBunchesMap. Channel: %i, Num: %i"), ChIndex, Connection->KeepProcessingActorChannelBunchesMap.Num());
+			{
+				UE_LOG(LogNet, VeryVerbose, TEXT("UActorChannel::CleanUp: Adding to KeepProcessingActorChannelBunchesMap. Channel: %i, Num: %i"), ChIndex, Connection->KeepProcessingActorChannelBunchesMap.Num());
 
-			// Remember the connection, since CleanUp below will NULL it
-			UNetConnection* OldConnection = Connection;
+				// Remember the connection, since CleanUp below will NULL it
+				UNetConnection* OldConnection = Connection;
 
-			// This will unregister the channel, and make it free for opening again
-			// We need to do this, since the server will assume this channel is free once we ack this packet
-			Super::CleanUp(bForDestroy, CloseReason);
+				// This will unregister the channel, and make it free for opening again
+				// We need to do this, since the server will assume this channel is free once we ack this packet
+				Super::CleanUp(bForDestroy, CloseReason);
 
-			// Restore connection property since we'll need it for processing bunches (the Super::CleanUp call above NULL'd it)
-			Connection = OldConnection;
+				// Restore connection property since we'll need it for processing bunches (the Super::CleanUp call above NULL'd it)
+				Connection = OldConnection;
 
-			QueuedCloseReason = CloseReason;
+				QueuedCloseReason = CloseReason;
 
-			// Add this channel to the KeepProcessingActorChannelBunchesMap list
-			ChannelsStillProcessing.Add(this);
+				// Add this channel to the KeepProcessingActorChannelBunchesMap list
+				ChannelsStillProcessing.Add(this);
 
-			// We set ChIndex to -1 to signify that we've already been "closed" but we aren't done processing bunches
-			ChIndex = -1;
+				// We set ChIndex to -1 to signify that we've already been "closed" but we aren't done processing bunches
+				ChIndex = -1;
 
-			// Return false so we won't do pending kill yet
-			return false;
+				// Return false so we won't do pending kill yet
+				return false;
+			}
 		}
 	}
 
