Add additional project-specific texture groups

diff --git a/Engine/Source/Runtime/Engine/Classes/Engine/TextureDefines.h b/Engine/Source/Runtime/Engine/Classes/Engine/TextureDefines.h
--- a/Engine/Source/Runtime/Engine/Classes/Engine/TextureDefines.h
+++ b/Engine/Source/Runtime/Engine/Classes/Engine/TextureDefines.h
@@ -89,6 +89,20 @@ enum TextureGroup : int
 	TEXTUREGROUP_Project16 UMETA(DisplayName="ini:Project Group 16"),
 	TEXTUREGROUP_Project17 UMETA(DisplayName="ini:Project Group 17"),
 	TEXTUREGROUP_Project18 UMETA(DisplayName="ini:Project Group 18"),
+	TEXTUREGROUP_Project19 UMETA(DisplayName="ini:Project Group 19"),
+	TEXTUREGROUP_Project20 UMETA(DisplayName="ini:Project Group 20"),
+	TEXTUREGROUP_Project21 UMETA(DisplayName="ini:Project Group 21"),
+	TEXTUREGROUP_Project22 UMETA(DisplayName="ini:Project Group 22"),
+	TEXTUREGROUP_Project23 UMETA(DisplayName="ini:Project Group 23"),
+	TEXTUREGROUP_Project24 UMETA(DisplayName="ini:Project Group 24"),
+	TEXTUREGROUP_Project25 UMETA(DisplayName="ini:Project Group 25"),
+	TEXTUREGROUP_Project26 UMETA(DisplayName="ini:Project Group 26"),
+	TEXTUREGROUP_Project27 UMETA(DisplayName="ini:Project Group 27"),
+	TEXTUREGROUP_Project28 UMETA(DisplayName="ini:Project Group 28"),
+	TEXTUREGROUP_Project29 UMETA(DisplayName="ini:Project Group 29"),
+	TEXTUREGROUP_Project30 UMETA(DisplayName="ini:Project Group 30"),
+	TEXTUREGROUP_Project31 UMETA(DisplayName="ini:Project Group 31"),
+	TEXTUREGROUP_Project32 UMETA(DisplayName="ini:Project Group 32"),
 	TEXTUREGROUP_MAX,
 };
 
