Fix a crash on editor shutdown where the transactors for ISM typed element handles were cleaned up too aggressively

diff --git a/Engine/Source/Runtime/Engine/Private/Elements/SMInstance/SMInstanceElementId.cpp b/Engine/Source/Runtime/Engine/Private/Elements/SMInstance/SMInstanceElementId.cpp
--- a/Engine/Source/Runtime/Engine/Private/Elements/SMInstance/SMInstanceElementId.cpp
+++ b/Engine/Source/Runtime/Engine/Private/Elements/SMInstance/SMInstanceElementId.cpp
@@ -364,10 +364,12 @@ void FSMInstanceElementIdMap::AddReferencedObjects(FReferenceCollector& Collecto
 #if WITH_EDITORONLY_DATA
 	FScopeLock Lock(&ISMComponentsCS);
 
+	Collector.AllowEliminatingReferences(false);
 	for (const TTuple<UInstancedStaticMeshComponent*, TSharedPtr<FSMInstanceElementIdMapEntry>>& ISMEntryPair : ISMComponents)
 	{
 		Collector.AddReferencedObject(ISMEntryPair.Value->Transactor);
 	}
+	Collector.AllowEliminatingReferences(true);
 #endif	// WITH_EDITORONLY_DATA
 }
 
