Adds a null check before evaluating the TemplateComponent type in ForEachComponentOfActorClassDefaults

diff --git a/Engine/Source/Runtime/Engine/Private/Blueprint.cpp b/Engine/Source/Runtime/Engine/Private/Blueprint.cpp
--- a/Engine/Source/Runtime/Engine/Private/Blueprint.cpp
+++ b/Engine/Source/Runtime/Engine/Private/Blueprint.cpp
@@ -1793,6 +1793,11 @@ void UBlueprint::ForEachComponentOfActorClassDefault(const TSubclassOf<AActor>&
 
 	auto FilterFunc = [&](const UActorComponent* TemplateComponent)
 	{
+		if (!TemplateComponent)
+		{
+			return true;
+		}
+
 		if (!InComponentClass.Get() || TemplateComponent->IsA(InComponentClass))
 		{
 			return InFunc(TemplateComponent);
