- Put back LevelStreaming.DefaultAllowClientUseMakingInvisibleTransactionRequests default value to true when using Iris.
- Fixed ULevelStreaming CanMakeVisible and CanMakeInvisible to prevent returning false if they already returned true once and became the world's CurrentLevelPendingVisibility/Invisibility
- Since this problem is very difficult to reproduce, I validated the fix by forcing artificially the case where CanMakeInvisible returns true the 1st time and false afterwards while called inside a BlockTillLevelStreamingCompleted.

diff --git a/Engine/Source/Runtime/Engine/Private/LevelStreaming.cpp b/Engine/Source/Runtime/Engine/Private/LevelStreaming.cpp
--- a/Engine/Source/Runtime/Engine/Private/LevelStreaming.cpp
+++ b/Engine/Source/Runtime/Engine/Private/LevelStreaming.cpp
@@ -45,7 +45,7 @@ namespace LevelStreamingCVars
 	// starting to replicate data based on an older visibility/streamingstatus update which can lead to broken channels
 	// to mitigate this problem we assign a TransactionId to each request/update to make sure that we are acting on the correct data
 #if UE_WITH_IRIS
-	static bool bDefaultAllowClientUseMakingInvisibleTransactionRequests = false;
+	static bool bDefaultAllowClientUseMakingInvisibleTransactionRequests = true;
 #else
 	static bool bDefaultAllowClientUseMakingInvisibleTransactionRequests = false;
 #endif
@@ -758,6 +758,29 @@ bool ULevelStreaming::RequestVisibilityChange(bool bVisible)
 
 bool ULevelStreaming::CanMakeInvisible()
 {
+	// Once the Level becomes the current pending invisibility level, this function must return true or else it could
+	// block indefinitely if called inside UWorld::BlockTillLevelStreamingCompleted
+	//
+	// Here's one example of why this could go wrong :
+	//      - The first call to CanMakeInvisible for a level returns true (when normally it would return false)
+	//			ShouldWaitForServerAckBeforeChangingVisibilityState would normally return true but
+	//			IsWaitingForNetVisibilityTransactionAck returns false only because ServerConnection is different than USOCK_Open 
+	//			and IsConcernedByNetVisibilityTransactionAck returns false
+	//		- RemoveFromWorld is called on this level and sets CurrentLevelPendingInvisibility
+	//		- Because time limit is exceeded, RemoveFromWorld exists with World's CurrentLevelPendingInvisibility 
+	//		  still set to this level streaming loaded level
+	//		- Since we are inside UWorld::BlockTillLevelStreamingCompleted, IsVisibilityRequestPending returns true because CurrentLevelPendingInvisibility is valid
+	//		- UpdateStreamingState is called again
+	//		- All future calls to CanMakeInvisible return false (ServerConnection is USOCK_Open) and we are waiting for server ack
+	//
+	// This is one hypothetical example. Detecting this will prevent any future case that could trigger an infinte loop in BlockTillLevelStreamingCompleted.
+	//
+	UWorld* World = GetWorld();
+	if (World && World->IsGameWorld() && LoadedLevel && (LoadedLevel == World->GetCurrentLevelPendingInvisibility()))
+	{
+		return true;
+	}
+
 	const bool bCanMakeInvisible = RequestVisibilityChange(false);
 
 	if (ShouldClientUseMakingInvisibleTransactionRequest())
@@ -773,6 +796,14 @@ bool ULevelStreaming::CanMakeInvisible()
 
 bool ULevelStreaming::CanMakeVisible()
 {
+	// Once the Level becomes the current pending visibility level, this function must return true or else it could
+	// block indefinitely if called inside UWorld::BlockTillLevelStreamingCompleted (same reason as CanMakeInvisible but with World's CurrentLevelPendingVisibility)
+	UWorld* World = GetWorld();
+	if (World && World->IsGameWorld() && LoadedLevel && (LoadedLevel == World->GetCurrentLevelPendingVisibility()))
+	{
+		return true;
+	}
+
 	const bool bCanMakeVisible = RequestVisibilityChange(true);
 
 	if (ShouldClientUseMakingVisibleTransactionRequest())
