PR #9622: Error if soundwave LoadingBehavior is not ForceInline (Contributed by [REDACTED_PERSON])

#rnx

diff --git a/Engine/Source/Runtime/Engine/Private/GameFramework/HapticFeedbackEffect.cpp b/Engine/Source/Runtime/Engine/Private/GameFramework/HapticFeedbackEffect.cpp
--- a/Engine/Source/Runtime/Engine/Private/GameFramework/HapticFeedbackEffect.cpp
+++ b/Engine/Source/Runtime/Engine/Private/GameFramework/HapticFeedbackEffect.cpp
@@ -183,6 +183,13 @@ void UHapticFeedbackEffect_SoundWave::PrepareSoundWaveBuffer()
 	{
 		return;
 	}
+
+	if (SoundWave->LoadingBehavior != ESoundWaveLoadingBehavior::ForceInline)
+	{
+		UE_LOG(LogHaptics, Error, TEXT("The LoadingBehavior of a SoundWave needs to be 'ESoundWaveLoadingBehavior::ForceInline' for use in Haptic Feedback. ('%s' requires a change)"), *GetNameSafe(SoundWave));
+		return;
+	}
+
 	AD->Precache(SoundWave, true, false, true);
 	// Remove call to InitAudioResource because AD->Precache calls the InitAudioResource code and calling it twice causes issues with Bulk data loading. 
 	uint8* PCMData = SoundWave->RawPCMData;
