Fix ensures from having disabled SoundField Submixes types

diff --git a/Engine/Source/Runtime/Engine/Classes/Sound/SoundSubmix.h b/Engine/Source/Runtime/Engine/Classes/Sound/SoundSubmix.h
--- a/Engine/Source/Runtime/Engine/Classes/Sound/SoundSubmix.h
+++ b/Engine/Source/Runtime/Engine/Classes/Sound/SoundSubmix.h
@@ -452,6 +452,8 @@ public:
 	void SanitizeLinks();
 
 protected:
+	virtual void PostLoad() override;
+
 #if WITH_EDITOR
 	virtual void PostEditChangeProperty(struct FPropertyChangedEvent& PropertyChangedEvent) override;
 #endif
diff --git a/Engine/Source/Runtime/Engine/Private/SoundSubmix.cpp b/Engine/Source/Runtime/Engine/Private/SoundSubmix.cpp
--- a/Engine/Source/Runtime/Engine/Private/SoundSubmix.cpp
+++ b/Engine/Source/Runtime/Engine/Private/SoundSubmix.cpp
@@ -927,9 +927,31 @@ void USoundfieldSubmix::PostEditChangeProperty(struct FPropertyChangedEvent& Pro
 		SanitizeLinks();
 	}
 }
-
 #endif // WITH_EDITOR
 
+void USoundfieldSubmix::PostLoad()
+{
+#if WITH_EDITOR
+
+	// Make sure the Encoding format is something we can use.
+	// Fallback to something that works otherwise and warn.
+
+	TArray<FName> FactoryNames = ISoundfieldFactory::GetAvailableSoundfieldFormats();
+	if (!FactoryNames.Contains(SoundfieldEncodingFormat))
+	{
+		const FName NoEncoding = ISoundfieldFactory::GetFormatNameForNoEncoding();
+		UE_LOG(LogAudio, Warning, TEXT("SoundfieldSubmix [%s] has Encoding format [%s] which is not currently currently enabled. Changing to [%s]"),
+			*GetName(), *SoundfieldEncodingFormat.ToString(), *NoEncoding.ToString());
+		SoundfieldEncodingFormat = NoEncoding;
+		SanitizeLinks();
+	}
+
+#endif //WITH_EDITOR
+
+	Super::PostLoad();
+}
+
+
 IAudioEndpointFactory* UEndpointSubmix::GetAudioEndpointForSubmix() const
 {
 	return IAudioEndpointFactory::Get(EndpointType);
