Bugfix: [Dolby][GAME-1093] When creating an Endpoint Submix Asset in editor, CreateNewEndpointInstance is not called

diff --git a/Engine/Source/Runtime/Engine/Classes/Sound/SoundSubmix.h b/Engine/Source/Runtime/Engine/Classes/Sound/SoundSubmix.h
--- a/Engine/Source/Runtime/Engine/Classes/Sound/SoundSubmix.h
+++ b/Engine/Source/Runtime/Engine/Classes/Sound/SoundSubmix.h
@@ -483,6 +483,9 @@ public:
 	UPROPERTY(EditAnywhere, Category = Endpoint)
 	TObjectPtr<UAudioEndpointSettingsBase> EndpointSettings;
 
+#if WITH_EDITOR
+	virtual void PostEditChangeProperty(struct FPropertyChangedEvent& PropertyChangedEvent) override;
+#endif
 };
 
 /**
diff --git a/Engine/Source/Runtime/Engine/Private/SoundSubmix.cpp b/Engine/Source/Runtime/Engine/Private/SoundSubmix.cpp
--- a/Engine/Source/Runtime/Engine/Private/SoundSubmix.cpp
+++ b/Engine/Source/Runtime/Engine/Private/SoundSubmix.cpp
@@ -927,6 +927,27 @@ void USoundfieldSubmix::PostEditChangeProperty(struct FPropertyChangedEvent& Pro
 		SanitizeLinks();
 	}
 }
+void UEndpointSubmix::PostEditChangeProperty(struct FPropertyChangedEvent& PropertyChangedEvent)
+{	
+	if (PropertyChangedEvent.Property != nullptr)
+	{
+		const FName PropertyName = PropertyChangedEvent.Property->GetFName();
+		if (PropertyName == GET_MEMBER_NAME_CHECKED(UEndpointSubmix, EndpointType)
+		|| PropertyName == GET_MEMBER_NAME_CHECKED(UEndpointSubmix, EndpointSettings)
+		|| PropertyName == GET_MEMBER_NAME_CHECKED(UEndpointSubmix, EndpointSettingsClass)
+		)
+		{
+			// Remove-re-add submix. Causes reinit with plugins.
+			if (FAudioDeviceManager* AudioDeviceManager = GEngine->GetAudioDeviceManager())
+			{
+				AudioDeviceManager->UnregisterSoundSubmix(this);
+				AudioDeviceManager->RegisterSoundSubmix(this);
+			}
+		}
+	}
+	Super::PostEditChangeProperty(PropertyChangedEvent);
+}
+
 #endif // WITH_EDITOR
 
 void USoundfieldSubmix::PostLoad()
