Make check run in Non-editor builds as its only done on startup.

diff --git a/Engine/Source/Runtime/Engine/Private/SoundSubmix.cpp b/Engine/Source/Runtime/Engine/Private/SoundSubmix.cpp
--- a/Engine/Source/Runtime/Engine/Private/SoundSubmix.cpp
+++ b/Engine/Source/Runtime/Engine/Private/SoundSubmix.cpp
@@ -952,8 +952,6 @@ void UEndpointSubmix::PostEditChangeProperty(struct FPropertyChangedEvent& Prope
 
 void USoundfieldSubmix::PostLoad()
 {
-#if WITH_EDITOR
-
 	// Make sure the Encoding format is something we can use.
 	// Fallback to something that works otherwise and warn.
 
@@ -967,8 +965,6 @@ void USoundfieldSubmix::PostLoad()
 		SanitizeLinks();
 	}
 
-#endif //WITH_EDITOR
-
 	Super::PostLoad();
 }
 
