Fix localization warning about duplicate "GCPROFILINGWARNING".

diff --git a/Engine/Source/Runtime/Engine/Private/UnrealEngine.cpp b/Engine/Source/Runtime/Engine/Private/UnrealEngine.cpp
--- a/Engine/Source/Runtime/Engine/Private/UnrealEngine.cpp
+++ b/Engine/Source/Runtime/Engine/Private/UnrealEngine.cpp
@@ -12287,7 +12287,7 @@ void DrawStatsHUD( UWorld* World, FViewport* Viewport, FCanvas* Canvas, UCanvas*
 		if (GRHIGlobals.IsDebugLayerEnabled)
 		{
 			SmallTextItem.SetColor(FLinearColor::Red);
-			SmallTextItem.Text = LOCTEXT("GCPROFILINGWARNING", "PROFILING WITH DEBUG DEVICE ON!");
+			SmallTextItem.Text = LOCTEXT("GRAPHICSDEBUGDEVICEWARNING", "PROFILING WITH DEBUG DEVICE ON!");
 			Canvas->DrawItem(SmallTextItem, FVector2D(MessageX, MessageY));
 			MessageY += FontSizeY;
 			SmallTextItem.SetColor(FLinearColor::White);
