Fix shipping build break

#rnx

diff --git a/Engine/Source/Runtime/Engine/Private/Tests/AutomationCommon.cpp b/Engine/Source/Runtime/Engine/Private/Tests/AutomationCommon.cpp
--- a/Engine/Source/Runtime/Engine/Private/Tests/AutomationCommon.cpp
+++ b/Engine/Source/Runtime/Engine/Private/Tests/AutomationCommon.cpp
@@ -38,12 +38,14 @@ static TAutoConsoleVariable<int32> CVarAutomationAllowFrameTraceCapture(
 
 //declare static variable
 FOnEditorAutomationMapLoad AutomationCommon::OnEditorAutomationMapLoad;
+#endif
 
 ///////////////////////////////////////////////////////////////////////
 // Common Latent commands
 
 namespace AutomationCommon
 {
+#if WITH_AUTOMATION_TESTS
 	FString GetRenderDetailsString()
 	{
 		FString HardwareDetailsString;
@@ -87,8 +89,6 @@ namespace AutomationCommon
 		return HardwareDetailsString;
 	}
 
-#if WITH_AUTOMATION_TESTS
-
 	/** Gets a path used for automation testing (PNG sent to the AutomationTest folder) */
 	FString GetScreenshotName(const FString& TestName)
 	{
@@ -313,8 +313,6 @@ namespace AutomationCommon
 		bool TaskCompleted;
 	};
 
-#endif
-
 	/** These save a PNG and get sent over the network */
 	static void SaveWindowAsScreenshot(TSharedRef<SWindow> Window, const FString& ScreenshotName)
 	{
@@ -349,6 +347,7 @@ namespace AutomationCommon
 
 		return TestWorld;
 	}
+#endif
 
 	UGameViewportClient* GetAnyGameViewportClient()
 	{
@@ -370,6 +369,7 @@ namespace AutomationCommon
 	}
 }
 
+#if WITH_AUTOMATION_TESTS
 bool AutomationOpenMap(const FString& MapName, bool bForceReload)
 {
 	bool bCanProceed = true;
