Fix Rewind Debugger bugs with incorrect track names

diff --git a/Engine/Source/Runtime/Engine/Private/Animation/AnimTrace.cpp b/Engine/Source/Runtime/Engine/Private/Animation/AnimTrace.cpp
--- a/Engine/Source/Runtime/Engine/Private/Animation/AnimTrace.cpp
+++ b/Engine/Source/Runtime/Engine/Private/Animation/AnimTrace.cpp
@@ -264,6 +264,7 @@ FCriticalSection GAnimTraceNameCriticalSection;
 void FAnimTrace::Reset()
 {
 	GSkeletalMeshTraceAnnotations.ClearAll();
+	GAnimTraceNames.Reset();
 	GAnimTraceCurrentNameId = 1;
 }
 
