Don't override max influences computed, otherwise alt skin weights may fail to upload.

#rnx

diff --git a/Engine/Source/Runtime/Engine/Private/SkeletalMeshLODRenderData.cpp b/Engine/Source/Runtime/Engine/Private/SkeletalMeshLODRenderData.cpp
--- a/Engine/Source/Runtime/Engine/Private/SkeletalMeshLODRenderData.cpp
+++ b/Engine/Source/Runtime/Engine/Private/SkeletalMeshLODRenderData.cpp
@@ -415,14 +415,12 @@ void FSkeletalMeshLODRenderData::BuildFromLODModel(
 		}
 	}
 
-	int32 MaxBoneInfluences = InLODModel->GetMaxBoneInfluences();
 	if (!bUseHighPrecisionWeights)
 	{
 		// Re-normalize the weights for 8-bits to ensure that they are distributed using the old algorithm
 		// from MeshUtilities.cpp
 		uint8 InfluenceWeights[MAX_TOTAL_INFLUENCES] = {0};
 
-		MaxBoneInfluences = 0;
 		for (FSoftSkinVertex& Vertex: Vertices)
 		{
 			uint32	TotalInfluenceWeight = 0;
@@ -438,8 +436,6 @@ void FSkeletalMeshLODRenderData::BuildFromLODModel(
 			}
 			InfluenceWeights[0] += 255 - TotalInfluenceWeight;
 
-			MaxBoneInfluences = FMath::Max(MaxBoneInfluences, MaxInfluenceIndex);
-
 			FMemory::Memzero(Vertex.InfluenceWeights);
 			for (int32 Index = 0; Index < MaxInfluenceIndex; Index++)
 			{
@@ -451,7 +447,7 @@ void FSkeletalMeshLODRenderData::BuildFromLODModel(
 
 	// Init skin weight buffer
 	SkinWeightVertexBuffer.SetNeedsCPUAccess(true);
-	SkinWeightVertexBuffer.SetMaxBoneInfluences(MaxBoneInfluences);
+	SkinWeightVertexBuffer.SetMaxBoneInfluences(InLODModel->GetMaxBoneInfluences());
 	SkinWeightVertexBuffer.SetUse16BitBoneIndex(InLODModel->DoSectionsUse16BitBoneIndex());
 	SkinWeightVertexBuffer.SetUse16BitBoneWeight(bUseHighPrecisionWeights);
 	SkinWeightVertexBuffer.Init(Vertices);
