Fixed: Errors in Cook ShooterGame Android

diff --git a/Engine/Source/Runtime/Engine/Private/ShaderCompiler/ShaderGenerationUtil.cpp b/Engine/Source/Runtime/Engine/Private/ShaderCompiler/ShaderGenerationUtil.cpp
--- a/Engine/Source/Runtime/Engine/Private/ShaderCompiler/ShaderGenerationUtil.cpp
+++ b/Engine/Source/Runtime/Engine/Private/ShaderCompiler/ShaderGenerationUtil.cpp
@@ -687,6 +687,9 @@ static int32 GetTargetNumChannels(EGBufferType Type)
 	case GBT_Float_16_16_16_16:
 		Ret = 4;
 		break;
+	case GBT_Float_32:
+		Ret = 1;
+		break;
 	default:
 		check(0);
 		break;
