World Partition: assert if actors are unexpectedly loaded when moving actors to streaming cell.

diff --git a/Engine/Source/Runtime/Engine/Private/WorldPartition/WorldPartitionLevelHelper.cpp b/Engine/Source/Runtime/Engine/Private/WorldPartition/WorldPartitionLevelHelper.cpp
--- a/Engine/Source/Runtime/Engine/Private/WorldPartition/WorldPartitionLevelHelper.cpp
+++ b/Engine/Source/Runtime/Engine/Private/WorldPartition/WorldPartitionLevelHelper.cpp
@@ -158,7 +158,17 @@ void FWorldPartitionLevelHelper::MoveExternalActorsToLevel(const TArray<FWorldPa
 
 	check(InLevel);
 	UPackage* LevelPackage = InLevel->GetPackage();
-	
+
+	// Gather existing actors to validate only the one we expect are added to the level
+	TSet<FName> LevelActors;
+	for (AActor* Actor : InLevel->Actors)
+	{
+		if (Actor)
+		{
+			LevelActors.Add(Actor->GetFName());
+		}
+	}
+
 	// Move all actors to Cell level
 	for (const FWorldPartitionRuntimeCellObjectMapping& PackageObjectMapping : InChildPackages)
 	{
@@ -203,12 +213,21 @@ void FWorldPartitionLevelHelper::MoveExternalActorsToLevel(const TArray<FWorldPa
 			}
 
 			OutModifiedPackages.Add(ActorExternalPackage);
+			LevelActors.Add(Actor->GetFName());
 		}
 		else
 		{
 			UE_LOG(LogWorldPartition, Warning, TEXT("Can't find actor %s."), *PackageObjectMapping.Path.ToString());
 		}
 	}
+
+	for (AActor* Actor : InLevel->Actors)
+	{
+		if (Actor && Actor->HasAllFlags(RF_WasLoaded))
+		{
+			checkf(LevelActors.Contains(Actor->GetFName()), TEXT("Actor %s(%s) was unexpectedly loaded when moving actors to streaming cell"), *Actor->GetActorNameOrLabel(), *Actor->GetName());
+		}
+	}
 }
 
 void FWorldPartitionLevelHelper::RemapLevelSoftObjectPaths(ULevel* InLevel, UWorldPartition* InWorldPartition)
@@ -584,4 +603,4 @@ bool FWorldPartitionLevelHelper::LoadActors(UWorld* InOuterWorld, ULevel* InDest
 	return (LoadProgress->NumPendingLoadRequests == 0);
 }
 
-#endif
+#endif
\ No newline at end of file
