Turn shader job submission into a Task to move this expensive work off the critical-path on the GameThread.

- This saves ~4.5% of the wall clock time of a CitySample/Lyra cook
- This change provides the base for moving other expensive or blocking work off the game thread:
    - Shader preprocessing.
    - Remote per-shader DDC requests
- Job->GetInputHash() is the most expensive function and offloading this work has been move to a async task.
- The locks guard access to the following members of the FShaderCompileJobCollection:
    - JobsInFlight, DuplicateJobsWaitList, GShaderCompilerStats, PendingJobs and the call to ProcessFinishedJob(...)
- I've opted for a simple design here.  I'm simply moved all the work happening per-job in the job loop in FShaderCompileJobCollection::SubmitJobs into a monolithic Task.
- Since these tasks are fire and forget, SubmitJobs is no longer keeping track of statistics for the number of jobs submitted or found in the cache.  I've removed this bookkeeping and remove the verbose log messages associated with them.
- I've remove the ExistingJobs optimization to searching the link list.  Each task doesn't know about the other tasks so there is no clean way to do this.  It was simpler to remove the optimization and just offload the work to an async thread.

diff --git a/Engine/Source/Runtime/Engine/Private/ShaderCompiler/ShaderCompiler.cpp b/Engine/Source/Runtime/Engine/Private/ShaderCompiler/ShaderCompiler.cpp
--- a/Engine/Source/Runtime/Engine/Private/ShaderCompiler/ShaderCompiler.cpp
+++ b/Engine/Source/Runtime/Engine/Private/ShaderCompiler/ShaderCompiler.cpp
@@ -573,23 +573,12 @@ void FShaderCompileJobCollection::SubmitJobs(const TArray<FShaderCommonCompileJo
 		// we may fulfill some of the jobs from the cache (and we will be subtracting them)
 		NumOutstandingJobs.Add(InJobs.Num());
 
-		int32 SubmittedJobsCount = 0;
-		int32 NumSubmittedJobs[NumShaderCompileJobPriorities] = { 0 };
+		for (FShaderCommonCompileJob* Job : InJobs)
 		{
-			// Just precompute the InputHash for each job in multiple-thread.
-			if (ShaderCompiler::IsJobCacheEnabled())
+			UE::Tasks::Launch(UE_SOURCE_LOCATION, [Job, this]()
 			{
-				TRACE_CPUPROFILER_EVENT_SCOPE(ShaderCompiler.GetInputHash);
-				ParallelFor( TEXT("ShaderCompiler.GetInputHash.PF"), InJobs.Num(),1, [&InJobs](int32 Index) { InJobs[Index]->GetInputHash(); }, EParallelForFlags::Unbalanced);
-			}
-
-			FWriteScopeLock Locker(Lock);
-
-			// Optimization: Only search the linked list once to prevent O(n^2) behavior
-			FShaderCommonCompileJob* ExistingJobs[NumShaderCompileJobPriorities] = { nullptr };
+				TRACE_CPUPROFILER_EVENT_SCOPE(ShaderJobTask);
 
-			for (FShaderCommonCompileJob* Job : InJobs)
-			{
 				check(Job->JobIndex != INDEX_NONE);
 				check(Job->Priority != EShaderCompileJobPriority::None);
 				check(Job->PendingPriority == EShaderCompileJobPriority::None);
@@ -603,68 +592,81 @@ void FShaderCompileJobCollection::SubmitJobs(const TArray<FShaderCommonCompileJo
 
 					const bool bCheckDDC = !(Job->bIsDefaultMaterial || Job->bIsGlobalShader);
 
-					// see if we can find the job in the cache first
-					if (FSharedBuffer* ExistingOutput = CompletedJobsCache.Find(InputHash, bCheckDDC))
+					FSharedBuffer* ExistingOutput = nullptr;
+					{
+						FReadScopeLock ReadLocker(Lock);
+
+						// see if we can find the job in the cache first
+						ExistingOutput = CompletedJobsCache.Find(InputHash, bCheckDDC);
+					}
+				
+					if (ExistingOutput)
 					{
 						UE_LOG(LogShaderCompilers, UE_SHADERCACHE_LOG_LEVEL, TEXT("There is already a cached job with the ihash %s, processing the new one immediately."), *LexToString(InputHash));
 						FMemoryReaderView MemReader(*ExistingOutput);
 						Job->SerializeOutput(MemReader);
 
-						// finish the job instantly
-						ProcessFinishedJob(Job, true);
+						{
+							FWriteScopeLock WriteLocker(Lock);
+							ProcessFinishedJob(Job, true);
+						}
 
-						continue;
+						bNewJob = false;
 					}
-					// see if another job with the same input hash is being worked on
-					else if (FShaderCommonCompileJob** DuplicateInFlight = JobsInFlight.Find(InputHash))
+					else 
 					{
-						UE_LOG(LogShaderCompilers, UE_SHADERCACHE_LOG_LEVEL, TEXT("There is an outstanding job with the ihash %s, not submitting another one (adding to wait list)."), *LexToString(InputHash));
+						FWriteScopeLock Locker(Lock);
+
+						FShaderCommonCompileJob** DuplicateInFlight = nullptr;
+						// see if another job with the same input hash is being worked on
+						DuplicateInFlight = JobsInFlight.Find(InputHash);
 
-						// because of the cloned jobs, we need to maintain a separate mapping
-						FShaderCommonCompileJob** WaitListHead = DuplicateJobsWaitList.Find(InputHash);
-						if (WaitListHead)
+						if (DuplicateInFlight)
 						{
-							Job->LinkAfter(*WaitListHead);
+							UE_LOG(LogShaderCompilers, UE_SHADERCACHE_LOG_LEVEL, TEXT("There is an outstanding job with the ihash %s, not submitting another one (adding to wait list)."), *LexToString(InputHash));
+
+							// because of the cloned jobs, we need to maintain a separate mapping
+							FShaderCommonCompileJob** WaitListHead = nullptr;
+							{
+								WaitListHead = DuplicateJobsWaitList.Find(InputHash);
+							}
+							
+							if (WaitListHead)
+							{
+								Job->LinkAfter(*WaitListHead);
+							}
+							else
+							{
+								DuplicateJobsWaitList.Add(InputHash, Job);
+							}
+							bNewJob = false;
 						}
 						else
 						{
-							DuplicateJobsWaitList.Add(InputHash, Job);
+							// track new jobs so we can de-dupe them
+							JobsInFlight.Add(InputHash, Job);
 						}
-						bNewJob = false;
-					}
-					else
-					{
-						// track new jobs so we can dedupe them
-						JobsInFlight.Add(InputHash, Job);
 					}
 				}
 
 				// new job
 				if (bNewJob)
 				{
+					FWriteScopeLock Locker(Lock);
+
 					GShaderCompilerStats->RegisterNewPendingJob(*Job);
 					ensure(!ShaderCompiler::IsJobCacheEnabled() || Job->bInputHashSet);
 #if UE_SHADERCOMPILER_FIFO_JOB_EXECUTION
 					// link the job at the end of pending list, so we're executing them in a FIFO and not LIFO order
 					if (PendingJobs[PriorityIndex])
 					{
-						// Optimization: Only search the linked list if we have to
-						if (ExistingJobs[PriorityIndex])
+						for (FShaderCommonCompileJob::TIterator It(PendingJobs[PriorityIndex]); It; ++It)
 						{
-							Job->LinkAfter(ExistingJobs[PriorityIndex]);
-							ExistingJobs[PriorityIndex] = Job;
-						}
-						else
-						{
-							for (FShaderCommonCompileJob::TIterator It(PendingJobs[PriorityIndex]); It; ++It)
+							FShaderCommonCompileJob& ExistingJob = *It;
+							if (ExistingJob.GetNextLink() == nullptr)
 							{
-								FShaderCommonCompileJob& ExistingJob = *It;
-								if (ExistingJob.GetNextLink() == nullptr)
-								{
-									Job->LinkAfter(&ExistingJob);
-									ExistingJobs[PriorityIndex] = Job;
-									break;
-								}
+								Job->LinkAfter(&ExistingJob);
+								break;
 							}
 						}
 					}
@@ -673,26 +675,11 @@ void FShaderCompileJobCollection::SubmitJobs(const TArray<FShaderCommonCompileJo
 					{
 						Job->LinkHead(PendingJobs[PriorityIndex]);
 					}
-					
 
 					NumPendingJobs[PriorityIndex]++;
-					NumSubmittedJobs[PriorityIndex]++;
 					Job->PendingPriority = Job->Priority;
-					++SubmittedJobsCount;
 				}
-			}
-		}
-
-		UE_LOG(LogShaderCompilers, UE_SHADERCACHE_LOG_LEVEL, TEXT("Actual jobs submitted %d (of %d new), total outstanding jobs: %d."), SubmittedJobsCount, InJobs.Num(), NumOutstandingJobs.GetValue());
-
-		for (int32 PriorityIndex = 0; PriorityIndex < NumShaderCompileJobPriorities; ++PriorityIndex)
-		{
-			if (NumSubmittedJobs[PriorityIndex] > 0)
-			{
-				UE_LOG(LogShaderCompilers, Verbose, TEXT("Submitted %d shader compile jobs with '%s' priority"),
-					NumSubmittedJobs[PriorityIndex],
-					ShaderCompileJobPriorityToString((EShaderCompileJobPriority)PriorityIndex));
-			}
+			});
 		}
 	}
 }
