Moved AActors CDO components utility methods previously added to UBlueprint, to AActor:
* GetActorClassDefaultComponents()
* GetActorClassDefaultComponent()
* GetActorClassDefaultComponentByName()
* ForEachComponentOfActorClassDefault()

diff --git a/Engine/Plugins/Experimental/PCG/Source/PCG/Private/Elements/PCGSpawnActor.cpp b/Engine/Plugins/Experimental/PCG/Source/PCG/Private/Elements/PCGSpawnActor.cpp
--- a/Engine/Plugins/Experimental/PCG/Source/PCG/Private/Elements/PCGSpawnActor.cpp
+++ b/Engine/Plugins/Experimental/PCG/Source/PCG/Private/Elements/PCGSpawnActor.cpp
@@ -148,7 +148,7 @@ UPCGGraphInterface* UPCGSpawnActorSettings::GetSubgraphInterface() const
 
 	UPCGGraphInterface* Result = nullptr;
 
-	UBlueprint::ForEachComponentOfActorClassDefault<UPCGComponent>(TemplateActorClass, [&](const UPCGComponent* PCGComponent)
+	AActor::ForEachComponentOfActorClassDefault<UPCGComponent>(TemplateActorClass, [&](const UPCGComponent* PCGComponent)
 	{
 		// If there is no graph, there is no graph instance
 		if (PCGComponent->GetGraph() && PCGComponent->bActivated)
@@ -445,7 +445,7 @@ bool FPCGSpawnActorElement::ExecuteInternal(FPCGContext* Context) const
 
 			TMap<FPCGISMCBuilderParameters, TArray<FTransform>> MeshDescriptorTransforms;
 
-			UBlueprint::ForEachComponentOfActorClassDefault<UStaticMeshComponent>(Settings->TemplateActorClass, [&](const UStaticMeshComponent* StaticMeshComponent)
+			AActor::ForEachComponentOfActorClassDefault<UStaticMeshComponent>(Settings->TemplateActorClass, [&](const UStaticMeshComponent* StaticMeshComponent)
 			{
 				FPCGISMCBuilderParameters Params;
 				Params.Descriptor.InitFrom(StaticMeshComponent);
diff --git a/Engine/Plugins/Experimental/PCG/Source/PCG/Private/Helpers/PCGActorHelpers.cpp b/Engine/Plugins/Experimental/PCG/Source/PCG/Private/Helpers/PCGActorHelpers.cpp
--- a/Engine/Plugins/Experimental/PCG/Source/PCG/Private/Helpers/PCGActorHelpers.cpp
+++ b/Engine/Plugins/Experimental/PCG/Source/PCG/Private/Helpers/PCGActorHelpers.cpp
@@ -9,7 +9,6 @@
 
 #include "EngineUtils.h"
 #include "Components/HierarchicalInstancedStaticMeshComponent.h"
-#include "Engine/Blueprint.h"
 #include "Engine/StaticMesh.h"
 #include "Materials/MaterialInterface.h"
 
@@ -219,7 +218,7 @@ bool UPCGActorHelpers::DeleteActors(UWorld* World, const TArray<TSoftObjectPtr<A
 void UPCGActorHelpers::GetActorClassDefaultComponents(const TSubclassOf<AActor>& ActorClass, TArray<UActorComponent*>& OutComponents, const TSubclassOf<UActorComponent>& InComponentClass)
 {
 	OutComponents.Reset();
-	UBlueprint::ForEachComponentOfActorClassDefault(ActorClass, InComponentClass, [&OutComponents](const UActorComponent* TemplateComponent)
+	AActor::ForEachComponentOfActorClassDefault(ActorClass, InComponentClass, [&OutComponents](const UActorComponent* TemplateComponent)
 	{
 		OutComponents.Add(const_cast<UActorComponent*>(TemplateComponent));
 		return true;
diff --git a/Engine/Source/Editor/UnrealEd/Private/ThumbnailRendering/BlueprintThumbnailRenderer.cpp b/Engine/Source/Editor/UnrealEd/Private/ThumbnailRendering/BlueprintThumbnailRenderer.cpp
--- a/Engine/Source/Editor/UnrealEd/Private/ThumbnailRendering/BlueprintThumbnailRenderer.cpp
+++ b/Engine/Source/Editor/UnrealEd/Private/ThumbnailRendering/BlueprintThumbnailRenderer.cpp
@@ -6,8 +6,6 @@
 #include "SceneView.h"
 #include "Misc/App.h"
 
-#include "Engine/Blueprint.h"
-
 UBlueprintThumbnailRenderer::UBlueprintThumbnailRenderer(const FObjectInitializer& ObjectInitializer)
 	: Super(ObjectInitializer)
 {
@@ -23,7 +21,7 @@ bool UBlueprintThumbnailRenderer::CanVisualizeAsset(UObject* Object)
 	// Only visualize actor based blueprints
 	if (Blueprint && Blueprint->GeneratedClass && Blueprint->GeneratedClass->IsChildOf(AActor::StaticClass()))
 	{
-		UBlueprint::ForEachComponentOfActorClassDefault<UActorComponent>(TSubclassOf<AActor>(Blueprint->GeneratedClass), [&](const UActorComponent* Component)
+		AActor::ForEachComponentOfActorClassDefault<UActorComponent>(TSubclassOf<AActor>(Blueprint->GeneratedClass), [&](const UActorComponent* Component)
 		{
 			bCanVisualizeAsset = FBlueprintThumbnailScene::IsValidComponentForVisualization(Component);
 			return !bCanVisualizeAsset;
diff --git a/Engine/Source/Runtime/AIModule/Private/AIHelpers.cpp b/Engine/Source/Runtime/AIModule/Private/AIHelpers.cpp
--- a/Engine/Source/Runtime/AIModule/Private/AIHelpers.cpp
+++ b/Engine/Source/Runtime/AIModule/Private/AIHelpers.cpp
@@ -2,8 +2,6 @@
 
 #include "AIHelpers.h"
 
-#include "Engine/Blueprint.h"
-
 namespace FAISystem
 {
 	FVector FindClosestLocation(const FVector& Origin, const TArray<FVector>& Locations)
@@ -118,7 +116,7 @@ TOptional<float> GetYawFromQuaternion(const FQuat& Quaternion)
 void GetActorClassDefaultComponents(const TSubclassOf<AActor>& ActorClass, TArray<UActorComponent*>& OutComponents, const TSubclassOf<UActorComponent>& InComponentClass)
 {
 	OutComponents.Reset();
-	UBlueprint::ForEachComponentOfActorClassDefault(ActorClass, InComponentClass, [&OutComponents](const UActorComponent* TemplateComponent)
+	AActor::ForEachComponentOfActorClassDefault(ActorClass, InComponentClass, [&OutComponents](const UActorComponent* TemplateComponent)
 	{
 		OutComponents.Add(const_cast<UActorComponent*>(TemplateComponent));
 		return true;
diff --git a/Engine/Source/Runtime/Engine/Classes/Engine/Blueprint.h b/Engine/Source/Runtime/Engine/Classes/Engine/Blueprint.h
--- a/Engine/Source/Runtime/Engine/Classes/Engine/Blueprint.h
+++ b/Engine/Source/Runtime/Engine/Classes/Engine/Blueprint.h
@@ -7,9 +7,7 @@
 #include "UObject/Object.h"
 #include "Misc/Guid.h"
 #include "UObject/Class.h"
-#include "Templates/EnableIf.h"
 #include "Templates/SubclassOf.h"
-#include "Templates/UnrealTypeTraits.h"
 #include "Engine/EngineTypes.h"
 #include "EdGraph/EdGraphPin.h"
 #include "Engine/BlueprintCore.h"
@@ -1096,15 +1094,6 @@ public:
 	 */
 	static bool GetBlueprintHierarchyFromClass(const UClass* InClass, TArray<UBlueprintGeneratedClass*>& OutBlueprintParents);
 
-	/**
-	 * Iterate over all BPGCs used to generate this class and its parents, calling InFunc on them. First element is the BPGC used to generate InClass
-	 *
-	 * @param InClass				The class to get the blueprint lineage for
-	 * @param InFunc				Function that will be called for each BPGC. Must return true to continue iteration, or false to stop.
-	 * @return						true if there were no status errors in any of the parent blueprints, otherwise false
-	 */
-	static bool ForEachBlueprintGeneratedClassInHierarchy(const UClass* InClass, TFunctionRef<bool(UBlueprintGeneratedClass*)> InFunc);
-
 private:
 	/**
 	 * Gets an array of all IBlueprintPropertyGuidProviders for this class and its parents.  0th elements is the IBlueprintPropertyGuidProvider for InClass
@@ -1121,85 +1110,6 @@ public:
 	static bool IsBlueprintHierarchyErrorFree(const UClass* InClass);
 #endif
 
-	/**
-	* Fetches all the components of ActorClass's CDO, including the ones added via the BP editor (which AActor.GetComponents fails to do for CDOs).
-	* @param InActorClass		Class of AActor for which we will retrieve all components.
-	* @param InComponentClass	Only retrieve components of this type.
-	* @param OutComponents this is where the found components will end up. Note that the preexisting contents of OutComponents will get overridden.
-	*/
-	static void GetActorClassDefaultComponents(const TSubclassOf<AActor>& InActorClass, const TSubclassOf<UActorComponent>& InComponentClass, TArray<const UActorComponent*>& OutComponents);
-
-	/**
-	* Fetches the first component of ActorClass's CDO which match the requested component class. Will include the components added via the BP editor (which AActor.GetComponents fails to do for CDOs).
-	* @param InActorClass		Class of AActor for which we will retrieve all components.
-	* @param InComponentClass	Only retrieve components of this type.
-	* @param OutComponents this is where the found components will end up. Note that the preexisting contents of OutComponents will get overridden.
-	*/
-	static const UActorComponent* GetActorClassDefaultComponent(const TSubclassOf<AActor>& InActorClass, const TSubclassOf<UActorComponent>& InComponentClass);
-
-	/**
-	* Get the component of ActorClass's CDO that matches the given object name. Will consider all components, including the ones added via the BP editor (which AActor.GetComponents fails to do for CDOs).
-	* @param InActorClass		Class of AActor for which we will search all components.
-	* @param InComponentClass	Only consider components of this type.
-	* @param OutComponents this is where the found components will end up. Note that the preexisting contents of OutComponents will get overridden.
-	*/
-	static const UActorComponent* GetActorClassDefaultComponentByName(const TSubclassOf<AActor>& InActorClass, const TSubclassOf<UActorComponent>& InComponentClass, FName InComponentName);
-
-	/**
-	* Iterate over the components of ActorClass's CDO, including the ones added via the BP editor (which AActor.GetComponents fails to return).
-	* @param InActorClass		Class of AActor for which we will retrieve all components.
-	* @param InComponentClass	Only consider components of this type.
-	* @param InFunc				Code that will be executed for each component. Must return true to continue iteration, or false to stop.
-	*/
-	static void ForEachComponentOfActorClassDefault(const TSubclassOf<AActor>& InActorClass, const TSubclassOf<UActorComponent>& InComponentClass, TFunctionRef<bool(const UActorComponent*)> InFunc);
-
-	/**
-	* Templated version of GetActorClassDefaultComponents()
-	* @see GetActorClassDefaultComponents
-	*/
-	template <typename TComponentClass = UActorComponent, typename = typename TEnableIf<TIsDerivedFrom<TComponentClass, UActorComponent>::IsDerived>::Type>
-	static void GetActorClassDefaultComponents(const TSubclassOf<AActor>& InActorClass, TArray<const TComponentClass*>& OutComponents)
-	{
-		ForEachComponentOfActorClassDefault(InActorClass, TComponentClass::StaticClass(), [&](const UActorComponent* TemplateComponent)
-		{
-			OutComponents.Add(CastChecked<TComponentClass>(TemplateComponent));
-			return true;
-		});
-	}
-
-	/**
-	* Templated version of GetActorClassDefaultComponent()
-	* @see GetActorClassDefaultComponent
-	*/
-	template <typename TComponentClass = UActorComponent, typename = typename TEnableIf<TIsDerivedFrom<TComponentClass, UActorComponent>::IsDerived>::Type>
-	static const TComponentClass* GetActorClassDefaultComponent(const TSubclassOf<AActor>& InActorClass)
-	{
-		return Cast<TComponentClass>(GetActorClassDefaultComponent(InActorClass, TComponentClass::StaticClass()));
-	}
-
-	/**
-	* Templated version of GetActorClassDefaultComponentByName()
-	* @see GetActorClassDefaultComponentByName
-	*/
-	template <typename TComponentClass = UActorComponent, typename = typename TEnableIf<TIsDerivedFrom<TComponentClass, UActorComponent>::IsDerived>::Type>
-	static const TComponentClass* GetActorClassDefaultComponentByName(const TSubclassOf<AActor>& InActorClass, FName InComponentName)
-	{
-		return Cast<TComponentClass>(GetActorClassDefaultComponentByName(InActorClass, TComponentClass::StaticClass(), InComponentName));
-	}
-
-	/**
-	* Templated version of ForEachComponentOfActorClassDefault()
-	* @see ForEachComponentOfActorClassDefault
-	*/
-	template <typename TComponentClass = UActorComponent, typename = typename TEnableIf<TIsDerivedFrom<TComponentClass, UActorComponent>::IsDerived>::Type>
-	static void ForEachComponentOfActorClassDefault(const TSubclassOf<AActor>& InActorClass, TFunctionRef<bool(const TComponentClass*)> InFunc)
-	{
-		ForEachComponentOfActorClassDefault(InActorClass, TComponentClass::StaticClass(), [&](const UActorComponent* TemplateComponent)
-		{
-			return InFunc(CastChecked<TComponentClass>(TemplateComponent));
-		});
-	}
-
 	template<class TFieldType>
 	static FName GetFieldNameFromClassByGuid(const UClass* InClass, const FGuid VarGuid)
 	{
diff --git a/Engine/Source/Runtime/Engine/Classes/Engine/BlueprintGeneratedClass.h b/Engine/Source/Runtime/Engine/Classes/Engine/BlueprintGeneratedClass.h
--- a/Engine/Source/Runtime/Engine/Classes/Engine/BlueprintGeneratedClass.h
+++ b/Engine/Source/Runtime/Engine/Classes/Engine/BlueprintGeneratedClass.h
@@ -737,6 +737,15 @@ public:
 	 */
 	static bool GetGeneratedClassesHierarchy(const UClass* InClass, TArray<const UBlueprintGeneratedClass*>& OutBPGClasses);
 
+	/**
+	 * Iterate over all BPGCs used to generate this class and its parents, calling InFunc on them. First element is the BPGC used to generate InClass
+	 *
+	 * @param InClass				The class to get the blueprint lineage for
+	 * @param InFunc				Function that will be called for each BPGC. Must return true to continue iteration, or false to stop.
+	 * @return						true if there were no status errors in any of the parent blueprints, otherwise false
+	 */
+	static bool ForEachGeneratedClassInHierarchy(const UClass* InClass, TFunctionRef<bool(const UBlueprintGeneratedClass*)> InFunc);
+
 	UInheritableComponentHandler* GetInheritableComponentHandler(const bool bCreateIfNecessary = false);
 
 	/** Find the object in the TemplateObjects array with the supplied name */
diff --git a/Engine/Source/Runtime/Engine/Classes/GameFramework/Actor.h b/Engine/Source/Runtime/Engine/Classes/GameFramework/Actor.h
--- a/Engine/Source/Runtime/Engine/Classes/GameFramework/Actor.h
+++ b/Engine/Source/Runtime/Engine/Classes/GameFramework/Actor.h
@@ -3837,6 +3837,86 @@ public:
 	/** Tells if an object owned by a component has been registered as a replicated subobject of the component */
 	bool IsActorComponentReplicatedSubObjectRegistered(const UActorComponent* OwnerComponent, const UObject* SubObject) const;
 
+
+	/**
+	* Fetches all the components of ActorClass's CDO, including the ones added via the BP editor (which AActor.GetComponents fails to do for CDOs).
+	* @param InActorClass		Class of AActor for which we will retrieve all components.
+	* @param InComponentClass	Only retrieve components of this type.
+	* @param OutComponents this is where the found components will end up. Note that the preexisting contents of OutComponents will get overridden.
+	*/
+	static void GetActorClassDefaultComponents(const TSubclassOf<AActor>& InActorClass, const TSubclassOf<UActorComponent>& InComponentClass, TArray<const UActorComponent*>& OutComponents);
+
+	/**
+	* Fetches the first component of ActorClass's CDO which match the requested component class. Will include the components added via the BP editor (which AActor.GetComponents fails to do for CDOs).
+	* @param InActorClass		Class of AActor for which we will retrieve all components.
+	* @param InComponentClass	Only retrieve components of this type.
+	* @param OutComponents this is where the found components will end up. Note that the preexisting contents of OutComponents will get overridden.
+	*/
+	static const UActorComponent* GetActorClassDefaultComponent(const TSubclassOf<AActor>& InActorClass, const TSubclassOf<UActorComponent>& InComponentClass);
+
+	/**
+	* Get the component of ActorClass's CDO that matches the given object name. Will consider all components, including the ones added via the BP editor (which AActor.GetComponents fails to do for CDOs).
+	* @param InActorClass		Class of AActor for which we will search all components.
+	* @param InComponentClass	Only consider components of this type.
+	* @param OutComponents this is where the found components will end up. Note that the preexisting contents of OutComponents will get overridden.
+	*/
+	static const UActorComponent* GetActorClassDefaultComponentByName(const TSubclassOf<AActor>& InActorClass, const TSubclassOf<UActorComponent>& InComponentClass, FName InComponentName);
+
+	/**
+	* Iterate over the components of ActorClass's CDO, including the ones added via the BP editor (which AActor.GetComponents fails to return).
+	* @param InActorClass		Class of AActor for which we will retrieve all components.
+	* @param InComponentClass	Only consider components of this type.
+	* @param InFunc				Code that will be executed for each component. Must return true to continue iteration, or false to stop.
+	*/
+	static void ForEachComponentOfActorClassDefault(const TSubclassOf<AActor>& InActorClass, const TSubclassOf<UActorComponent>& InComponentClass, TFunctionRef<bool(const UActorComponent*)> InFunc);
+
+	/**
+	* Templated version of GetActorClassDefaultComponents()
+	* @see GetActorClassDefaultComponents
+	*/
+	template <typename TComponentClass = UActorComponent, typename = typename TEnableIf<TIsDerivedFrom<TComponentClass, UActorComponent>::IsDerived>::Type>
+	static void GetActorClassDefaultComponents(const TSubclassOf<AActor>& InActorClass, TArray<const TComponentClass*>& OutComponents)
+	{
+		ForEachComponentOfActorClassDefault(InActorClass, TComponentClass::StaticClass(), [&](const UActorComponent* TemplateComponent)
+		{
+			OutComponents.Add(CastChecked<TComponentClass>(TemplateComponent));
+			return true;
+		});
+	}
+
+	/**
+	* Templated version of GetActorClassDefaultComponent()
+	* @see GetActorClassDefaultComponent
+	*/
+	template <typename TComponentClass = UActorComponent, typename = typename TEnableIf<TIsDerivedFrom<TComponentClass, UActorComponent>::IsDerived>::Type>
+	static const TComponentClass* GetActorClassDefaultComponent(const TSubclassOf<AActor>& InActorClass)
+	{
+		return Cast<TComponentClass>(GetActorClassDefaultComponent(InActorClass, TComponentClass::StaticClass()));
+	}
+
+	/**
+	* Templated version of GetActorClassDefaultComponentByName()
+	* @see GetActorClassDefaultComponentByName
+	*/
+	template <typename TComponentClass = UActorComponent, typename = typename TEnableIf<TIsDerivedFrom<TComponentClass, UActorComponent>::IsDerived>::Type>
+	static const TComponentClass* GetActorClassDefaultComponentByName(const TSubclassOf<AActor>& InActorClass, FName InComponentName)
+	{
+		return Cast<TComponentClass>(GetActorClassDefaultComponentByName(InActorClass, TComponentClass::StaticClass(), InComponentName));
+	}
+
+	/**
+	* Templated version of ForEachComponentOfActorClassDefault()
+	* @see ForEachComponentOfActorClassDefault
+	*/
+	template <typename TComponentClass = UActorComponent, typename = typename TEnableIf<TIsDerivedFrom<TComponentClass, UActorComponent>::IsDerived>::Type>
+	static void ForEachComponentOfActorClassDefault(const TSubclassOf<AActor>& InActorClass, TFunctionRef<bool(const TComponentClass*)> InFunc)
+	{
+		ForEachComponentOfActorClassDefault(InActorClass, TComponentClass::StaticClass(), [&](const UActorComponent* TemplateComponent)
+		{
+			return InFunc(CastChecked<TComponentClass>(TemplateComponent));
+		});
+	}
+
 private:
 	/** Collection of SubObjects that get replicated when this actor gets replicated. */
 	UE::Net::FSubObjectRegistry ReplicatedSubObjects;
diff --git a/Engine/Source/Runtime/Engine/Private/Actor.cpp b/Engine/Source/Runtime/Engine/Private/Actor.cpp
--- a/Engine/Source/Runtime/Engine/Private/Actor.cpp
+++ b/Engine/Source/Runtime/Engine/Private/Actor.cpp
@@ -55,6 +55,7 @@
 #include "PrimitiveSceneProxy.h"
 #include "UObject/MetaData.h"
 #include "Engine/DamageEvents.h"
+#include "Engine/SCS_Node.h"
 
 #if WITH_EDITOR
 #include "FoliageHelper.h"
@@ -6304,5 +6305,112 @@ bool AActor::UseWorldPartitionRuntimeCellDataLayers() const
 // DataLayers (end)
 //---------------------------------------------------------------------------
 
+
+void AActor::GetActorClassDefaultComponents(const TSubclassOf<AActor>& InActorClass, const TSubclassOf<UActorComponent>& InComponentClass, TArray<const UActorComponent*>& OutComponents)
+{
+	OutComponents.Reset();
+
+	ForEachComponentOfActorClassDefault(InActorClass, InComponentClass, [&](const UActorComponent* TemplateComponent)
+	{
+		OutComponents.Add(TemplateComponent);
+		return true;
+	});
+}
+
+const UActorComponent* AActor::GetActorClassDefaultComponentByName(const TSubclassOf<AActor>& InActorClass, const TSubclassOf<UActorComponent>& InComponentClass, FName InComponentName)
+{
+	const UActorComponent* Result = nullptr;
+
+	ForEachComponentOfActorClassDefault(InActorClass, InComponentClass, [&Result, InComponentName](const UActorComponent* TemplateComponent)
+	{
+		// Try to strip suffix used to identify template component instances
+		FString StrippedName = TemplateComponent->GetName();
+		if (StrippedName.RemoveFromEnd(UActorComponent::ComponentTemplateNameSuffix))
+		{
+			if (StrippedName == InComponentName.ToString())
+			{
+				Result = TemplateComponent;
+				return false;
+			}
+		}
+		else if (TemplateComponent->GetFName() == InComponentName)
+		{
+			Result = TemplateComponent;
+			return false;
+		}
+
+		return true;
+	});
+
+	return Result;
+}
+
+const UActorComponent* AActor::GetActorClassDefaultComponent(const TSubclassOf<AActor>& InActorClass, const TSubclassOf<UActorComponent>& InComponentClass)
+{
+	const UActorComponent* Result = nullptr;
+
+	ForEachComponentOfActorClassDefault(InActorClass, InComponentClass, [&Result](const UActorComponent* TemplateComponent)
+	{
+		Result = TemplateComponent;
+		return false;
+	});
+
+	return Result;
+}
+
+void AActor::ForEachComponentOfActorClassDefault(const TSubclassOf<AActor>& ActorClass, const TSubclassOf<UActorComponent>& InComponentClass, TFunctionRef<bool(const UActorComponent*)> InFunc)
+{
+	if (!ActorClass.Get())
+	{
+		return;
+	}
+
+	auto FilterFunc = [&](const UActorComponent* TemplateComponent)
+	{
+		if (!TemplateComponent)
+		{
+			return true;
+		}
+		
+		if (!InComponentClass.Get() || TemplateComponent->IsA(InComponentClass))
+		{
+			return InFunc(TemplateComponent);
+		}
+
+		return true;
+	};
+
+	// Process native components
+	const AActor* CDO = ActorClass->GetDefaultObject<AActor>();
+	for (const UActorComponent* Component : CDO->GetComponents())
+	{
+		if (!FilterFunc(Component))
+		{
+			return;
+		}
+	}
+
+	// Process blueprint components
+	if (UBlueprintGeneratedClass* ActorBlueprintGeneratedClass = Cast<UBlueprintGeneratedClass>(ActorClass))
+	{
+		UBlueprintGeneratedClass::ForEachGeneratedClassInHierarchy(ActorClass, [&](const UBlueprintGeneratedClass* CurrentBPGC)
+		{
+			if (const USimpleConstructionScript* const ConstructionScript = CurrentBPGC->SimpleConstructionScript)
+			{
+				// Gets all BP added components
+				for (const USCS_Node* const Node : ConstructionScript->GetAllNodes())
+				{
+					if (!FilterFunc(Node->GetActualComponentTemplate(ActorBlueprintGeneratedClass)))
+					{
+						return false;
+					}
+				}
+			}
+			return true;
+		});
+	}
+}
+
+
 #undef LOCTEXT_NAMESPACE
 
diff --git a/Engine/Source/Runtime/Engine/Private/Blueprint.cpp b/Engine/Source/Runtime/Engine/Private/Blueprint.cpp
--- a/Engine/Source/Runtime/Engine/Private/Blueprint.cpp
+++ b/Engine/Source/Runtime/Engine/Private/Blueprint.cpp
@@ -13,12 +13,10 @@
 #include "Components/TimelineComponent.h"
 #include "Modules/ModuleManager.h"
 #include "UObject/TextProperty.h"
-#include "GameFramework/Actor.h"
-#include "Engine/SimpleConstructionScript.h"
-#include "Engine/SCS_Node.h"
 
 #if WITH_EDITOR
 #include "BlueprintCompilationManager.h"
+#include "Engine/SCS_Node.h"
 #include "Kismet2/BlueprintEditorUtils.h"
 #include "Kismet2/ComponentEditorUtils.h"
 #include "Kismet2/KismetEditorUtilities.h"
@@ -1632,41 +1630,6 @@ bool UBlueprint::GetBlueprintHierarchyFromClass(const UClass* InClass, TArray<UB
 	return bNoErrors;
 }
 
-bool UBlueprint::ForEachBlueprintGeneratedClassInHierarchy(const UClass* InClass, TFunctionRef<bool(UBlueprintGeneratedClass*)> InFunc)
-{
-	bool bNoErrors = true;
-	UBlueprintGeneratedClass* CurrentClass = Cast<UBlueprintGeneratedClass>(const_cast<UClass*>(InClass));
-	while (CurrentClass)
-	{
-		if (!InFunc(CurrentClass))
-		{
-			return bNoErrors;
-		}
-
-#if WITH_EDITORONLY_DATA
-		UBlueprint* BP = UBlueprint::GetBlueprintFromClass(CurrentClass);
-
-		if (BP)
-		{
-			bNoErrors &= (BP->Status != BS_Error);
-		}
-
-		// If valid, use stored ParentClass rather than the actual UClass::GetSuperClass(); handles the case when the class has not been recompiled yet after a reparent operation.
-		if (BP && BP->ParentClass)
-		{
-			CurrentClass = Cast<UBlueprintGeneratedClass>(BP->ParentClass);
-		}
-		else
-#endif // #if WITH_EDITORONLY_DATA
-		{
-			check(CurrentClass);
-			CurrentClass = Cast<UBlueprintGeneratedClass>(CurrentClass->GetSuperClass());
-		}
-	}
-
-	return bNoErrors;
-}
-
 bool UBlueprint::GetBlueprintHierarchyFromClass(const UClass* InClass, TArray<IBlueprintPropertyGuidProvider*>& OutBlueprintParents)
 {
 	OutBlueprintParents.Reset();
@@ -1732,111 +1695,6 @@ bool UBlueprint::IsBlueprintHierarchyErrorFree(const UClass* InClass)
 }
 #endif
 
-void UBlueprint::GetActorClassDefaultComponents(const TSubclassOf<AActor>& InActorClass, const TSubclassOf<UActorComponent>& InComponentClass, TArray<const UActorComponent*>& OutComponents)
-{
-	OutComponents.Reset();
-
-	ForEachComponentOfActorClassDefault(InActorClass, InComponentClass, [&](const UActorComponent* TemplateComponent)
-	{
-		OutComponents.Add(TemplateComponent);
-		return true;
-	});
-}
-
-const UActorComponent* UBlueprint::GetActorClassDefaultComponentByName(const TSubclassOf<AActor>& InActorClass, const TSubclassOf<UActorComponent>& InComponentClass, FName InComponentName)
-{
-	const UActorComponent* Result = nullptr;
-
-	ForEachComponentOfActorClassDefault(InActorClass, InComponentClass, [&Result, InComponentName](const UActorComponent* TemplateComponent)
-	{
-		// Try to strip suffix used to identify template component instances
-		FString StrippedName = TemplateComponent->GetName();
-		if (StrippedName.RemoveFromEnd(UActorComponent::ComponentTemplateNameSuffix))
-		{
-			if (StrippedName == InComponentName.ToString())
-			{
-				Result = TemplateComponent;
-				return false;
-			}
-		}
-		else if (TemplateComponent->GetFName() == InComponentName)
-		{
-			Result = TemplateComponent;
-			return false;
-		}
-
-		return true;
-	});
-
-	return Result;
-}
-
-const UActorComponent* UBlueprint::GetActorClassDefaultComponent(const TSubclassOf<AActor>& InActorClass, const TSubclassOf<UActorComponent>& InComponentClass)
-{
-	const UActorComponent* Result = nullptr;
-
-	ForEachComponentOfActorClassDefault(InActorClass, InComponentClass, [&Result](const UActorComponent* TemplateComponent)
-	{
-		Result = TemplateComponent;
-		return false;
-	});
-
-	return Result;
-}
-
-void UBlueprint::ForEachComponentOfActorClassDefault(const TSubclassOf<AActor>& ActorClass, const TSubclassOf<UActorComponent>& InComponentClass, TFunctionRef<bool(const UActorComponent*)> InFunc)
-{
-	if (!ActorClass.Get())
-	{
-		return;
-	}
-
-	auto FilterFunc = [&](const UActorComponent* TemplateComponent)
-	{
-		if (!TemplateComponent)
-		{
-			return true;
-		}
-
-		if (!InComponentClass.Get() || TemplateComponent->IsA(InComponentClass))
-		{
-			return InFunc(TemplateComponent);
-		}
-
-		return true;
-	};
-
-	// Process native components
-	const AActor* CDO = ActorClass->GetDefaultObject<AActor>();
-	for (const UActorComponent* Component : CDO->GetComponents())
-	{
-		if (!FilterFunc(Component))
-		{
-			return;
-		}
-	}
-
-	// Process blueprint components
-	if (UBlueprintGeneratedClass* ActorBlueprintGeneratedClass = Cast<UBlueprintGeneratedClass>(ActorClass))
-	{
-		ForEachBlueprintGeneratedClassInHierarchy(ActorClass, [&](UBlueprintGeneratedClass* CurrentBPGC)
-		{
-			if (const USimpleConstructionScript* const ConstructionScript = CurrentBPGC->SimpleConstructionScript)
-			{
-				// Gets all BP added components
-				for (const USCS_Node* const Node : ConstructionScript->GetAllNodes())
-				{
-					if (!FilterFunc(Node->GetActualComponentTemplate(ActorBlueprintGeneratedClass)))
-					{
-						return false;
-					}
-				}
-			}
-			return true;
-		});
-	}
-}
-
 FName UBlueprint::FindBlueprintPropertyNameFromGuid(const FGuid& PropertyGuid) const
 {
 #if WITH_EDITORONLY_DATA
diff --git a/Engine/Source/Runtime/Engine/Private/BlueprintGeneratedClass.cpp b/Engine/Source/Runtime/Engine/Private/BlueprintGeneratedClass.cpp
--- a/Engine/Source/Runtime/Engine/Private/BlueprintGeneratedClass.cpp
+++ b/Engine/Source/Runtime/Engine/Private/BlueprintGeneratedClass.cpp
@@ -1267,18 +1267,32 @@ void UBlueprintGeneratedClass::UnbindDynamicDelegatesForProperty(UObject* InInst
 bool UBlueprintGeneratedClass::GetGeneratedClassesHierarchy(const UClass* InClass, TArray<const UBlueprintGeneratedClass*>& OutBPGClasses)
 {
 	OutBPGClasses.Empty();
+
+	return ForEachGeneratedClassInHierarchy(InClass, [&OutBPGClasses](const UBlueprintGeneratedClass* BPGClass)
+	{
+		OutBPGClasses.Add(BPGClass);
+		return true;
+	});
+}
+
+bool UBlueprintGeneratedClass::ForEachGeneratedClassInHierarchy(const UClass* InClass, TFunctionRef<bool(const UBlueprintGeneratedClass*)> InFunc)
+{
 	bool bNoErrors = true;
-	while(const UBlueprintGeneratedClass* BPGClass = Cast<const UBlueprintGeneratedClass>(InClass))
+	while (const UBlueprintGeneratedClass* BPGClass = Cast<const UBlueprintGeneratedClass>(InClass))
 	{
+		if (!InFunc(BPGClass))
+		{
+			return bNoErrors;
+		}
+
 #if WITH_EDITORONLY_DATA
 		// A cooked class has already been validated and will not have a source Blueprint asset.
-		if(!BPGClass->bCooked)
+		if (!BPGClass->bCooked)
 		{
 			const UBlueprint* BP = Cast<const UBlueprint>(BPGClass->ClassGeneratedBy);
 			bNoErrors &= (NULL != BP) && (BP->Status != BS_Error);
 		}
 #endif
-		OutBPGClasses.Add(BPGClass);
 		InClass = BPGClass->GetSuperClass();
 	}
 	return bNoErrors;
diff --git a/Engine/Source/Runtime/Engine/Private/WorldPartition/HLOD/HLODTemplatedInstancedStaticMeshComponent.cpp b/Engine/Source/Runtime/Engine/Private/WorldPartition/HLOD/HLODTemplatedInstancedStaticMeshComponent.cpp
--- a/Engine/Source/Runtime/Engine/Private/WorldPartition/HLOD/HLODTemplatedInstancedStaticMeshComponent.cpp
+++ b/Engine/Source/Runtime/Engine/Private/WorldPartition/HLOD/HLODTemplatedInstancedStaticMeshComponent.cpp
@@ -68,7 +68,7 @@ void UHLODTemplatedInstancedStaticMeshComponent::SetTemplateComponentName(const
 
 void UHLODTemplatedInstancedStaticMeshComponent::RestoreAssetsFromActorTemplate()
 {
-	const UStaticMeshComponent* TemplateSMC = UBlueprint::GetActorClassDefaultComponentByName<UStaticMeshComponent>(TemplateActorClass, TemplateComponentName);
+	const UStaticMeshComponent* TemplateSMC = AActor::GetActorClassDefaultComponentByName<UStaticMeshComponent>(TemplateActorClass, TemplateComponentName);
 	if (TemplateSMC)
 	{
 		// StaticMesh
