LineBatchComponent: add a BatchID
Move some Draw primitives functions from DrawDebugHelpers to LineBatchComponent
#rnx

diff --git a/Engine/Source/Runtime/Engine/Classes/Components/LineBatchComponent.h b/Engine/Source/Runtime/Engine/Classes/Components/LineBatchComponent.h
--- a/Engine/Source/Runtime/Engine/Classes/Components/LineBatchComponent.h
+++ b/Engine/Source/Runtime/Engine/Classes/Components/LineBatchComponent.h
@@ -38,6 +38,9 @@ struct FBatchedLine
 	UPROPERTY()
 	uint8 DepthPriority;
 
+	UPROPERTY()
+	uint32 BatchID;
+
 	FBatchedLine()
 		: Start(ForceInit)
 		, End(ForceInit)
@@ -45,14 +48,16 @@ struct FBatchedLine
 		, Thickness(0)
 		, RemainingLifeTime(0)
 		, DepthPriority(0)
+		, BatchID(0)
 	{}
-	FBatchedLine(const FVector& InStart, const FVector& InEnd, const FLinearColor& InColor, float InLifeTime, float InThickness, uint8 InDepthPriority)
+	FBatchedLine(const FVector& InStart, const FVector& InEnd, const FLinearColor& InColor, float InLifeTime, float InThickness, uint8 InDepthPriority, uint32 InBatchID = 0)
 		:	Start(InStart)
 		,	End(InEnd)
 		,	Color(InColor)
 		,	Thickness(InThickness)
 		,	RemainingLifeTime(InLifeTime)
 		,	DepthPriority(InDepthPriority)
+		,	BatchID(InBatchID)
 	{}
 };
 
@@ -75,6 +80,9 @@ struct FBatchedPoint
 
 	UPROPERTY()
 	uint8 DepthPriority;
+	
+	UPROPERTY()
+	uint32 BatchID;
 
 	FBatchedPoint()
 		: Position(ForceInit)
@@ -82,37 +90,37 @@ struct FBatchedPoint
 		, PointSize(0)
 		, RemainingLifeTime(0)
 		, DepthPriority(0)
+		, BatchID(0)
 	{}
-	FBatchedPoint(const FVector& InPosition, const FLinearColor& InColor, float InPointSize, float InLifeTime, uint8 InDepthPriority)
+	FBatchedPoint(const FVector& InPosition, const FLinearColor& InColor, float InPointSize, float InLifeTime, uint8 InDepthPriority, uint32 InBatchID = 0)
 		:	Position(InPosition)
 		,	Color(InColor)
 		,	PointSize(InPointSize)
 		,	RemainingLifeTime(InLifeTime)
 		,	DepthPriority(InDepthPriority)
+		,	BatchID(InBatchID)
 	{}
 };
 
 struct FBatchedMesh
 {
-public:
-	FBatchedMesh()
-		: RemainingLifeTime(0)
-	{};
+	FBatchedMesh() = default;
 
 	/**
 	 * MeshVerts - linear array of world space vertex positions
 	 * MeshIndices - array of indices into MeshVerts.  Each triplet is a tri.  i.e. [0,1,2] is first tri, [3,4,5] is 2nd tri, etc
 	 */
-	FBatchedMesh(TArray<FVector> const& InMeshVerts, TArray<int32> const& InMeshIndices, FColor const& InColor, uint8 InDepthPriority, float LifeTime)
+	FBatchedMesh(TArray<FVector> const& InMeshVerts, TArray<int32> const& InMeshIndices, FColor const& InColor, uint8 InDepthPriority, float LifeTime, uint32 InBatchID = 0)
 		: MeshVerts(InMeshVerts), MeshIndices(InMeshIndices), 
-		  Color(InColor), DepthPriority(InDepthPriority), RemainingLifeTime(LifeTime)
+		  Color(InColor), RemainingLifeTime(LifeTime), DepthPriority(InDepthPriority), BatchID(InBatchID)
 	{}
 
 	TArray<FVector> MeshVerts;
 	TArray<int32> MeshIndices;
 	FColor Color;
-	uint8 DepthPriority;
-	float RemainingLifeTime;
+	float RemainingLifeTime = 0.f;
+	uint8 DepthPriority = 0;
+	uint32 BatchID = 0;
 };
 
 /** 
@@ -125,27 +133,51 @@ class ULineBatchComponent : public UPrimitiveComponent
 	GENERATED_UCLASS_BODY()
 
 	/** Buffer of lines to draw */
-	TArray<struct FBatchedLine> BatchedLines;
+	TArray<FBatchedLine> BatchedLines;
 	/** Buffer or points to draw */
-	TArray<struct FBatchedPoint> BatchedPoints;
+	TArray<FBatchedPoint> BatchedPoints;
 	/** Default time that lines/points will draw for */
 	float DefaultLifeTime;
 	/** Buffer of simple meshes to draw */
-	TArray<struct FBatchedMesh> BatchedMeshes;
+	TArray<FBatchedMesh> BatchedMeshes;
 	/** Whether to calculate a tight accurate bounds (encompassing all points), or use a giant bounds that is fast to compute. */
 	uint32 bCalculateAccurateBounds:1;
 
+	/** Defines the value for an invalid id */
+	static constexpr uint32 INVALID_ID = 0;
+
 	/** Provide many lines to draw - faster than calling DrawLine many times. */
-	ENGINE_API void DrawLines(const TArray<FBatchedLine>& InLines);
+	ENGINE_API void DrawLines(TArrayView<FBatchedLine> InLines);
 
 	/** Draw a box */
-	ENGINE_API void DrawBox(const FBox& Box, const FMatrix& TM, const FColor& Color, uint8 InDepthPriorityGroup);
+	ENGINE_API void DrawBox(const FBox& Box, const FMatrix& TM, FLinearColor Color, uint8 InDepthPriorityGroup, uint32 InBatchID = INVALID_ID);
+
+	/** Draw a box */
+	ENGINE_API void DrawBox(const FVector& Center, const FVector& Box, FLinearColor Color, float LifeTime, uint8 DepthPriority, float Thickness, uint32 InBatchID = INVALID_ID);
+	
+	/** Draw a box */
+	ENGINE_API void DrawBox(const FVector& Center, const FVector& Box, const FQuat& Rotation, FLinearColor Color, float LifeTime, uint8 DepthPriority, float Thickness, uint32 InBatchID = INVALID_ID);
+
+	/** Draw an arrow */
+	ENGINE_API void DrawDirectionalArrow(const FMatrix& ArrowToWorld, FLinearColor InColor, float Length, float ArrowSize, uint8 DepthPriority, uint32 InBatchID = INVALID_ID);
 
 	/** Draw an arrow */
-	ENGINE_API void DrawDirectionalArrow(const FMatrix& ArrowToWorld,FColor InColor,float Length,float ArrowSize,uint8 DepthPriority);
+	ENGINE_API void DrawDirectionalArrow(const FVector& LineStart, const FVector& LineEnd, float ArrowSize, FLinearColor Color, float LifeTime, uint8 DepthPriority, float Thickness, uint32 InBatchID = INVALID_ID);
 
 	/** Draw a circle */
-	ENGINE_API void DrawCircle(const FVector& Base, const FVector& X, const FVector& Y, FColor Color, float Radius, int32 NumSides, uint8 DepthPriority);
+	ENGINE_API void DrawCircle(const FVector& Base, const FVector& X, const FVector& Y, FLinearColor Color, float Radius, int32 NumSides, uint8 DepthPriority, uint32 InBatchID = INVALID_ID);
+
+	/** Draw a sphere */
+	ENGINE_API void DrawSphere(const FVector& Center, float Radius, int32 Segments, FLinearColor Color, float LifeTime, uint8 DepthPriority, float Thickness, uint32 InBatchID = INVALID_ID);
+
+	/** Draw a cylinder */
+	ENGINE_API void DrawCylinder(const FVector& Start, const FVector& End, float Radius, int32 Segments, FLinearColor  Color, float LifeTime, uint8 DepthPriority, float Thickness, uint32 InBatchID = INVALID_ID);
+
+	/** Draw a cone */
+	ENGINE_API void DrawCone(const FVector& Origin, const FVector& Direction, float Length, float AngleWidth, float AngleHeight, int32 NumSides, FLinearColor DrawColor, float LifeTime, uint8 DepthPriority, float Thickness, uint32 InBatchID = INVALID_ID);
+
+	/** Draw a cone */
+	ENGINE_API void DrawCapsule(const FVector& Center, float HalfHeight, float Radius, const FQuat& Rotation, FLinearColor Color, float LifeTime, uint8 DepthPriority, float Thickness, uint32 InBatchID = INVALID_ID);
 
 	ENGINE_API virtual void DrawLine(
 		const FVector& Start,
@@ -153,20 +185,22 @@ class ULineBatchComponent : public UPrimitiveComponent
 		const FLinearColor& Color,
 		uint8 DepthPriority,
 		float Thickness = 0.0f,
-		float LifeTime = 0.0f
+		float LifeTime = 0.0f,
+		uint32 InBatchID = INVALID_ID
 		);
 	ENGINE_API virtual void DrawPoint(
 		const FVector& Position,
 		const FLinearColor& Color,
 		float PointSize,
 		uint8 DepthPriority,
-		float LifeTime = 0.0f
+		float LifeTime = 0.0f,
+		uint32 InBatchID = INVALID_ID
 		);
 
 	/** Draw a box */
-	ENGINE_API void DrawSolidBox(FBox const& Box, FTransform const& Xform, const FColor& Color, uint8 DepthPriority, float LifeTime);
+	ENGINE_API void DrawSolidBox(FBox const& Box, FTransform const& Xform, const FColor& Color, uint8 DepthPriority, float LifeTime, uint32 InBatchID = INVALID_ID);
 	/** Draw a mesh */
-	ENGINE_API void DrawMesh(TArray<FVector> const& Verts, TArray<int32> const& Indices, FColor const& Color, uint8 DepthPriority, float LifeTime);
+	ENGINE_API void DrawMesh(TArray<FVector> const& Verts, TArray<int32> const& Indices, FColor const& Color, uint8 DepthPriority, float LifeTime, uint32 InBatchID = INVALID_ID);
 
 	//~ Begin UPrimitiveComponent Interface.
 	ENGINE_API virtual FPrimitiveSceneProxy* CreateSceneProxy() override;
@@ -181,4 +215,11 @@ class ULineBatchComponent : public UPrimitiveComponent
 
 	/** Clear all batched lines, points and meshes */
 	ENGINE_API void Flush();
+	
+	/** Remove batched lines, points and meshes with given ID */
+	ENGINE_API void ClearBatch(uint32 InBatchID);
+
+protected:
+	void AddHalfCircle(const FVector& Base, const FVector& X, const FVector& Y, const FLinearColor& Color, const float Radius, int32 NumSides, const float LifeTime, uint8 DepthPriority, const float Thickness, const uint32 InBatchID);
+	void AddCircle(const FVector& Base, const FVector& X, const FVector& Y, const FLinearColor& Color, const float Radius, int32 NumSides, const float LifeTime, uint8 DepthPriority, const float Thickness, const uint32 InBatchID);
 };
diff --git a/Engine/Source/Runtime/Engine/Private/Components/LineBatchComponent.cpp b/Engine/Source/Runtime/Engine/Private/Components/LineBatchComponent.cpp
--- a/Engine/Source/Runtime/Engine/Private/Components/LineBatchComponent.cpp
+++ b/Engine/Source/Runtime/Engine/Private/Components/LineBatchComponent.cpp
@@ -137,9 +137,6 @@ uint32 FLineBatcherSceneProxy::GetAllocatedSize( void ) const
 	return( FPrimitiveSceneProxy::GetAllocatedSize() + Lines.GetAllocatedSize() + Points.GetAllocatedSize() + Meshes.GetAllocatedSize() ); 
 }
 
-
-
-
 ULineBatchComponent::ULineBatchComponent(const FObjectInitializer& ObjectInitializer)
 	: Super(ObjectInitializer)
 {
@@ -158,15 +155,15 @@ ULineBatchComponent::ULineBatchComponent(const FObjectInitializer& ObjectInitial
 	bIgnoreStreamingManagerUpdate = true;
 }
 
-void ULineBatchComponent::DrawLine(const FVector& Start, const FVector& End, const FLinearColor& Color, uint8 DepthPriority, const float Thickness, const float LifeTime)
+void ULineBatchComponent::DrawLine(const FVector& Start, const FVector& End, const FLinearColor& Color, uint8 DepthPriority, const float Thickness, const float LifeTime, uint32 InBatchID)
 {
-	new(BatchedLines) FBatchedLine(Start, End, Color, LifeTime, Thickness, DepthPriority);
+	new(BatchedLines) FBatchedLine(Start, End, Color, LifeTime, Thickness, DepthPriority, InBatchID);
 
 	// LineBatcher and PersistentLineBatcher components will be updated at the end of UWorld::Tick
 	MarkRenderStateDirty();
 }
 
-void ULineBatchComponent::DrawLines(const TArray<FBatchedLine>& InLines)
+void ULineBatchComponent::DrawLines(TArrayView<FBatchedLine> InLines)
 {
 	BatchedLines.Append(InLines);
 	
@@ -179,15 +176,16 @@ void ULineBatchComponent::DrawPoint(
 	const FLinearColor& Color,
 	float PointSize,
 	uint8 DepthPriority,
-	float LifeTime
+	float LifeTime,
+	uint32 InBatchID
 	)
 {
-	new(BatchedPoints) FBatchedPoint(Position, Color, PointSize, LifeTime, DepthPriority);
+	new(BatchedPoints) FBatchedPoint(Position, Color, PointSize, LifeTime, DepthPriority, InBatchID);
 	// LineBatcher and PersistentLineBatcher components will be updated at the end of UWorld::Tick
 	MarkRenderStateDirty();
 }
 
-void ULineBatchComponent::DrawBox(const FBox& Box, const FMatrix& TM, const FColor& Color, uint8 InDepthPriorityGroup)
+void ULineBatchComponent::DrawBox(const FBox& Box, const FMatrix& TM, FLinearColor Color, uint8 InDepthPriorityGroup, uint32 InBatchID)
 {
 	FVector	B[2], P, Q;
 	int32 ai, aj;
@@ -199,23 +197,97 @@ void ULineBatchComponent::DrawBox(const FBox& Box, const FMatrix& TM, const FCol
 		P.X=B[ai].X; Q.X=B[ai].X;
 		P.Y=B[aj].Y; Q.Y=B[aj].Y;
 		P.Z=B[0].Z; Q.Z=B[1].Z;
-		new(BatchedLines) FBatchedLine(TM.TransformPosition(P), TM.TransformPosition(Q), Color, DefaultLifeTime, 0.0f, InDepthPriorityGroup);
+		new(BatchedLines) FBatchedLine(TM.TransformPosition(P), TM.TransformPosition(Q), Color, DefaultLifeTime, 0.0f, InDepthPriorityGroup, InBatchID);
 
 		P.Y=B[ai].Y; Q.Y=B[ai].Y;
 		P.Z=B[aj].Z; Q.Z=B[aj].Z;
 		P.X=B[0].X; Q.X=B[1].X;
-		new(BatchedLines) FBatchedLine(TM.TransformPosition(P), TM.TransformPosition(Q), Color, DefaultLifeTime, 0.0f, InDepthPriorityGroup);
+		new(BatchedLines) FBatchedLine(TM.TransformPosition(P), TM.TransformPosition(Q), Color, DefaultLifeTime, 0.0f, InDepthPriorityGroup, InBatchID);
 
 		P.Z=B[ai].Z; Q.Z=B[ai].Z;
 		P.X=B[aj].X; Q.X=B[aj].X;
 		P.Y=B[0].Y; Q.Y=B[1].Y;
-		new(BatchedLines) FBatchedLine(TM.TransformPosition(P), TM.TransformPosition(Q), Color, DefaultLifeTime, 0.0f, InDepthPriorityGroup);
+		new(BatchedLines) FBatchedLine(TM.TransformPosition(P), TM.TransformPosition(Q), Color, DefaultLifeTime, 0.0f, InDepthPriorityGroup, InBatchID);
 	}
 	// LineBatcher and PersistentLineBatcher components will be updated at the end of UWorld::Tick
 	MarkRenderStateDirty();
 }
 
-void ULineBatchComponent::DrawSolidBox(const FBox& Box, const FTransform& Xform, const FColor& Color, uint8 DepthPriority, float LifeTime)
+void ULineBatchComponent::DrawBox(FVector const& Center, FVector const& Box, FLinearColor Color, float LifeTime, uint8 DepthPriority, float Thickness, uint32 InBatchID)
+{
+	new(BatchedLines) FBatchedLine(Center + FVector( Box.X,  Box.Y,  Box.Z), Center + FVector( Box.X, -Box.Y, Box.Z), Color, LifeTime, Thickness, DepthPriority, InBatchID);
+	new(BatchedLines) FBatchedLine(Center + FVector( Box.X, -Box.Y,  Box.Z), Center + FVector(-Box.X, -Box.Y, Box.Z), Color, LifeTime, Thickness, DepthPriority, InBatchID);
+	new(BatchedLines) FBatchedLine(Center + FVector(-Box.X, -Box.Y,  Box.Z), Center + FVector(-Box.X,  Box.Y, Box.Z) ,Color, LifeTime, Thickness, DepthPriority, InBatchID);
+	new(BatchedLines) FBatchedLine(Center + FVector(-Box.X,  Box.Y,  Box.Z), Center + FVector( Box.X,  Box.Y, Box.Z), Color, LifeTime, Thickness, DepthPriority, InBatchID);
+
+	new(BatchedLines) FBatchedLine(Center + FVector( Box.X,  Box.Y, -Box.Z), Center + FVector( Box.X, -Box.Y, -Box.Z), Color, LifeTime, Thickness, DepthPriority, InBatchID);
+	new(BatchedLines) FBatchedLine(Center + FVector( Box.X, -Box.Y, -Box.Z), Center + FVector(-Box.X, -Box.Y, -Box.Z), Color, LifeTime, Thickness, DepthPriority, InBatchID);
+	new(BatchedLines) FBatchedLine(Center + FVector(-Box.X, -Box.Y, -Box.Z), Center + FVector(-Box.X,  Box.Y, -Box.Z), Color, LifeTime, Thickness, DepthPriority, InBatchID);
+	new(BatchedLines) FBatchedLine(Center + FVector(-Box.X,  Box.Y, -Box.Z), Center + FVector( Box.X,  Box.Y, -Box.Z), Color, LifeTime, Thickness, DepthPriority, InBatchID);
+
+	new(BatchedLines) FBatchedLine(Center + FVector( Box.X,  Box.Y,  Box.Z), Center + FVector( Box.X,  Box.Y, -Box.Z), Color, LifeTime, Thickness, DepthPriority, InBatchID);
+	new(BatchedLines) FBatchedLine(Center + FVector( Box.X, -Box.Y,  Box.Z), Center + FVector( Box.X, -Box.Y, -Box.Z), Color, LifeTime, Thickness, DepthPriority, InBatchID);
+	new(BatchedLines) FBatchedLine(Center + FVector(-Box.X, -Box.Y,  Box.Z), Center + FVector(-Box.X, -Box.Y, -Box.Z), Color, LifeTime, Thickness, DepthPriority, InBatchID);
+	new(BatchedLines) FBatchedLine(Center + FVector(-Box.X,  Box.Y,  Box.Z), Center + FVector(-Box.X,  Box.Y, -Box.Z), Color, LifeTime, Thickness, DepthPriority, InBatchID);
+	
+	MarkRenderStateDirty();
+}
+
+void ULineBatchComponent::DrawBox(FVector const& Center, FVector const& Box, const FQuat& Rotation, FLinearColor Color, float LifeTime, uint8 DepthPriority, float Thickness, uint32 InBatchID)
+{
+	FTransform const Transform(Rotation);
+	FVector Start = Transform.TransformPosition(FVector( Box.X,  Box.Y,  Box.Z));
+	FVector End = Transform.TransformPosition(FVector( Box.X, -Box.Y, Box.Z));
+	new(BatchedLines) FBatchedLine(Center + Start, Center + End, Color, LifeTime, Thickness, DepthPriority, InBatchID);
+
+	Start = Transform.TransformPosition(FVector( Box.X, -Box.Y,  Box.Z));
+	End = Transform.TransformPosition(FVector(-Box.X, -Box.Y, Box.Z));
+	new(BatchedLines) FBatchedLine(Center + Start, Center + End, Color, LifeTime, Thickness, DepthPriority, InBatchID);
+
+	Start = Transform.TransformPosition(FVector(-Box.X, -Box.Y,  Box.Z));
+	End = Transform.TransformPosition(FVector(-Box.X,  Box.Y, Box.Z));
+	new(BatchedLines) FBatchedLine(Center + Start, Center + End, Color, LifeTime, Thickness, DepthPriority, InBatchID);
+
+	Start = Transform.TransformPosition(FVector(-Box.X,  Box.Y,  Box.Z));
+	End = Transform.TransformPosition(FVector( Box.X,  Box.Y, Box.Z));
+	new(BatchedLines) FBatchedLine(Center + Start, Center + End, Color, LifeTime, Thickness, DepthPriority, InBatchID);
+
+	Start = Transform.TransformPosition(FVector( Box.X,  Box.Y, -Box.Z));
+	End = Transform.TransformPosition(FVector( Box.X, -Box.Y, -Box.Z));
+	new(BatchedLines) FBatchedLine(Center + Start, Center + End, Color, LifeTime, Thickness, DepthPriority, InBatchID);
+
+	Start = Transform.TransformPosition(FVector( Box.X, -Box.Y, -Box.Z));
+	End = Transform.TransformPosition(FVector(-Box.X, -Box.Y, -Box.Z));
+	new(BatchedLines) FBatchedLine(Center + Start, Center + End, Color, LifeTime, Thickness, DepthPriority, InBatchID);
+
+	Start = Transform.TransformPosition(FVector(-Box.X, -Box.Y, -Box.Z));
+	End = Transform.TransformPosition(FVector(-Box.X,  Box.Y, -Box.Z));
+	new(BatchedLines) FBatchedLine(Center + Start, Center + End, Color, LifeTime, Thickness, DepthPriority, InBatchID);
+
+	Start = Transform.TransformPosition(FVector(-Box.X,  Box.Y, -Box.Z));
+	End = Transform.TransformPosition(FVector( Box.X,  Box.Y, -Box.Z));
+	new(BatchedLines) FBatchedLine(Center + Start, Center + End, Color, LifeTime, Thickness, DepthPriority, InBatchID);
+
+	Start = Transform.TransformPosition(FVector( Box.X,  Box.Y,  Box.Z));
+	End = Transform.TransformPosition(FVector( Box.X,  Box.Y, -Box.Z));
+	new(BatchedLines) FBatchedLine(Center + Start, Center + End, Color, LifeTime, Thickness, DepthPriority, InBatchID);
+
+	Start = Transform.TransformPosition(FVector( Box.X, -Box.Y,  Box.Z));
+	End = Transform.TransformPosition(FVector( Box.X, -Box.Y, -Box.Z));
+	new(BatchedLines) FBatchedLine(Center + Start, Center + End, Color, LifeTime, Thickness, DepthPriority, InBatchID);
+
+	Start = Transform.TransformPosition(FVector(-Box.X, -Box.Y,  Box.Z));
+	End = Transform.TransformPosition(FVector(-Box.X, -Box.Y, -Box.Z));
+	new(BatchedLines) FBatchedLine(Center + Start, Center + End, Color, LifeTime, Thickness, DepthPriority, InBatchID);
+
+	Start = Transform.TransformPosition(FVector(-Box.X,  Box.Y,  Box.Z));
+	End = Transform.TransformPosition(FVector(-Box.X,  Box.Y, -Box.Z));
+	new(BatchedLines) FBatchedLine(Center + Start, Center + End, Color, LifeTime, Thickness, DepthPriority, InBatchID);
+
+	MarkRenderStateDirty();
+}
+
+void ULineBatchComponent::DrawSolidBox(const FBox& Box, const FTransform& Xform, const FColor& Color, uint8 DepthPriority, float LifeTime, uint32 InBatchID)
 {
 	int32 const NewMeshIdx = BatchedMeshes.Add(FBatchedMesh());
 	FBatchedMesh& BM = BatchedMeshes[NewMeshIdx];
@@ -223,6 +295,7 @@ void ULineBatchComponent::DrawSolidBox(const FBox& Box, const FTransform& Xform,
 	BM.Color = Color;
 	BM.DepthPriority = DepthPriority;
 	BM.RemainingLifeTime = LifeTime;
+	BM.BatchID = InBatchID;
 
 	BM.MeshVerts.AddUninitialized(8);
 	BM.MeshVerts[0] = Xform.TransformPosition( FVector(Box.Min.X, Box.Min.Y, Box.Max.Z) );
@@ -236,7 +309,7 @@ void ULineBatchComponent::DrawSolidBox(const FBox& Box, const FTransform& Xform,
 
 	// clockwise
 	BM.MeshIndices.AddUninitialized(36);
-	int32 const Indices[36] = {	3,2,0,
+	constexpr int32 Indices[36] = {	3,2,0,
 		3,0,1,
 		7,3,1,
 		7,1,5,
@@ -257,7 +330,7 @@ void ULineBatchComponent::DrawSolidBox(const FBox& Box, const FTransform& Xform,
 	MarkRenderStateDirty();
 }
 
-void ULineBatchComponent::DrawMesh(TArray<FVector> const& Verts, TArray<int32> const& Indices, FColor const& Color, uint8 DepthPriority, float LifeTime)
+void ULineBatchComponent::DrawMesh(TArray<FVector> const& Verts, TArray<int32> const& Indices, FColor const& Color, uint8 DepthPriority, float LifeTime, uint32 InBatchID)
 {
 	// modifying array element directly to avoid copying arrays
 	int32 const NewMeshIdx = BatchedMeshes.Add(FBatchedMesh());
@@ -268,34 +341,272 @@ void ULineBatchComponent::DrawMesh(TArray<FVector> const& Verts, TArray<int32> c
 	BM.Color = Color;
 	BM.DepthPriority = DepthPriority;
 	BM.RemainingLifeTime = LifeTime;
+	BM.BatchID = InBatchID;
 
 	MarkRenderStateDirty();
 }
 
-void ULineBatchComponent::DrawDirectionalArrow(const FMatrix& ArrowToWorld,FColor InColor,float Length,float ArrowSize,uint8 DepthPriority)
+void ULineBatchComponent::DrawDirectionalArrow(const FMatrix& ArrowToWorld, FLinearColor InColor, float Length, float ArrowSize, uint8 DepthPriority, uint32 InBatchID)
 {
 	const FVector Tip = ArrowToWorld.TransformPosition(FVector(Length,0,0));
-	new(BatchedLines) FBatchedLine(Tip,ArrowToWorld.TransformPosition(FVector::ZeroVector),InColor,DefaultLifeTime,0.0f,DepthPriority);
-	new(BatchedLines) FBatchedLine(Tip,ArrowToWorld.TransformPosition(FVector(Length-ArrowSize,+ArrowSize,+ArrowSize)),InColor,DefaultLifeTime,0.0f,DepthPriority);
-	new(BatchedLines) FBatchedLine(Tip,ArrowToWorld.TransformPosition(FVector(Length-ArrowSize,+ArrowSize,-ArrowSize)),InColor,DefaultLifeTime,0.0f,DepthPriority);
-	new(BatchedLines) FBatchedLine(Tip,ArrowToWorld.TransformPosition(FVector(Length-ArrowSize,-ArrowSize,+ArrowSize)),InColor,DefaultLifeTime,0.0f,DepthPriority);
-	new(BatchedLines) FBatchedLine(Tip,ArrowToWorld.TransformPosition(FVector(Length-ArrowSize,-ArrowSize,-ArrowSize)),InColor,DefaultLifeTime,0.0f,DepthPriority);
+	new(BatchedLines) FBatchedLine(Tip,ArrowToWorld.TransformPosition(FVector::ZeroVector),InColor,DefaultLifeTime,0.0f,DepthPriority, InBatchID);
+	new(BatchedLines) FBatchedLine(Tip,ArrowToWorld.TransformPosition(FVector(Length-ArrowSize,+ArrowSize,+ArrowSize)),InColor,DefaultLifeTime,0.0f,DepthPriority, InBatchID);
+	new(BatchedLines) FBatchedLine(Tip,ArrowToWorld.TransformPosition(FVector(Length-ArrowSize,+ArrowSize,-ArrowSize)),InColor,DefaultLifeTime,0.0f,DepthPriority, InBatchID);
+	new(BatchedLines) FBatchedLine(Tip,ArrowToWorld.TransformPosition(FVector(Length-ArrowSize,-ArrowSize,+ArrowSize)),InColor,DefaultLifeTime,0.0f,DepthPriority, InBatchID);
+	new(BatchedLines) FBatchedLine(Tip,ArrowToWorld.TransformPosition(FVector(Length-ArrowSize,-ArrowSize,-ArrowSize)),InColor,DefaultLifeTime,0.0f,DepthPriority, InBatchID);
+	
 	MarkRenderStateDirty();
 }
 
-/** Draw a circle */
-void ULineBatchComponent::DrawCircle(const FVector& Base,const FVector& X,const FVector& Y,FColor Color,float Radius,int32 NumSides,uint8 DepthPriority)
+void ULineBatchComponent::DrawDirectionalArrow(FVector const& LineStart, FVector const& LineEnd, float ArrowSize, FLinearColor Color, float LifeTime, uint8 DepthPriority, float Thickness, uint32 InBatchID)
 {
+	FVector Dir = (LineEnd-LineStart);
+	Dir.Normalize();
+	FVector Up(0, 0, 1);
+	FVector Right = Dir ^ Up;
+	if (!Right.IsNormalized())
+	{
+		Dir.FindBestAxisVectors(Up, Right);
+	}
+	const FVector Origin = FVector::ZeroVector;
+	FMatrix TM;
+	// get matrix with dir/right/up
+	TM.SetAxes(&Dir, &Right, &Up, &Origin);
+
+	// since dir is x direction, my arrow will be pointing +y, -x and -y, -x
+	const float ArrowSqrt = FMath::Sqrt(ArrowSize);
+
+	new(BatchedLines) FBatchedLine(LineStart,LineEnd, Color, LifeTime, Thickness, DepthPriority, InBatchID);
+	new(BatchedLines) FBatchedLine(LineEnd,LineEnd + TM.TransformPosition(FVector(-ArrowSqrt, ArrowSqrt, 0)), Color, LifeTime, Thickness, DepthPriority, InBatchID);
+	new(BatchedLines) FBatchedLine(LineEnd,LineEnd + TM.TransformPosition(FVector(-ArrowSqrt, -ArrowSqrt, 0)), Color, LifeTime, Thickness, DepthPriority, InBatchID);
+
+	MarkRenderStateDirty();
+}
+
+
+void ULineBatchComponent::AddHalfCircle(const FVector& Base, const FVector& X, const FVector& Y, const FLinearColor& Color, const float Radius, int32 NumSides, const float LifeTime, uint8 DepthPriority, const float Thickness, const uint32 InBatchID)
+{
+	// Need at least 2 sides
+	NumSides = FMath::Max(NumSides, 2);
+	const float AngleDelta = 2.0f * UE_PI / NumSides;
+	FVector	LastVertex = Base + X * Radius;
+
+	for( int32 SideIndex = 0; SideIndex < (NumSides/2); SideIndex++)
+	{
+		FVector	Vertex = Base + (X * FMath::Cos(AngleDelta * (SideIndex + 1)) + Y * FMath::Sin(AngleDelta * (SideIndex + 1))) * Radius;
+		new(BatchedLines) FBatchedLine(LastVertex, Vertex, Color, LifeTime, Thickness, DepthPriority, InBatchID);
+		LastVertex = Vertex;
+	}	
+}
+
+void ULineBatchComponent::AddCircle(const FVector& Base, const FVector& X, const FVector& Y, const FLinearColor& Color, const float Radius, int32 NumSides, const float LifeTime, uint8 DepthPriority, const float Thickness, const uint32 InBatchID)
+{
+	// Need at least 2 sides
+	NumSides = FMath::Max(NumSides, 2);
 	const float	AngleDelta = 2.0f * UE_PI / NumSides;
 	FVector	LastVertex = Base + X * Radius;
 
-	for(int32 SideIndex = 0;SideIndex < NumSides;SideIndex++)
+	for (int32 SideIndex = 0; SideIndex < NumSides; SideIndex++)
 	{
 		const FVector Vertex = Base + (X * FMath::Cos(AngleDelta * (SideIndex + 1)) + Y * FMath::Sin(AngleDelta * (SideIndex + 1))) * Radius;
-		new(BatchedLines) FBatchedLine(LastVertex,Vertex,Color,DefaultLifeTime,0.0f,DepthPriority);
+		new(BatchedLines) FBatchedLine(LastVertex, Vertex, Color, LifeTime, Thickness, DepthPriority, InBatchID);
 		LastVertex = Vertex;
 	}
+}
+
+/** Draw a circle */
+void ULineBatchComponent::DrawCircle(const FVector& Base, const FVector& X, const FVector& Y, FLinearColor Color, float Radius, int32 NumSides, uint8 DepthPriority, uint32 InBatchID)
+{
+	AddHalfCircle(Base, X, Y, Color, Radius, NumSides, DefaultLifeTime, DepthPriority, 0.f, InBatchID);
+
+	MarkRenderStateDirty();
+}
+
+/** Draw a sphere */
+void ULineBatchComponent::DrawSphere(FVector const& Center, float Radius, int32 Segments, FLinearColor Color, float LifeTime, uint8 DepthPriority, float Thickness, uint32 InBatchID)
+{
+	// Need at least 4 segments
+	Segments = FMath::Max(Segments, 4);
+
+	const float AngleInc = 2.f * UE_PI / Segments;
+	int32 NumSegmentsY = Segments;
+	float Latitude = AngleInc;
+	float SinY1 = 0.0f, CosY1 = 1.0f;
+
+	TArray<FBatchedLine> Lines;
+	Lines.Empty(NumSegmentsY * Segments * 2);
+	while (NumSegmentsY--)
+	{
+		const float SinY2 = FMath::Sin(Latitude);
+		const float CosY2 = FMath::Cos(Latitude);
+
+		FVector Vertex1 = FVector(SinY1, 0.0f, CosY1) * Radius + Center;
+		FVector Vertex3 = FVector(SinY2, 0.0f, CosY2) * Radius + Center;
+		float Longitude = AngleInc;
+
+		int32 NumSegmentsX = Segments;
+		while (NumSegmentsX--)
+		{
+			const float SinX = FMath::Sin(Longitude);
+			const float CosX = FMath::Cos(Longitude);
+
+			const FVector Vertex2 = FVector((CosX * SinY1), (SinX * SinY1), CosY1) * Radius + Center;
+			const FVector Vertex4 = FVector((CosX * SinY2), (SinX * SinY2), CosY2) * Radius + Center;
+
+			new(BatchedLines) FBatchedLine(Vertex1, Vertex2, Color, LifeTime, Thickness, DepthPriority, InBatchID);
+			new(BatchedLines) FBatchedLine(Vertex1, Vertex3, Color, LifeTime, Thickness, DepthPriority, InBatchID);
+
+			Vertex1 = Vertex2;
+			Vertex3 = Vertex4;
+			Longitude += AngleInc;
+		}
+		SinY1 = SinY2;
+		CosY1 = CosY2;
+		Latitude += AngleInc;
+	}
+	MarkRenderStateDirty();
+}
+
+void ULineBatchComponent::DrawCylinder(FVector const& Start, FVector const& End, float Radius, int32 Segments, FLinearColor Color, float LifeTime, uint8 DepthPriority, float Thickness, uint32 InBatchID)
+{
+	// Need at least 4 segments
+	Segments = FMath::Max(Segments, 4);
+
+	// Rotate a point around axis to form cylinder segments
+	const float AngleInc = 360.f / Segments;
+	float Angle = AngleInc;
+
+	// Default for Axis is up
+	FVector Axis = (End - Start).GetSafeNormal();
+	if( Axis.IsZero() )
+	{
+		Axis = FVector(0.f, 0.f, 1.f);
+	}
 
+	FVector Perpendicular, Dummy;
+	Axis.FindBestAxisVectors(Perpendicular, Dummy);
+		
+	FVector Segment = Perpendicular.RotateAngleAxis(0, Axis) * Radius;
+	FVector P1 = Segment + Start;
+	FVector P3 = Segment + End;
+
+	while( Segments-- )
+	{
+		Segment = Perpendicular.RotateAngleAxis(Angle, Axis) * Radius;
+		FVector P2 = Segment + Start;
+		FVector P4 = Segment + End;
+
+		new(BatchedLines) FBatchedLine(P2, P4, Color, LifeTime, Thickness, DepthPriority, InBatchID);
+		new(BatchedLines) FBatchedLine(P1, P2, Color, LifeTime, Thickness, DepthPriority, InBatchID);
+		new(BatchedLines) FBatchedLine(P3, P4, Color, LifeTime, Thickness, DepthPriority, InBatchID);
+
+		P1 = P2;
+		P3 = P4;
+		Angle += AngleInc;
+	}
+	MarkRenderStateDirty();
+}
+
+void ULineBatchComponent::DrawCone(FVector const& Origin, FVector const& Direction, float Length, float AngleWidth, float AngleHeight, int32 NumSides, FLinearColor DrawColor, float LifeTime, uint8 DepthPriority, float Thickness, uint32 InBatchID)
+{
+	// Need at least 4 sides
+	NumSides = FMath::Max(NumSides, 4);
+
+	const float Angle1 = FMath::Clamp<float>(AngleHeight, UE_KINDA_SMALL_NUMBER, (UE_PI - UE_KINDA_SMALL_NUMBER));
+	const float Angle2 = FMath::Clamp<float>(AngleWidth, UE_KINDA_SMALL_NUMBER, (UE_PI - UE_KINDA_SMALL_NUMBER));
+
+	const float SinX_2 = FMath::Sin(0.5f * Angle1);
+	const float SinY_2 = FMath::Sin(0.5f * Angle2);
+
+	const float SinSqX_2 = SinX_2 * SinX_2;
+	const float SinSqY_2 = SinY_2 * SinY_2;
+
+	TArray<FVector> ConeVerts;
+	ConeVerts.AddUninitialized(NumSides);
+
+	for(int32 i = 0; i < NumSides; i++)
+	{
+		const float Fraction	= (float)i/(float)(NumSides);
+		const float Thi			= 2.f * UE_PI * Fraction;
+		const float Phi			= FMath::Atan2(FMath::Sin(Thi)*SinY_2, FMath::Cos(Thi)*SinX_2);
+		const float SinPhi		= FMath::Sin(Phi);
+		const float CosPhi		= FMath::Cos(Phi);
+		const float SinSqPhi	= SinPhi*SinPhi;
+		const float CosSqPhi	= CosPhi*CosPhi;
+
+		const float RSq			= SinSqX_2*SinSqY_2 / (SinSqX_2*SinSqPhi + SinSqY_2*CosSqPhi);
+		const float R			= FMath::Sqrt(RSq);
+		const float Sqr			= FMath::Sqrt(1-RSq);
+		const float Alpha		= R*CosPhi;
+		const float Beta		= R*SinPhi;
+
+		ConeVerts[i].X = (1 - 2*RSq);
+		ConeVerts[i].Y = 2 * Sqr * Alpha;
+		ConeVerts[i].Z = 2 * Sqr * Beta;
+	}
+
+	// Calculate transform for cone.
+	FVector YAxis, ZAxis;
+	const FVector DirectionNorm = Direction.GetSafeNormal();
+	DirectionNorm.FindBestAxisVectors(YAxis, ZAxis);
+	const FMatrix ConeToWorld = FScaleMatrix(FVector(Length)) * FMatrix(DirectionNorm, YAxis, ZAxis, Origin);
+
+	FVector CurrentPoint, PrevPoint, FirstPoint;
+	for(int32 i = 0; i < NumSides; i++)
+	{
+		CurrentPoint = ConeToWorld.TransformPosition(ConeVerts[i]);
+		new(BatchedLines) FBatchedLine(ConeToWorld.GetOrigin(), CurrentPoint, DrawColor, LifeTime, Thickness, DepthPriority, InBatchID);
+
+		// PrevPoint must be defined to draw junctions
+		if( i > 0 )
+		{
+			new(BatchedLines) FBatchedLine(PrevPoint, CurrentPoint, DrawColor, LifeTime, Thickness, DepthPriority, InBatchID);
+		}
+		else
+		{
+			FirstPoint = CurrentPoint;
+		}
+
+		PrevPoint = CurrentPoint;
+	}
+	// Connect last junction to first
+	new(BatchedLines) FBatchedLine(CurrentPoint, FirstPoint, DrawColor, LifeTime, Thickness, DepthPriority, InBatchID);
+
+	MarkRenderStateDirty();
+}
+
+void ULineBatchComponent::DrawCapsule(FVector const& Center, float HalfHeight, float Radius, const FQuat& Rotation, FLinearColor Color, float LifeTime, uint8 DepthPriority, float Thickness, uint32 InBatchID)
+{
+	constexpr int32 DrawCollisionSides = 16;
+	const FVector Origin = Center;
+	const FMatrix Axes = FQuatRotationTranslationMatrix(Rotation, FVector::ZeroVector);
+	const FVector XAxis = Axes.GetScaledAxis( EAxis::X );
+	const FVector YAxis = Axes.GetScaledAxis( EAxis::Y );
+	const FVector ZAxis = Axes.GetScaledAxis( EAxis::Z ); 
+
+	// Draw top and bottom circles
+	const float HalfAxis = FMath::Max<float>(HalfHeight - Radius, 1.f);
+	const FVector TopEnd = Origin + HalfAxis*ZAxis;
+	const FVector BottomEnd = Origin - HalfAxis*ZAxis;
+
+	AddCircle(TopEnd, XAxis, YAxis, Color, Radius, DrawCollisionSides, LifeTime, Thickness, DepthPriority,  InBatchID);
+	AddCircle( BottomEnd, XAxis, YAxis, Color, Radius, DrawCollisionSides, LifeTime, Thickness, DepthPriority, InBatchID);
+
+		// Draw domed caps
+	AddHalfCircle( TopEnd, YAxis, ZAxis, Color, Radius, DrawCollisionSides, LifeTime, Thickness, DepthPriority, InBatchID);
+	AddHalfCircle( TopEnd, XAxis, ZAxis, Color, Radius, DrawCollisionSides, LifeTime, Thickness, DepthPriority, InBatchID);
+
+	const FVector NegZAxis = -ZAxis;
+
+	AddHalfCircle( BottomEnd, YAxis, NegZAxis, Color, Radius, DrawCollisionSides, LifeTime, Thickness, DepthPriority, InBatchID);
+	AddHalfCircle( BottomEnd, XAxis, NegZAxis, Color, Radius, DrawCollisionSides, LifeTime, Thickness, DepthPriority, InBatchID);
+
+	// Draw connected lines
+	new(BatchedLines) FBatchedLine(TopEnd + Radius*XAxis, BottomEnd + Radius*XAxis, Color, LifeTime, Thickness, DepthPriority, InBatchID);
+	new(BatchedLines) FBatchedLine(TopEnd - Radius*XAxis, BottomEnd - Radius*XAxis, Color, LifeTime, Thickness, DepthPriority, InBatchID);
+	new(BatchedLines) FBatchedLine(TopEnd + Radius*XAxis, BottomEnd + Radius*XAxis, Color, LifeTime, Thickness, DepthPriority, InBatchID);
+	new(BatchedLines) FBatchedLine(TopEnd - Radius*XAxis, BottomEnd - Radius*XAxis, Color, LifeTime, Thickness, DepthPriority, InBatchID);
+	
 	MarkRenderStateDirty();
 }
 
@@ -449,3 +760,19 @@ void ULineBatchComponent::Flush()
 	}
 }
 
+void ULineBatchComponent::ClearBatch(uint32 InBatchID)
+{
+	if (InBatchID == INVALID_ID)
+	{
+		return;
+	}
+	
+	bool bDirty = BatchedLines.RemoveAllSwap([InBatchID](const FBatchedLine& Element) { return Element.BatchID == InBatchID; }) > 0;
+	bDirty |= BatchedPoints.RemoveAllSwap([InBatchID](const FBatchedPoint& Pt) { return Pt.BatchID == InBatchID; }) > 0;
+	bDirty |= BatchedMeshes.RemoveAllSwap([InBatchID](const FBatchedMesh& Mesh) { return Mesh.BatchID == InBatchID; }) > 0;
+
+	if(bDirty)
+	{
+		MarkRenderStateDirty();
+	}
+}
diff --git a/Engine/Source/Runtime/Engine/Private/DrawDebugHelpers.cpp b/Engine/Source/Runtime/Engine/Private/DrawDebugHelpers.cpp
--- a/Engine/Source/Runtime/Engine/Private/DrawDebugHelpers.cpp
+++ b/Engine/Source/Runtime/Engine/Private/DrawDebugHelpers.cpp
@@ -14,7 +14,6 @@
 #include "Components/LineBatchComponent.h"
 #include "Engine/Canvas.h"
 #include "GameFramework/HUD.h"
-#include "Math/ScaleMatrix.h"
 
 #if ENABLE_DRAW_DEBUG
 
@@ -139,31 +138,16 @@ void DrawDebugDirectionalArrow(const UWorld* InWorld, FVector const& LineStart,
 	// no debug line drawing on dedicated server
 	if (GEngine->GetNetMode(InWorld) != NM_DedicatedServer)
 	{
-		if (ArrowSize <= 0)
-		{
-			ArrowSize = 10.f;
-		}
-
-		DrawDebugLine(InWorld, LineStart, LineEnd, Color, bPersistentLines, LifeTime, DepthPriority, Thickness);
-
-		FVector Dir = (LineEnd-LineStart);
-		Dir.Normalize();
-		FVector Up(0, 0, 1);
-		FVector Right = Dir ^ Up;
-		if (!Right.IsNormalized())
+		if (ULineBatchComponent* const LineBatcher = GetDebugLineBatcher(InWorld, bPersistentLines, LifeTime, (DepthPriority == SDPG_Foreground)))
 		{
-			Dir.FindBestAxisVectors(Up, Right);
+			float const LineLifeTime = GetDebugLineLifeTime(LineBatcher, LifeTime, bPersistentLines);
+			if (ArrowSize <= 0)
+			{
+				ArrowSize = 10.f;
+			}
+			
+			LineBatcher->DrawDirectionalArrow(LineStart, LineEnd, ArrowSize, Color, LineLifeTime, DepthPriority, Thickness);
 		}
-		FVector Origin = FVector::ZeroVector;
-		FMatrix TM;
-		// get matrix with dir/right/up
-		TM.SetAxes(&Dir, &Right, &Up, &Origin);
-
-		// since dir is x direction, my arrow will be pointing +y, -x and -y, -x
-		float ArrowSqrt = FMath::Sqrt(ArrowSize);
-		FVector ArrowPos;
-		DrawDebugLine(InWorld, LineEnd, LineEnd + TM.TransformPosition(FVector(-ArrowSqrt, ArrowSqrt, 0)), Color, bPersistentLines, LifeTime, DepthPriority, Thickness);
-		DrawDebugLine(InWorld, LineEnd, LineEnd + TM.TransformPosition(FVector(-ArrowSqrt, -ArrowSqrt, 0)), Color, bPersistentLines, LifeTime, DepthPriority, Thickness);
 	}
 	else
 	{
@@ -179,22 +163,8 @@ void DrawDebugBox(const UWorld* InWorld, FVector const& Center, FVector const& B
 		// this means foreground lines can't be persistent 
 		if (ULineBatchComponent* const LineBatcher = GetDebugLineBatcher(InWorld, bPersistentLines, LifeTime, (DepthPriority == SDPG_Foreground)))
 		{
-			float LineLifeTime = GetDebugLineLifeTime(LineBatcher, LifeTime, bPersistentLines);
-
-			LineBatcher->DrawLine(Center + FVector( Box.X,  Box.Y,  Box.Z), Center + FVector( Box.X, -Box.Y, Box.Z), Color, DepthPriority, Thickness, LineLifeTime);
-			LineBatcher->DrawLine(Center + FVector( Box.X, -Box.Y,  Box.Z), Center + FVector(-Box.X, -Box.Y, Box.Z), Color, DepthPriority, Thickness, LineLifeTime);
-			LineBatcher->DrawLine(Center + FVector(-Box.X, -Box.Y,  Box.Z), Center + FVector(-Box.X,  Box.Y, Box.Z), Color, DepthPriority, Thickness, LineLifeTime);
-			LineBatcher->DrawLine(Center + FVector(-Box.X,  Box.Y,  Box.Z), Center + FVector( Box.X,  Box.Y, Box.Z), Color, DepthPriority, Thickness, LineLifeTime);
-
-			LineBatcher->DrawLine(Center + FVector( Box.X,  Box.Y, -Box.Z), Center + FVector( Box.X, -Box.Y, -Box.Z), Color, DepthPriority, Thickness, LineLifeTime);
-			LineBatcher->DrawLine(Center + FVector( Box.X, -Box.Y, -Box.Z), Center + FVector(-Box.X, -Box.Y, -Box.Z), Color, DepthPriority, Thickness, LineLifeTime);
-			LineBatcher->DrawLine(Center + FVector(-Box.X, -Box.Y, -Box.Z), Center + FVector(-Box.X,  Box.Y, -Box.Z), Color, DepthPriority, Thickness, LineLifeTime);
-			LineBatcher->DrawLine(Center + FVector(-Box.X,  Box.Y, -Box.Z), Center + FVector( Box.X,  Box.Y, -Box.Z), Color, DepthPriority, Thickness, LineLifeTime);
-
-			LineBatcher->DrawLine(Center + FVector( Box.X,  Box.Y,  Box.Z), Center + FVector( Box.X,  Box.Y, -Box.Z), Color, DepthPriority, Thickness, LineLifeTime);
-			LineBatcher->DrawLine(Center + FVector( Box.X, -Box.Y,  Box.Z), Center + FVector( Box.X, -Box.Y, -Box.Z), Color, DepthPriority, Thickness, LineLifeTime);
-			LineBatcher->DrawLine(Center + FVector(-Box.X, -Box.Y,  Box.Z), Center + FVector(-Box.X, -Box.Y, -Box.Z), Color, DepthPriority, Thickness, LineLifeTime);
-			LineBatcher->DrawLine(Center + FVector(-Box.X,  Box.Y,  Box.Z), Center + FVector(-Box.X,  Box.Y, -Box.Z), Color, DepthPriority, Thickness, LineLifeTime);
+			const float BoxLifeTime = GetDebugLineLifeTime(LineBatcher, LifeTime, bPersistentLines);
+			LineBatcher->DrawBox(Center, Box, Color, BoxLifeTime, DepthPriority, Thickness);
 		}
 	}
 	else
@@ -211,59 +181,8 @@ void DrawDebugBox(const UWorld* InWorld, FVector const& Center, FVector const& B
 		// this means foreground lines can't be persistent 
 		if (ULineBatchComponent* const LineBatcher = GetDebugLineBatcher(InWorld, bPersistentLines, LifeTime, (DepthPriority == SDPG_Foreground)))
 		{
-			float const LineLifeTime = GetDebugLineLifeTime(LineBatcher, LifeTime, bPersistentLines);
-			TArray<struct FBatchedLine> Lines;
-
-			FTransform const Transform(Rotation);
-			FVector Start = Transform.TransformPosition(FVector( Box.X,  Box.Y,  Box.Z));
-			FVector End = Transform.TransformPosition(FVector( Box.X, -Box.Y, Box.Z));
-			new(Lines) FBatchedLine(Center + Start, Center + End, Color, LineLifeTime, Thickness, DepthPriority);
-
-			Start = Transform.TransformPosition(FVector( Box.X, -Box.Y,  Box.Z));
-			End = Transform.TransformPosition(FVector(-Box.X, -Box.Y, Box.Z));
-			new(Lines) FBatchedLine(Center + Start, Center + End, Color, LineLifeTime, Thickness, DepthPriority);
-
-			Start = Transform.TransformPosition(FVector(-Box.X, -Box.Y,  Box.Z));
-			End = Transform.TransformPosition(FVector(-Box.X,  Box.Y, Box.Z));
-			new(Lines) FBatchedLine(Center + Start, Center + End, Color, LineLifeTime, Thickness, DepthPriority);
-
-			Start = Transform.TransformPosition(FVector(-Box.X,  Box.Y,  Box.Z));
-			End = Transform.TransformPosition(FVector( Box.X,  Box.Y, Box.Z));
-			new(Lines) FBatchedLine(Center + Start, Center + End, Color, LineLifeTime, Thickness, DepthPriority);
-
-			Start = Transform.TransformPosition(FVector( Box.X,  Box.Y, -Box.Z));
-			End = Transform.TransformPosition(FVector( Box.X, -Box.Y, -Box.Z));
-			new(Lines) FBatchedLine(Center + Start, Center + End, Color, LineLifeTime, Thickness, DepthPriority);
-
-			Start = Transform.TransformPosition(FVector( Box.X, -Box.Y, -Box.Z));
-			End = Transform.TransformPosition(FVector(-Box.X, -Box.Y, -Box.Z));
-			new(Lines) FBatchedLine(Center + Start, Center + End, Color, LineLifeTime, Thickness, DepthPriority);
-
-			Start = Transform.TransformPosition(FVector(-Box.X, -Box.Y, -Box.Z));
-			End = Transform.TransformPosition(FVector(-Box.X,  Box.Y, -Box.Z));
-			new(Lines) FBatchedLine(Center + Start, Center + End, Color, LineLifeTime, Thickness, DepthPriority);
-
-			Start = Transform.TransformPosition(FVector(-Box.X,  Box.Y, -Box.Z));
-			End = Transform.TransformPosition(FVector( Box.X,  Box.Y, -Box.Z));
-			new(Lines )FBatchedLine(Center + Start, Center + End, Color, LineLifeTime, Thickness, DepthPriority);
-
-			Start = Transform.TransformPosition(FVector( Box.X,  Box.Y,  Box.Z));
-			End = Transform.TransformPosition(FVector( Box.X,  Box.Y, -Box.Z));
-			new(Lines) FBatchedLine(Center + Start, Center + End, Color, LineLifeTime, Thickness, DepthPriority);
-
-			Start = Transform.TransformPosition(FVector( Box.X, -Box.Y,  Box.Z));
-			End = Transform.TransformPosition(FVector( Box.X, -Box.Y, -Box.Z));
-			new(Lines) FBatchedLine(Center + Start, Center + End, Color, LineLifeTime, Thickness, DepthPriority);
-
-			Start = Transform.TransformPosition(FVector(-Box.X, -Box.Y,  Box.Z));
-			End = Transform.TransformPosition(FVector(-Box.X, -Box.Y, -Box.Z));
-			new(Lines) FBatchedLine(Center + Start, Center + End, Color, LineLifeTime, Thickness, DepthPriority);
-
-			Start = Transform.TransformPosition(FVector(-Box.X,  Box.Y,  Box.Z));
-			End = Transform.TransformPosition(FVector(-Box.X,  Box.Y, -Box.Z));
-			new(Lines) FBatchedLine(Center + Start, Center + End, Color, LineLifeTime, Thickness, DepthPriority);
-
-			LineBatcher->DrawLines(Lines);
+			float const BoxLifeTime = GetDebugLineLifeTime(LineBatcher, LifeTime, bPersistentLines);
+			LineBatcher->DrawBox(Center, Box, Rotation, Color, BoxLifeTime, DepthPriority, Thickness);
 		}
 	}
 	else
@@ -578,52 +497,8 @@ void DrawDebugSphere(const UWorld* InWorld, FVector const& Center, float Radius,
 		// this means foreground lines can't be persistent 
 		if (ULineBatchComponent* const LineBatcher = GetDebugLineBatcher(InWorld, bPersistentLines, LifeTime, (DepthPriority == SDPG_Foreground)))
 		{
-			float LineLifeTime = GetDebugLineLifeTime(LineBatcher, LifeTime, bPersistentLines);
-
-			// Need at least 4 segments
-			Segments = FMath::Max(Segments, 4);
-
-			FVector Vertex1, Vertex2, Vertex3, Vertex4;
-			const float AngleInc = 2.f * UE_PI / float(Segments);
-			int32 NumSegmentsY = Segments;
-			float Latitude = AngleInc;
-			int32 NumSegmentsX;
-			float Longitude;
-			float SinY1 = 0.0f, CosY1 = 1.0f, SinY2, CosY2;
-			float SinX, CosX;
-
-			TArray<FBatchedLine> Lines;
-			Lines.Empty(NumSegmentsY * Segments * 2);
-			while (NumSegmentsY--)
-			{
-				SinY2 = FMath::Sin(Latitude);
-				CosY2 = FMath::Cos(Latitude);
-
-				Vertex1 = FVector(SinY1, 0.0f, CosY1) * Radius + Center;
-				Vertex3 = FVector(SinY2, 0.0f, CosY2) * Radius + Center;
-				Longitude = AngleInc;
-
-				NumSegmentsX = Segments;
-				while (NumSegmentsX--)
-				{
-					SinX = FMath::Sin(Longitude);
-					CosX = FMath::Cos(Longitude);
-
-					Vertex2 = FVector((CosX * SinY1), (SinX * SinY1), CosY1) * Radius + Center;
-					Vertex4 = FVector((CosX * SinY2), (SinX * SinY2), CosY2) * Radius + Center;
-
-					Lines.Add(FBatchedLine(Vertex1, Vertex2, Color, LineLifeTime, Thickness, DepthPriority));
-					Lines.Add(FBatchedLine(Vertex1, Vertex3, Color, LineLifeTime, Thickness, DepthPriority));
-
-					Vertex1 = Vertex2;
-					Vertex3 = Vertex4;
-					Longitude += AngleInc;
-				}
-				SinY1 = SinY2;
-				CosY1 = CosY2;
-				Latitude += AngleInc;
-			}
-			LineBatcher->DrawLines(Lines);
+			const float SphereLifeTime = GetDebugLineLifeTime(LineBatcher, LifeTime, bPersistentLines);
+			LineBatcher->DrawSphere(Center, Radius, Segments, Color, SphereLifeTime, DepthPriority, Thickness);
 		}
 	}
 	else
@@ -640,46 +515,8 @@ void DrawDebugCylinder(const UWorld* InWorld, FVector const& Start, FVector cons
 		// this means foreground lines can't be persistent 
 		if (ULineBatchComponent* const LineBatcher = GetDebugLineBatcher(InWorld, bPersistentLines, LifeTime, (DepthPriority == SDPG_Foreground)))
 		{
-			// Need at least 4 segments
-			Segments = FMath::Max(Segments, 4);
-
-			// Rotate a point around axis to form cylinder segments
-			FVector Segment;
-			FVector P1, P2, P3, P4;
-			const float AngleInc = 360.f / Segments;
-			float Angle = AngleInc;
-
-			// Default for Axis is up
-			FVector Axis = (End - Start).GetSafeNormal();
-			if( Axis.IsZero() )
-			{
-				Axis = FVector(0.f, 0.f, 1.f);
-			}
-
-			FVector Perpendicular;
-			FVector Dummy;
-
-			Axis.FindBestAxisVectors(Perpendicular, Dummy);
-		
-			Segment = Perpendicular.RotateAngleAxis(0, Axis) * Radius;
-			P1 = Segment + Start;
-			P3 = Segment + End;
-
-			const float LineLifeTime = GetDebugLineLifeTime(LineBatcher, LifeTime, bPersistentLines);
-			while( Segments-- )
-			{
-				Segment = Perpendicular.RotateAngleAxis(Angle, Axis) * Radius;
-				P2 = Segment + Start;
-				P4 = Segment + End;
-
-				LineBatcher->DrawLine(P2, P4, Color, DepthPriority, Thickness, LineLifeTime);
-				LineBatcher->DrawLine(P1, P2, Color, DepthPriority, Thickness, LineLifeTime);
-				LineBatcher->DrawLine(P3, P4, Color, DepthPriority, Thickness, LineLifeTime);
-
-				P1 = P2;
-				P3 = P4;
-				Angle += AngleInc;
-			}
+			const float CylinderLifeTime = GetDebugLineLifeTime(LineBatcher, LifeTime, bPersistentLines);
+			LineBatcher->DrawCylinder(Start, End, Radius, Segments, Color, CylinderLifeTime, DepthPriority, Thickness);
 		}
 	}
 	else
@@ -762,81 +599,11 @@ void DrawDebugCone(const UWorld* InWorld, FVector const& Origin, FVector const&
 	// no debug line drawing on dedicated server
 	if (GEngine->GetNetMode(InWorld) != NM_DedicatedServer)
 	{
-		// Need at least 4 sides
-		NumSides = FMath::Max(NumSides, 4);
-
-		const float Angle1 = FMath::Clamp<float>(AngleHeight, (float)UE_KINDA_SMALL_NUMBER, (float)(UE_PI - UE_KINDA_SMALL_NUMBER));
-		const float Angle2 = FMath::Clamp<float>(AngleWidth, (float)UE_KINDA_SMALL_NUMBER, (float)(UE_PI - UE_KINDA_SMALL_NUMBER));
-
-		const float SinX_2 = FMath::Sin(0.5f * Angle1);
-		const float SinY_2 = FMath::Sin(0.5f * Angle2);
-
-		const float SinSqX_2 = SinX_2 * SinX_2;
-		const float SinSqY_2 = SinY_2 * SinY_2;
-
-		const float TanX_2 = FMath::Tan(0.5f * Angle1);
-		const float TanY_2 = FMath::Tan(0.5f * Angle2);
-
-		TArray<FVector> ConeVerts;
-		ConeVerts.AddUninitialized(NumSides);
-
-		for(int32 i = 0; i < NumSides; i++)
-		{
-			const float Fraction	= (float)i/(float)(NumSides);
-			const float Thi			= 2.f * UE_PI * Fraction;
-			const float Phi			= FMath::Atan2(FMath::Sin(Thi)*SinY_2, FMath::Cos(Thi)*SinX_2);
-			const float SinPhi		= FMath::Sin(Phi);
-			const float CosPhi		= FMath::Cos(Phi);
-			const float SinSqPhi	= SinPhi*SinPhi;
-			const float CosSqPhi	= CosPhi*CosPhi;
-
-			const float RSq			= SinSqX_2*SinSqY_2 / (SinSqX_2*SinSqPhi + SinSqY_2*CosSqPhi);
-			const float R			= FMath::Sqrt(RSq);
-			const float Sqr			= FMath::Sqrt(1-RSq);
-			const float Alpha		= R*CosPhi;
-			const float Beta		= R*SinPhi;
-
-			ConeVerts[i].X = (1 - 2*RSq);
-			ConeVerts[i].Y = 2 * Sqr * Alpha;
-			ConeVerts[i].Z = 2 * Sqr * Beta;
-		}
-
-		// Calculate transform for cone.
-		FVector YAxis, ZAxis;
-		FVector DirectionNorm = Direction.GetSafeNormal();
-		DirectionNorm.FindBestAxisVectors(YAxis, ZAxis);
-		const FMatrix ConeToWorld = FScaleMatrix(FVector(Length)) * FMatrix(DirectionNorm, YAxis, ZAxis, Origin);
-
 		// this means foreground lines can't be persistent 
 		if (ULineBatchComponent* const LineBatcher = GetDebugLineBatcher(InWorld, bPersistentLines, LifeTime, (DepthPriority == SDPG_Foreground)))
 		{
-			float const LineLifeTime = GetDebugLineLifeTime(LineBatcher, LifeTime, bPersistentLines);
-
-			TArray<FBatchedLine> Lines;
-			Lines.Empty(NumSides);
-
-			FVector CurrentPoint, PrevPoint, FirstPoint;
-			for(int32 i = 0; i < NumSides; i++)
-			{
-				CurrentPoint = ConeToWorld.TransformPosition(ConeVerts[i]);
-				Lines.Add(FBatchedLine(ConeToWorld.GetOrigin(), CurrentPoint, DrawColor, LineLifeTime, Thickness, DepthPriority));
-
-				// PrevPoint must be defined to draw junctions
-				if( i > 0 )
-				{
-					Lines.Add(FBatchedLine(PrevPoint, CurrentPoint, DrawColor, LineLifeTime, Thickness, DepthPriority));
-				}
-				else
-				{
-					FirstPoint = CurrentPoint;
-				}
-
-				PrevPoint = CurrentPoint;
-			}
-			// Connect last junction to first
-			Lines.Add(FBatchedLine(CurrentPoint, FirstPoint, DrawColor, LineLifeTime, Thickness, DepthPriority));
-
-			LineBatcher->DrawLines(Lines);
+			float const ConeLifeTime = GetDebugLineLifeTime(LineBatcher, LifeTime, bPersistentLines);
+			LineBatcher->DrawCone(Origin, Direction, Length, AngleWidth, AngleHeight, NumSides, DrawColor, ConeLifeTime, DepthPriority, Thickness);
 		}
 	}
 	else
@@ -931,21 +698,6 @@ void DrawDebugFrustum(const UWorld* InWorld, const FMatrix& FrustumToWorld, FCol
 	}
 }
 
-
-
-static void DrawHalfCircle(const UWorld* InWorld, const FVector& Base, const FVector& X, const FVector& Y, const FColor& Color, float Radius, int32 NumSides, bool bPersistentLines, float LifeTime, uint8 DepthPriority, float Thickness)
-{
-	float	AngleDelta = 2.0f * (float)UE_PI / ((float)NumSides);
-	FVector	LastVertex = Base + X * Radius;
-
-	for(int32 SideIndex = 0; SideIndex < (NumSides/2); SideIndex++)
-	{
-		FVector	Vertex = Base + (X * FMath::Cos(AngleDelta * (SideIndex + 1)) + Y * FMath::Sin(AngleDelta * (SideIndex + 1))) * Radius;
-		DrawDebugLine(InWorld, LastVertex, Vertex, Color, bPersistentLines, LifeTime, DepthPriority, Thickness);
-		LastVertex = Vertex;
-	}	
-}
-
 void DrawCircle(const UWorld* InWorld, const FVector& Base, const FVector& X, const FVector& Y, const FColor& Color, float Radius, int32 NumSides, bool bPersistentLines, float LifeTime, uint8 DepthPriority, float Thickness)
 {
 	const float	AngleDelta = 2.0f * UE_PI / NumSides;
@@ -964,36 +716,11 @@ void DrawDebugCapsule(const UWorld* InWorld, FVector const& Center, float HalfHe
 	// no debug line drawing on dedicated server
 	if (GEngine->GetNetMode(InWorld) != NM_DedicatedServer)
 	{
-		const int32 DrawCollisionSides = 16;
-
-		FVector Origin = Center;
-		FMatrix Axes = FQuatRotationTranslationMatrix(Rotation, FVector::ZeroVector);
-		FVector XAxis = Axes.GetScaledAxis( EAxis::X );
-		FVector YAxis = Axes.GetScaledAxis( EAxis::Y );
-		FVector ZAxis = Axes.GetScaledAxis( EAxis::Z ); 
-
-		// Draw top and bottom circles
-		float HalfAxis = FMath::Max<float>(HalfHeight - Radius, 1.f);
-		FVector TopEnd = Origin + HalfAxis*ZAxis;
-		FVector BottomEnd = Origin - HalfAxis*ZAxis;
-
-		DrawCircle(InWorld, TopEnd, XAxis, YAxis, Color, Radius, DrawCollisionSides, bPersistentLines, LifeTime, DepthPriority, Thickness);
-		DrawCircle(InWorld, BottomEnd, XAxis, YAxis, Color, Radius, DrawCollisionSides, bPersistentLines, LifeTime, DepthPriority, Thickness);
-
-		// Draw domed caps
-		DrawHalfCircle(InWorld, TopEnd, YAxis, ZAxis, Color, Radius, DrawCollisionSides, bPersistentLines, LifeTime, DepthPriority, Thickness);
-		DrawHalfCircle(InWorld, TopEnd, XAxis, ZAxis, Color, Radius, DrawCollisionSides, bPersistentLines, LifeTime, DepthPriority, Thickness);
-
-		FVector NegZAxis = -ZAxis;
-
-		DrawHalfCircle(InWorld, BottomEnd, YAxis, NegZAxis, Color, Radius, DrawCollisionSides, bPersistentLines, LifeTime, DepthPriority, Thickness);
-		DrawHalfCircle(InWorld, BottomEnd, XAxis, NegZAxis, Color, Radius, DrawCollisionSides, bPersistentLines, LifeTime, DepthPriority, Thickness);
-
-		// Draw connected lines
-		DrawDebugLine(InWorld, TopEnd + Radius*XAxis, BottomEnd + Radius*XAxis, Color, bPersistentLines, LifeTime, DepthPriority, Thickness);
-		DrawDebugLine(InWorld, TopEnd - Radius*XAxis, BottomEnd - Radius*XAxis, Color, bPersistentLines, LifeTime, DepthPriority, Thickness);
-		DrawDebugLine(InWorld, TopEnd + Radius*YAxis, BottomEnd + Radius*YAxis, Color, bPersistentLines, LifeTime, DepthPriority, Thickness);
-		DrawDebugLine(InWorld, TopEnd - Radius*YAxis, BottomEnd - Radius*YAxis, Color, bPersistentLines, LifeTime, DepthPriority, Thickness);
+		if (ULineBatchComponent* const LineBatcher = GetDebugLineBatcher(InWorld, bPersistentLines, LifeTime, (DepthPriority == SDPG_Foreground)))
+		{
+			float const CapsuleLifeTime = GetDebugLineLifeTime(LineBatcher, LifeTime, bPersistentLines);
+			LineBatcher->DrawCapsule(Center, HalfHeight, Radius, Rotation, Color, CapsuleLifeTime, DepthPriority, Thickness);
+		}
 	}
 	else
 	{
@@ -1303,7 +1030,7 @@ void DrawDebugCanvasWireBox(UCanvas* Canvas, const FMatrix& Transform, const FBo
 
 void DrawDebugCanvasCapsule(UCanvas* Canvas, const FMatrix& Transform, float HalfLength, float Radius, const FColor& LineColor)
 {
-	const int32 DrawCollisionSides = 16;
+	constexpr int32 DrawCollisionSides = 16;
 
 	FVector Origin = Transform.GetOrigin();
 	FVector XAxis = Transform.GetScaledAxis(EAxis::X);
