PR #9329: Added a few clarifications to TouchInterface to make it a bit clearer what things do at a glance. (Contributed by [REDACTED_PERSON])

diff --git a/Engine/Source/Runtime/Engine/Classes/GameFramework/TouchInterface.h b/Engine/Source/Runtime/Engine/Classes/GameFramework/TouchInterface.h
--- a/Engine/Source/Runtime/Engine/Classes/GameFramework/TouchInterface.h
+++ b/Engine/Source/Runtime/Engine/Classes/GameFramework/TouchInterface.h
@@ -21,13 +21,13 @@ struct FTouchInputControl
 	TObjectPtr<UTexture2D> Image1;
 	UPROPERTY(EditAnywhere, Category="Control", meta=(ToolTip="For sticks, this is the Background"))
 	TObjectPtr<UTexture2D> Image2;
-	UPROPERTY(EditAnywhere, Category="Control", meta=(ToolTip="The center point of the control (if <= 1.0, it's relative to screen, > 1.0 is absolute)"))
+	UPROPERTY(EditAnywhere, Category="Control", meta=(ToolTip="The initial center point of the control. If Time Until Reset is < 0, control resets back to here.\nUse negative numbers to invert positioning from top to bottom, left to right. (if <= 1.0, it's relative to screen, > 1.0 is absolute)"))
 	FVector2D Center;
 	UPROPERTY(EditAnywhere, Category="Control", meta=(ToolTip="The size of the control (if <= 1.0, it's relative to screen, > 1.0 is absolute)"))
 	FVector2D VisualSize;
 	UPROPERTY(EditAnywhere, Category="Control", meta=(ToolTip="For sticks, the size of the thumb (if <= 1.0, it's relative to screen, > 1.0 is absolute)"))
 	FVector2D ThumbSize;
-	UPROPERTY(EditAnywhere, Category="Control", meta=(ToolTip="The interactive size of the control (if <= 1.0, it's relative to screen, > 1.0 is absolute)"))
+	UPROPERTY(EditAnywhere, Category="Control", meta=(ToolTip="The interactive size of the control. Measured outward from Center. (if <= 1.0, it's relative to screen, > 1.0 is absolute)"))
 	FVector2D InteractionSize;
 	UPROPERTY(EditAnywhere, Category = "Control", meta = (ToolTip = "The scale for control input"))
 	FVector2D InputScale;
@@ -75,7 +75,7 @@ class ENGINE_API UTouchInterface : public UObject
 	UPROPERTY(EditAnywhere, Category="TouchInterface", meta=(ToolTip="How long after joystick enabled for touch (0.0 will disable this feature)"))
 	float ActivationDelay;
 
-	UPROPERTY(EditAnywhere, Category="TouchInterface", meta=(ToolTip="Whether to prevent joystick re-center"))
+	UPROPERTY(EditAnywhere, Category="TouchInterface", meta=(ToolTip="Prevent joystick re-centering and moving from Center through user taps"))
 	bool bPreventRecenter;
 
 	UPROPERTY(EditAnywhere, Category = "TouchInterface", meta = (ToolTip = "Delay at startup before virtual joystick is drawn"))
