Expose CreateSinglePrimitiveSceneUniformBuffer in the IRendererModule to make it possible to create a Scene UniformBuffer using public interface.

diff --git a/Engine/Source/Runtime/Engine/Public/PrimitiveSceneShaderData.h b/Engine/Source/Runtime/Engine/Public/PrimitiveSceneShaderData.h
--- a/Engine/Source/Runtime/Engine/Public/PrimitiveSceneShaderData.h
+++ b/Engine/Source/Runtime/Engine/Public/PrimitiveSceneShaderData.h
@@ -69,47 +69,47 @@ public:
 
 	EShaderPlatform ShaderPlatform=SP_NumPlatforms;
 
-	UE_DEPRECATED(5.3, "Use FRendererModule::CreateSinglePrimitiveSceneUniformBuffer instead")
+	UE_DEPRECATED(5.3, "Use IRendererModule::CreateSinglePrimitiveSceneUniformBuffer instead")
 	FPrimitiveSceneShaderData PrimitiveSceneData;
-	UE_DEPRECATED(5.3, "Use FRendererModule::CreateSinglePrimitiveSceneUniformBuffer instead")
+	UE_DEPRECATED(5.3, "Use IRendererModule::CreateSinglePrimitiveSceneUniformBuffer instead")
 	FInstanceSceneShaderData InstanceSceneData;
-	UE_DEPRECATED(5.3, "Use FRendererModule::CreateSinglePrimitiveSceneUniformBuffer instead")
+	UE_DEPRECATED(5.3, "Use IRendererModule::CreateSinglePrimitiveSceneUniformBuffer instead")
 	FLightmapSceneShaderData LightmapSceneData;
 
-	UE_DEPRECATED(5.3, "Use FRendererModule::CreateSinglePrimitiveSceneUniformBuffer instead")
+	UE_DEPRECATED(5.3, "Use IRendererModule::CreateSinglePrimitiveSceneUniformBuffer instead")
 	FBufferRHIRef PrimitiveSceneDataBufferRHI;
-	UE_DEPRECATED(5.3, "Use FRendererModule::CreateSinglePrimitiveSceneUniformBuffer instead")
+	UE_DEPRECATED(5.3, "Use IRendererModule::CreateSinglePrimitiveSceneUniformBuffer instead")
 	FShaderResourceViewRHIRef PrimitiveSceneDataBufferSRV;
 
 	FBufferRHIRef SkyIrradianceEnvironmentMapRHI;
 	FShaderResourceViewRHIRef SkyIrradianceEnvironmentMapSRV;
 
-	UE_DEPRECATED(5.3, "Use FRendererModule::CreateSinglePrimitiveSceneUniformBuffer instead")
+	UE_DEPRECATED(5.3, "Use IRendererModule::CreateSinglePrimitiveSceneUniformBuffer instead")
 	FBufferRHIRef InstanceSceneDataBufferRHI;
-	UE_DEPRECATED(5.3, "Use FRendererModule::CreateSinglePrimitiveSceneUniformBuffer instead")
+	UE_DEPRECATED(5.3, "Use IRendererModule::CreateSinglePrimitiveSceneUniformBuffer instead")
 	FShaderResourceViewRHIRef InstanceSceneDataBufferSRV;
 
-	UE_DEPRECATED(5.3, "Use FRendererModule::CreateSinglePrimitiveSceneUniformBuffer instead")
+	UE_DEPRECATED(5.3, "Use IRendererModule::CreateSinglePrimitiveSceneUniformBuffer instead")
 	FBufferRHIRef InstancePayloadDataBufferRHI;
-	UE_DEPRECATED(5.3, "Use FRendererModule::CreateSinglePrimitiveSceneUniformBuffer instead")
+	UE_DEPRECATED(5.3, "Use IRendererModule::CreateSinglePrimitiveSceneUniformBuffer instead")
 	FShaderResourceViewRHIRef InstancePayloadDataBufferSRV;
 
 	FTextureRHIRef PrimitiveSceneDataTextureRHI;
 	FShaderResourceViewRHIRef PrimitiveSceneDataTextureSRV;
 
-	UE_DEPRECATED(5.3, "Use FRendererModule::CreateSinglePrimitiveSceneUniformBuffer instead")
+	UE_DEPRECATED(5.3, "Use IRendererModule::CreateSinglePrimitiveSceneUniformBuffer instead")
 	FBufferRHIRef LightmapSceneDataBufferRHI;
-	UE_DEPRECATED(5.3, "Use FRendererModule::CreateSinglePrimitiveSceneUniformBuffer instead")
+	UE_DEPRECATED(5.3, "Use IRendererModule::CreateSinglePrimitiveSceneUniformBuffer instead")
 	FShaderResourceViewRHIRef LightmapSceneDataBufferSRV;
 
-	UE_DEPRECATED(5.3, "Use FRendererModule::CreateSinglePrimitiveSceneUniformBuffer instead")
+	UE_DEPRECATED(5.3, "Use IRendererModule::CreateSinglePrimitiveSceneUniformBuffer instead")
 	FBufferRHIRef EditorVisualizeLevelInstanceDataBufferRHI;
-	UE_DEPRECATED(5.3, "Use FRendererModule::CreateSinglePrimitiveSceneUniformBuffer instead")
+	UE_DEPRECATED(5.3, "Use IRendererModule::CreateSinglePrimitiveSceneUniformBuffer instead")
 	FShaderResourceViewRHIRef EditorVisualizeLevelInstanceDataBufferSRV;
 
-	UE_DEPRECATED(5.3, "Use FRendererModule::CreateSinglePrimitiveSceneUniformBuffer instead")
+	UE_DEPRECATED(5.3, "Use IRendererModule::CreateSinglePrimitiveSceneUniformBuffer instead")
 	FBufferRHIRef EditorSelectedDataBufferRHI;
-	UE_DEPRECATED(5.3, "Use FRendererModule::CreateSinglePrimitiveSceneUniformBuffer instead")
+	UE_DEPRECATED(5.3, "Use IRendererModule::CreateSinglePrimitiveSceneUniformBuffer instead")
 	FShaderResourceViewRHIRef EditorSelectedDataBufferSRV;
 };
 
@@ -118,5 +118,5 @@ public:
 * This is used when the VF is used for rendering outside normal mesh passes, where there is no valid scene.
 */
 extern ENGINE_API TGlobalResource<FSinglePrimitiveStructured> GIdentityPrimitiveBuffer;
-UE_DEPRECATED(5.3, "Use FRendererModule::CreateSinglePrimitiveSceneUniformBuffer instead")
+UE_DEPRECATED(5.3, "Use IRendererModule::CreateSinglePrimitiveSceneUniformBuffer instead")
 extern ENGINE_API TGlobalResource<FSinglePrimitiveStructured> GTilePrimitiveBuffer;
diff --git a/Engine/Source/Runtime/RenderCore/Public/RendererInterface.h b/Engine/Source/Runtime/RenderCore/Public/RendererInterface.h
--- a/Engine/Source/Runtime/RenderCore/Public/RendererInterface.h
+++ b/Engine/Source/Runtime/RenderCore/Public/RendererInterface.h
@@ -38,6 +38,7 @@ class FMaterialRenderProxy;
 class FGPUScenePrimitiveCollector;
 class FViewInfo;
 template<typename ShaderType, typename PointerTableType> class TShaderRefBase;
+class FSceneUniformBuffer;
 
 namespace Nanite
 {
@@ -711,6 +712,9 @@ public:
 
 	virtual void InitializeSystemTextures(FRHICommandListImmediate& RHICmdList) = 0;
 
+	/** Create a Scene Uniform Buffer containing only the scene representation for a single primitive */
+	virtual FSceneUniformBuffer* CreateSinglePrimitiveSceneUniformBuffer(FRDGBuilder& GraphBuilder, const FViewInfo& SceneView, FMeshBatch& Mesh) = 0;
+
 	/** Draws a tile mesh element with the specified view. */
 	virtual void DrawTileMesh(FCanvasRenderContext& RenderContext, struct FMeshPassProcessorRenderState& DrawRenderState, const FSceneView& View, FMeshBatch& Mesh, bool bIsHitTesting, const class FHitProxyId& HitProxyId, bool bUse128bitRT = false) = 0;
 
diff --git a/Engine/Source/Runtime/Renderer/Private/RendererModule.h b/Engine/Source/Runtime/Renderer/Private/RendererModule.h
--- a/Engine/Source/Runtime/Renderer/Private/RendererModule.h
+++ b/Engine/Source/Runtime/Renderer/Private/RendererModule.h
@@ -49,7 +49,7 @@ public:
 	virtual uint32 GetNumDynamicLightsAffectingPrimitive(const FPrimitiveSceneInfo* PrimitiveSceneInfo,const FLightCacheInterface* LCI) override;
 	virtual void OnWorldCleanup(UWorld* World, bool bSessionEnded, bool bCleanupResources, bool bWorldChanged) override;
 	virtual void InitializeSystemTextures(FRHICommandListImmediate& RHICmdList);
-	FSceneUniformBuffer* CreateSinglePrimitiveSceneUniformBuffer(FRDGBuilder& GraphBuilder, const FViewInfo& SceneView, FMeshBatch& Mesh);
+	virtual FSceneUniformBuffer* CreateSinglePrimitiveSceneUniformBuffer(FRDGBuilder& GraphBuilder, const FViewInfo& SceneView, FMeshBatch& Mesh) override;
 	virtual void DrawTileMesh(FCanvasRenderContext& RenderContext, FMeshPassProcessorRenderState& DrawRenderState, const FSceneView& View, FMeshBatch& Mesh, bool bIsHitTesting, const FHitProxyId& HitProxyId, bool bUse128bitRT = false) override;
 	virtual void DebugLogOnCrash() override;
 	virtual void GPUBenchmark(FSynthBenchmarkResults& InOut, float WorkScale) override;
