World Partition: ASpatialHashRuntimeGridInfo is now editor only.

diff --git a/Engine/Source/Runtime/Engine/Public/WorldPartition/WorldPartitionRuntimeSpatialHash.h b/Engine/Source/Runtime/Engine/Public/WorldPartition/WorldPartitionRuntimeSpatialHash.h
--- a/Engine/Source/Runtime/Engine/Public/WorldPartition/WorldPartitionRuntimeSpatialHash.h
+++ b/Engine/Source/Runtime/Engine/Public/WorldPartition/WorldPartitionRuntimeSpatialHash.h
@@ -230,6 +230,8 @@ class ASpatialHashRuntimeGridInfo : public AInfo
 	GENERATED_UCLASS_BODY()
 
 public:
+	virtual bool IsEditorOnly() const override { return true; }
+
 	UPROPERTY()
 	FSpatialHashRuntimeGrid	GridSettings;
 };
