Adding CPU scope marker to catch creating debug info.

- This is a slow operation so useful to see in Insights

diff --git a/Engine/Source/Runtime/Engine/Private/ShaderCompiler/ShaderCompiler.cpp b/Engine/Source/Runtime/Engine/Private/ShaderCompiler/ShaderCompiler.cpp
--- a/Engine/Source/Runtime/Engine/Private/ShaderCompiler/ShaderCompiler.cpp
+++ b/Engine/Source/Runtime/Engine/Private/ShaderCompiler/ShaderCompiler.cpp
@@ -4332,6 +4332,8 @@ EShaderDebugInfoFlags FShaderCompilingManager::GetDumpShaderDebugInfoFlags() con
 
 FString FShaderCompilingManager::CreateShaderDebugInfoPath(const FShaderCompilerInput& ShaderCompilerInput) const
 {
+	TRACE_CPUPROFILER_EVENT_SCOPE(FShaderCompilingManager::CreateShaderDebugInfoPath);
+
 	FString DumpDebugInfoPath = ShaderCompilerInput.DumpDebugInfoRootPath / ShaderCompilerInput.DebugGroupName + ShaderCompilerInput.DebugExtension;
 
 	// Sanitize the name to be used as a path
