Fixed Blueprint API of "spawn system with parameter" functions

diff --git a/Engine/Plugins/FX/Niagara/Source/Niagara/Private/NiagaraFunctionLibrary.cpp b/Engine/Plugins/FX/Niagara/Source/Niagara/Private/NiagaraFunctionLibrary.cpp
--- a/Engine/Plugins/FX/Niagara/Source/Niagara/Private/NiagaraFunctionLibrary.cpp
+++ b/Engine/Plugins/FX/Niagara/Source/Niagara/Private/NiagaraFunctionLibrary.cpp
@@ -107,7 +107,7 @@ UNiagaraComponent* UNiagaraFunctionLibrary::SpawnSystemAtLocation(const UObject*
 * Spawns a Niagara System at the specified world location/rotation
 * @return			The spawned UNiagaraComponent
 */
-UNiagaraComponent* UNiagaraFunctionLibrary::SpawnSystemAtLocationWithParams(FFXSystemSpawnParameters& SpawnParams)
+UNiagaraComponent* UNiagaraFunctionLibrary::SpawnSystemAtLocationWithParams(const FFXSystemSpawnParameters& SpawnParams)
 {
 	UNiagaraComponent* PSC = NULL;
 	UNiagaraSystem* NiagaraSystem = Cast<UNiagaraSystem>(SpawnParams.SystemTemplate);
@@ -204,7 +204,7 @@ UNiagaraComponent* UNiagaraFunctionLibrary::SpawnSystemAttached(
 * @return			The spawned UNiagaraComponent
 */
 
-UNiagaraComponent* UNiagaraFunctionLibrary::SpawnSystemAttachedWithParams(FFXSystemSpawnParameters& SpawnParams)
+UNiagaraComponent* UNiagaraFunctionLibrary::SpawnSystemAttachedWithParams(const FFXSystemSpawnParameters& SpawnParams)
 {
 	LLM_SCOPE(ELLMTag::Niagara);
 	
diff --git a/Engine/Plugins/FX/Niagara/Source/Niagara/Public/NiagaraFunctionLibrary.h b/Engine/Plugins/FX/Niagara/Source/Niagara/Public/NiagaraFunctionLibrary.h
--- a/Engine/Plugins/FX/Niagara/Source/Niagara/Public/NiagaraFunctionLibrary.h
+++ b/Engine/Plugins/FX/Niagara/Source/Niagara/Public/NiagaraFunctionLibrary.h
@@ -29,10 +29,10 @@ class NIAGARA_API UNiagaraFunctionLibrary : public UBlueprintFunctionLibrary
 public:
 
 	UFUNCTION(BlueprintCallable, Category = Niagara, meta = (Keywords = "niagara System", WorldContext = "WorldContextObject", UnsafeDuringActorConstruction = "true"))
-	static UNiagaraComponent* SpawnSystemAtLocationWithParams(FFXSystemSpawnParameters& SpawnParams);
+	static UNiagaraComponent* SpawnSystemAtLocationWithParams(const FFXSystemSpawnParameters& SpawnParams);
 
 	UFUNCTION(BlueprintCallable, Category = Niagara, meta = (Keywords = "niagara System", UnsafeDuringActorConstruction = "true"))
-	static UNiagaraComponent* SpawnSystemAttachedWithParams(FFXSystemSpawnParameters& SpawnParams);
+	static UNiagaraComponent* SpawnSystemAttachedWithParams(const FFXSystemSpawnParameters& SpawnParams);
 
 	/**
 	* Spawns a Niagara System at the specified world location/rotation
diff --git a/Engine/Source/Runtime/Engine/Classes/Particles/ParticleSystemComponent.h b/Engine/Source/Runtime/Engine/Classes/Particles/ParticleSystemComponent.h
--- a/Engine/Source/Runtime/Engine/Classes/Particles/ParticleSystemComponent.h
+++ b/Engine/Source/Runtime/Engine/Classes/Particles/ParticleSystemComponent.h
@@ -334,20 +334,43 @@ struct FFXSystemSpawnParameters
 {
 	GENERATED_BODY()
 
-	const UObject* WorldContextObject = nullptr;
-	UFXSystemAsset* SystemTemplate = nullptr;
+	UPROPERTY(EditAnywhere, BlueprintReadWrite, AdvancedDisplay, Category = "Effects|Components|ParticleSystem")
+	TObjectPtr<const UObject> WorldContextObject;
+
+	UPROPERTY(EditAnywhere, BlueprintReadWrite, Category = "Effects|Components|ParticleSystem")
+	TObjectPtr<UFXSystemAsset> SystemTemplate;
+
+	UPROPERTY(EditAnywhere, BlueprintReadWrite, Category = "Effects|Components|ParticleSystem")
 	FVector Location = FVector::ZeroVector;
+
+	UPROPERTY(EditAnywhere, BlueprintReadWrite, Category = "Effects|Components|ParticleSystem")
 	FRotator Rotation = FRotator::ZeroRotator;
+
+	UPROPERTY(EditAnywhere, BlueprintReadWrite, Category = "Effects|Components|ParticleSystem")
 	FVector Scale = FVector(1.f);
 
-	USceneComponent* AttachToComponent = nullptr;
+	UPROPERTY(EditAnywhere, BlueprintReadWrite, AdvancedDisplay, Category = "Effects|Components|ParticleSystem")
+	TObjectPtr<USceneComponent> AttachToComponent;
+	
+	UPROPERTY(EditAnywhere, BlueprintReadWrite, AdvancedDisplay, Category = "Effects|Components|ParticleSystem")
 	FName AttachPointName = NAME_None;
-	EAttachLocation::Type LocationType = EAttachLocation::KeepWorldPosition;
+	
+	UPROPERTY(EditAnywhere, BlueprintReadWrite, AdvancedDisplay, Category = "Effects|Components|ParticleSystem")
+	TEnumAsByte<EAttachLocation::Type> LocationType = EAttachLocation::KeepWorldPosition;
 
+	UPROPERTY(EditAnywhere, BlueprintReadWrite, Category = "Effects|Components|ParticleSystem")
 	bool bAutoDestroy = true;
+
+	UPROPERTY(EditAnywhere, BlueprintReadWrite, Category = "Effects|Components|ParticleSystem")
 	bool bAutoActivate = true;
+
+	UPROPERTY(EditAnywhere, BlueprintReadWrite, AdvancedDisplay, Category = "Effects|Components|ParticleSystem")
 	EPSCPoolMethod PoolingMethod = EPSCPoolMethod::None;
+
+	UPROPERTY(EditAnywhere, BlueprintReadWrite, AdvancedDisplay, Category = "Effects|Components|ParticleSystem")
 	bool bPreCullCheck = true;
+
+	UPROPERTY(EditAnywhere, BlueprintReadWrite, AdvancedDisplay, Category = "Effects|Components|ParticleSystem")
 	bool bIsPlayerEffect = false;
 };
 
