[DebugDrawService] use synchronization object for delegate unregistration too since it can happen from multiple threads from RecreateRenderState_Concurrent that will unregister and register delegates from multiple threads.

diff --git a/Engine/Source/Runtime/Engine/Classes/Debug/DebugDrawService.h b/Engine/Source/Runtime/Engine/Classes/Debug/DebugDrawService.h
--- a/Engine/Source/Runtime/Engine/Classes/Debug/DebugDrawService.h
+++ b/Engine/Source/Runtime/Engine/Classes/Debug/DebugDrawService.h
@@ -3,17 +3,19 @@
 
 #include "ShowFlags.h"
 #include "Kismet/BlueprintFunctionLibrary.h"
-#include "Misc/MTAccessDetector.h"
 #include "HAL/CriticalSection.h"
 #include "DebugDrawService.generated.h"
 
+class APlayerController;
 class FCanvas;
 class FSceneView;
+class FViewport;
+class UCanvas;
 
 /** 
- * 
+ * Delegate called when rendering a viewport and the associated engine flag is set.
  */
-DECLARE_DELEGATE_TwoParams(FDebugDrawDelegate, class UCanvas*, class APlayerController*);
+DECLARE_DELEGATE_TwoParams(FDebugDrawDelegate, UCanvas*, APlayerController*);
 
 UCLASS(config=Engine)
 class ENGINE_API UDebugDrawService : public UBlueprintFunctionLibrary
@@ -23,30 +25,30 @@ class ENGINE_API UDebugDrawService : public UBlueprintFunctionLibrary
 	static FDelegateHandle Register(const TCHAR* Name, const FDebugDrawDelegate& NewDelegate);
 	static void Unregister(FDelegateHandle HandleToRemove);
 
-	// Draws debug canvas that has already been initialized to a viewport
-	static void Draw(const FEngineShowFlags Flags, class UCanvas* Canvas);
+	/** Draws debug canvas that has already been initialized to a viewport */
+	static void Draw(const FEngineShowFlags Flags, UCanvas* Canvas);
 
-	// Initialize a debug canvas object then calls above draw. If CanvasObject is null it will find/create it for you
-	static void Draw(const FEngineShowFlags Flags, class FViewport* Viewport, FSceneView* View, FCanvas* Canvas, class UCanvas* CanvasObject = nullptr);
+	/** Initialize a debug canvas object then calls above draw. If CanvasObject is null it will find/create it for you */
+	static void Draw(const FEngineShowFlags Flags, FViewport* Viewport, FSceneView* View, FCanvas* Canvas, UCanvas* CanvasObject = nullptr);
 
 private:
 
 	/**
-	 * Synchronization object for delegate registration since it can happen from multiple threads from
-	 * primitives added in batch through
-	 *		ParallelFor(AddPrimitiveBatches.Num(), [&](int32 Index)
+	 * Synchronization object for delegate registration since it can happen from multiple threads.
+	 * e.g.
+	 *		primitives added in batch through
+	 *			ParallelFor(AddPrimitiveBatches.Num(), [&](int32 Index)
 	 *			{
 	 *				FOptionalTaskTagScope Scope(ETaskTag::EParallelGameThread);
 	 *				Scene->AddPrimitive(...);
 	 *			}
-	 *
-	 * Critical section used only for registration since it should not be required for the other accesses (i.e. Unregister &Draw).
-	 * In those cases we use the AccessDetector.
+	 *		or
+	 *			RecreateRenderState_Concurrent that will unregister and register delegates from multiple threads.
 	 */
 	static FCriticalSection DelegatesLock;
-#if ENABLE_MT_DETECTOR
-	static FRWRecursiveAccessDetector DelegatesDetector;
-#endif
-	static TArray<TArray<FDebugDrawDelegate> > Delegates;
+
+	DECLARE_MULTICAST_DELEGATE_TwoParams(FDebugDrawMulticastDelegate, UCanvas*, APlayerController*);
+
+	static TArray<FDebugDrawMulticastDelegate> Delegates;
 	static FEngineShowFlags ObservedFlags;
 };
diff --git a/Engine/Source/Runtime/Engine/Private/Debug/DebugDrawService.cpp b/Engine/Source/Runtime/Engine/Private/Debug/DebugDrawService.cpp
--- a/Engine/Source/Runtime/Engine/Private/Debug/DebugDrawService.cpp
+++ b/Engine/Source/Runtime/Engine/Private/Debug/DebugDrawService.cpp
@@ -8,11 +8,8 @@
 #include UE_INLINE_GENERATED_CPP_BY_NAME(DebugDrawService)
 
 FCriticalSection UDebugDrawService::DelegatesLock;
-#if ENABLE_MT_DETECTOR
-FRWRecursiveAccessDetector UDebugDrawService::DelegatesDetector;
-#endif
 
-TArray<TArray<FDebugDrawDelegate> > UDebugDrawService::Delegates;
+TArray<UDebugDrawService::FDebugDrawMulticastDelegate> UDebugDrawService::Delegates;
 FEngineShowFlags UDebugDrawService::ObservedFlags(ESFIM_Editor);
 
 UDebugDrawService::UDebugDrawService(const FObjectInitializer& ObjectInitializer)
@@ -28,14 +25,12 @@ FDelegateHandle UDebugDrawService::Register(const TCHAR* Name, const FDebugDrawD
 	FDelegateHandle Result;
 	if (Index != INDEX_NONE)
 	{
-		FScopeLock PathLock(&DelegatesLock);
-		UE_MT_SCOPED_WRITE_ACCESS(DelegatesDetector);
+		FScopeLock ScopeLock(&DelegatesLock);
 		if (Index >= Delegates.Num())
 		{
 			Delegates.AddZeroed(Index - Delegates.Num() + 1);
 		}
-		Delegates[Index].Add(NewDelegate);
-		Result = Delegates[Index].Last().GetHandle();
+		Result = Delegates[Index].Add(NewDelegate);
 		ObservedFlags.SetSingleFlag(Index, true);
 	}
 	return Result;
@@ -43,19 +38,18 @@ FDelegateHandle UDebugDrawService::Register(const TCHAR* Name, const FDebugDrawD
 
 void UDebugDrawService::Unregister(const FDelegateHandle HandleToRemove)
 {
-	UE_MT_SCOPED_WRITE_ACCESS(DelegatesDetector);
-	TArray<FDebugDrawDelegate>* DelegatesArray = Delegates.GetData();
-	for (int32 Flag = 0; Flag < Delegates.Num(); ++Flag, ++DelegatesArray)
+	FScopeLock ScopeLock(&DelegatesLock);
+	for (int32 Flag = 0; Flag < Delegates.Num(); ++Flag)
 	{
-		check(DelegatesArray); //it shouldn't happen, but to be sure
-		const uint32 Index = DelegatesArray->IndexOfByPredicate([=](const FDebugDrawDelegate& Delegate){ return Delegate.GetHandle() == HandleToRemove; });
-		if (Index != INDEX_NONE)
+		FDebugDrawMulticastDelegate& MulticastDelegate = Delegates[Flag];
+		if (MulticastDelegate.Remove(HandleToRemove))
 		{
-			DelegatesArray->RemoveAtSwap(Index, 1, false);
-			if (DelegatesArray->Num() == 0)
+			if (MulticastDelegate.IsBound() == false)
 			{
 				ObservedFlags.SetSingleFlag(Flag, false);
 			}
+
+			break;
 		}
 	}	
 }
@@ -89,24 +83,14 @@ void UDebugDrawService::Draw(const FEngineShowFlags Flags, UCanvas* Canvas)
 		return;
 	}
 
-	UE_MT_SCOPED_WRITE_ACCESS(DelegatesDetector);
+	FScopeLock ScopeLock(&DelegatesLock);
 	for (int32 FlagIndex = 0; FlagIndex < Delegates.Num(); ++FlagIndex)
 	{
-		if (Flags.GetSingleFlag(FlagIndex) && ObservedFlags.GetSingleFlag(FlagIndex) && Delegates[FlagIndex].Num() > 0)
+		FDebugDrawMulticastDelegate& MulticastDelegate = Delegates[FlagIndex];
+
+		if (Flags.GetSingleFlag(FlagIndex) && ObservedFlags.GetSingleFlag(FlagIndex))
 		{
-			for (int32 i = Delegates[FlagIndex].Num() - 1; i >= 0; --i)
-			{
-				FDebugDrawDelegate& Delegate = Delegates[FlagIndex][i];
-
-				if (Delegate.IsBound())
-				{
-					Delegate.Execute(Canvas, nullptr);
-				}
-				else
-				{
-					Delegates[FlagIndex].RemoveAtSwap(i, 1, /*bAllowShrinking=*/false);
-				}
-			}
+			MulticastDelegate.Broadcast(Canvas, nullptr);
 		}
 	}
 }
