Don't destroy orphaned default subobjects on load, instead wait for construction script rerunning, and then destroy them as we do for blueprint defined components

diff --git a/Engine/Source/Runtime/CoreUObject/Private/UObject/LinkerLoad.cpp b/Engine/Source/Runtime/CoreUObject/Private/UObject/LinkerLoad.cpp
--- a/Engine/Source/Runtime/CoreUObject/Private/UObject/LinkerLoad.cpp
+++ b/Engine/Source/Runtime/CoreUObject/Private/UObject/LinkerLoad.cpp
@@ -5163,23 +5163,12 @@ UObject* FLinkerLoad::CreateExport( int32 Index )
 		// default subobjects; that is, the subobject owner's archetype is expected to also contain a matching default subobject
 		// instance with the same type/name. However, if the instanced subobject is based on an archetype that's owned by
 		// something other than it's owner's archetype (e.g. Blueprint-added component archetypes, which are owned by the
-		// Blueprint class object), such a match would not exist, and the export would then fail to be created below as a result.
+		// Blueprint class object), such a match would not exist.
 		if (Export.bIsInheritedInstance && Template->GetOuter()->IsA<UClass>())
 		{
 			Export.bIsInheritedInstance = false;
 		}
 
-		// If this export was an instance inherited from another class, 
-		// it should not be recreated unless the archetype also has it
-		if (Export.bIsInheritedInstance)
-		{
-			if (FindObjectWithOuter(ThisParent->GetArchetype(), LoadClass, NewName) == nullptr)
-			{
-				Export.bExportLoadFailed = true;
-				return nullptr;
-			}
-		}
-
 		LoadClass->GetDefaultObject();
 
 		FStaticConstructObjectParameters Params(LoadClass);
diff --git a/Engine/Source/Runtime/Engine/Private/ActorConstruction.cpp b/Engine/Source/Runtime/Engine/Private/ActorConstruction.cpp
--- a/Engine/Source/Runtime/Engine/Private/ActorConstruction.cpp
+++ b/Engine/Source/Runtime/Engine/Private/ActorConstruction.cpp
@@ -182,6 +182,19 @@ void AActor::DestroyConstructedComponents()
 				}
 			}
 
+			if (!bDestroyComponent)
+			{
+				// check for orphaned natively created components:
+				if (Component->CreationMethod == EComponentCreationMethod::Native && Component->HasAnyFlags(RF_DefaultSubObject))
+				{
+					UObject* ComponentArchetype = Component->GetArchetype();
+					if (ComponentArchetype == ComponentArchetype->GetClass()->ClassDefaultObject)
+					{
+						bDestroyComponent = true;
+					}
+				}
+			}
+
 			if (bDestroyComponent)
 			{
 				if (Component == RootComponent)
diff --git a/Engine/Source/Runtime/Engine/Private/Components/ActorComponent.cpp b/Engine/Source/Runtime/Engine/Private/Components/ActorComponent.cpp
--- a/Engine/Source/Runtime/Engine/Private/Components/ActorComponent.cpp
+++ b/Engine/Source/Runtime/Engine/Private/Components/ActorComponent.cpp
@@ -354,36 +354,7 @@ void UActorComponent::PostInitProperties()
 	// Instance components will be added during the owner's initialization
 	if (OwnerPrivate && CreationMethod != EComponentCreationMethod::Instance)
 	{
-		if (!FPlatformProperties::RequiresCookedData() && !OwnerPrivate->GetClass()->bCooked && CreationMethod == EComponentCreationMethod::Native && HasAllFlags(RF_NeedLoad|RF_DefaultSubObject))
-		{
-			UObject* MyArchetype = GetArchetype();
-			if (IsValid(MyArchetype) && MyArchetype != GetClass()->ClassDefaultObject)
-			{
-				OwnerPrivate->AddOwnedComponent(this);
-			}
-			else
-			{
-				// else: this is a natively created component that thinks its archetype is the CDO of
-				// this class, rather than a template component and this isn't the template component.
-				// Delete this stale component
-#if WITH_EDITOR
-				if (HasAnyInternalFlags(EInternalObjectFlags::AsyncLoading))
-				{
-					// Async loading components cannot be pending kill, or the async loading code will assert when trying to postload them.
-					// Instead, wait until the postload and mark pending kill at that time
-					bMarkPendingKillOnPostLoad = true;
-				}
-				else
-#endif // WITH_EDITOR
-				{
-					MarkAsGarbage();
-				}
-			}
-		}
-		else
-		{
-			OwnerPrivate->AddOwnedComponent(this);
-		}
+		OwnerPrivate->AddOwnedComponent(this);
 	}
 }
 
