Don't conform/migrate sparse data when loading a cooked BPGC in the editor

#rnx

diff --git a/Engine/Source/Runtime/Engine/Private/BlueprintGeneratedClass.cpp b/Engine/Source/Runtime/Engine/Private/BlueprintGeneratedClass.cpp
--- a/Engine/Source/Runtime/Engine/Private/BlueprintGeneratedClass.cpp
+++ b/Engine/Source/Runtime/Engine/Private/BlueprintGeneratedClass.cpp
@@ -566,7 +566,7 @@ void UBlueprintGeneratedClass::SerializeDefaultObject(UObject* Object, FStructur
 	}
 
 #if WITH_EDITORONLY_DATA
-	if (bIsSparseClassDataSerializable || GetOutermost()->bIsCookedForEditor)
+	if (bIsSparseClassDataSerializable || GetPackage()->HasAnyPackageFlags(PKG_Cooked))
 #endif
 	{
 		if(bSkipSparseClassDataSerialization)
@@ -596,7 +596,10 @@ void UBlueprintGeneratedClass::SerializeDefaultObject(UObject* Object, FStructur
 		}
 
 #if WITH_EDITOR
-		PrepareToConformSparseClassData();
+		if (!GetPackage()->HasAnyPackageFlags(PKG_Cooked))
+		{
+			PrepareToConformSparseClassData();
+		}
 #endif
 	}
 }
@@ -620,15 +623,17 @@ void UBlueprintGeneratedClass::PostLoadDefaultObject(UObject* Object)
 	}
 
 #if WITH_EDITOR
-	Object->MoveDataToSparseClassDataStruct();
+	if (!GetPackage()->HasAnyPackageFlags(PKG_Cooked))
+	{
+		ConformSparseClassData(Object);
+		Object->MoveDataToSparseClassDataStruct();
+	}
 
 	if (Object->GetSparseClassDataStruct())
 	{
 		// now that any data has been moved into the sparse data structure we can safely serialize it
 		bIsSparseClassDataSerializable = true;
 	}
-
-	ConformSparseClassData(Object);
 #endif
 }
 
