[DebugRenderProxy] extract code so that sub classes can use it :
- add Draw method for each type of shape
- expose MaterialCache class
- add virtual method GetDynamicMeshElementsForView to avoid duplicating the view filtering logic

diff --git a/Engine/Source/Runtime/Engine/Private/DebugRenderSceneProxy.cpp b/Engine/Source/Runtime/Engine/Private/DebugRenderSceneProxy.cpp
--- a/Engine/Source/Runtime/Engine/Private/DebugRenderSceneProxy.cpp
+++ b/Engine/Source/Runtime/Engine/Private/DebugRenderSceneProxy.cpp
@@ -129,186 +129,98 @@ void FDebugRenderSceneProxy::GetDynamicMeshElements(const TArray<const FSceneVie
 {
 	QUICK_SCOPE_CYCLE_COUNTER( STAT_DebugRenderSceneProxy_GetDynamicMeshElements );
 
-	// Draw solid spheres
-	struct FMaterialCache
-	{
-		FMaterialCache()
-			: bUseFakeLight(false)
-		{}
+	FMaterialCache DefaultMaterialCache(Collector);
+	FMaterialCache SolidMeshMaterialCache(Collector, /**bUseLight*/ true, SolidMeshMaterial.Get());
 
-		FMaterialRenderProxy* operator[](FLinearColor Color)
+	for (int32 ViewIndex = 0; ViewIndex < Views.Num(); ViewIndex++)
+	{
+		if (VisibilityMap & (1 << ViewIndex))
 		{
-			FMaterialRenderProxy* MeshColor = NULL;
-			const uint32 HashKey = GetTypeHashHelper(Color);
-			if (MeshColorInstances.Contains(HashKey))
-			{
-				MeshColor = *MeshColorInstances.Find(HashKey);
-			}
-			else
-			{
-				if (bUseFakeLight && SolidMeshMaterial.IsValid())
-				{
-					
-					MeshColor = &Collector->AllocateOneFrameResource<FColoredMaterialRenderProxy>(
-						SolidMeshMaterial->GetRenderProxy(),
-						Color,
-						"GizmoColor"
-						);
-				}
-				else
-				{
-					MeshColor = &Collector->AllocateOneFrameResource<FColoredMaterialRenderProxy>(GEngine->DebugMeshMaterial->GetRenderProxy(), Color);
-				}
-
-				MeshColorInstances.Add(HashKey, MeshColor);
-			}
-
-			return MeshColor;
+			const FSceneView* View = Views[ViewIndex];
+			GetDynamicMeshElementsForView(View, ViewIndex, ViewFamily, VisibilityMap, Collector, DefaultMaterialCache, SolidMeshMaterialCache);
 		}
+	}
+}
 
-		void Init(FMeshElementCollector& InCollector, bool UseLight = false, class UMaterial* InMaterial = nullptr)
-		{
-			Collector = &InCollector;
-			bUseFakeLight = UseLight;
-			SolidMeshMaterial = InMaterial;
-		}
+void FDebugRenderSceneProxy::GetDynamicMeshElementsForView(const FSceneView* View, const int32 ViewIndex, const FSceneViewFamily& ViewFamily, const uint32 VisibilityMap, FMeshElementCollector& Collector, FMaterialCache& DefaultMaterialCache, FMaterialCache& SolidMeshMaterialCache) const
+{
+	FPrimitiveDrawInterface* PDI = Collector.GetPDI(ViewIndex);
+
+	// Draw Lines
+	const int32 LinesNum = Lines.Num();
+	PDI->AddReserveLines(SDPG_World, LinesNum, false, false);
+	for (const FDebugLine& Line : Lines)
+	{
+		Line.Draw(PDI);
+	}
 
-		FMeshElementCollector* Collector = nullptr;
-		TMap<uint32, FMaterialRenderProxy*> MeshColorInstances;
-		TWeakObjectPtr<class UMaterial> SolidMeshMaterial;
-		bool bUseFakeLight;
-	};
+	// Draw Dashed Lines
+	for (const FDashedLine& Dash : DashedLines)
+	{
+		Dash.Draw(PDI);
+	}
 
-	FMaterialCache MaterialCache[2];
-	MaterialCache[0].Init(Collector);
-	MaterialCache[1].Init(Collector, true, SolidMeshMaterial.Get());
+	// Draw Arrows
+	const uint32 ArrowsNum = ArrowLines.Num();
+	PDI->AddReserveLines(SDPG_World, 5 * ArrowsNum, false, false);
+	for (const FArrowLine& ArrowLine : ArrowLines)
+	{
+		ArrowLine.Draw(PDI, 8.f);
+	}
 
-	for (int32 ViewIndex = 0; ViewIndex < Views.Num(); ViewIndex++)
+	// Draw Stars
+	for (const FWireStar& Star : Stars)
 	{
-		if (VisibilityMap & (1 << ViewIndex))
+		Star.Draw(PDI);
+	}
+
+	// Draw Cylinders
+	for (const FWireCylinder& Cylinder : Cylinders)
+	{
+		Cylinder.Draw(PDI, DrawType, DrawAlpha, DefaultMaterialCache, ViewIndex, Collector);
+	}
+
+	// Draw Boxes
+	for (const FDebugBox& Box : Boxes)
+	{
+		Box.Draw(PDI, DrawType, DrawAlpha, DefaultMaterialCache, ViewIndex, Collector);
+	}
+
+	// Draw Cones
+	TArray<FVector> Verts;
+	for (const FCone& Cone : Cones)
+	{
+		Cone.Draw(PDI, DrawType, DrawAlpha, DefaultMaterialCache, ViewIndex, Collector, &Verts);
+	}
+
+	// Draw spheres
+	for (const FSphere& Sphere : Spheres)
+	{
+		if (PointInView(Sphere.Location, View))
 		{
-			const FSceneView* View = Views[ViewIndex];
-			FPrimitiveDrawInterface* PDI = Collector.GetPDI(ViewIndex);
-
-			// Draw Lines
-			const int32 LinesNum = Lines.Num();
-			PDI->AddReserveLines(SDPG_World, LinesNum, false, false);
-			for (const auto& CurrentLine : Lines)
-			{
-				PDI->DrawLine(CurrentLine.Start, CurrentLine.End, CurrentLine.Color, SDPG_World, CurrentLine.Thickness, 0, CurrentLine.Thickness > 0);
-			}
-
-			// Draw Dashed Lines
-			for(int32 DashIdx=0; DashIdx<DashedLines.Num(); DashIdx++)
-			{
-				const FDashedLine& Dash = DashedLines[DashIdx];
-
-				DrawDashedLine(PDI, Dash.Start, Dash.End, Dash.Color, Dash.DashSize, SDPG_World);
-			}
-
-			// Draw Arrows
-			const uint32 ArrowsNum = ArrowLines.Num();
-			PDI->AddReserveLines(SDPG_World, 5 * ArrowsNum, false, false);
-			for (const auto& CurrentArrow : ArrowLines)
-			{
-				DrawLineArrow(PDI, CurrentArrow.Start, CurrentArrow.End, CurrentArrow.Color, 8.0f);
-			}
-
-			// Draw Stars
-			for(int32 StarIdx=0; StarIdx<Stars.Num(); StarIdx++)
-			{
-				const FWireStar& Star = Stars[StarIdx];
-
-				DrawWireStar(PDI, Star.Position, Star.Size, Star.Color, SDPG_World);
-			}
-
-			// Draw Cylinders
-			for(const auto& Cylinder : Cylinders)
-			{
-				if (DrawType == SolidAndWireMeshes || DrawType == WireMesh)
-				{
-					DrawWireCylinder(PDI, Cylinder.Base, FVector(1, 0, 0), FVector(0, 1, 0), FVector(0, 0, 1), Cylinder.Color, Cylinder.Radius, Cylinder.HalfHeight, (DrawType == SolidAndWireMeshes) ? 9 : 16, SDPG_World, DrawType == SolidAndWireMeshes ? 2 : 0, 0, true);
-				}
-
-				if (DrawType == SolidAndWireMeshes || DrawType == SolidMesh)
-				{
-					GetCylinderMesh(Cylinder.Base, FVector(1, 0, 0), FVector(0, 1, 0), FVector(0, 0, 1), Cylinder.Radius, Cylinder.HalfHeight, 16, MaterialCache[0][Cylinder.Color.WithAlpha(DrawAlpha)], SDPG_World, ViewIndex, Collector);
-				}
-			}
-
-			// Draw Boxes
-			for(const auto& Box :  Boxes)
-			{
-				if (DrawType == SolidAndWireMeshes || DrawType == WireMesh)
-				{
-					DrawWireBox(PDI, Box.Transform.ToMatrixWithScale(), Box.Box, Box.Color, SDPG_World, DrawType == SolidAndWireMeshes ? 2 : 0, 0, true);
-				}
-				if (DrawType == SolidAndWireMeshes || DrawType == SolidMesh)
-				{
-					GetBoxMesh(FTransform(Box.Box.GetCenter()).ToMatrixNoScale() * Box.Transform.ToMatrixWithScale(), Box.Box.GetExtent(), MaterialCache[0][Box.Color.WithAlpha(DrawAlpha)], SDPG_World, ViewIndex, Collector);
-				}
-			}
-
-			// Draw Boxes
-			TArray<FVector> Verts;
-			for (auto& CurrentCone : Cones)
-			{
-				if (DrawType == SolidAndWireMeshes || DrawType == WireMesh)
-				{
-					DrawWireCone(PDI, Verts, CurrentCone.ConeToWorld, 1, CurrentCone.Angle2, (DrawType == SolidAndWireMeshes) ? 9 : 16, CurrentCone.Color, SDPG_World, DrawType == SolidAndWireMeshes ? 2 : 0, 0, true);
-				}
-				if (DrawType == SolidAndWireMeshes || DrawType == SolidMesh)
-				{
-					GetConeMesh(CurrentCone.ConeToWorld, CurrentCone.Angle1, CurrentCone.Angle2, 16, MaterialCache[0][CurrentCone.Color.WithAlpha(DrawAlpha)], SDPG_World, ViewIndex, Collector);
-				}
-			}
-
-			for (auto It = Spheres.CreateConstIterator(); It; ++It)
-			{
-				if (PointInView(It->Location, View))
-				{
-					if (DrawType == SolidAndWireMeshes || DrawType == WireMesh)
-					{
-						DrawWireSphere(PDI, It->Location, It->Color.WithAlpha(255), It->Radius, 20, SDPG_World, DrawType == SolidAndWireMeshes ? 2 : 0, 0, true);
-					}
-					if (DrawType == SolidAndWireMeshes || DrawType == SolidMesh)
-					{
-						GetSphereMesh(It->Location, FVector(It->Radius), 20, 7, MaterialCache[0][It->Color.WithAlpha(DrawAlpha)], SDPG_World, false, ViewIndex, Collector);
-					}
-				}
-			}
-
-			for (auto It = Capsules.CreateConstIterator(); It; ++It)
-			{
-				if (PointInView(It->Base, View))
-				{
-					if (DrawType == SolidAndWireMeshes || DrawType == WireMesh)
-					{
-						const float HalfAxis = FMath::Max<float>(It->HalfHeight - It->Radius, 1.f);
-						const FVector BottomEnd = It->Base + It->Radius * It->Z;
-						const FVector TopEnd = BottomEnd + (2 * HalfAxis) * It->Z;
-						const float CylinderHalfHeight = (TopEnd - BottomEnd).Size() * 0.5;
-						const FVector CylinderLocation = BottomEnd + CylinderHalfHeight * It->Z;
-						DrawWireCapsule(PDI, CylinderLocation, It->X, It->Y, It->Z, It->Color, It->Radius, It->HalfHeight, (DrawType == SolidAndWireMeshes) ? 9 : 16, SDPG_World, DrawType == SolidAndWireMeshes ? 2 : 0, 0, true);
-					}
-					if (DrawType == SolidAndWireMeshes || DrawType == SolidMesh)
-					{
-						GetCapsuleMesh(It->Base, It->X, It->Y, It->Z, It->Color, It->Radius, It->HalfHeight, 16, MaterialCache[0][It->Color.WithAlpha(DrawAlpha)], SDPG_World, false, ViewIndex, Collector);
-					}
-				}
-			}
-
-			for (const auto& Mesh : Meshes)
-			{
-				FDynamicMeshBuilder MeshBuilder(View->GetFeatureLevel());
-				MeshBuilder.AddVertices(Mesh.Vertices);
-				MeshBuilder.AddTriangles(Mesh.Indices);
-
-				MeshBuilder.GetMesh(FMatrix::Identity, MaterialCache[Mesh.Color.A == 255 ? 1 : 0][Mesh.Color], SDPG_World, false, false, ViewIndex, Collector);
-			}
+			Sphere.Draw(PDI, DrawType, DrawAlpha, DefaultMaterialCache, ViewIndex, Collector);
+		}
+	}
 
+	// Draw Capsules
+	for (const FCapsule& Capsule : Capsules)
+	{
+		if (PointInView(Capsule.Base, View))
+		{
+			Capsule.Draw(PDI, DrawType, DrawAlpha, DefaultMaterialCache, ViewIndex, Collector);
 		}
 	}
+
+	// Draw Meshes
+	for (const FMesh& Mesh : Meshes)
+	{
+		FDynamicMeshBuilder MeshBuilder(View->GetFeatureLevel());
+		MeshBuilder.AddVertices(Mesh.Vertices);
+		MeshBuilder.AddTriangles(Mesh.Indices);
+
+		FMaterialCache& MeshMaterialCache = Mesh.Color.A == 255 ? SolidMeshMaterialCache : DefaultMaterialCache; 
+		MeshBuilder.GetMesh(FMatrix::Identity, MeshMaterialCache[Mesh.Color], SDPG_World, false, false, ViewIndex, Collector);
+	}
 }
 
 /**
@@ -320,6 +232,58 @@ void FDebugRenderSceneProxy::GetDynamicMeshElements(const TArray<const FSceneVie
 * @param Mag		Size of the arrow.
 */
 void FDebugRenderSceneProxy::DrawLineArrow(FPrimitiveDrawInterface* PDI,const FVector &Start,const FVector &End,const FColor &Color,float Mag) const
+{
+	const FArrowLine ArrowLine = {Start, End, Color};
+	ArrowLine.Draw(PDI, Mag);
+}
+
+FDebugRenderSceneProxy::FMaterialCache::FMaterialCache(FMeshElementCollector& InCollector, bool bUseLight, UMaterial* InMaterial)
+	: Collector(InCollector)
+	, SolidMeshMaterial(InMaterial)
+	, bUseFakeLight(bUseLight)
+{
+}
+
+FMaterialRenderProxy* FDebugRenderSceneProxy::FMaterialCache::operator[](FLinearColor Color)
+{
+	FMaterialRenderProxy* MeshColor = nullptr;
+	const uint32 HashKey = GetTypeHashHelper(Color);
+	if (MeshColorInstances.Contains(HashKey))
+	{
+		MeshColor = *MeshColorInstances.Find(HashKey);
+	}
+	else
+	{
+		if (bUseFakeLight && SolidMeshMaterial.IsValid())
+		{
+			MeshColor = &Collector.AllocateOneFrameResource<FColoredMaterialRenderProxy>(
+				SolidMeshMaterial->GetRenderProxy(),
+				Color,
+				"GizmoColor"
+				);
+		}
+		else
+		{
+			MeshColor = &Collector.AllocateOneFrameResource<FColoredMaterialRenderProxy>(GEngine->DebugMeshMaterial->GetRenderProxy(), Color);
+		}
+
+		MeshColorInstances.Add(HashKey, MeshColor);
+	}
+
+	return MeshColor;
+}
+
+void FDebugRenderSceneProxy::FDebugLine::Draw(FPrimitiveDrawInterface* PDI) const
+{
+	PDI->DrawLine(Start, End, Color, SDPG_World, Thickness, 0, Thickness > 0);
+}
+
+void FDebugRenderSceneProxy::FWireStar::Draw(FPrimitiveDrawInterface* PDI) const
+{
+	DrawWireStar(PDI, Position, Size, Color, SDPG_World);
+}
+
+void FDebugRenderSceneProxy::FArrowLine::Draw(FPrimitiveDrawInterface* PDI, const float Mag) const
 {
 	// draw a pretty arrow
 	FVector Dir = End - Start;
@@ -330,3 +294,77 @@ void FDebugRenderSceneProxy::DrawLineArrow(FPrimitiveDrawInterface* PDI,const FV
 	FMatrix ArrowTM(Dir,YAxis,ZAxis,Start);
 	DrawDirectionalArrow(PDI,ArrowTM,Color,DirMag,Mag,SDPG_World);
 }
+
+void FDebugRenderSceneProxy::FDashedLine::Draw(FPrimitiveDrawInterface* PDI) const
+{
+	DrawDashedLine(PDI, Start, End, Color, DashSize, SDPG_World);
+}
+
+void FDebugRenderSceneProxy::FDebugBox::Draw(FPrimitiveDrawInterface* PDI, EDrawType InDrawType, uint32 InDrawAlpha, FMaterialCache& MaterialCache, int32 ViewIndex, FMeshElementCollector& Collector) const
+{
+	if (InDrawType == SolidAndWireMeshes || InDrawType == WireMesh)
+	{
+		DrawWireBox(PDI, Transform.ToMatrixWithScale(), Box, Color, SDPG_World, InDrawType == SolidAndWireMeshes ? 2 : 0, 0, true);
+	}
+	if (InDrawType == SolidAndWireMeshes || InDrawType == SolidMesh)
+	{
+		GetBoxMesh(FTransform(Box.GetCenter()).ToMatrixNoScale() * Transform.ToMatrixWithScale(), Box.GetExtent(), MaterialCache[Color.WithAlpha(InDrawAlpha)], SDPG_World, ViewIndex, Collector);
+	}
+}
+
+void FDebugRenderSceneProxy::FWireCylinder::Draw(FPrimitiveDrawInterface* PDI, EDrawType InDrawType, uint32 InDrawAlpha, FMaterialCache& MaterialCache, int32 ViewIndex, FMeshElementCollector& Collector) const
+{
+	if (InDrawType == SolidAndWireMeshes || InDrawType == WireMesh)
+	{
+		DrawWireCylinder(PDI, Base, FVector(1, 0, 0), FVector(0, 1, 0), FVector(0, 0, 1), Color, Radius, HalfHeight, (InDrawType == SolidAndWireMeshes) ? 9 : 16, SDPG_World, InDrawType == SolidAndWireMeshes ? 2 : 0, 0, true);
+	}
+
+	if (InDrawType == SolidAndWireMeshes || InDrawType == SolidMesh)
+	{
+		GetCylinderMesh(Base, FVector(1, 0, 0), FVector(0, 1, 0), FVector(0, 0, 1), Radius, HalfHeight, 16, MaterialCache[Color.WithAlpha(InDrawAlpha)], SDPG_World, ViewIndex, Collector);
+	}
+}
+
+void FDebugRenderSceneProxy::FCone::Draw(FPrimitiveDrawInterface* PDI, EDrawType InDrawType, uint32 InDrawAlpha, FMaterialCache& MaterialCache, int32 ViewIndex, FMeshElementCollector& Collector, TArray<FVector>* VertsCache) const
+{
+	if (InDrawType == SolidAndWireMeshes || InDrawType == WireMesh)
+	{
+		TArray<FVector> LocalVertsCache;
+		TArray<FVector>& Verts = VertsCache != nullptr ? *VertsCache : LocalVertsCache;
+		DrawWireCone(PDI, Verts, ConeToWorld, 1, Angle2, (InDrawType == SolidAndWireMeshes) ? 9 : 16, Color, SDPG_World, InDrawType == SolidAndWireMeshes ? 2 : 0, 0, true);
+	}
+	if (InDrawType == SolidAndWireMeshes || InDrawType == SolidMesh)
+	{
+		GetConeMesh(ConeToWorld, Angle1, Angle2, 16, MaterialCache[Color.WithAlpha(InDrawAlpha)], SDPG_World, ViewIndex, Collector);
+	}
+}
+
+
+void FDebugRenderSceneProxy::FSphere::Draw(FPrimitiveDrawInterface* PDI, EDrawType InDrawType, uint32 InDrawAlpha, FMaterialCache& MaterialCache, int32 ViewIndex, FMeshElementCollector& Collector) const
+{
+	if (InDrawType == SolidAndWireMeshes || InDrawType == WireMesh)
+	{
+		DrawWireSphere(PDI, Location, Color.WithAlpha(255), Radius, 20, SDPG_World, InDrawType == SolidAndWireMeshes ? 2 : 0, 0, true);
+	}
+	if (InDrawType == SolidAndWireMeshes || InDrawType == SolidMesh)
+	{
+		GetSphereMesh(Location, FVector(Radius), 20, 7, MaterialCache[Color.WithAlpha(InDrawAlpha)], SDPG_World, false, ViewIndex, Collector);
+	}
+}
+
+void FDebugRenderSceneProxy::FCapsule::Draw(FPrimitiveDrawInterface* PDI, EDrawType InDrawType, uint32 InDrawAlpha, FMaterialCache& MaterialCache, int32 ViewIndex, FMeshElementCollector& Collector) const
+{
+	if (InDrawType == SolidAndWireMeshes || InDrawType == WireMesh)
+	{
+		const float HalfAxis = FMath::Max<float>(HalfHeight - Radius, 1.f);
+		const FVector BottomEnd = Base + Radius * Z;
+		const FVector TopEnd = BottomEnd + (2 * HalfAxis) * Z;
+		const float CylinderHalfHeight = (TopEnd - BottomEnd).Size() * 0.5;
+		const FVector CylinderLocation = BottomEnd + CylinderHalfHeight * Z;
+		DrawWireCapsule(PDI, CylinderLocation, X, Y, Z, Color, Radius, HalfHeight, (InDrawType == SolidAndWireMeshes) ? 9 : 16, SDPG_World, InDrawType == SolidAndWireMeshes ? 2 : 0, 0, true);
+	}
+	if (InDrawType == SolidAndWireMeshes || InDrawType == SolidMesh)
+	{
+		GetCapsuleMesh(Base, X, Y, Z, Color, Radius, HalfHeight, 16, MaterialCache[Color.WithAlpha(InDrawAlpha)], SDPG_World, false, ViewIndex, Collector);
+	}
+}
\ No newline at end of file
diff --git a/Engine/Source/Runtime/Engine/Public/DebugRenderSceneProxy.h b/Engine/Source/Runtime/Engine/Public/DebugRenderSceneProxy.h
--- a/Engine/Source/Runtime/Engine/Public/DebugRenderSceneProxy.h
+++ b/Engine/Source/Runtime/Engine/Public/DebugRenderSceneProxy.h
@@ -12,13 +12,15 @@
 #include "DynamicMeshBuilder.h"
 
 class APlayerController;
+class FMaterialRenderProxy;
 class FMeshElementCollector;
 class FPrimitiveDrawInterface;
 class FRegisterComponentContext;
 class UCanvas;
+class UMaterial;
 class UPrimitiveComponent;
 
-DECLARE_DELEGATE_TwoParams(FDebugDrawDelegate, class UCanvas*, class APlayerController*);
+DECLARE_DELEGATE_TwoParams(FDebugDrawDelegate, UCanvas*, APlayerController*);
 
 class FDebugRenderSceneProxy : public FPrimitiveSceneProxy
 {
@@ -67,6 +69,17 @@ public:
 		return FVector::DistSquared(Start, View->ViewMatrices.GetViewOrigin()) <= FMath::Square(Range);
 	}
 
+	struct FMaterialCache
+	{
+		FMaterialCache(FMeshElementCollector& InCollector, bool bUseLight = false, UMaterial* InMaterial = nullptr);
+		FMaterialRenderProxy* operator[](FLinearColor Color);
+
+		FMeshElementCollector& Collector;
+		TMap<uint32, FMaterialRenderProxy*> MeshColorInstances;
+		TWeakObjectPtr<UMaterial> SolidMeshMaterial;
+		bool bUseFakeLight = false;
+	};
+
 	/** Struct to hold info about lines to render. */
 	struct FDebugLine
 	{
@@ -80,6 +93,8 @@ public:
 		FVector End;
 		FColor Color;
 		float Thickness;
+
+		ENGINE_API void Draw(FPrimitiveDrawInterface* PDI) const;
 	};
 
 	/** Struct to hold info about boxes to render. */
@@ -95,6 +110,8 @@ public:
 		{
 		}
 
+		ENGINE_API void Draw(FPrimitiveDrawInterface* PDI, EDrawType InDrawType, uint32 InDrawAlpha, FMaterialCache& MaterialCache, int32 ViewIndex, FMeshElementCollector& Collector) const;
+
 		FBox Box;
 		FColor Color;
 		FTransform Transform;
@@ -109,6 +126,8 @@ public:
 		HalfHeight(InHalfHeight),
 		Color(InColor) {}
 
+		ENGINE_API void Draw(FPrimitiveDrawInterface* PDI, EDrawType InDrawType, uint32 InDrawAlpha, FMaterialCache& MaterialCache, int32 ViewIndex, FMeshElementCollector& Collector) const;
+
 		FVector Base;
 		float Radius;
 		float HalfHeight;
@@ -123,6 +142,8 @@ public:
 		Color(InColor),
 		Size(InSize) {}
 
+		ENGINE_API void Draw(FPrimitiveDrawInterface* PDI) const;
+
 		FVector Position;
 		FColor Color;
 		float Size;
@@ -136,6 +157,8 @@ public:
 		End(InEnd),
 		Color(InColor) {}
 
+		ENGINE_API void Draw(FPrimitiveDrawInterface* PDI, const float Mag) const;
+
 		FVector Start;
 		FVector End;
 		FColor Color;
@@ -150,6 +173,8 @@ public:
 		Color(InColor),
 		DashSize(InDashSize) {}
 
+		ENGINE_API void Draw(FPrimitiveDrawInterface* PDI) const;
+
 		FVector Start;
 		FVector End;
 		FColor Color;
@@ -165,6 +190,8 @@ public:
 		Location(InLocation),
 		Color(InColor.ToFColor(true)) {}
 
+		ENGINE_API void Draw(FPrimitiveDrawInterface* PDI, EDrawType InDrawType, uint32 InDrawAlpha, FMaterialCache& MaterialCache, int32 ViewIndex, FMeshElementCollector& Collector) const;
+
 		float Radius;
 		FVector Location;
 		FColor Color;
@@ -193,6 +220,8 @@ public:
 			Angle2(InAngle2),
 			Color(InColor.ToFColor(true)) {}
 
+		ENGINE_API void Draw(FPrimitiveDrawInterface* PDI, EDrawType InDrawType, uint32 InDrawAlpha, FMaterialCache& MaterialCache, int32 ViewIndex, FMeshElementCollector& Collector, TArray<FVector>* VertsCache = nullptr) const;
+
 		FMatrix ConeToWorld;
 		float Angle1;
 		float Angle2;
@@ -216,10 +245,9 @@ public:
 			, HalfHeight(InHalfHeight)
 			, X(x)
 			, Y(y)
-			, Z(z)
-		{
+			, Z(z) {}
 
-		}
+		ENGINE_API void Draw(FPrimitiveDrawInterface* PDI, EDrawType InDrawType, uint32 InDrawAlpha, FMaterialCache& MaterialCache, int32 ViewIndex, FMeshElementCollector& Collector) const;
 
 		float Radius;
 		FVector Base; //Center point of the base of the cylinder.
@@ -248,7 +276,10 @@ public:
 	EDrawType DrawType;
 	uint32 DrawAlpha;
 
-	TWeakObjectPtr<class UMaterial> SolidMeshMaterial;
+	TWeakObjectPtr<UMaterial> SolidMeshMaterial;
+
+protected:
+	ENGINE_API virtual void GetDynamicMeshElementsForView(const FSceneView* View, const int32 ViewIndex, const FSceneViewFamily& ViewFamily, const uint32 VisibilityMap, FMeshElementCollector& Collector, FMaterialCache& DefaultMaterialCache, FMaterialCache& SolidMeshMaterialCache) const;
 };
 
 
