hitting a cluster union in TSceneCastCommonImp needs to be unsuccessful if we don't hit any child components (e.g. because of ignored actors/components)

diff --git a/Engine/Source/Runtime/Engine/Private/Collision/SceneQuery.cpp b/Engine/Source/Runtime/Engine/Private/Collision/SceneQuery.cpp
--- a/Engine/Source/Runtime/Engine/Private/Collision/SceneQuery.cpp
+++ b/Engine/Source/Runtime/Engine/Private/Collision/SceneQuery.cpp
@@ -472,6 +472,7 @@ bool TSceneCastCommonImp(const UWorld* World, typename Traits::TOutHits& OutHits
 						}
 
 						OutHits.Append(AllNewHits);
+						bBlockingHit &= !OutHits.IsEmpty();
 					}
 					else
 					{
@@ -480,6 +481,7 @@ bool TSceneCastCommonImp(const UWorld* World, typename Traits::TOutHits& OutHits
 						if (ClusterUnionHit.bIsClusterUnion)
 						{
 							Traits::ResetOutHits(OutHits, Start, End);
+							bBlockingHit = ClusterUnionHit.bHit;
 							if (ClusterUnionHit.bHit)
 							{
 								OutHits = NewHit;
