Optimization to UHierarchicalInstancedStaticMeshComponent::GetOverlappingBoxTransforms()

diff --git a/Engine/Source/Runtime/Engine/Private/HierarchicalInstancedStaticMesh.cpp b/Engine/Source/Runtime/Engine/Private/HierarchicalInstancedStaticMesh.cpp
--- a/Engine/Source/Runtime/Engine/Private/HierarchicalInstancedStaticMesh.cpp
+++ b/Engine/Source/Runtime/Engine/Private/HierarchicalInstancedStaticMesh.cpp
@@ -3499,21 +3499,17 @@ int32 UHierarchicalInstancedStaticMeshComponent::GetOverlappingBoxCount(const FB
 
 void UHierarchicalInstancedStaticMeshComponent::GetOverlappingBoxTransforms(const FBox& Box, TArray<FTransform>& OutTransforms) const
 {
+	OutTransforms.Reset();
+	
 	GatherInstanceTransformsInArea(*this, Box, 0, OutTransforms);
 
 	const FBoxSphereBounds MeshBounds = GetStaticMesh()->GetBounds();
-	int32 NumTransforms = OutTransforms.Num();
-	for(int32 Idx = NumTransforms - 1 ; Idx >= 0 ; --Idx)
+	OutTransforms.RemoveAllSwap([&MeshBounds, &Box](const FTransform& Transform) -> bool
 	{
-		FTransform& TM = OutTransforms[Idx];
-		const FVector Centre = TM.GetLocation();
+		const FVector Centre = Transform.GetLocation();
 		const FBox OtherBox(FVector(Centre - MeshBounds.BoxExtent), FVector(Centre + MeshBounds.BoxExtent));
-
-		if(!Box.Intersect(OtherBox))
-		{
-			OutTransforms.RemoveAt(Idx);
-		}
-	}
+		return !Box.Intersect(OtherBox); 
+	});
 }
 
 void UHierarchicalInstancedStaticMeshComponent::GetTree(TArray<FClusterNode>& OutClusterTree)
