Do not auto-generate GBuffer layout for mobile platforms. Mobile uses hardcoded layout anyway

diff --git a/Engine/Shaders/Private/DeferredShadingCommon.ush b/Engine/Shaders/Private/DeferredShadingCommon.ush
--- a/Engine/Shaders/Private/DeferredShadingCommon.ush
+++ b/Engine/Shaders/Private/DeferredShadingCommon.ush
@@ -430,8 +430,9 @@ bool CheckerFromSceneColorUV(float2 UVSceneColor)
 
 #if GBUFFER_REFACTOR
 #include "GBufferHelpers.ush"
-
-#include "/Engine/Generated/ShaderAutogen/AutogenShaderHeaders.ush"
+	#if FEATURE_LEVEL > FEATURE_LEVEL_ES3_1
+	#include "/Engine/Generated/ShaderAutogen/AutogenShaderHeaders.ush"
+	#endif
 #endif
 
 struct FScreenSpaceData
diff --git a/Engine/Source/Editor/UnrealEd/Private/EditorEngine.cpp b/Engine/Source/Editor/UnrealEd/Private/EditorEngine.cpp
--- a/Engine/Source/Editor/UnrealEd/Private/EditorEngine.cpp
+++ b/Engine/Source/Editor/UnrealEd/Private/EditorEngine.cpp
@@ -7756,7 +7756,7 @@ void UEditorEngine::SetPreviewPlatform(const FPreviewPlatformInfo& NewPreviewPla
 	if (NewPreviewPlatform.PreviewShaderFormatName != NAME_None)
 	{
 		// Force generation of the autogen files if they don't already exist
-		FShaderCompileUtilities::WriteGBufferInfoAutogen(ShaderPlatform, MaxFeatureLevel);
+		FShaderCompileUtilities::GenerateBrdfHeaders(ShaderPlatform);
 	}
 
 	// Record the new preview platform
diff --git a/Engine/Source/Runtime/Engine/Private/ShaderCompiler/ShaderGenerationUtil.cpp b/Engine/Source/Runtime/Engine/Private/ShaderCompiler/ShaderGenerationUtil.cpp
--- a/Engine/Source/Runtime/Engine/Private/ShaderCompiler/ShaderGenerationUtil.cpp
+++ b/Engine/Source/Runtime/Engine/Private/ShaderCompiler/ShaderGenerationUtil.cpp
@@ -2095,9 +2095,12 @@ void FShaderCompileUtilities::WriteGBufferInfoAutogen(EShaderPlatform TargetPlat
 void FShaderCompileUtilities::GenerateBrdfHeaders(EShaderPlatform TargetPlatform)
 {
 	ERHIFeatureLevel::Type FeatureLevel = GetMaxSupportedFeatureLevel(TargetPlatform);
-
-	// Writes the GBuffer format .ush file if it's out of date.
-	WriteGBufferInfoAutogen(TargetPlatform, FeatureLevel);
+	// auto-generated GBuffer layout is not used by mobile rendering
+	if (FeatureLevel > ERHIFeatureLevel::ES3_1)
+	{
+		// Writes the GBuffer format .ush file if it's out of date.
+		WriteGBufferInfoAutogen(TargetPlatform, FeatureLevel);
+	}
 }
 
 void FShaderCompileUtilities::GenerateBrdfHeaders(const FName& ShaderFormat)
