Deprecate EnableAutoLoadingSplashScreen functionality.

Original cl https://p4-swarm.epicgames.net/reviews/23018413/ by [REDACTED_PERSON].

diff --git a/Engine/Source/Runtime/Engine/Classes/Kismet/StereoLayerFunctionLibrary.h b/Engine/Source/Runtime/Engine/Classes/Kismet/StereoLayerFunctionLibrary.h
--- a/Engine/Source/Runtime/Engine/Classes/Kismet/StereoLayerFunctionLibrary.h
+++ b/Engine/Source/Runtime/Engine/Classes/Kismet/StereoLayerFunctionLibrary.h
@@ -44,6 +44,7 @@ public:
 	 *
 	 * @param	bAutoShowEnabled	(in)	True, if automatic showing of splash screens is enabled when map is being loaded.
 	 */
-	UFUNCTION(BlueprintCallable, Category = "VR")
+	UE_DEPRECATED(5.3, "We don't recommend using Auto Show Loading Screen any longer, and it will be removed in a future update.")
+	UFUNCTION(BlueprintCallable, Category = "VR", meta = (DeprecatedFunction, DeprecationMessage = "We don't recommend using Auto Show Loading Screen any longer, and it will be removed in a future update."))
 	static void EnableAutoLoadingSplashScreen(bool InAutoShowEnabled);
 };
