Fix race condition in shader compiler

diff --git a/Engine/Source/Runtime/Engine/Private/ShaderCompiler/ShaderCompiler.cpp b/Engine/Source/Runtime/Engine/Private/ShaderCompiler/ShaderCompiler.cpp
--- a/Engine/Source/Runtime/Engine/Private/ShaderCompiler/ShaderCompiler.cpp
+++ b/Engine/Source/Runtime/Engine/Private/ShaderCompiler/ShaderCompiler.cpp
@@ -594,7 +594,7 @@ void FShaderCompileJobCollection::SubmitJobs(const TArray<FShaderCommonCompileJo
 
 					FSharedBuffer* ExistingOutput = nullptr;
 					{
-						FReadScopeLock ReadLocker(Lock);
+						FWriteScopeLock Locker(Lock);
 
 						// see if we can find the job in the cache first
 						ExistingOutput = CompletedJobsCache.Find(InputHash, bCheckDDC);
