HLOD declaration of dependencies during cook: Use the AssetRegistry's dynamic GatherDependencies system to identify which HLOD packages exist for a UWorld at cooktime and add them as dependencies.
#rnx

diff --git a/Engine/Source/Developer/HierarchicalLODUtilities/Private/HierarchicalLODUtilities.cpp b/Engine/Source/Developer/HierarchicalLODUtilities/Private/HierarchicalLODUtilities.cpp
--- a/Engine/Source/Developer/HierarchicalLODUtilities/Private/HierarchicalLODUtilities.cpp
+++ b/Engine/Source/Developer/HierarchicalLODUtilities/Private/HierarchicalLODUtilities.cpp
@@ -112,6 +112,13 @@ static FString GetHLODPackageName(const FString& InLevelPackageName, const uint3
 	return FString::Printf(TEXT("%s/HLOD/%s"), *PathName, *InOutHLODProxyName);
 }
 
+FString FHierarchicalLODUtilities::GetWildcardOfHLODPackagesForPackage(const FString& PackageName)
+{
+	const FString PathName = FPackageName::GetLongPackagePath(PackageName);
+	const FString BaseName = FPackageName::GetShortName(PackageName);
+	return FString::Printf(TEXT("%s/HLOD/%s_*_HLOD"), *PathName, *BaseName);
+}
+
 static FString GetHLODPackageName(const ULevel* InLevel, const uint32 InHLODLevelIndex, FString& InOutHLODProxyName)
 {
 	// Strip out any PIE or level instance prefix from the given level package name
@@ -134,12 +141,6 @@ static FString GetHLODPackageName(const ULevel* InLevel, const uint32 InHLODLeve
 	return GetHLODPackageName(LevelPackageName, InHLODLevelIndex, InOutHLODProxyName);
 }
 
-FString FHierarchicalLODUtilities::GetHLODPackageName(const ULevel* InLevel, const uint32 InHLODLevelIndex)
-{
-	FString HLODProxyName;
-	return ::GetHLODPackageName(InLevel, InHLODLevelIndex, HLODProxyName);
-}
-
 void FHierarchicalLODUtilities::CleanStandaloneAssetsInPackage(UPackage* InPackage)
 {
 	TArray<UObject*> Objects;
@@ -188,7 +189,7 @@ UPackage* FHierarchicalLODUtilities::CreateOrRetrieveLevelHLODPackage(const ULev
 	checkf(InLevel != nullptr, TEXT("Invalid Level supplied"));
 
 	FString HLODProxyName;
-	const FString HLODLevelPackageName = ::GetHLODPackageName(InLevel, HLODLevelIndex, HLODProxyName);
+	const FString HLODLevelPackageName = GetHLODPackageName(InLevel, HLODLevelIndex, HLODProxyName);
 
 	// Find existing package
 	bool bCreatedNewPackage = false;
@@ -207,7 +208,7 @@ UHLODProxy* FHierarchicalLODUtilities::RetrieveLevelHLODProxy(const ULevel* InLe
 	checkf(InLevel != nullptr, TEXT("Invalid Level supplied"));
 	FString HLODProxyName;
 	FCookLoadScope CookLoadScope(ECookLoadType::UsedInGame);
-	const FString HLODLevelPackageName = ::GetHLODPackageName(InLevel, HLODLevelIndex, HLODProxyName);
+	const FString HLODLevelPackageName = GetHLODPackageName(InLevel, HLODLevelIndex, HLODProxyName);
 
 	UHLODProxy* HLODProxy = LoadObject<UHLODProxy>(nullptr, *HLODLevelPackageName, nullptr, LOAD_Quiet | LOAD_NoWarn);
 	return HLODProxy;
@@ -248,7 +249,7 @@ UPackage* FHierarchicalLODUtilities::CreateOrRetrieveLevelHLODPackage(const ULev
 FString FHierarchicalLODUtilities::GetLevelHLODProxyName(const FString& InLevelPackageName, const uint32 InHLODLevelIndex)
 {
 	FString HLODProxyName;
-	FString HLODPackageName = ::GetHLODPackageName(InLevelPackageName, InHLODLevelIndex, HLODProxyName);
+	FString HLODPackageName = GetHLODPackageName(InLevelPackageName, InHLODLevelIndex, HLODProxyName);
 	return HLODPackageName + TEXT(".") + HLODProxyName;
 }
 
diff --git a/Engine/Source/Developer/HierarchicalLODUtilities/Public/HierarchicalLODUtilities.h b/Engine/Source/Developer/HierarchicalLODUtilities/Public/HierarchicalLODUtilities.h
--- a/Engine/Source/Developer/HierarchicalLODUtilities/Public/HierarchicalLODUtilities.h
+++ b/Engine/Source/Developer/HierarchicalLODUtilities/Public/HierarchicalLODUtilities.h
@@ -33,7 +33,7 @@ public:
 	virtual UPackage* RetrieveLevelHLODPackage(const ULevel* InLevel, const uint32 HLODLevelIndex) override;
 	virtual UHLODProxy* CreateOrRetrieveLevelHLODProxy(const ULevel* InLevel, const uint32 HLODLevelIndex) override;
 	virtual UHLODProxy* RetrieveLevelHLODProxy(const ULevel* InLevel, const uint32 HLODLevelIndex) override;
-	virtual FString GetHLODPackageName(const ULevel* InLevel, const uint32 InHLODLevelIndex) override;
+	virtual FString GetWildcardOfHLODPackagesForPackage(const FString& PackageName) override;
 	virtual FString GetLevelHLODProxyName(const FString& InLevelPackageName, const uint32 HLODLevelIndex) override;
 	UE_DEPRECATED(4.20, "Use different signature containing BaseMaterial parameter")
 	virtual bool BuildStaticMeshForLODActor(ALODActor* LODActor, UPackage* AssetsOuter, const FHierarchicalSimplification& LODSetup) override;
diff --git a/Engine/Source/Developer/HierarchicalLODUtilities/Public/IHierarchicalLODUtilities.h b/Engine/Source/Developer/HierarchicalLODUtilities/Public/IHierarchicalLODUtilities.h
--- a/Engine/Source/Developer/HierarchicalLODUtilities/Public/IHierarchicalLODUtilities.h
+++ b/Engine/Source/Developer/HierarchicalLODUtilities/Public/IHierarchicalLODUtilities.h
@@ -75,7 +75,8 @@ public:
 	/** Retrieves the HLOD proxy that is created for the given level */
 	virtual UHLODProxy* RetrieveLevelHLODProxy(const ULevel* InLevel, const uint32 HLODLevelIndex) = 0;
 
-	virtual FString GetHLODPackageName(const ULevel* InLevel, const uint32 InHLODLevelIndex) = 0;
+	/** Retrieves a wildcard that includes all possible HLOD packagenames that could be used by the given package. */
+	virtual FString GetWildcardOfHLODPackagesForPackage(const FString& PackageName) = 0;
 
 	/** Get the name of a HLOD proxy given the name of a level's package */
 	virtual FString GetLevelHLODProxyName(const FString& InLevelPackageName, const uint32 HLODLevelIndex) = 0;
diff --git a/Engine/Source/Runtime/AssetRegistry/Private/AssetRegistry.cpp b/Engine/Source/Runtime/AssetRegistry/Private/AssetRegistry.cpp
--- a/Engine/Source/Runtime/AssetRegistry/Private/AssetRegistry.cpp
+++ b/Engine/Source/Runtime/AssetRegistry/Private/AssetRegistry.cpp
@@ -4138,9 +4138,12 @@ void FAssetRegistryImpl::LoadCalculatedDependencies(FName PackageName,
 		return CompiledFilter;
 	};
 
+	TArray<UE::AssetDependencyGatherer::Private::FRegisteredAssetDependencyGatherer*, TInlineAllocator<2>> Gatherers;
 	for (const FAssetData* AssetData : State.GetAssetsByPackageName(PackageName))
 	{
-		if (UE::AssetDependencyGatherer::Private::FRegisteredAssetDependencyGatherer** AssetDependencyGatherer = RegisteredDependencyGathererClasses.Find(AssetData->GetClass()))
+		Gatherers.Reset();
+		RegisteredDependencyGathererClasses.MultiFind(AssetData->GetClass(), Gatherers);
+		for (UE::AssetDependencyGatherer::Private::FRegisteredAssetDependencyGatherer* Gatherer : Gatherers)
 		{
 			if (!bOutHadActivity)
 			{
@@ -4151,7 +4154,7 @@ void FAssetRegistryImpl::LoadCalculatedDependencies(FName PackageName,
 			TArray<IAssetDependencyGatherer::FGathereredDependency> GatheredDependencies;
 			TArray<FString> DirectoryReferences;
 
-			(*AssetDependencyGatherer)->GatherDependencies(*AssetData, State, GetCompiledFilter, GatheredDependencies, DirectoryReferences);
+			Gatherer->GatherDependencies(*AssetData, State, GetCompiledFilter, GatheredDependencies, DirectoryReferences);
 			
 			if (GatheredDependencies.Num())
 			{
diff --git a/Engine/Source/Runtime/AssetRegistry/Private/AssetRegistryImpl.h b/Engine/Source/Runtime/AssetRegistry/Private/AssetRegistryImpl.h
--- a/Engine/Source/Runtime/AssetRegistry/Private/AssetRegistryImpl.h
+++ b/Engine/Source/Runtime/AssetRegistry/Private/AssetRegistryImpl.h
@@ -449,7 +449,7 @@ private:
 	TMultiMap<FString, FName> DirectoryReferencers;
 
 	/** A map of per asset class dependency gatherer called in LoadCalculatedDependencies */
-	TMap<UClass*, UE::AssetDependencyGatherer::Private::FRegisteredAssetDependencyGatherer*> RegisteredDependencyGathererClasses;
+	TMultiMap<UClass*, UE::AssetDependencyGatherer::Private::FRegisteredAssetDependencyGatherer*> RegisteredDependencyGathererClasses;
 	bool bRegisteredDependencyGathererClassesDirty;
 #endif
 #if WITH_ENGINE && WITH_EDITOR
diff --git a/Engine/Source/Runtime/Engine/Private/HLOD/HLODEngineSubsystem.cpp b/Engine/Source/Runtime/Engine/Private/HLOD/HLODEngineSubsystem.cpp
--- a/Engine/Source/Runtime/Engine/Private/HLOD/HLODEngineSubsystem.cpp
+++ b/Engine/Source/Runtime/Engine/Private/HLOD/HLODEngineSubsystem.cpp
@@ -6,12 +6,16 @@
 
 #if WITH_EDITOR
 
+#include "AssetRegistry/AssetData.h"
+#include "AssetRegistry/AssetDependencyGatherer.h"
 #include "EngineUtils.h"
 #include "Engine/LODActor.h"
 #include "Engine/HLODProxy.h"
+#include "Engine/World.h"
 #include "Editor.h"
 #include "UnrealEngine.h"
 #include "HierarchicalLOD.h"
+#include "Misc/PathViews.h"
 #include "Modules/ModuleManager.h"
 #include "IHierarchicalLODUtilities.h"
 #include "HierarchicalLODUtilitiesModule.h"
@@ -64,7 +68,6 @@ void UHLODEngineSubsystem::RegisterRecreateLODActorsDelegates()
 		OnPostWorldInitializationDelegateHandle = FWorldDelegates::OnPostWorldInitialization.AddUObject(this, &UHLODEngineSubsystem::RecreateLODActorsForWorld);
 		OnLevelAddedToWorldDelegateHandle = FWorldDelegates::LevelAddedToWorld.AddUObject(this, &UHLODEngineSubsystem::RecreateLODActorsForLevel);
 		OnPreSaveWorlDelegateHandle = FEditorDelegates::PreSaveWorldWithContext.AddUObject(this, &UHLODEngineSubsystem::OnPreSaveWorld);
-		OnWorldCollectSaveReferencesDelegateHandle = FWorldDelegates::OnCollectSaveReferences.AddUObject(this, &UHLODEngineSubsystem::OnWorldCollectSaveReferences);
 	}	
 }
 
@@ -195,34 +198,49 @@ void UHLODEngineSubsystem::OnPreSaveWorld(UWorld* InWorld, FObjectPreSaveContext
 	}
 }
 
-void UHLODEngineSubsystem::OnWorldCollectSaveReferences(UWorld* World, FArchive& Ar)
+class FHLODDependencyGatherer : public IAssetDependencyGatherer
 {
-	// Disabled for now as this introduce validation errors
-#if 0
-	// Record a dependency on the paths to HLODProxy packages that would be associated with
+public:
+	virtual void GatherDependencies(const FAssetData& AssetData, const FAssetRegistryState& AssetRegistryState,
+		TFunctionRef<FARCompiledFilter(const FARFilter&)> CompileFilterFunc,
+		TArray<IAssetDependencyGatherer::FGathereredDependency>& OutDependencies,
+		TArray<FString>& OutDependencyDirectories) const override;
+};
+
+void FHLODDependencyGatherer::GatherDependencies(const FAssetData& AssetData,
+	const FAssetRegistryState& AssetRegistryState, TFunctionRef<FARCompiledFilter(const FARFilter&)> CompileFilterFunc,
+	TArray<IAssetDependencyGatherer::FGathereredDependency>& OutDependencies,
+	TArray<FString>& OutDependencyDirectories) const
+{
+	if (!GetDefault<UHierarchicalLODSettings>()->bSaveLODActorsToHLODPackages)
+	{
+		return;
+	}
+
+	// Record a dependency on the paths to HLODProxy packages that can be associated with
 	// this level if they exist
 	FHierarchicalLODUtilitiesModule& Module = FModuleManager::LoadModuleChecked<FHierarchicalLODUtilitiesModule>("HierarchicalLODUtilities");
 	IHierarchicalLODUtilities* Utilities = Module.GetUtilities();
 
-	for (ULevel* Level : World->GetLevels())
+	// TODO: GetHLODPackageName is usually constructed from the World's packagename, but this can be replaced by ULevelStreaming::PackageNameToLoad
+	// We need to write the list of LevelStreaming PackageNameToLoad into the AssetData so we can read it from here.
+	FString Wildcard = Utilities->GetWildcardOfHLODPackagesForPackage(AssetData.PackageName.ToString());
+	FARFilter PossibleAssetsFilter;
+	PossibleAssetsFilter.PackagePaths.Add(FName(FPathViews::GetPath(Wildcard)));
+	TArray<FAssetData> FilteredAssets;
+	AssetRegistryState.GetAssets(CompileFilterFunc(PossibleAssetsFilter), {}, FilteredAssets);
+	FString PackageName;
+	for (const FAssetData& FilteredAsset : FilteredAssets)
 	{
-		AWorldSettings* WorldSettings = Level->GetWorldSettings(false);
-		if (WorldSettings)
+		FilteredAsset.PackageName.ToString(PackageName);
+		if (PackageName.MatchesWildcard(Wildcard))
 		{
-			int32 NumLODLevels = WorldSettings->GetHierarchicalLODSetup().Num();
-			for (int32 LODIndex = 0; LODIndex < NumLODLevels; ++LODIndex)
-			{
-				const FString HLODLevelPackageName = Utilities->GetHLODPackageName(Level, LODIndex);
-				if (!HLODLevelPackageName.IsEmpty())
-				{
-					FSoftObjectPath Path(HLODLevelPackageName);
-					Ar << Path;
-				}
-			}
+			OutDependencies.Emplace(IAssetDependencyGatherer::FGathereredDependency{ FilteredAsset.PackageName,
+				UE::AssetRegistry::EDependencyProperty::Game | UE::AssetRegistry::EDependencyProperty::Build });
 		}
 	}
-#endif
 }
+REGISTER_ASSETDEPENDENCY_GATHERER(FHLODDependencyGatherer, UWorld);
 
 #endif // WITH_EDITOR
 
diff --git a/Engine/Source/Runtime/Engine/Public/HLOD/HLODEngineSubsystem.h b/Engine/Source/Runtime/Engine/Public/HLOD/HLODEngineSubsystem.h
--- a/Engine/Source/Runtime/Engine/Public/HLOD/HLODEngineSubsystem.h
+++ b/Engine/Source/Runtime/Engine/Public/HLOD/HLODEngineSubsystem.h
@@ -54,7 +54,6 @@ private:
 	void RecreateLODActorsForLevel(ULevel* InLevel, UWorld* InWorld);
 
 	void OnPreSaveWorld(UWorld* InWorld, FObjectPreSaveContext ObjectSaveContext);
-	void OnWorldCollectSaveReferences(UWorld* World, FArchive& Ar);
 
 	void UnregisterRecreateLODActorsDelegates();
 	void RegisterRecreateLODActorsDelegates();
@@ -65,7 +64,6 @@ private:
 	FDelegateHandle OnPostWorldInitializationDelegateHandle;
 	FDelegateHandle OnLevelAddedToWorldDelegateHandle;
 	FDelegateHandle OnPreSaveWorlDelegateHandle;
-	FDelegateHandle OnWorldCollectSaveReferencesDelegateHandle;
 
 	bool bDisableHLODCleanupOnLoad;
 	bool bDisableHLODSpawningOnLoad;
