Removing temporary animation comrpession debugging used for tracking down cook issue (CLs 25194942, 25163840, 25133548, 25114030)

diff --git a/Engine/Plugins/Animation/AnimationData/Source/AnimationData/Private/AnimSequencerDataModel.cpp b/Engine/Plugins/Animation/AnimationData/Source/AnimationData/Private/AnimSequencerDataModel.cpp
--- a/Engine/Plugins/Animation/AnimationData/Source/AnimationData/Private/AnimSequencerDataModel.cpp
+++ b/Engine/Plugins/Animation/AnimationData/Source/AnimationData/Private/AnimSequencerDataModel.cpp
@@ -186,18 +186,6 @@ void UAnimationSequencerDataModel::PostLoad()
 
 		CachedRawDataGUID.Invalidate();
 
-#if UE_DEBUG_ANIMATION_COMPRESSION				
-		const bool bOutputDebugData = UAnimationSettings::Get()->CompressionDebugAssetPaths.Contains(GetOuter()->GetPathName());
-		if (bOutputDebugData)
-		{
-			UE_LOG(LogAnimationCompression, Display, TEXT("CachedRawDataGUID.Invalidate %s"), *GetPathName());
-			FStackTracker TempTracker(nullptr, nullptr, nullptr, true);
-			TempTracker.CaptureStackTrace(1);
-			TempTracker.DumpStackTraces(0, *GLog);
-			TempTracker.ResetTracking();
-		}
-#endif // UE_DEBUG_ANIMATION_COMPRESSION
-
 		RemoveOutOfDateControls();
 
 		ValidateData();
@@ -693,19 +681,6 @@ FGuid UAnimationSequencerDataModel::GenerateGuid() const
 		ReturnGuid = Guid;
 	}
 
-
-#if UE_DEBUG_ANIMATION_COMPRESSION
-	const bool bOutputDebugData = UAnimationSettings::Get()->CompressionDebugGUIDs.Contains(ReturnGuid) || UAnimationSettings::Get()->CompressionDebugAssetPaths.Contains(GetOuter()->GetPathName());
-	if (bOutputDebugData)
-	{
-		UE_LOG(LogAnimationCompression, Display, TEXT("UAnimationSequencerDataModel::GenerateGuid %s - %s"), *GetPathName(), *ReturnGuid.ToString());
-		FStackTracker TempTracker(nullptr, nullptr, nullptr, true);
-		TempTracker.CaptureStackTrace(1);
-		TempTracker.DumpStackTraces(0, *GLog);
-		TempTracker.ResetTracking();
-	}
-#endif // UE_DEBUG_ANIMATION_COMPRESSION
-
 	return ReturnGuid;
 }
 
@@ -763,18 +738,7 @@ void UAnimationSequencerDataModel::OnNotify(const EAnimDataModelNotifyType& Noti
 			// Prevent reset when being populated inside of upgrade path (always happens in UAnimSequenceBase::PostLoad)
 			if (CachedRawDataGUID.IsValid() && (!Collector.Contains(EAnimDataModelNotifyType::Populated) || !FUObjectThreadContext::Get().IsRoutingPostLoad))
 			{
-				CachedRawDataGUID.Invalidate();			
-#if UE_DEBUG_ANIMATION_COMPRESSION
-				const bool bOutputDebugData = UAnimationSettings::Get()->CompressionDebugAssetPaths.Contains(GetOuter()->GetPathName());
-				if (bOutputDebugData)
-				{
-					UE_LOG(LogAnimationCompression, Display, TEXT("CachedRawDataGUID.Invalidate %s"), *GetPathName());
-					FStackTracker TempTracker(nullptr, nullptr, nullptr, true);
-					TempTracker.CaptureStackTrace(1);
-					TempTracker.DumpStackTraces(0, *GLog);
-					TempTracker.ResetTracking();
-				}
-#endif // UE_DEBUG_ANIMATION_COMPRESSION
+				CachedRawDataGUID.Invalidate();
 			}
 		};
 
diff --git a/Engine/Source/Runtime/Engine/Classes/Animation/AnimSequence.h b/Engine/Source/Runtime/Engine/Classes/Animation/AnimSequence.h
--- a/Engine/Source/Runtime/Engine/Classes/Animation/AnimSequence.h
+++ b/Engine/Source/Runtime/Engine/Classes/Animation/AnimSequence.h
@@ -27,8 +27,6 @@
 #include "AnimData/IAnimationDataModel.h"
 #endif // WITH_EDITOR
 
-#define UE_DEBUG_ANIMATION_COMPRESSION 1
-
 #include "AnimSequence.generated.h"
 
 typedef TArray<FTransform> FTransformArrayA2;
@@ -895,18 +893,6 @@ protected:
 	TMap<FIoHash, TUniquePtr<FCompressedAnimSequence>> DataByPlatformKeyHash;
 	TMap<FIoHash, TPimplPtr<UE::Anim::FAnimationSequenceAsyncCacheTask>> CacheTasksByKeyHash;
 
-#if UE_DEBUG_ANIMATION_COMPRESSION
-	void LogCompressionRequestTimings() const;
-	struct FCompressionRequestTimings
-	{
-		FDateTime LastBeginCacheRequest;
-		FDateTime LastApplyCacheData;
-		FDateTime LastClearCache;
-		FDateTime LastWillNeverCook;
-	};
-	TMap<FIoHash, FCompressionRequestTimings> TimingsByKeyHash;
-	TMap<FIoHash, FString> DDCStringByKeyHash;
-#endif // UE_DEBUG_ANIMATION_COMPRESSION
 protected:
 	bool TryCancelAsyncTasks();
 	void FinishAsyncTasks();
diff --git a/Engine/Source/Runtime/Engine/Classes/Animation/AnimationSettings.h b/Engine/Source/Runtime/Engine/Classes/Animation/AnimationSettings.h
--- a/Engine/Source/Runtime/Engine/Classes/Animation/AnimationSettings.h
+++ b/Engine/Source/Runtime/Engine/Classes/Animation/AnimationSettings.h
@@ -103,20 +103,7 @@ class ENGINE_API UAnimationSettings : public UDeveloperSettings
 
 	/** Project specific default frame-rate used when (re)initializing any animation based data */
 	UPROPERTY(config, EditAnywhere, Category = AnimationData)
-	FFrameRate DefaultFrameRate;
-
-	/** Data model GUIDs which should be tracked by animation compression debugging/logging */
-	UPROPERTY(config, EditAnywhere, Category = AnimationData)
-	TArray<FGuid> CompressionDebugGUIDs;
-
-	/** Animation sequence DDC hashes which should be tracked by animation compression debugging/logging */
-	UPROPERTY(config, EditAnywhere, Category = AnimationData)
-	TArray<FString> CompressionDebugHashes;
-
-	/** Animation sequence full paths which should be tracked by animation compression debugging/logging */
-	UPROPERTY(config, EditAnywhere, Category = AnimationData)
-	TArray<FString> CompressionDebugAssetPaths;
-	
+	FFrameRate DefaultFrameRate;	
 public:
 	static UAnimationSettings * Get() { return CastChecked<UAnimationSettings>(UAnimationSettings::StaticClass()->GetDefaultObject()); }
 
diff --git a/Engine/Source/Runtime/Engine/Private/Animation/AnimSequence.cpp b/Engine/Source/Runtime/Engine/Private/Animation/AnimSequence.cpp
--- a/Engine/Source/Runtime/Engine/Private/Animation/AnimSequence.cpp
+++ b/Engine/Source/Runtime/Engine/Private/Animation/AnimSequence.cpp
@@ -466,26 +466,13 @@ void UAnimSequence::WillNeverCacheCookedPlatformDataAgain()
 {
 	Super::WillNeverCacheCookedPlatformDataAgain();
 
-#if UE_DEBUG_ANIMATION_COMPRESSION
-	const bool bOutputDebugData = UAnimationSettings::Get()->CompressionDebugAssetPaths.Contains(GetPathName());
-	if (bOutputDebugData)
-	{
-		UE_LOG(LogAnimationCompression, Display, TEXT("UAnimSequence::WillNeverCacheCookedPlatformDataAgain %s"), *GetPathName());
-
-		for (TPair<FIoHash, FCompressionRequestTimings>& TimingsPair : TimingsByKeyHash)
-		{
-			TimingsPair.Value.LastWillNeverCook = FDateTime::Now();	
-		}
-	}
-#endif // UE_DEBUG_ANIMATION_COMPRESSION
-
 	UE::Anim::FAnimSequenceCompilingManager::Get().FinishCompilation({this});	
 	ClearCompressedCurveData();
 	ClearCompressedBoneData();
 
 	// Clear out current platform, and any target platform data
 	CompressedData.Reset();
-	
+		
 	CacheTasksByKeyHash.Empty();
 	DataByPlatformKeyHash.Empty();
 	PRAGMA_DISABLE_DEPRECATION_WARNINGS
@@ -497,18 +484,6 @@ void UAnimSequence::ClearAllCachedCookedPlatformData()
 {
 	Super::ClearAllCachedCookedPlatformData();
 	
-#if UE_DEBUG_ANIMATION_COMPRESSION
-	const bool bOutputDebugData = UAnimationSettings::Get()->CompressionDebugAssetPaths.Contains(GetPathName());
-	if (bOutputDebugData)
-	{
-		UE_LOG(LogAnimationCompression, Display, TEXT("UAnimSequence::ClearAllCachedCookedPlatformData %s"), *GetPathName());
-		for (TPair<FIoHash, FCompressionRequestTimings>& TimingsPair : TimingsByKeyHash)
-		{
-			TimingsPair.Value.LastClearCache = FDateTime::Now();	
-		}
-	}
-#endif // UE_DEBUG_ANIMATION_COMPRESSION
-	
 	// Delete any cache tasks first because the destructor will cancel the cache and build tasks,
 	// and drop their pointers to the data.
 	CacheTasksByKeyHash.Empty();
@@ -1183,27 +1158,12 @@ bool UAnimSequence::IsCachedCookedPlatformDataLoaded(const ITargetPlatform* Targ
 		EndCacheDerivedData(KeyHash);
 	}
 
-	auto LogCookedDataWarning = [this, TargetPlatform, KeyHash]()
+	auto RecompressForPlatform = [this, TargetPlatform]()
 	{
-#if UE_DEBUG_ANIMATION_COMPRESSION
-		const bool bOutputDebugData = UAnimationSettings::Get()->CompressionDebugHashes.Contains(LexToString(KeyHash)) || UAnimationSettings::Get()->CompressionDebugAssetPaths.Contains(GetPathName());
-		if (bOutputDebugData)
-		{
-			// This means cooked data is expected, but it it currently invalid and no task is running to generate it
-			UE_LOG(LogAnimationCompression, Warning, TEXT("Expecting either valid compressed Animation Data, or an in-flight compression task for asset %s\nPlatform: %s\nHash: %s"), *GetPathName(), *TargetPlatform->DisplayName().ToString(), *LexToString(KeyHash));
-
-			if (bOutputDebugData)
-			{
-				UE_LOG(LogAnimationCompression, Display, TEXT("UAnimSequence::IsCachedCookedPlatformDataLoaded %s - %s"), *GetPathName(), *LexToString(KeyHash));
-			}
-		}
-		LogCompressionRequestTimings();
-#endif // UE_DEBUG_ANIMATION_COMPRESSION
-
 		if (IsInGameThread())
 		{
-			// Temporary fix for soft-lock, to be removed when root-cause is addressed
-			BeginCacheForCookedPlatformData(TargetPlatform);			
+			// This is required until we can ensure inter-animsequence dependencies are handled correctly (see UE-184634)
+			BeginCacheForCookedPlatformData(TargetPlatform);
 		}
 	};
 
@@ -1215,7 +1175,7 @@ bool UAnimSequence::IsCachedCookedPlatformDataLoaded(const ITargetPlatform* Targ
 		}
 		else if (!CacheTasksByKeyHash.Contains(KeyHash))
 		{
-			LogCookedDataWarning();
+			RecompressForPlatform();
 		}
 	}
 	else
@@ -1228,12 +1188,12 @@ bool UAnimSequence::IsCachedCookedPlatformDataLoaded(const ITargetPlatform* Targ
 			}
 			else if (!CacheTasksByKeyHash.Contains(KeyHash))
 			{
-				LogCookedDataWarning();
+				RecompressForPlatform();
 			}
 		}
 		else if (!CacheTasksByKeyHash.Contains(KeyHash))
 		{
-			LogCookedDataWarning();
+			RecompressForPlatform();
 		}
 	}
 
@@ -5139,16 +5099,7 @@ FIoHash UAnimSequence::CreateDerivedDataKeyHash(const ITargetPlatform* TargetPla
 
 	Writer << Ret;
 
-	const FIoHash FinalHash = Writer.Finalize();
-
-#if UE_DEBUG_ANIMATION_COMPRESSION
-	const bool bOutputDebugData = UAnimationSettings::Get()->CompressionDebugHashes.Contains(LexToString(FinalHash)) || UAnimationSettings::Get()->CompressionDebugAssetPaths.Contains(GetPathName());
-	if (bOutputDebugData)
-	{
-		DDCStringByKeyHash.FindOrAdd(FinalHash) = Ret;
-	}
-#endif // UE_DEBUG_ANIMATION_COMPRESSION
-	return FinalHash;
+	return Writer.Finalize();
 }
 
 FIoHash UAnimSequence::BeginCacheDerivedData(const ITargetPlatform* TargetPlatform)
@@ -5206,15 +5157,6 @@ FIoHash UAnimSequence::BeginCacheDerivedData(const ITargetPlatform* TargetPlatfo
     {
     	TargetData = DataByPlatformKeyHash.Emplace(KeyHash, MakeUnique<FCompressedAnimSequence>()).Get();
     }
-
-#if UE_DEBUG_ANIMATION_COMPRESSION
-	const bool bOutputDebugData = UAnimationSettings::Get()->CompressionDebugHashes.Contains(LexToString(KeyHash)) || UAnimationSettings::Get()->CompressionDebugAssetPaths.Contains(GetPathName());
-	if (bOutputDebugData)
-	{
-		UE_LOG(LogAnimationCompression, Display, TEXT("UAnimSequence::BeginCacheDerivedData %s - %s"), *GetPathName(), *LexToString(KeyHash));
-		TimingsByKeyHash.FindOrAdd(KeyHash).LastBeginCacheRequest = FDateTime::Now();
-	}
-#endif // UE_DEBUG_ANIMATION_COMPRESSION
 	check(TargetData);
 	TargetData->Reset();
 
@@ -5283,40 +5225,10 @@ void UAnimSequence::EndCacheDerivedData(const FIoHash& KeyHash)
 	{
 		return;
 	}
-
-#if UE_DEBUG_ANIMATION_COMPRESSION
-	const bool bOutputDebugData = UAnimationSettings::Get()->CompressionDebugHashes.Contains(LexToString(KeyHash)) || UAnimationSettings::Get()->CompressionDebugAssetPaths.Contains(GetPathName());
-	if (bOutputDebugData)
-	{
-		UE_LOG(LogAnimationCompression, Display, TEXT("UAnimSequence::EndCacheDerivedData %s - %s"), *GetPathName(), *LexToString(KeyHash));
-	}
-#endif // UE_DEBUG_ANIMATION_COMPRESSION
 	
 	UE::Anim::FAnimSequenceCompilingManager::Get().FinishCompilation({this});
 }
-
-#if UE_DEBUG_ANIMATION_COMPRESSION
-void UAnimSequence::LogCompressionRequestTimings() const
-{
-	for (const TPair<FIoHash, FCompressionRequestTimings>& TimingsPair : TimingsByKeyHash)
-	{
-		UE_LOG(LogAnimation, Warning, TEXT("Hash: %s\nBeginCache: %s\nApplyCache: %s\nClearCache: %s\nWillNeverCook: %s"),
-			*LexToString(TimingsPair.Key),
-			*TimingsPair.Value.LastBeginCacheRequest.ToString(),
-			*TimingsPair.Value.LastApplyCacheData.ToString(),
-			*TimingsPair.Value.LastClearCache.ToString(),
-			*TimingsPair.Value.LastWillNeverCook.ToString());
-	}
 	
-	for (const TPair<FIoHash, FString>& StringPair : DDCStringByKeyHash)
-	{
-		UE_LOG(LogAnimation, Warning, TEXT("Hash: %s\nDDC key: %s"),
-		*LexToString(StringPair.Key),
-		*StringPair.Value);
-	}
-}
-#endif // UE_DEBUG_ANIMATION_COMPRESSION
-
 FCompressedAnimSequence& UAnimSequence::CacheDerivedData(const ITargetPlatform* TargetPlatform)
 {
 	const FIoHash KeyHash = BeginCacheDerivedData(TargetPlatform);
@@ -5422,15 +5334,6 @@ void UAnimSequence::FinishAsyncTasks()
 					bUseRawDataOnly = false;
                     PRAGMA_ENABLE_DEPRECATION_WARNINGS
 				}
-
-#if UE_DEBUG_ANIMATION_COMPRESSION
-				const bool bOutputDebugData = UAnimationSettings::Get()->CompressionDebugHashes.Contains(LexToString(It->Key)) || UAnimationSettings::Get()->CompressionDebugAssetPaths.Contains(GetPathName());
-				if (bOutputDebugData)
-				{
-					UE_LOG(LogAnimationCompression, Display, TEXT("UAnimSequence::FinishAsyncTasks %s - %s"), *GetPathName(), *LexToString(It->Key));					
-					TimingsByKeyHash.FindChecked(It->Key).LastApplyCacheData = FDateTime::Now();
-				}
-#endif // UE_DEBUG_ANIMATION_COMPRESSION
 			}
 			else
 			{
