Crash Fix:
repro: Deleting a component that was a base actor for foliage, undoing and deleting that same component again.

The reason why it was crashing is that the USMInstanceElementIdMapTransactor was created RF_Transactional inside a transaction but the Entry that owns it isn't removed on Undo leaving the Entry with a nullptr Transactor.
Instead we set the RF_Transactional flag on the USMInstanceElementIdMapTransactor after the NewObject call so that it can be transacted but it will not deleted by the undo of the operation that created it.

#rnx

diff --git a/Engine/Source/Runtime/Engine/Private/Elements/SMInstance/SMInstanceElementId.cpp b/Engine/Source/Runtime/Engine/Private/Elements/SMInstance/SMInstanceElementId.cpp
--- a/Engine/Source/Runtime/Engine/Private/Elements/SMInstance/SMInstanceElementId.cpp
+++ b/Engine/Source/Runtime/Engine/Private/Elements/SMInstance/SMInstanceElementId.cpp
@@ -20,7 +20,8 @@ FSMInstanceElementIdMapEntry::FSMInstanceElementIdMapEntry(FSMInstanceElementIdM
 	checkf(Component, TEXT("FSMInstanceElementIdMapEntry must be constructed with a valid component!"));
 
 #if WITH_EDITORONLY_DATA
-	Transactor = NewObject<USMInstanceElementIdMapTransactor>(GetTransientPackage(), NAME_None, RF_Transient | RF_Transactional);
+	Transactor = NewObject<USMInstanceElementIdMapTransactor>(GetTransientPackage(), NAME_None, RF_Transient);
+	Transactor->SetFlags(RF_Transactional);
 	Transactor->SetOwnerEntry(this);
 #endif	// WITH_EDITORONLY_DATA
 }
