Spelling fix (Centre -> Center)
#rnx

diff --git a/Engine/Source/Runtime/Engine/Private/HierarchicalInstancedStaticMesh.cpp b/Engine/Source/Runtime/Engine/Private/HierarchicalInstancedStaticMesh.cpp
--- a/Engine/Source/Runtime/Engine/Private/HierarchicalInstancedStaticMesh.cpp
+++ b/Engine/Source/Runtime/Engine/Private/HierarchicalInstancedStaticMesh.cpp
@@ -3485,8 +3485,8 @@ int32 UHierarchicalInstancedStaticMeshComponent::GetOverlappingBoxCount(const FB
 	const FBoxSphereBounds MeshBounds = GetStaticMesh()->GetBounds();
 	for(FTransform& T : Transforms)
 	{
-		const FVector Centre = T.GetLocation();
-		const FBox OtherBox(FVector(Centre - MeshBounds.BoxExtent), FVector(Centre + MeshBounds.BoxExtent));
+		const FVector Center = T.GetLocation();
+		const FBox OtherBox(FVector(Center - MeshBounds.BoxExtent), FVector(Center + MeshBounds.BoxExtent));
 
 		if(Box.Intersect(OtherBox))
 		{
@@ -3506,8 +3506,8 @@ void UHierarchicalInstancedStaticMeshComponent::GetOverlappingBoxTransforms(cons
 	const FBoxSphereBounds MeshBounds = GetStaticMesh()->GetBounds();
 	OutTransforms.RemoveAllSwap([&MeshBounds, &Box](const FTransform& Transform) -> bool
 	{
-		const FVector Centre = Transform.GetLocation();
-		const FBox OtherBox(FVector(Centre - MeshBounds.BoxExtent), FVector(Centre + MeshBounds.BoxExtent));
+		const FVector Center = Transform.GetLocation();
+		const FBox OtherBox(FVector(Center - MeshBounds.BoxExtent), FVector(Center + MeshBounds.BoxExtent));
 		return !Box.Intersect(OtherBox); 
 	});
 }
