Fix crash trying to access unreachable (not destroyed) skeleton

diff --git a/Engine/Source/Runtime/Engine/Private/BoneContainer.cpp b/Engine/Source/Runtime/Engine/Private/BoneContainer.cpp
--- a/Engine/Source/Runtime/Engine/Private/BoneContainer.cpp
+++ b/Engine/Source/Runtime/Engine/Private/BoneContainer.cpp
@@ -273,7 +273,7 @@ void FBoneContainer::CacheRequiredAnimCurves(const UE::Anim::FCurveFilterSetting
 
 	CurveFlags.Empty();
 
-	if (USkeleton* Skeleton = AssetSkeleton.Get())
+	if (USkeleton* Skeleton = AssetSkeleton.GetEvenIfUnreachable())
 	{
 		// Copy filter curves.
 		TArray<FName, TInlineAllocator<32>> FilterCurves;
@@ -403,7 +403,7 @@ void FBoneContainer::RegenerateSerialNumber()
 const FRetargetSourceCachedData& FBoneContainer::GetRetargetSourceCachedData(const FName& InRetargetSourceName) const
 {
 	LLM_SCOPE_BYNAME(TEXT("Animation/BoneContainer"));
-	const TArray<FTransform>& RetargetTransforms = AssetSkeleton->GetRefLocalPoses(InRetargetSourceName);
+	const TArray<FTransform>& RetargetTransforms = AssetSkeleton.GetEvenIfUnreachable()->GetRefLocalPoses(InRetargetSourceName);
 	return GetRetargetSourceCachedData(InRetargetSourceName, FSkeletonRemapping(), RetargetTransforms);
 }
 
@@ -428,7 +428,7 @@ const FRetargetSourceCachedData& FBoneContainer::GetRetargetSourceCachedData(con
 			const int32 TargetSkeletonBoneIndex = CompactPoseToSkeletonIndex[CompactBoneIndex];
 			const int32 SourceSkeletonBoneIndex = InRemapping.IsValid() ? InRemapping.GetSourceSkeletonBoneIndex(TargetSkeletonBoneIndex) : TargetSkeletonBoneIndex;
 
-			if (SourceSkeletonBoneIndex != INDEX_NONE && AssetSkeleton->GetBoneTranslationRetargetingMode(SourceSkeletonBoneIndex, bDisableRetargeting) == EBoneTranslationRetargetingMode::OrientAndScale)
+			if (SourceSkeletonBoneIndex != INDEX_NONE && AssetSkeleton.GetEvenIfUnreachable()->GetBoneTranslationRetargetingMode(SourceSkeletonBoneIndex, bDisableRetargeting) == EBoneTranslationRetargetingMode::OrientAndScale)
 			{
 				if(AuthoredOnRefSkeleton.IsValidIndex(SourceSkeletonBoneIndex))
 				{
