Fix duplicating & moving GameFeatureData leading to incorrect data in AddWPContent_GameFeatureAction
ContentBundleDescriptor Guid (and debug color) are now regenerated on PostDuplicate
PackageRoot is now computed at call site instead of stored in the ContentBundleDescriptor

diff --git a/Engine/Plugins/Experimental/GameFeatures/Source/GameFeatures/Private/GameFeatureAction_AddWPContent.cpp b/Engine/Plugins/Experimental/GameFeatures/Source/GameFeatures/Private/GameFeatureAction_AddWPContent.cpp
--- a/Engine/Plugins/Experimental/GameFeatures/Source/GameFeatures/Private/GameFeatureAction_AddWPContent.cpp
+++ b/Engine/Plugins/Experimental/GameFeatures/Source/GameFeatures/Private/GameFeatureAction_AddWPContent.cpp
@@ -19,7 +19,7 @@ UGameFeatureAction_AddWPContent::UGameFeatureAction_AddWPContent(const FObjectIn
 #if WITH_EDITOR
 	if (UGameFeatureData* GameFeatureData = GetTypedOuter<UGameFeatureData>())
 	{
-		ContentBundleDescriptor->InitializeObject(GameFeatureData->GetName(), FPackageName::GetPackageMountPoint(GetPackage()->GetName()).ToString());
+		ContentBundleDescriptor->InitializeObject(GameFeatureData->GetName());
 	}
 #endif
 }
diff --git a/Engine/Source/Runtime/Engine/Private/WorldPartition/ContentBundle/ContentBundleDescriptor.cpp b/Engine/Source/Runtime/Engine/Private/WorldPartition/ContentBundle/ContentBundleDescriptor.cpp
--- a/Engine/Source/Runtime/Engine/Private/WorldPartition/ContentBundle/ContentBundleDescriptor.cpp
+++ b/Engine/Source/Runtime/Engine/Private/WorldPartition/ContentBundle/ContentBundleDescriptor.cpp
@@ -7,13 +7,17 @@
 UContentBundleDescriptor::UContentBundleDescriptor(const FObjectInitializer& ObjectInitializer)
 	: DebugColor(FColor::Black)
 {
+
+}
+
+FString UContentBundleDescriptor::GetPackageRoot() const
+{
+	return FPackageName::GetPackageMountPoint(GetPackage()->GetName()).ToString();
 }
 
 bool UContentBundleDescriptor::IsValid() const
 {
-	return Guid.IsValid()
-		&& !DisplayName.IsEmpty()
-		&& !PackageRoot.IsEmpty();
+	return Guid.IsValid() && !DisplayName.IsEmpty();
 }
 
 FString UContentBundleDescriptor::GetContentBundleCompactString(const FGuid& InContentBundleID)
@@ -23,11 +27,10 @@ FString UContentBundleDescriptor::GetContentBundleCompactString(const FGuid& InC
 }
 
 #if WITH_EDITOR
-void UContentBundleDescriptor::InitializeObject(const FString& InContentBundleName, const FString& InPackageRoot)
+void UContentBundleDescriptor::InitializeObject(const FString& InContentBundleName)
 {
 	Guid = FGuid::NewGuid();
 	DisplayName = InContentBundleName;
-	PackageRoot = InPackageRoot;
 	InitDebugColor();
 }
 
@@ -38,6 +41,19 @@ void UContentBundleDescriptor::PostLoad()
 	Super::PostLoad();
 }
 
+#if WITH_EDITOR
+void UContentBundleDescriptor::PostDuplicate(bool bDuplicateForPIE)
+{
+	Super::PostDuplicate(bDuplicateForPIE);
+
+	if (!bDuplicateForPIE)
+	{
+		Guid = FGuid::NewGuid();
+		InitDebugColor();
+	}
+}
+#endif
+
 void UContentBundleDescriptor::InitDebugColor()
 {
 	// If not set, generate a color based on guid
diff --git a/Engine/Source/Runtime/Engine/Public/WorldPartition/ContentBundle/ContentBundleDescriptor.h b/Engine/Source/Runtime/Engine/Public/WorldPartition/ContentBundle/ContentBundleDescriptor.h
--- a/Engine/Source/Runtime/Engine/Public/WorldPartition/ContentBundle/ContentBundleDescriptor.h
+++ b/Engine/Source/Runtime/Engine/Public/WorldPartition/ContentBundle/ContentBundleDescriptor.h
@@ -19,17 +19,15 @@ class ENGINE_API UContentBundleDescriptor : public UObject
 public:
 	const FString& GetDisplayName() const { return DisplayName; }
 	const FColor& GetDebugColor() const { return DebugColor; }
-	const FString& GetPackageRoot() const { return PackageRoot; }
+	FString GetPackageRoot() const;
 	const FGuid& GetGuid() const { return Guid; }
 
 #if WITH_EDITOR
-	void InitializeObject(const FString& InContentBundleName, const FString& InPackageRoot);
-	void SetDisplayName(const FString& InName) { DisplayName = InName; }
-	void SetPackageRoot(const FString& InPackageRoot) { PackageRoot = InPackageRoot; }
-	void SetDebugColor(const FColor& InDebugColor) { DebugColor = InDebugColor; }
+	void InitializeObject(const FString& InContentBundleName);
 
 	//~ Begin UObject Interface
 	virtual void PostLoad() override;
+	virtual void PostDuplicate(bool bDuplicateForPIE) override;
 	//~ End UObject Interface
 #endif
 
@@ -44,12 +42,9 @@ private:
 	UPROPERTY(EditAnywhere, Category = BaseInformation)
 	FString DisplayName;
 
-	UPROPERTY(EditAnywhere, Category = BaseInformation)
+	UPROPERTY(EditAnywhere, DuplicateTransient, Category = BaseInformation)
 	FColor DebugColor;
 
-	UPROPERTY(VisibleAnywhere, Category = BaseInformation, AdvancedDisplay)
+	UPROPERTY(VisibleAnywhere, DuplicateTransient, Category = BaseInformation, AdvancedDisplay)
 	FGuid Guid;
-
-	UPROPERTY(VisibleAnywhere, Category = BaseInformation, AdvancedDisplay)
-	FString PackageRoot;
-};
\ No newline at end of file
+};
