Add missing include

diff --git a/Engine/Source/Runtime/Engine/Private/WorldPartition/ContentBundle/ContentBundleDescriptor.cpp b/Engine/Source/Runtime/Engine/Private/WorldPartition/ContentBundle/ContentBundleDescriptor.cpp
--- a/Engine/Source/Runtime/Engine/Private/WorldPartition/ContentBundle/ContentBundleDescriptor.cpp
+++ b/Engine/Source/Runtime/Engine/Private/WorldPartition/ContentBundle/ContentBundleDescriptor.cpp
@@ -1,5 +1,8 @@
 // Copyright Epic Games, Inc. All Rights Reserved.
 #include "WorldPartition/ContentBundle/ContentBundleDescriptor.h"
+
+#include "Misc/PackageName.h"
+#include "UObject/Package.h"
 #include "WorldPartition/ContentBundle/ContentBundleLog.h"
 
 #include UE_INLINE_GENERATED_CPP_BY_NAME(ContentBundleDescriptor)
