World Partition: don't keep raw pointers to always loaded actors during PIE, as these actors can be destroyed.
We recently fixed a known case, so this change is making it future proof for future cases, if any.

diff --git a/Engine/Source/Runtime/Engine/Private/WorldPartition/WorldPartitionRuntimeHash.cpp b/Engine/Source/Runtime/Engine/Private/WorldPartition/WorldPartitionRuntimeHash.cpp
--- a/Engine/Source/Runtime/Engine/Private/WorldPartition/WorldPartitionRuntimeHash.cpp
+++ b/Engine/Source/Runtime/Engine/Private/WorldPartition/WorldPartitionRuntimeHash.cpp
@@ -361,7 +361,7 @@ void UWorldPartitionRuntimeHash::ForceExternalActorLevelReference(bool bForceExt
 	{
 		for (const FAlwaysLoadedActorForPIE& AlwaysLoadedActor : AlwaysLoadedActorsForPIE)
 		{
-			if (AActor* Actor = AlwaysLoadedActor.Actor)
+			if (AActor* Actor = AlwaysLoadedActor.Actor.Get())
 			{
 				Actor->SetForceExternalActorLevelReferenceForPIE(bForceExternalActorLevelReferenceForPIE);
 			}
diff --git a/Engine/Source/Runtime/Engine/Public/WorldPartition/WorldPartitionRuntimeHash.h b/Engine/Source/Runtime/Engine/Public/WorldPartition/WorldPartitionRuntimeHash.h
--- a/Engine/Source/Runtime/Engine/Public/WorldPartition/WorldPartitionRuntimeHash.h
+++ b/Engine/Source/Runtime/Engine/Public/WorldPartition/WorldPartitionRuntimeHash.h
@@ -168,7 +168,7 @@ protected:
 		{}
 
 		FWorldPartitionReference Reference;
-		AActor* Actor;
+		TWeakObjectPtr<AActor> Actor;
 	};
 
 	TArray<FAlwaysLoadedActorForPIE> AlwaysLoadedActorsForPIE;
