Operate on single block at a time during VT build to reduce peak memory usage.

Source pixels for each individal block of each layer can be freed once all tiles for it was created.

diff --git a/Engine/Source/Runtime/Engine/Private/TextureDerivedDataTask.cpp b/Engine/Source/Runtime/Engine/Private/TextureDerivedDataTask.cpp
--- a/Engine/Source/Runtime/Engine/Private/TextureDerivedDataTask.cpp
+++ b/Engine/Source/Runtime/Engine/Private/TextureDerivedDataTask.cpp
@@ -899,12 +899,11 @@ static int64 GetBuildRequiredMemoryEstimate(UTexture* InTexture,
 		// for each layer/block :
 		//    generate mips (requires F32 copy)
 		//    output to intermediate format
-		// (at this point source could be freed; currently it is not)
-		// intermediate format copy is then used to make tiles
-		// for each tile :
-		//    make padded tile in intermediate format
-		//    encode to output format
-		//    discard padded tile in intermediate format
+		//    intermediate format copy is then used to make tiles
+		//    for each tile :
+		//       make padded tile in intermediate format
+		//       encode to output format
+		//       discard padded tile in intermediate format
 		// all output tiles are then aggregated
 
 		// Compute the memory it should take to uncompress the bulkdata in memory
@@ -953,7 +952,8 @@ static int64 GetBuildRequiredMemoryEstimate(UTexture* InTexture,
 
 		int64 NumOutputPixelsPerLayer = NumTilesPerLayer * TilePixels;
 
-		int64 VTIntermediateSizeBytes = 0;
+		// intermediate is created just once per block, use max size estimate
+		int64 VTIntermediateSizeBytes = IntermediateFloatColorBytes;
 		int64 OutputSizeBytes = 0;
 	
 		int64 MaxPerPixelEncoderMemUse = 0;
@@ -968,8 +968,6 @@ static int64 GetBuildRequiredMemoryEstimate(UTexture* InTexture,
 
 			int64 IntermediateBytesPerPixel = ERawImageFormat::GetBytesPerPixel(IntermediateImageFormat);
 
-			VTIntermediateSizeBytes += TotalPixelsPerLayer * IntermediateBytesPerPixel;
-
 			// + output bytes? (but can overlap with IntermediateFloatColorBytes)
 			//	almost always less than IntermediateFloatColorBytes
 			//  exception would be lots of udim blocks + lots of layers
diff --git a/Engine/Source/Runtime/Engine/Private/VT/VirtualTextureDataBuilder.cpp b/Engine/Source/Runtime/Engine/Private/VT/VirtualTextureDataBuilder.cpp
--- a/Engine/Source/Runtime/Engine/Private/VT/VirtualTextureDataBuilder.cpp
+++ b/Engine/Source/Runtime/Engine/Private/VT/VirtualTextureDataBuilder.cpp
@@ -15,9 +15,12 @@
 
 // Debugging aid to dump tiles to disc as png files
 #define SAVE_TILES 0
+#define SAVE_CHUNKS 0
 
-#if SAVE_TILES
+#if SAVE_TILES || SAVE_CHUNKS
 #include "ImageUtils.h"
+#include "Misc/Paths.h"
+#include "Misc/FileHelper.h"
 #endif
 
 static TAutoConsoleVariable<int32> CVarVTParallelTileCompression(
@@ -238,7 +241,6 @@ FVirtualTextureDataBuilder::FVirtualTextureDataBuilder(FVirtualTextureBuiltData
 
 FVirtualTextureDataBuilder::~FVirtualTextureDataBuilder()
 {
-	FreeSourcePixels();
 }
 
 bool FVirtualTextureBuilderDerivedInfo::InitializeFromBuildSettings(const FTextureSourceData& InSourceData, const FTextureBuildSettings* InSettingsPerLayer)
@@ -320,6 +322,8 @@ bool FVirtualTextureBuilderDerivedInfo::InitializeFromBuildSettings(const FTextu
 
 bool FVirtualTextureDataBuilder::Build(FTextureSourceData& InSourceData, FTextureSourceData& InCompositeSourceData, const FTextureBuildSettings* InSettingsPerLayer, bool bAllowAsync)
 {
+	const int32 NumBlocks = InSourceData.Blocks.Num();
+
 	const int32 NumLayers = InSourceData.Layers.Num();
 	checkf(NumLayers <= (int32)VIRTUALTEXTURE_DATA_MAXLAYERS, TEXT("The maximum amount of layers is exceeded."));
 	checkf(NumLayers > 0, TEXT("No layers to build."));
@@ -354,67 +358,83 @@ bool FVirtualTextureDataBuilder::Build(FTextureSourceData& InSourceData, FTextur
 	OutData.Chunks.Empty();
 	OutData.NumMips = DerivedInfo.NumMips;
 
-	// process original source to intermediate format, make mips, appply processing options
-	BuildSourcePixels(InSourceData, InCompositeSourceData);
-
-	// processed data is now in SourceBlocks , original InSourceData not needed anymore
-	// can free InSourceData (and InCompositeSourceData) now
-	InSourceData.ReleaseMemory();
-	InCompositeSourceData.ReleaseMemory();
-
 	// override async compression if requested
 	bAllowAsync = bAllowAsync && CVarVTParallelTileCompression.GetValueOnAnyThread();
 
-	// encode SourceBlocks to VT Tiles and pack into chunks :
-	bool ok = BuildPagesMacroBlocks(bAllowAsync);
-	// free SourceBlocks :
-	FreeSourcePixels();
-
-	return ok;
-}
+	LayerPayload.SetNum(NumLayers);
 
-bool FVirtualTextureDataBuilder::BuildPagesForChunk(const TArray<FVTSourceTileEntry>& ActiveTileList, bool bAllowAsync)
-{
-	TArray<FLayerData> LayerData;
-	LayerData.AddDefaulted(SourceLayers.Num());
-
-	for (int32 LayerIndex = 0; LayerIndex < LayerData.Num(); LayerIndex++)
 	{
-		BuildTiles(ActiveTileList, LayerIndex, LayerData[LayerIndex], bAllowAsync);
-	}
+		FScopedSlowTask BuildTask(NumLayers * NumBlocks);
 
-	// Fill out tile offsets per layer if we haven't yet and if all layers are raw uncompressed data.
-	if (OutData.TileDataOffsetPerLayer.Num() == 0)
-	{
-		bool bIsRawGPUData = true;
-		for (int32 LayerIndex = 0; LayerIndex < LayerData.Num(); LayerIndex++)
+		// Process source texture layer by layer
+		// Layer blocks will be freed from inside of BuildLayerBlocks() as soon as they are done
+		for (int32 LayerIndex = 0; LayerIndex < NumLayers; ++LayerIndex)
 		{
-			if (LayerData[LayerIndex].Codec != EVirtualTextureCodec::RawGPU)
+			const FTextureBuildSettings& BuildSettingsForLayer = SettingsPerLayer[LayerIndex];
+			FVirtualTextureSourceLayerData LayerData;
+
+			// Specify the format we are processing to in this step :
+			LayerData.ImageFormat = UE::TextureBuildUtilities::GetVirtualTextureBuildIntermediateFormat(BuildSettingsForLayer);
+			LayerData.GammaSpace = BuildSettingsForLayer.GetDestGammaSpace();
+			// Gamma correction can either be applied in step 1 or step 2 of the VT build
+			//	depending on whether the Intermediate format is U8 or not
+			if (!ERawImageFormat::GetFormatNeedsGammaSpace(LayerData.ImageFormat))
 			{
-				bIsRawGPUData = false;
-				break;
+				LayerData.GammaSpace = EGammaSpace::Linear;
 			}
-		}
-		if (bIsRawGPUData)
-		{
-			int64 TileDataOffset = 0;
-			OutData.TileDataOffsetPerLayer.Reserve(LayerData.Num());
-			for (int32 LayerIndex = 0; LayerIndex < LayerData.Num(); LayerIndex++)
+
+			// LayerData.bHasAlpha will be updated after the Build to detect if there is actual alpha in the layers
+			LayerData.bHasAlpha = BuildSettingsForLayer.bForceAlphaChannel;
+
+			LayerData.FormatName = FImageCoreUtils::ConvertToUncompressedTextureFormatName(LayerData.ImageFormat);
+			LayerData.PixelFormat = FImageCoreUtils::GetPixelFormatForRawImageFormat(LayerData.ImageFormat);
+			LayerData.SourceFormat = FImageCoreUtils::ConvertToTextureSourceFormat(LayerData.ImageFormat);
+
+			// Don't want platform specific swizzling for VT tile data, this tends to add extra padding for textures with odd dimensions
+			// (VT physical tiles generally not power-of-2 after adding border)
+			// Must match TextureDerivedData.cpp
+			LayerData.TextureFormatName = UE::TextureBuildUtilities::TextureFormatRemovePlatformPrefixFromName(BuildSettingsForLayer.TextureFormatName);
+
+			// bHasAlpha was previously set to true if bForceAlphaChannel
+			// if it's false and not bForceNoAlphaChannel, we scan each block for alpha
+			// if alpha is in any tile of any block, it gets enabled for all so they have consistent pixel format
+			if (!LayerData.bHasAlpha && !BuildSettingsForLayer.bForceNoAlphaChannel)
 			{
-				TileDataOffset += LayerData[LayerIndex].TilePayload[0].Num();
-				OutData.TileDataOffsetPerLayer.Add( IntCastChecked<uint32>(TileDataOffset) );
+				// Note that check is a bit wrong to do at this point, because it does require all blocks to be present in the memory
+				// This should've been stored somewhere way earlier, maybe at import process
+				for (int32 BlockIndex = 0; BlockIndex < NumBlocks; ++BlockIndex)
+				{
+					const TArray<FImage>& SourceMips = InSourceData.Blocks[BlockIndex].MipsPerLayer[LayerIndex];
+					if (!SourceMips.IsEmpty())
+					{
+						LayerData.bHasAlpha = DetectAlphaChannel(SourceMips[0]);
+						if (LayerData.bHasAlpha)
+						{
+							break;
+						}
+					}
+				}
 			}
+
+			// Building happens in following order: [Layers] -> [Blocks (with creating mips)] -> [Tiles] -> [Mips]
+			BuildLayerBlocks(BuildTask, LayerIndex, LayerData, InSourceData, InCompositeSourceData, bAllowAsync);
 		}
 	}
 
-	// Write tiles out to chunk.
-	return PushDataToChunk(ActiveTileList, LayerData);
+	// Rearrange compressed VT tiles into chunks for output
+	// Chunks contain multiple tiles: [Tiles] -> [Mips] -> [Layers]
+	bool ok = BuildChunks();
+
+	// Release memory used during build process
+	LayerPayload.Empty();
+
+	return ok;
 }
 
-bool FVirtualTextureDataBuilder::BuildPagesMacroBlocks(bool bAllowAsync)
+bool FVirtualTextureDataBuilder::BuildChunks()
 {
 	static const uint32 MinSizePerChunk = 1024u; // Each chunk will contain a mip level of at least this size (MinSizePerChunk x MinSizePerChunk)
-	const uint32 NumLayers = SourceLayers.Num();
+	const uint32 NumLayers = LayerPayload.Num();
 	const int32 TileSize = SettingsPerLayer[0].VirtualTextureTileSize;
 	const uint32 MinSizePerChunkInTiles = FMath::DivideAndRoundUp<uint32>(MinSizePerChunk, TileSize);
 	const uint32 MinTilesPerChunk = MinSizePerChunkInTiles * MinSizePerChunkInTiles;
@@ -433,12 +453,14 @@ bool FVirtualTextureDataBuilder::BuildPagesMacroBlocks(bool bAllowAsync)
 		MipHeightInTiles = FMath::DivideAndRoundUp(MipHeightInTiles, 2u);
 	}
 
-	// loop over each macro block and assemble the tiles
 	FScopedSlowTask BuildTask(NumTiles);
 
 	TArray<FVTSourceTileEntry> TilesInChunk;
 	TilesInChunk.Reserve(NumTiles);
 
+	// Loop over Tiles in Morton order, and assemble the tiles into chunks
+	// This only moves memory into chunks, packing tile payload into same
+	// order as older version of code to maintain identical output
 	{
 		uint32 TileIndex = 0u;
 		bool bInFinalChunk = false;
@@ -478,17 +500,24 @@ bool FVirtualTextureDataBuilder::BuildPagesMacroBlocks(bool bAllowAsync)
 					const int32 BlockX = TileX / MipBlockSizeInTilesX;
 					const int32 BlockY = TileY / MipBlockSizeInTilesY;
 
-					const int32 BlockIndex = FindSourceBlockIndex(Mip, BlockX, BlockY);
-					if (BlockIndex != INDEX_NONE)
+					FVTBlockPayload* Block = FindSourceBlock(Mip, BlockX, BlockY);
+					if (Block)
 					{
-						const FTextureSourceBlockData& Block = SourceBlocks[BlockIndex];
+						const int32 MipIndex = Mip - Block->MipBias;
+						const int32 BlockWidth = FMath::Max(Block->SizeX >> MipIndex, 1);
+						const int32 BlockWidthInTiles = FMath::DivideAndRoundUp(BlockWidth, TileSize);
+
 						FVTSourceTileEntry* TileEntry = new(TilesInChunk) FVTSourceTileEntry;
-						TileEntry->BlockIndex = BlockIndex;
-						TileEntry->TileIndex = TileIndex;
+						TileEntry->Block = Block;
 						TileEntry->MipIndex = Mip;
-						TileEntry->MipIndexInBlock = Mip - Block.MipBias;
-						TileEntry->TileInBlockX = TileX - Block.BlockX * MipBlockSizeInTilesX;
-						TileEntry->TileInBlockY = TileY - Block.BlockY * MipBlockSizeInTilesY;
+						TileEntry->MipIndexInBlock = MipIndex;
+						int32 TileInBlockX = TileX - Block->BlockX * MipBlockSizeInTilesX;
+						int32 TileInBlockY = TileY - Block->BlockY * MipBlockSizeInTilesY;
+
+						// TileIndex is destination index to use in output chunks
+						// TileIndexInBlock is source index into compressed data array of block/mip
+						TileEntry->TileIndexInBlock = TileInBlockY * BlockWidthInTiles + TileInBlockX;
+						TileEntry->TileIndex = TileIndex;
 
 						OffsetData.AddTile(TileIndexInMip);
 					}
@@ -501,7 +530,7 @@ bool FVirtualTextureDataBuilder::BuildPagesMacroBlocks(bool bAllowAsync)
 			if (!bInFinalChunk && TilesInChunk.Num() >= (int32)MinTilesPerChunk)
 			{
 				OutData.TileIndexPerChunk.Add(TileIndex);
-				if ( ! BuildPagesForChunk(TilesInChunk, bAllowAsync) )
+				if ( ! BuildPagesForChunk(TilesInChunk) )
 					return false;
 				TilesInChunk.Reset();
 			}
@@ -520,7 +549,7 @@ bool FVirtualTextureDataBuilder::BuildPagesMacroBlocks(bool bAllowAsync)
 
 		if (TilesInChunk.Num() > 0)
 		{
-			if ( ! BuildPagesForChunk(TilesInChunk, bAllowAsync) )
+			if ( ! BuildPagesForChunk(TilesInChunk) )
 				return false;
 		}
 
@@ -583,19 +612,16 @@ static const TCHAR* GetSafePixelFormatName(EPixelFormat Format)
 	}
 }
 
-void FVirtualTextureDataBuilder::BuildTiles(const TArray<FVTSourceTileEntry>& TileList, uint32 LayerIndex, FLayerData& GeneratedData, bool bAllowAsync)
+void FVirtualTextureDataBuilder::BuildBlockTiles(uint32 LayerIndex, uint32 BlockIndex, FVTBlockPayload& Block, const FVirtualTextureSourceLayerData& LayerData, bool bAllowAsync)
 {
-	TRACE_CPUPROFILER_EVENT_SCOPE(Texture.VT.BuildTiles);
+	TRACE_CPUPROFILER_EVENT_SCOPE(Texture.VT.BuildBlockTiles);
 
 	const FTextureBuildSettings& BuildSettingsLayer0 = SettingsPerLayer[0];
 	const FTextureBuildSettings& BuildSettingsForLayer = SettingsPerLayer[LayerIndex];
-	const FVirtualTextureSourceLayerData& LayerData = SourceLayers[LayerIndex];
 
 	const int32 TileSize = BuildSettingsLayer0.VirtualTextureTileSize;
 	const int32 BorderSize = BuildSettingsLayer0.VirtualTextureBorderSize;
 	const int32 PhysicalTileSize = TileSize + BorderSize * 2;
-	
-	UE_LOG(LogVirtualTexturing, Verbose, TEXT("VT BuildTiles LayerIndex=%d TileList.Num=%d Pixels=%lld"), LayerIndex, TileList.Num(), (int64)PhysicalTileSize*PhysicalTileSize*TileList.Num() );
 
 	FThreadSafeBool bCompressionError = false;
 	EPixelFormat CompressedFormat = PF_Unknown;
@@ -633,79 +659,477 @@ void FVirtualTextureDataBuilder::BuildTiles(const TArray<FVTSourceTileEntry>& Ti
 		TBSettings.OodleTextureSdkVersion = BuildSettingsForLayer.OodleTextureSdkVersion;
 
 		check(TBSettings.GetDestGammaSpace() == BuildSettingsForLayer.GetDestGammaSpace());
- 
-		GeneratedData.TilePayload.AddDefaulted(TileList.Num());
 
-		// ParallelFor is on TaskGraph for VT tiles
-		//	TextureFormats should disable their own internal use of TaskGraph for VT tiles if necessary
-		
-		const bool bIsSingleThreaded = !bAllowAsync;
+		// Mip levels start at Block.MipBias in case provided Texture is smaller than block size used in VT
+		for (int32 MipIndex = 0; MipIndex < Block.NumMips; ++MipIndex)
+		{
+			int32 Mip = MipIndex + Block.MipBias;
+			const int32 MipBlockSizeX = FMath::Max(Block.SizeX >> MipIndex, 1);
+			const int32 MipBlockSizeY = FMath::Max(Block.SizeY >> MipIndex, 1);
+			const int32 MipBlockSizeInTilesX = FMath::DivideAndRoundUp(MipBlockSizeX, TileSize);
+			const int32 MipBlockSizeInTilesY = FMath::DivideAndRoundUp(MipBlockSizeY, TileSize);
+			const int32 NumTiles = MipBlockSizeInTilesY * MipBlockSizeInTilesX;
+
+			if (MipIndex == 0)
+			{
+				Block.Tiles.SetNum(NumTiles);
+			}
+
+			// ParallelFor is on TaskGraph for VT tiles
+			//	TextureFormats should disable their own internal use of TaskGraph for VT tiles if necessary
+			const bool bIsSingleThreaded = !bAllowAsync;
+
+			// Build all tiles for this mip level
+			ParallelFor(TEXT("Texture.VT.BuildTiles.PF"), NumTiles, 1, [&](int32 TileIndex)
+			{
+				int32 TileY = TileIndex / MipBlockSizeInTilesX;
+				int32 TileX = TileIndex % MipBlockSizeInTilesX;
+
+				Block.Tiles[TileIndex].Mips.SetNum(Block.NumMips);
+
+				const FImage& SourceMip = Block.Mips[MipIndex];
+				check(SourceMip.Format == LayerData.ImageFormat);
+				check(SourceMip.GammaSpace == LayerData.GammaSpace);
 
-		int NumTiles = TileList.Num();
-		ParallelFor( TEXT("Texture.VT.BuildTiles.PF"), NumTiles,1, [&](int32 TileIndex)
+				FPixelDataRectangle SourceData(LayerData.SourceFormat,
+					SourceMip.SizeX,
+					SourceMip.SizeY,
+					const_cast<uint8*>(SourceMip.RawData.GetData()));
+
+				TArray<FImage> TileImages;
+				FImage* TileImage = new(TileImages) FImage(PhysicalTileSize, PhysicalTileSize, LayerData.ImageFormat, LayerData.GammaSpace);
+				FPixelDataRectangle TileData(LayerData.SourceFormat, PhysicalTileSize, PhysicalTileSize, TileImage->RawData.GetData());
+
+				TileData.Clear();
+				TileData.CopyRectangleBordered(0, 0, SourceData,
+					TileX * TileSize - BorderSize,
+					TileY * TileSize - BorderSize,
+					PhysicalTileSize,
+					PhysicalTileSize,
+					(TextureAddress)BuildSettingsLayer0.VirtualAddressingModeX,
+					(TextureAddress)BuildSettingsLayer0.VirtualAddressingModeY);
+
+#if SAVE_TILES
+				{
+					FString DebugName = FPaths::MakeValidFileName(*DebugTexturePathName, TEXT('_'));
+					FString BasePath = FPaths::ProjectUserDir();
+					FString TileFileName = BasePath / FString::Format(TEXT("{0}_{1}_{2}_{3}_{4}_{5}.png"), TArray<FStringFormatArg>({ *DebugName, BlockIndex, TileX, TileY, Mip, LayerIndex }));
+					TileData.Save(TileFileName, ImageWrapper);
+				}
+#endif // SAVE_TILES
+
+				// give each tile a unique DebugTexturePathName for DebugDump option :
+				FString DebugTilePathName = FString::Printf(TEXT("%s_L%d_VT%04d"), *DebugTexturePathName, LayerIndex, TileIndex);
+
+				TArray<FCompressedImage2D> CompressedMip;
+				TArray<FImage> EmptyList;
+				uint32 NumMipsInTail, ExtData;
+				// this is the Build for Tiles to do the encode to GPU formats, with no processing
+				if (!ensure(Compressor->BuildTexture(TileImages, EmptyList, TBSettings, DebugTilePathName, CompressedMip, NumMipsInTail, ExtData, nullptr)))
+				{
+					bCompressionError = true;
+				}
+
+				check(CompressedMip.Num() == 1);
+				checkf(CompressedFormat == PF_Unknown || CompressedFormat == CompressedMip[0].PixelFormat,
+					TEXT("CompressedFormat: %s (%d), CompressedMip[0].PixelFormat: %s (%d)"),
+					GetSafePixelFormatName(CompressedFormat), (int32)CompressedFormat,
+					GetSafePixelFormatName((EPixelFormat)CompressedMip[0].PixelFormat), (int32)CompressedMip[0].PixelFormat);
+
+				FVTTileMipPayload& MipPayload = Block.Tiles[TileIndex].Mips[MipIndex];
+				MipPayload.Payload = MoveTemp(CompressedMip[0].RawData);
+				MipPayload.CompressedFormat = (EPixelFormat)CompressedMip[0].PixelFormat;
+			},
+			(bIsSingleThreaded ? EParallelForFlags::ForceSingleThread : EParallelForFlags::None));
+		}
+	}
+}
+
+void FVirtualTextureDataBuilder::BuildLayerBlocks(FSlowTask& BuildTask, uint32 LayerIndex, const FVirtualTextureSourceLayerData& LayerData, FTextureSourceData& SourceData, FTextureSourceData& CompositeSourceData, bool bAllowAsync)
+{
+	TRACE_CPUPROFILER_EVENT_SCOPE(Texture.VT.BuildLayerBlocks);
+
+	const TArray<FImage> EmptyImageArray;
+
+	const int32 TileSize = SettingsPerLayer[0].VirtualTextureTileSize;
+	const int32 NumLayers = SourceData.Layers.Num();
+	const int32 NumBlocks = SourceData.Blocks.Num();
+
+	// If we have more than 1 block, need to create miptail that contains mips made from multiple blocks
+	const bool bNeedsMiptailBlock = (NumBlocks > 1);
+
+	// Miptail
+	TArray<FImage> MiptailInputImages;
+	FPixelDataRectangle MiptailPixelData{ LayerData.SourceFormat, 0, 0, 0 };
+	const uint32 BlockSize = FMath::Max(DerivedInfo.BlockSizeX, DerivedInfo.BlockSizeY);
+	const uint32 BlockSizeInTiles = FMath::DivideAndRoundUp<uint32>(BlockSize, TileSize);
+	const uint32 MaxMipInBlock = FMath::CeilLogTwo(BlockSizeInTiles);
+	const uint32 MipWidthInBlock = FMath::Max<uint32>(DerivedInfo.BlockSizeX >> MaxMipInBlock, 1);
+	const uint32 MipHeightInBlock = FMath::Max<uint32>(DerivedInfo.BlockSizeY >> MaxMipInBlock, 1);
+	const uint32 MipInputSizeX = FMath::RoundUpToPowerOfTwo(DerivedInfo.SizeInBlocksX * MipWidthInBlock);
+	const uint32 MipInputSizeY = FMath::RoundUpToPowerOfTwo(DerivedInfo.SizeInBlocksY * MipHeightInBlock);
+	const uint32 MipInputSize = FMath::Max(MipInputSizeX, MipInputSizeY);
+
+	if (bNeedsMiptailBlock)
+	{
+		MiptailInputImages.Reset(1);
+		FImage* MiptailInputImage = new(MiptailInputImages) FImage();
+		MiptailInputImage->Init(MipInputSizeX, MipInputSizeY, LayerData.ImageFormat, LayerData.GammaSpace);
+		MiptailPixelData = FPixelDataRectangle(LayerData.SourceFormat, MipInputSizeX, MipInputSizeY, MiptailInputImage->RawData.GetData());
+		MiptailPixelData.Clear();
+	}
+
+	LayerPayload[LayerIndex].Blocks.SetNum(NumBlocks + (bNeedsMiptailBlock ? 1 : 0));
+
+	// Process source texture block by block from same layer
+	// Each block is released as soon as possible at end of each iteration
+	for (int32 BlockIndex = 0; BlockIndex < NumBlocks; ++BlockIndex)
+	{
+		BuildTask.EnterProgressFrame();
+
+		const FTextureSourceBlockData& SourceBlockData = SourceData.Blocks[BlockIndex];
+
+		// Current lock + mips that will be compressed to tiles
+		FVTBlockPayload& BlockData = LayerPayload[LayerIndex].Blocks[BlockIndex];
+
+		BlockData.BlockX = SourceBlockData.BlockX;
+		// UE applies a (1-y) transform to imported UVs, so apply a similar transform to UDIM block locations here
+		// This ensures that UDIM tiles will appear in the correct location when sampled with transformed UVs
+		BlockData.BlockY = (DerivedInfo.SizeInBlocksY - SourceBlockData.BlockY) % DerivedInfo.SizeInBlocksY;
+		BlockData.NumMips = SourceBlockData.NumMips;
+		BlockData.NumSlices = SourceBlockData.NumSlices;
+		BlockData.MipBias = SourceBlockData.MipBias;
+		BlockData.SizeX = 0u;
+		BlockData.SizeY = 0u;
+
+		const FTextureBuildSettings& BuildSettingsForLayer = SettingsPerLayer[LayerIndex];
+
+		const TArray<FImage>& SourceMips = SourceBlockData.MipsPerLayer[LayerIndex];
+		const TArray<FImage>* CompositeSourceMips = &EmptyImageArray;
+		if (CompositeSourceData.Blocks.Num() > 0)
 		{
-			const FVTSourceTileEntry& Tile = TileList[TileIndex];
+			CompositeSourceMips = &CompositeSourceData.Blocks[BlockIndex].MipsPerLayer[LayerIndex];
+		}
+
+		// Adjust the build settings to generate an uncompressed texture with mips but leave other settings
+		// like color correction, ... in place
+
+		// TBSettings starts with the full Texture settings, so we get all options
+		//  then we change FormatName to be == Source format, so no Compression is done
+		FTextureBuildSettings TBSettings = SettingsPerLayer[0];
+		//TBSettings.MaxTextureResolution = FTextureBuildSettings::MaxTextureResolutionDefault;
+		TBSettings.TextureFormatName = LayerData.FormatName;
+
+		if (LayerIndex != 0)
+		{
+			// @todo Oodle : this looks fragile
+			//	some of the processing options are copied from BuildSettingsForLayer
+			//	but some are NOT
+			//	it seems semi-random
+			//  In the common case of NumLayers==1 , then it doesn't matter
+			//	so this would be rarely observed
+
+			// TBSettings was set from Layer 0, copy in some settings from this Layer ?
+			TBSettings.bSRGB = BuildSettingsForLayer.bSRGB;
+			TBSettings.bUseLegacyGamma = BuildSettingsForLayer.bUseLegacyGamma;
+			TBSettings.bForceAlphaChannel = BuildSettingsForLayer.bForceAlphaChannel;
+			TBSettings.bForceNoAlphaChannel = BuildSettingsForLayer.bForceNoAlphaChannel;
+			TBSettings.bHDRSource = BuildSettingsForLayer.bHDRSource;
+			TBSettings.bApplyYCoCgBlockScale = BuildSettingsForLayer.bApplyYCoCgBlockScale;
+			TBSettings.bReplicateRed = BuildSettingsForLayer.bReplicateRed;
+			TBSettings.bReplicateAlpha = BuildSettingsForLayer.bReplicateAlpha;
+		}
+
+		// Make sure the output of the texture builder is in the same gamma space as we expect it.
+		check(TBSettings.GetDestGammaSpace() == BuildSettingsForLayer.GetDestGammaSpace());
+
+		// Leave original mip settings alone unless it's none at which point we will just generate them using a simple average
+		if (TBSettings.MipGenSettings == TMGS_NoMipmaps)
+		{
+			TBSettings.MipGenSettings = TMGS_SimpleAverage;
+		}
 
-			const FTextureSourceBlockData& Block = SourceBlocks[Tile.BlockIndex];
-			const FImage& SourceMip = Block.MipsPerLayer[LayerIndex][Tile.MipIndexInBlock];
-			check(SourceMip.Format == LayerData.ImageFormat);
-			check(SourceMip.GammaSpace == LayerData.GammaSpace);
-
-			FPixelDataRectangle SourceData(LayerData.SourceFormat,
-				SourceMip.SizeX,
-				SourceMip.SizeY,
-				const_cast<uint8*>(SourceMip.RawData.GetData()));
-
-			TArray<FImage> TileImages;
-			FImage* TileImage = new(TileImages) FImage(PhysicalTileSize, PhysicalTileSize, LayerData.ImageFormat, LayerData.GammaSpace);
-			FPixelDataRectangle TileData(LayerData.SourceFormat, PhysicalTileSize, PhysicalTileSize, TileImage->RawData.GetData());
-
-			TileData.Clear();
-			TileData.CopyRectangleBordered(0, 0, SourceData,
-				Tile.TileInBlockX * TileSize - BorderSize,
-				Tile.TileInBlockY * TileSize - BorderSize,
-				PhysicalTileSize,
-				PhysicalTileSize,
-				(TextureAddress)BuildSettingsLayer0.VirtualAddressingModeX,
-				(TextureAddress)BuildSettingsLayer0.VirtualAddressingModeY);
-
-#if 0//SAVE_TILES
+		// For multi-block images, we may have scaled the max block size to be tile-sized, but individual blocks may still be smaller than 1 tile
+		// These need to be scaled up as well (scaling up individual blocks has the effect of reducing the block's mip-bias)
+		int32 LocalBlockSizeScale = DerivedInfo.BlockSizeScale;
+		while (SourceMips[0].SizeX * LocalBlockSizeScale < TileSize || SourceMips[0].SizeY * LocalBlockSizeScale < TileSize)
+		{
+			check(BlockData.MipBias > 0u);
+			--BlockData.MipBias;
+			LocalBlockSizeScale *= 2;
+		}
+
+		// give each tile a unique DebugTexturePathName for DebugDump option :
+		FString CurDebugTexturePathName = FString::Printf(TEXT("%s_L%d_B%d"), *DebugTexturePathName, LayerIndex, BlockIndex);
+
+		// Use the texture compressor module to do all the hard work
+		// this is the Build to Uncompressed to apply processing to create the source for the tiles
+		TArray<FCompressedImage2D> CompressedMips;
+		bool bBuildTextureResult = false;
+		if (LocalBlockSizeScale == 1)
+		{
+			uint32 NumMipsInTail, ExtData;
+			bBuildTextureResult = Compressor->BuildTexture(SourceMips, *CompositeSourceMips, TBSettings, CurDebugTexturePathName, CompressedMips, NumMipsInTail, ExtData, nullptr);
+		}
+		else
+		{
+			// Need to generate scaled source images before building mips
+			// Typically this is only needed to scale very small source images to be at least tile-sized, so performance shouldn't be a big concern here
+			TArray<FImage> ScaledSourceMips;
+			TArray<FImage> ScaledCompositeMips;
+			ScaledSourceMips.Reserve(SourceMips.Num());
+			ScaledCompositeMips.Reserve(CompositeSourceMips->Num());
+			for (const FImage& SrcMip : SourceMips)
 			{
-				FString BasePath = FPaths::ProjectUserDir();
-				FString TileFileName = BasePath / FString::Format(TEXT("{0}_{1}_{2}_{3}_{4}"), TArray<FStringFormatArg>({ Settings.DebugName, Tile.X, Tile.Y, Tile.Mip, LayerIndex }));
-				TileData.Save(TileFileName, ImageWrapper);
-				FString TileSourceFileName = BasePath / FString::Format(TEXT("{0}_{1}_{2}_{3}_{4}_Src"), TArray<FStringFormatArg>({ Settings.DebugName, Tile.X, Tile.Y, Tile.Mip, LayerIndex }));
-				SourceData.Save(TileSourceFileName, ImageWrapper);
+				FImage* ScaledMip = new(ScaledSourceMips) FImage;
+				// Pow22 cannot be used as a destination gamma, so change it to sRGB now :
+				EGammaSpace GammaSpace = (SrcMip.GammaSpace == EGammaSpace::Pow22) ? EGammaSpace::sRGB : SrcMip.GammaSpace;
+				SrcMip.ResizeTo(*ScaledMip, SrcMip.SizeX * LocalBlockSizeScale, SrcMip.SizeY * LocalBlockSizeScale, SrcMip.Format, GammaSpace);
+			}
+
+			for (const FImage& SrcMip : *CompositeSourceMips)
+			{
+				FImage* ScaledMip = new(ScaledCompositeMips) FImage;
+				// Pow22 cannot be used as a destination gamma, so change it to sRGB now :
+				EGammaSpace GammaSpace = (SrcMip.GammaSpace == EGammaSpace::Pow22) ? EGammaSpace::sRGB : SrcMip.GammaSpace;
+				SrcMip.ResizeTo(*ScaledMip, SrcMip.SizeX * LocalBlockSizeScale, SrcMip.SizeY * LocalBlockSizeScale, SrcMip.Format, GammaSpace);
 			}
-#endif // SAVE_TILES
 
-			// give each tile a unique DebugTexturePathName for DebugDump option :
-			FString DebugTilePathName = FString::Printf(TEXT("%s_L%d_VT%04d"), *DebugTexturePathName, LayerIndex, TileIndex);
+			// Pow22 was converted to sRGB by Resize :
+			TBSettings.bUseLegacyGamma = false;
 
-			TArray<FCompressedImage2D> CompressedMip;
-			TArray<FImage> EmptyList;
 			uint32 NumMipsInTail, ExtData;
-			// this is the Build for Tiles to do the encode to GPU formats, with no processing
-			if (!ensure(Compressor->BuildTexture(TileImages, EmptyList, TBSettings, DebugTilePathName, CompressedMip, NumMipsInTail, ExtData, nullptr)))
+			bBuildTextureResult = Compressor->BuildTexture(ScaledSourceMips, ScaledCompositeMips, TBSettings, CurDebugTexturePathName, CompressedMips, NumMipsInTail, ExtData, nullptr);
+		}
+
+		check(bBuildTextureResult);
+
+		// Get size of block from Compressor output, since it may have been padded/adjusted
+		{
+			BlockData.SizeX = CompressedMips[0].SizeX;
+			BlockData.SizeY = CompressedMips[0].SizeY;
+
+			// re-compute mip bias to account for any resizing of this block (typically due to clamped max size)
+			const int32 MipBiasX = FMath::CeilLogTwo(DerivedInfo.BlockSizeX / BlockData.SizeX);
+			const int32 MipBiasY = FMath::CeilLogTwo(DerivedInfo.BlockSizeY / BlockData.SizeY);
+			checkf(MipBiasX == MipBiasY, TEXT("Mismatched aspect ratio (%d x %d), (%d x %d)"), DerivedInfo.BlockSizeX, DerivedInfo.BlockSizeY, BlockData.SizeX, BlockData.SizeY);
+			BlockData.MipBias = MipBiasX;
+		}
+
+		check(BlockData.SizeX << BlockData.MipBias == DerivedInfo.BlockSizeX);
+		check(BlockData.SizeY << BlockData.MipBias == DerivedInfo.BlockSizeY);
+
+		// Use actual block size (not the the one UDIM's passe here) to determine how many
+		// mips you'll have. As different blocks can be smaller than full UDIM block size
+		const uint32 BlockSizeXY = FMath::Max(BlockData.SizeX, BlockData.SizeY);
+		if (NumBlocks == 1u)
+		{
+			const uint32 MaxMipInBlockXY = FMath::CeilLogTwo(BlockSizeXY);
+			BlockData.NumMips = FMath::Min<int32>(CompressedMips.Num(), MaxMipInBlockXY + 1);
+		}
+		else
+		{
+			const uint32 BlockSizeInTilesXY = FMath::DivideAndRoundUp<uint32>(BlockSizeXY, TileSize);
+			const uint32 MaxMipInBlockXY = FMath::CeilLogTwo(BlockSizeInTilesXY);
+			BlockData.NumMips = FMath::Min<int32>(CompressedMips.Num(), MaxMipInBlockXY + 1);
+		}
+
+		BlockData.Mips.Reserve(BlockData.NumMips);
+		for (int32 MipIndex = 0; MipIndex < BlockData.NumMips; ++MipIndex)
+		{
+			FCompressedImage2D& CompressedMip = CompressedMips[MipIndex];
+			check(CompressedMip.PixelFormat == LayerData.PixelFormat);
+			FImage* Image = new(BlockData.Mips) FImage();
+			Image->SizeX = CompressedMip.SizeX;
+			Image->SizeY = CompressedMip.SizeY;
+			Image->Format = LayerData.ImageFormat;
+			Image->GammaSpace = LayerData.GammaSpace;
+			Image->NumSlices = 1;
+			check(Image->IsImageInfoValid());
+			Image->RawData = MoveTemp(CompressedMip.RawData);
+		}
+
+		if (bNeedsMiptailBlock)
+		{
+			const FImage& SrcMipImage = BlockData.Mips[MaxMipInBlock - BlockData.MipBias];
+			check(SrcMipImage.SizeX == MipWidthInBlock);
+			check(SrcMipImage.SizeY == MipHeightInBlock);
+
+			FPixelDataRectangle SrcPixelData(LayerData.SourceFormat, SrcMipImage.SizeX, SrcMipImage.SizeY, const_cast<uint8*>(SrcMipImage.RawData.GetData()));
+			MiptailPixelData.CopyRectangle(BlockData.BlockX * MipWidthInBlock, BlockData.BlockY * MipHeightInBlock, SrcPixelData, 0, 0, MipWidthInBlock, MipHeightInBlock);
+		}
+		else
+		{
+			// Extract fallback color from last mip.
 			{
-				bCompressionError = true;
+				// this actually just samples one pixel ; it comes from last mip so it's often small already
+				// @todo Oodle : just use a "get average color" function
+				FImage OnePixelImage(1, 1, 1, ERawImageFormat::RGBA32F);
+				BlockData.Mips.Last().ResizeTo(OnePixelImage, 1, 1, ERawImageFormat::RGBA32F, EGammaSpace::Linear);
+				OutData.LayerFallbackColors[LayerIndex] = OnePixelImage.AsRGBA32F()[0];
 			}
+		}
 
-			check(CompressedMip.Num() == 1);
-			checkf(CompressedFormat == PF_Unknown || CompressedFormat == CompressedMip[0].PixelFormat, 
-				TEXT("CompressedFormat: %s (%d), CompressedMip[0].PixelFormat: %s (%d)"),
-				GetSafePixelFormatName(CompressedFormat), (int32)CompressedFormat, 
-				GetSafePixelFormatName((EPixelFormat)CompressedMip[0].PixelFormat), (int32)CompressedMip[0].PixelFormat);
+		BuildBlockTiles(LayerIndex, BlockIndex, BlockData, LayerData, bAllowAsync);
 
-			CompressedFormat = (EPixelFormat)CompressedMip[0].PixelFormat;
+		// Mips not needed anymore
+		BlockData.Mips.Empty();
 
-			GeneratedData.TilePayload[TileIndex] = MoveTemp(CompressedMip[0].RawData);
+		// SourceData for this Block & Layer not needed anymore, can free SourceData (and CompositeSourceData) now
+		SourceData.Blocks[BlockIndex].MipsPerLayer[LayerIndex].Empty();
+		if (CompositeSourceData.Blocks.Num() > 0)
+		{
+			CompositeSourceData.Blocks[BlockIndex].MipsPerLayer[LayerIndex].Empty();
+		}
+	}
 
-			// if the SourceBlocks we used has no more tiles that need to read from it, free it now?
-			//	  what we'd like to do is do all the tiles within each Block/Layer at a time and free that source Block/Layer
-			//	  immediately when all its tiles are done before starting the next Block/Layer
-		}, 
-		(bIsSingleThreaded ? EParallelForFlags::ForceSingleThread : EParallelForFlags::None));
+	if (bNeedsMiptailBlock)
+	{
+#if SAVE_TILES
+		{
+			FString DebugName = FPaths::MakeValidFileName(*DebugTexturePathName, TEXT('_'));
+			const FString BasePath = FPaths::ProjectUserDir();
+			const FString MipFileName = BasePath / FString::Format(TEXT("{0}_{1}.png"), TArray<FStringFormatArg>({ *DebugName, LayerIndex }));
+			MiptailPixelData.Save(MipFileName, ImageWrapper);
+		}
+#endif // SAVE_TILES
+
+		FVTBlockPayload& BlockData = LayerPayload[LayerIndex].Blocks.Last();
+		BlockData.BlockX = 0;
+		BlockData.BlockY = 0;
+		BlockData.SizeInBlocksX = DerivedInfo.SizeInBlocksX; // miptail block covers the entire logical source texture
+		BlockData.SizeInBlocksY = DerivedInfo.SizeInBlocksY;
+		BlockData.SizeX = FMath::Max(MipInputSizeX >> 1, 1u);
+		BlockData.SizeY = FMath::Max(MipInputSizeY >> 1, 1u);
+		BlockData.NumMips = OutData.NumMips - MaxMipInBlock - 1;//   FMath::CeilLogTwo(MipInputSize); // Don't add 1, since 'MipInputSize' is one mip larger
+		BlockData.NumSlices = 1; // TODO?
+		BlockData.MipBias = MaxMipInBlock + 1;
+		check(BlockData.NumMips > 0);
+
+		// Total number of mips should be equal to number of mips per block plus number of miptail mips
+		check(MaxMipInBlock + BlockData.NumMips + 1 == OutData.NumMips);
+
+
+		const FTextureBuildSettings& BuildSettingsForLayer = SettingsPerLayer[LayerIndex];
+
+		// Adjust the build settings to generate an uncompressed texture with mips but leave other settings
+		// like color correction, ... in place
+		FTextureBuildSettings TBSettings = SettingsPerLayer[0];
+		TBSettings.MaxTextureResolution = FTextureBuildSettings::MaxTextureResolutionDefault;; // don't limit the size of the mip-tail, this limit only applies to each source block
+		TBSettings.TextureFormatName = LayerData.FormatName;
+		TBSettings.bSRGB = BuildSettingsForLayer.bSRGB;
+		TBSettings.bUseLegacyGamma = BuildSettingsForLayer.bUseLegacyGamma;
+
+		// Make sure the output of the texture builder is in the same gamma space as we expect it.
+		check(TBSettings.GetDestGammaSpace() == BuildSettingsForLayer.GetDestGammaSpace());
+
+		// Leave original mip settings alone unless it's none at which point we will just generate them using a simple average
+		if (TBSettings.MipGenSettings == TMGS_NoMipmaps || TBSettings.MipGenSettings == TMGS_LeaveExistingMips)
+		{
+			TBSettings.MipGenSettings = TMGS_SimpleAverage;
+		}
+
+		// give each tile a unique DebugTexturePathName for DebugDump option :
+		FString CurDebugTexturePathName = FString::Printf(TEXT("%s_L%d_MT"), *DebugTexturePathName, LayerIndex);
+
+		// Use the texture compressor module to do all the hard work
+		// TODO - composite images?
+		TArray<FCompressedImage2D> CompressedMips;
+		uint32 NumMipsInTail, ExtData;
+		// this is a Build to uncompressed, to apply processing
+		if (!Compressor->BuildTexture(MiptailInputImages, EmptyImageArray, TBSettings, CurDebugTexturePathName, CompressedMips, NumMipsInTail, ExtData, nullptr))
+		{
+			check(false);
+		}
+
+		// We skip the first compressed mip output, since that will just be a copy of the input
+		check(CompressedMips.Num() >= BlockData.NumMips + 1);
+		check(BlockData.SizeX == CompressedMips[1].SizeX);
+		check(BlockData.SizeY == CompressedMips[1].SizeY);
+
+		BlockData.Mips.Reserve(CompressedMips.Num() - 1);
+		for (int32 MipIndex = 1; MipIndex < BlockData.NumMips + 1; ++MipIndex)
+		{
+			FCompressedImage2D& CompressedMip = CompressedMips[MipIndex];
+			check(CompressedMip.PixelFormat == LayerData.PixelFormat);
+			FImage* Image = new(BlockData.Mips) FImage();
+			Image->SizeX = CompressedMip.SizeX;
+			Image->SizeY = CompressedMip.SizeY;
+			Image->Format = LayerData.ImageFormat;
+			Image->GammaSpace = LayerData.GammaSpace;
+			Image->NumSlices = 1;
+			check(Image->IsImageInfoValid());
+			Image->RawData = MoveTemp(CompressedMip.RawData);
+		}
+
+		BuildBlockTiles(LayerIndex, NumBlocks, BlockData, LayerData, bAllowAsync);
+
+		// Extract fallback color from last mip.
+		{
+			// this actually just samples one pixel ; it comes from last mip so it's often small already
+			// @todo Oodle : just use a "get average color" function
+			FImage OnePixelImage(1, 1, 1, ERawImageFormat::RGBA32F);
+			BlockData.Mips.Last().ResizeTo(OnePixelImage, 1, 1, ERawImageFormat::RGBA32F, EGammaSpace::Linear);
+			OutData.LayerFallbackColors[LayerIndex] = OnePixelImage.AsRGBA32F()[0];
+		}
+	}
+}
+
+bool FVirtualTextureDataBuilder::BuildPagesForChunk(const TArray<FVTSourceTileEntry>& ActiveTileList)
+{
+	TArray<FLayerData> LayerData;
+	LayerData.AddDefaulted(LayerPayload.Num());
+
+	for (int32 LayerIndex = 0; LayerIndex < LayerData.Num(); LayerIndex++)
+	{
+		BuildTiles(ActiveTileList, LayerIndex, LayerData[LayerIndex]);
+	}
+
+	// Fill out tile offsets per layer if we haven't yet and if all layers are raw uncompressed data.
+	if (OutData.TileDataOffsetPerLayer.Num() == 0)
+	{
+		bool bIsRawGPUData = true;
+		for (int32 LayerIndex = 0; LayerIndex < LayerData.Num(); LayerIndex++)
+		{
+			if (LayerData[LayerIndex].Codec != EVirtualTextureCodec::RawGPU)
+			{
+				bIsRawGPUData = false;
+				break;
+			}
+		}
+		if (bIsRawGPUData)
+		{
+			int64 TileDataOffset = 0;
+			OutData.TileDataOffsetPerLayer.Reserve(LayerData.Num());
+			for (int32 LayerIndex = 0; LayerIndex < LayerData.Num(); LayerIndex++)
+			{
+				TileDataOffset += LayerData[LayerIndex].TilePayload[0].Num();
+				OutData.TileDataOffsetPerLayer.Add( IntCastChecked<uint32>(TileDataOffset) );
+			}
+		}
+	}
+
+	// Write tiles out to chunk.
+	return PushDataToChunk(ActiveTileList, LayerData);
+}
+
+void FVirtualTextureDataBuilder::BuildTiles(const TArray<FVTSourceTileEntry>& TileList, uint32 LayerIndex, FLayerData& GeneratedData)
+{
+	TRACE_CPUPROFILER_EVENT_SCOPE(Texture.VT.BuildTiles);
+
+	FThreadSafeBool bCompressionError = false;
+	EPixelFormat CompressedFormat = PF_Unknown;
+
+	{
+		GeneratedData.TilePayload.AddDefaulted(TileList.Num());
+		
+		for (int32 TileIndex = 0; TileIndex < TileList.Num(); TileIndex++)
+		{
+			const FVTSourceTileEntry& Tile = TileList[TileIndex];
+			FVTBlockPayload* Block = Tile.Block;
+			GeneratedData.TilePayload[TileIndex] = MoveTemp(Block->Tiles[Tile.TileIndexInBlock].Mips[Tile.MipIndexInBlock].Payload);
+			CompressedFormat = Block->Tiles[Tile.TileIndexInBlock].Mips[Tile.MipIndexInBlock].CompressedFormat;
+		}
 
 		GeneratedData.Codec = EVirtualTextureCodec::RawGPU;
 	}
@@ -741,7 +1165,7 @@ void FVirtualTextureDataBuilder::BuildTiles(const TArray<FVTSourceTileEntry>& Ti
 
 bool FVirtualTextureDataBuilder::PushDataToChunk(const TArray<FVTSourceTileEntry>& Tiles, const TArray<FLayerData>& LayerData)
 {
-	const int32 NumLayers = SourceLayers.Num();
+	const int32 NumLayers = LayerPayload.Num();
 
 	int64 TotalSize = sizeof(FVirtualTextureChunkHeader);
 	for (int32 Layer = 0; Layer < NumLayers; ++Layer)
@@ -815,375 +1239,39 @@ bool FVirtualTextureDataBuilder::PushDataToChunk(const TArray<FVTSourceTileEntry
 
 	FSHA1::HashBuffer(NewChunkData, TotalSize, Chunk.BulkDataHash.Hash);
 
+#if SAVE_CHUNKS
+	{
+		FString DebugName = FPaths::MakeValidFileName(*DebugTexturePathName, TEXT('_'));
+		FString BasePath = FPaths::ProjectUserDir();
+		FString Name = BasePath / FString::Format(TEXT("chunk_{0}_{1}.bin"), TArray<FStringFormatArg>({ ChunkDumpIndex++, Chunk.BulkDataHash.ToString() }));
+		FFileHelper::SaveArrayToFile(TArrayView<uint8> { NewChunkData, (int32)TotalSize }, * Name);
+	}
+#endif
+
 	BulkData.Unlock();
 	BulkData.SetBulkDataFlags(BULKDATA_Force_NOT_InlinePayload);
 
 	return true;
 }
 
-int32 FVirtualTextureDataBuilder::FindSourceBlockIndex(int32 MipIndex, int32 BlockX, int32 BlockY) const
+FVTBlockPayload* FVirtualTextureDataBuilder::FindSourceBlock(int32 MipIndex, int32 BlockX, int32 BlockY)
 {
-	for(int32 BlockIndex = 0; BlockIndex < SourceBlocks.Num(); ++BlockIndex)
-	{
-		const FTextureSourceBlockData& Block = SourceBlocks[BlockIndex];
-		if (BlockX >= Block.BlockX && BlockX < Block.BlockX + Block.SizeInBlocksX &&
-			BlockY >= Block.BlockY && BlockY < Block.BlockY + Block.SizeInBlocksY &&
-			MipIndex >= Block.MipBias &&
-			(MipIndex - Block.MipBias) < Block.NumMips)
-		{
-			return BlockIndex;
-		}
-	}
-	return INDEX_NONE;
-}
-
-// This builds an uncompressed version of the texture containing all other build settings baked in
-// color corrections, mip sharpening, ....
-void FVirtualTextureDataBuilder::BuildSourcePixels(const FTextureSourceData& SourceData, const FTextureSourceData& CompositeSourceData)
-{
-	static const TArray<FImage> EmptyImageArray;
-
-	const int32 TileSize = SettingsPerLayer[0].VirtualTextureTileSize;
-	const int32 NumBlocks = SourceData.Blocks.Num();
-	const int32 NumLayers = SourceData.Layers.Num();
-
-	SourceLayers.AddDefaulted(NumLayers);
-	for (int32 LayerIndex = 0; LayerIndex < NumLayers; ++LayerIndex)
-	{
-		const FTextureBuildSettings& BuildSettingsForLayer = SettingsPerLayer[LayerIndex];
-		FVirtualTextureSourceLayerData& LayerData = SourceLayers[LayerIndex];
-
-		// Specify the format we are processing to in this step :
-		LayerData.ImageFormat = UE::TextureBuildUtilities::GetVirtualTextureBuildIntermediateFormat(BuildSettingsForLayer);
-		LayerData.GammaSpace = BuildSettingsForLayer.GetDestGammaSpace();
-		// Gamma correction can either be applied in step 1 or step 2 of the VT build
-		//	depending on whether the Intermediate format is U8 or not
-		if ( ! ERawImageFormat::GetFormatNeedsGammaSpace(LayerData.ImageFormat) )
-		{
-			LayerData.GammaSpace = EGammaSpace::Linear;
-		}
-
-		// LayerData.bHasAlpha will be updated after the Build to detect if there is actual alpha in the layers
-		LayerData.bHasAlpha = BuildSettingsForLayer.bForceAlphaChannel;
-		
-		LayerData.FormatName = FImageCoreUtils::ConvertToUncompressedTextureFormatName(LayerData.ImageFormat);
-		LayerData.PixelFormat = FImageCoreUtils::GetPixelFormatForRawImageFormat(LayerData.ImageFormat);
-		LayerData.SourceFormat = FImageCoreUtils::ConvertToTextureSourceFormat(LayerData.ImageFormat);
-	}
-
-	SourceBlocks.AddDefaulted(NumBlocks);
-	for (int32 BlockIndex = 0; BlockIndex < NumBlocks; ++BlockIndex)
-	{
-		const FTextureSourceBlockData& SourceBlockData = SourceData.Blocks[BlockIndex];
-
-		FTextureSourceBlockData& BlockData = SourceBlocks[BlockIndex];
-		BlockData.BlockX = SourceBlockData.BlockX;
-		// UE applies a (1-y) transform to imported UVs, so apply a similar transform to UDIM block locations here
-		// This ensures that UDIM tiles will appear in the correct location when sampled with transformed UVs
-		BlockData.BlockY = (DerivedInfo.SizeInBlocksY - SourceBlockData.BlockY) % DerivedInfo.SizeInBlocksY;
-		BlockData.NumMips = SourceBlockData.NumMips;
-		BlockData.NumSlices = SourceBlockData.NumSlices;
-		BlockData.MipBias = SourceBlockData.MipBias;
-		BlockData.SizeX = 0u;
-		BlockData.SizeY = 0u;
-		BlockData.MipsPerLayer.AddDefaulted(NumLayers);
-		for (int32 LayerIndex = 0; LayerIndex < NumLayers; ++LayerIndex)
-		{
-			const FTextureBuildSettings& BuildSettingsForLayer = SettingsPerLayer[LayerIndex];
-			FVirtualTextureSourceLayerData& LayerData = SourceLayers[LayerIndex];
-
-			const TArray<FImage>& SourceMips = SourceBlockData.MipsPerLayer[LayerIndex];
-			const TArray<FImage>* CompositeSourceMips = &EmptyImageArray;
-			if (CompositeSourceData.Blocks.Num() > 0)
-			{
-				CompositeSourceMips = &CompositeSourceData.Blocks[BlockIndex].MipsPerLayer[LayerIndex];
-			}
-
-			// Adjust the build settings to generate an uncompressed texture with mips but leave other settings
-			// like color correction, ... in place
-
-			// TBSettings starts with the full Texture settings, so we get all options
-			//  then we change FormatName to be == Source format, so no Compression is done
-			FTextureBuildSettings TBSettings = SettingsPerLayer[0];
-			//TBSettings.MaxTextureResolution = FTextureBuildSettings::MaxTextureResolutionDefault;
-			TBSettings.TextureFormatName = LayerData.FormatName;
-			
-			if ( LayerIndex != 0 )
-			{
-				//	some of the processing options are copied from BuildSettingsForLayer
-				//	but some are NOT
-				//	it seems semi-random
-				//  In the common case of NumLayers==1 , then it doesn't matter
-				//	so this would be rarely observed
-
-				// TBSettings was set from Layer 0, copy in some settings from this Layer ?
-				TBSettings.bSRGB = BuildSettingsForLayer.bSRGB;
-				TBSettings.bUseLegacyGamma = BuildSettingsForLayer.bUseLegacyGamma;
-				TBSettings.bForceAlphaChannel = BuildSettingsForLayer.bForceAlphaChannel;
-				TBSettings.bForceNoAlphaChannel = BuildSettingsForLayer.bForceNoAlphaChannel;
-				TBSettings.bHDRSource = BuildSettingsForLayer.bHDRSource;
-				TBSettings.bApplyYCoCgBlockScale = BuildSettingsForLayer.bApplyYCoCgBlockScale;
-				TBSettings.bReplicateRed = BuildSettingsForLayer.bReplicateRed;
-				TBSettings.bReplicateAlpha = BuildSettingsForLayer.bReplicateAlpha;
-			}
-
-			// Make sure the output of the texture builder is in the same gamma space as we expect it.
-			check(TBSettings.GetDestGammaSpace() == BuildSettingsForLayer.GetDestGammaSpace());
-
-			// Leave original mip settings alone unless it's none at which point we will just generate them using a simple average
-			if (TBSettings.MipGenSettings == TMGS_NoMipmaps)
-			{
-				TBSettings.MipGenSettings = TMGS_SimpleAverage;
-			}
-
-			// For multi-block images, we may have scaled the max block size to be tile-sized, but individual blocks may still be smaller than 1 tile
-			// These need to be scaled up as well (scaling up individual blocks has the effect of reducing the block's mip-bias)
-			int32 LocalBlockSizeScale = DerivedInfo.BlockSizeScale;
-			while (SourceMips[0].SizeX * LocalBlockSizeScale < TileSize || SourceMips[0].SizeY * LocalBlockSizeScale < TileSize)
-			{
-				check(BlockData.MipBias > 0u);
-				--BlockData.MipBias;
-				LocalBlockSizeScale *= 2;
-			}
-			
-			// give each tile a unique DebugTexturePathName for DebugDump option :
-			FString CurDebugTexturePathName = FString::Printf(TEXT("%s_L%d_B%d"), *DebugTexturePathName, LayerIndex, BlockIndex);
-
-			// Use the texture compressor module to do all the hard work
-			// this is the Build to Uncompressed to apply processing to create the source for the tiles
-			TArray<FCompressedImage2D> CompressedMips;
-			bool bBuildTextureResult = false;
-			if (LocalBlockSizeScale == 1)
-			{
-				uint32 NumMipsInTail, ExtData;
-				bBuildTextureResult = Compressor->BuildTexture(SourceMips, *CompositeSourceMips, TBSettings, CurDebugTexturePathName, CompressedMips, NumMipsInTail, ExtData, nullptr);
-			}
-			else
-			{
-				// Need to generate scaled source images before building mips
-				// Typically this is only needed to scale very small source images to be at least tile-sized, so performance shouldn't be a big concern here
-				TArray<FImage> ScaledSourceMips;
-				TArray<FImage> ScaledCompositeMips;
-				ScaledSourceMips.Reserve(SourceMips.Num());
-				ScaledCompositeMips.Reserve(CompositeSourceMips->Num());
-				for (const FImage& SrcMip : SourceMips)
-				{
-					FImage* ScaledMip = new(ScaledSourceMips) FImage;
-					// Pow22 cannot be used as a destination gamma, so change it to sRGB now :
-					EGammaSpace GammaSpace = (SrcMip.GammaSpace == EGammaSpace::Pow22) ? EGammaSpace::sRGB : SrcMip.GammaSpace;
-					SrcMip.ResizeTo(*ScaledMip, SrcMip.SizeX * LocalBlockSizeScale, SrcMip.SizeY * LocalBlockSizeScale, SrcMip.Format, GammaSpace);
-				}
-
-				for (const FImage& SrcMip : *CompositeSourceMips)
-				{
-					FImage* ScaledMip = new(ScaledCompositeMips) FImage;
-					// Pow22 cannot be used as a destination gamma, so change it to sRGB now :
-					EGammaSpace GammaSpace = (SrcMip.GammaSpace == EGammaSpace::Pow22) ? EGammaSpace::sRGB : SrcMip.GammaSpace;
-					SrcMip.ResizeTo(*ScaledMip, SrcMip.SizeX * LocalBlockSizeScale, SrcMip.SizeY * LocalBlockSizeScale, SrcMip.Format, GammaSpace);
-				}
-
-				// Pow22 was converted to sRGB by Resize :
-				TBSettings.bUseLegacyGamma = false;
-
-				uint32 NumMipsInTail, ExtData;
-				bBuildTextureResult = Compressor->BuildTexture(ScaledSourceMips, ScaledCompositeMips, TBSettings, CurDebugTexturePathName, CompressedMips, NumMipsInTail, ExtData, nullptr);
-			}
-
-			check(bBuildTextureResult);
-
-			// Get size of block from Compressor output, since it may have been padded/adjusted
-			{
-				BlockData.SizeX = CompressedMips[0].SizeX;
-				BlockData.SizeY = CompressedMips[0].SizeY;
-
-				// re-compute mip bias to account for any resizing of this block (typically due to clamped max size)
-				const int32 MipBiasX = FMath::CeilLogTwo(DerivedInfo.BlockSizeX / BlockData.SizeX);
-				const int32 MipBiasY = FMath::CeilLogTwo(DerivedInfo.BlockSizeY / BlockData.SizeY);
-				checkf(MipBiasX == MipBiasY, TEXT("Mismatched aspect ratio (%d x %d), (%d x %d)"), DerivedInfo.BlockSizeX, DerivedInfo.BlockSizeY, BlockData.SizeX, BlockData.SizeY);
-				BlockData.MipBias = MipBiasX;
-			}
-
-			check(BlockData.SizeX << BlockData.MipBias == DerivedInfo.BlockSizeX);
-			check(BlockData.SizeY << BlockData.MipBias == DerivedInfo.BlockSizeY);
-
-			const uint32 BlockSize = FMath::Max(BlockData.SizeX, BlockData.SizeY);
-			if (NumBlocks == 1u)
-			{
-				const uint32 MaxMipInBlock = FMath::CeilLogTwo(BlockSize);
-				BlockData.NumMips = FMath::Min<int32>(CompressedMips.Num(), MaxMipInBlock + 1);
-			}
-			else
-			{
-				const uint32 BlockSizeInTiles = FMath::DivideAndRoundUp<uint32>(BlockSize, TileSize);
-				const uint32 MaxMipInBlock = FMath::CeilLogTwo(BlockSizeInTiles);
-				BlockData.NumMips = FMath::Min<int32>(CompressedMips.Num(), MaxMipInBlock + 1);
-			}
-			BlockData.MipsPerLayer[LayerIndex].Reserve(BlockData.NumMips);
-			for (int32 MipIndex = 0; MipIndex < BlockData.NumMips; ++MipIndex)
-			{
-				FCompressedImage2D& CompressedMip = CompressedMips[MipIndex];
-				check(CompressedMip.PixelFormat == LayerData.PixelFormat);
-				FImage* Image = new(BlockData.MipsPerLayer[LayerIndex]) FImage();
-				Image->SizeX = CompressedMip.SizeX;
-				Image->SizeY = CompressedMip.SizeY;
-				Image->Format = LayerData.ImageFormat;
-				Image->GammaSpace = LayerData.GammaSpace;
-				Image->NumSlices = 1;
-				check( Image->IsImageInfoValid() );
-				Image->RawData = MoveTemp(CompressedMip.RawData);
-			}
-
-			// bHasAlpha was previously set to true if bForceAlphaChannel
-			// if it's false and not bForceNoAlphaChannel, we scan each block for alpha
-			// if alpha is in any tile of any block, it gets enabled for all so they have consistent pixel format
-			if (!LayerData.bHasAlpha && !BuildSettingsForLayer.bForceNoAlphaChannel )
-			{
-				LayerData.bHasAlpha = DetectAlphaChannel(BlockData.MipsPerLayer[LayerIndex][0]);
-			}
-		}
-	}
-
-	// If we have more than 1 block, need to create miptail that contains mips made from multiple blocks
-	if (NumBlocks > 1)
+	for (int32 LayerIndex = 0; LayerIndex < LayerPayload.Num(); ++LayerIndex)
 	{
-		const uint32 BlockSize = FMath::Max(DerivedInfo.BlockSizeX, DerivedInfo.BlockSizeY);
-		const uint32 BlockSizeInTiles = FMath::DivideAndRoundUp<uint32>(BlockSize, TileSize);
-		const uint32 MaxMipInBlock = FMath::CeilLogTwo(BlockSizeInTiles);
-		const uint32 MipWidthInBlock = FMath::Max<uint32>(DerivedInfo.BlockSizeX >> MaxMipInBlock, 1);
-		const uint32 MipHeightInBlock = FMath::Max<uint32>(DerivedInfo.BlockSizeY >> MaxMipInBlock, 1);
-		const uint32 MipInputSizeX = FMath::RoundUpToPowerOfTwo(DerivedInfo.SizeInBlocksX * MipWidthInBlock);
-		const uint32 MipInputSizeY = FMath::RoundUpToPowerOfTwo(DerivedInfo.SizeInBlocksY * MipHeightInBlock);
-		const uint32 MipInputSize = FMath::Max(MipInputSizeX, MipInputSizeY);
-		//const uint32 MipInputSizeInTiles = FMath::DivideAndRoundUp<uint32>(MipInputSize, TileSize);
-
-		FTextureSourceBlockData& SourceMiptailBlock = SourceBlocks.AddDefaulted_GetRef();
-		SourceMiptailBlock.BlockX = 0;
-		SourceMiptailBlock.BlockY = 0;
-		SourceMiptailBlock.SizeInBlocksX = DerivedInfo.SizeInBlocksX; // miptail block covers the entire logical source texture
-		SourceMiptailBlock.SizeInBlocksY = DerivedInfo.SizeInBlocksY;
-		SourceMiptailBlock.SizeX = FMath::Max(MipInputSizeX >> 1, 1u);
-		SourceMiptailBlock.SizeY = FMath::Max(MipInputSizeY >> 1, 1u);
-		SourceMiptailBlock.NumMips = OutData.NumMips - MaxMipInBlock - 1;//   FMath::CeilLogTwo(MipInputSize); // Don't add 1, since 'MipInputSize' is one mip larger
-		SourceMiptailBlock.NumSlices = 1; // TODO?
-		SourceMiptailBlock.MipBias = MaxMipInBlock + 1;
-		SourceMiptailBlock.MipsPerLayer.AddDefaulted(NumLayers);
-		check(SourceMiptailBlock.NumMips > 0);
-
-		// Total number of mips should be equal to number of mips per block plus number of miptail mips
-		check(MaxMipInBlock + SourceMiptailBlock.NumMips + 1 == OutData.NumMips);
-
-		TArray<FImage> MiptailInputImages;
-		for (int32 LayerIndex = 0u; LayerIndex < NumLayers; ++LayerIndex)
+		TArray<FVTBlockPayload>& Blocks = LayerPayload[LayerIndex].Blocks;
+		for (int32 BlockIndex = 0; BlockIndex < Blocks.Num(); ++BlockIndex)
 		{
-			const FTextureBuildSettings& BuildSettingsForLayer = SettingsPerLayer[LayerIndex];
-			const FVirtualTextureSourceLayerData& LayerData = SourceLayers[LayerIndex];
-
-			MiptailInputImages.Reset(1);
-			FImage* MiptailInputImage = new(MiptailInputImages) FImage();
-			MiptailInputImage->Init(MipInputSizeX, MipInputSizeY, LayerData.ImageFormat, LayerData.GammaSpace);
-			FPixelDataRectangle DstPixelData(LayerData.SourceFormat, MipInputSizeX, MipInputSizeY, MiptailInputImage->RawData.GetData());
-			DstPixelData.Clear();
-
-			for (int32 BlockIndex = 0; BlockIndex < NumBlocks; ++BlockIndex)
-			{
-				const FTextureSourceBlockData& BlockData = SourceBlocks[BlockIndex];
-				const FImage& SrcMipImage = BlockData.MipsPerLayer[LayerIndex][MaxMipInBlock - BlockData.MipBias];
-				check(SrcMipImage.SizeX == MipWidthInBlock);
-				check(SrcMipImage.SizeY == MipHeightInBlock);
-
-				FPixelDataRectangle SrcPixelData(LayerData.SourceFormat, SrcMipImage.SizeX, SrcMipImage.SizeY, const_cast<uint8*>(SrcMipImage.RawData.GetData()));
-				DstPixelData.CopyRectangle(BlockData.BlockX * MipWidthInBlock, BlockData.BlockY * MipHeightInBlock, SrcPixelData, 0, 0, MipWidthInBlock, MipHeightInBlock);
-			}
-
-#if SAVE_TILES
+			FVTBlockPayload& Block = Blocks[BlockIndex];
+			if (BlockX >= Block.BlockX && BlockX < Block.BlockX + Block.SizeInBlocksX &&
+				BlockY >= Block.BlockY && BlockY < Block.BlockY + Block.SizeInBlocksY &&
+				MipIndex >= Block.MipBias &&
+				(MipIndex - Block.MipBias) < Block.NumMips)
 			{
-				const FString BasePath = FPaths::ProjectUserDir();
-				int32 ObjectNameStart;
-				SourceData.TextureFullName.FindLastChar('.', ObjectNameStart);
-				FString ObjectName = SourceData.TextureFullName.Right(SourceData.TextureFullName.Len() - ObjectNameStart - 1);
-				const FString MipFileName = BasePath / FString::Format(TEXT("{0}_{1}"), TArray<FStringFormatArg>({MoveTemp(ObjectName), LayerIndex}));
-				DstPixelData.Save(MipFileName, ImageWrapper);
-			}
-#endif // SAVE_TILES
-
-			// Adjust the build settings to generate an uncompressed texture with mips but leave other settings
-			// like color correction, ... in place
-			FTextureBuildSettings TBSettings = SettingsPerLayer[0];
-			TBSettings.MaxTextureResolution = FTextureBuildSettings::MaxTextureResolutionDefault;; // don't limit the size of the mip-tail, this limit only applies to each source block
-			TBSettings.TextureFormatName = LayerData.FormatName;
-			TBSettings.bSRGB = BuildSettingsForLayer.bSRGB;
-			TBSettings.bUseLegacyGamma = BuildSettingsForLayer.bUseLegacyGamma;
-
-			// Make sure the output of the texture builder is in the same gamma space as we expect it.
-			check(TBSettings.GetDestGammaSpace() == BuildSettingsForLayer.GetDestGammaSpace());
-
-			// Leave original mip settings alone unless it's none at which point we will just generate them using a simple average
-			if (TBSettings.MipGenSettings == TMGS_NoMipmaps || TBSettings.MipGenSettings == TMGS_LeaveExistingMips)
-			{
-				TBSettings.MipGenSettings = TMGS_SimpleAverage;
-			}
-
-			// give each tile a unique DebugTexturePathName for DebugDump option :
-			FString CurDebugTexturePathName = FString::Printf(TEXT("%s_L%d_MT"), *DebugTexturePathName, LayerIndex);
-
-			// Use the texture compressor module to do all the hard work
-			// TODO - composite images?
-			TArray<FCompressedImage2D> CompressedMips;
-			uint32 NumMipsInTail, ExtData;
-			// this is a Build to uncompressed, to apply processing
-			if (!Compressor->BuildTexture(MiptailInputImages, EmptyImageArray, TBSettings, CurDebugTexturePathName, CompressedMips, NumMipsInTail, ExtData, nullptr))
-			{
-				check(false);
-			}
-
-			// We skip the first compressed mip output, since that will just be a copy of the input
-			check(CompressedMips.Num() >= SourceMiptailBlock.NumMips + 1);
-			check(SourceMiptailBlock.SizeX == CompressedMips[1].SizeX);
-			check(SourceMiptailBlock.SizeY == CompressedMips[1].SizeY);
-
-			SourceMiptailBlock.MipsPerLayer[LayerIndex].Reserve(CompressedMips.Num() - 1);
-			for(int32 MipIndex = 1; MipIndex < SourceMiptailBlock.NumMips + 1; ++MipIndex)
-			{
-				FCompressedImage2D& CompressedMip = CompressedMips[MipIndex];
-				check(CompressedMip.PixelFormat == LayerData.PixelFormat);
-				FImage* Image = new(SourceMiptailBlock.MipsPerLayer[LayerIndex]) FImage();
-				Image->SizeX = CompressedMip.SizeX;
-				Image->SizeY = CompressedMip.SizeY;
-				Image->Format = LayerData.ImageFormat;
-				Image->GammaSpace = LayerData.GammaSpace;
-				Image->NumSlices = 1;
-				check( Image->IsImageInfoValid() );
-				Image->RawData = MoveTemp(CompressedMip.RawData);
+				return &Block;
 			}
 		}
 	}
-
-	// Extract fallback color from last mip.
-	for (int32 LayerIndex = 0; LayerIndex < NumLayers; ++LayerIndex)
-	{
-		// this actually just samples one pixel ; it comes from last mip so it's often small already
-		// @todo Oodle : just use a "get average color" function
-		FImage OnePixelImage(1, 1, 1, ERawImageFormat::RGBA32F);
-		SourceBlocks.Last().MipsPerLayer[LayerIndex].Last().ResizeTo(OnePixelImage, 1, 1, ERawImageFormat::RGBA32F, EGammaSpace::Linear);
-		OutData.LayerFallbackColors[LayerIndex] = OnePixelImage.AsRGBA32F()[0];
-	}
-
-	for (int32 LayerIndex = 0; LayerIndex < NumLayers; ++LayerIndex)
-	{
-		const FTextureBuildSettings& BuildSettingsForLayer = SettingsPerLayer[LayerIndex];
-		FVirtualTextureSourceLayerData& LayerData = SourceLayers[LayerIndex];
-
-		// Don't want platform specific swizzling for VT tile data, this tends to add extra padding for textures with odd dimensions
-		// (VT physical tiles generally not power-of-2 after adding border)
-		// Must match TextureDerivedData.cpp
-		LayerData.TextureFormatName = UE::TextureBuildUtilities::TextureFormatRemovePlatformPrefixFromName(BuildSettingsForLayer.TextureFormatName);
-	}
-}
-
-void FVirtualTextureDataBuilder::FreeSourcePixels()
-{
-	SourceBlocks.Empty();
-	SourceLayers.Empty();
+	return nullptr;
 }
 
 // Leaving this code here for now, in case we want to build a new/better system for creating/storing miptails
@@ -1265,7 +1353,22 @@ bool FVirtualTextureDataBuilder::DetectAlphaChannel(const FImage &Image)
 		const FFloat16Color* LastColor = SrcColors + (Image.SizeX * Image.SizeY * Image.NumSlices);
 		while (SrcColors < LastColor)
 		{
-			if (SrcColors->A <  (1.0f - UE_SMALL_NUMBER))
+			if (SrcColors->A < (1.0f - UE_SMALL_NUMBER))
+			{
+				return true;
+			}
+			++SrcColors;
+		}
+		return false;
+	}
+	else if (Image.Format == ERawImageFormat::RGBA32F)
+	{
+		const FLinearColor* SrcColors = (&Image.AsRGBA32F()[0]);
+		const FLinearColor* LastColor = SrcColors + (Image.SizeX * Image.SizeY * Image.NumSlices);
+		while (SrcColors < LastColor)
+		{
+			// this comparison matches to would happen if RGBA32F would be converted to BGRA8 format
+			if (SrcColors->QuantizeRound().A < 255)
 			{
 				return true;
 			}
@@ -1273,7 +1376,25 @@ bool FVirtualTextureDataBuilder::DetectAlphaChannel(const FImage &Image)
 		}
 		return false;
 	}
-	else if (Image.Format == ERawImageFormat::G16)
+	else if (Image.Format == ERawImageFormat::RGBA16)
+	{
+		const uint16* SrcColors = (&Image.AsRGBA16()[0]);
+		const uint16* LastColor = SrcColors + (Image.SizeX * Image.SizeY * Image.NumSlices);
+		while (SrcColors < LastColor)
+		{
+			if (SrcColors[3] < 65535)
+			{
+				return true;
+			}
+			SrcColors += 4;
+		}
+		return false;
+	}
+	else if (Image.Format == ERawImageFormat::G16 ||
+			 Image.Format == ERawImageFormat::G8 ||
+			 Image.Format == ERawImageFormat::BGRE8 ||
+			 Image.Format == ERawImageFormat::R16F ||
+			 Image.Format == ERawImageFormat::R32F)
 	{
 		return false;
 	}
diff --git a/Engine/Source/Runtime/Engine/Private/VT/VirtualTextureDataBuilder.h b/Engine/Source/Runtime/Engine/Private/VT/VirtualTextureDataBuilder.h
--- a/Engine/Source/Runtime/Engine/Private/VT/VirtualTextureDataBuilder.h
+++ b/Engine/Source/Runtime/Engine/Private/VT/VirtualTextureDataBuilder.h
@@ -12,15 +12,54 @@ class IImageWrapperModule;
 struct FTextureSourceData;
 struct FTextureSourceBlockData;
 struct FTextureBuildSettings;
+struct FSlowTask;
+
+struct FVTTileMipPayload
+{
+	EPixelFormat CompressedFormat = EPixelFormat::PF_Unknown;
+	TArray<uint8> Payload;
+};
+
+struct FVTTilePayload
+{
+	TArray<FVTTileMipPayload> Mips;
+};
+
+struct FVTBlockPayload
+{
+	// Block size in pixels (each block can be different size)
+	int32 SizeX = 0;
+	int32 SizeY = 0;
+
+	// Block coordinate (in block units) in texture where block is located
+	int32 BlockX = 0;
+	int32 BlockY = 0;
+
+	// Normally each blocks covers a 1x1 block area in output, but it can cover more.
+	// For example, miptail covers multiple input blocks. Thease are in block units.
+	int32 SizeInBlocksX = 1;
+	int32 SizeInBlocksY = 1;
+
+	int32 MipBias = 0; // First mip level in block that contains actual data (because blocks can be smaller than others)
+	int32 NumMips = 0; // Count of mips (may be less than source texture reqests because of MipBias)
+	int32 NumSlices = 0;
+	TArray<FVTTilePayload> Tiles;
+	TArray<FImage> Mips;
+};
+
+struct FVTLayerPayload
+{
+	TArray<FVTBlockPayload> Blocks;
+};
+
 
 struct FVTSourceTileEntry
 {
-	int32 BlockIndex;
-	int32 TileIndex;
-	int32 MipIndex;
-	int32 MipIndexInBlock;
-	int32 TileInBlockX;
-	int32 TileInBlockY;
+	FVTBlockPayload* Block = nullptr;
+	int32 TileIndex = 0;
+	int32 TileIndexInBlock = 0;
+	int32 MipIndex = 0;
+	int32 MipIndexInBlock = 0;
 };
 
 struct FLayerData
@@ -86,18 +125,18 @@ public:
 private:
 	friend struct FAsyncMacroBlockTask;
 
-	bool BuildPagesMacroBlocks(bool bAllowAsync);
-	bool BuildPagesForChunk(const TArray<FVTSourceTileEntry>& ActiveTileList, bool bAllowAsync);
-	void BuildTiles(const TArray<FVTSourceTileEntry>& TileList, uint32 layer, FLayerData& GeneratedData, bool bAllowAsync);
-	bool PushDataToChunk(const TArray<FVTSourceTileEntry> &Tiles, const TArray<FLayerData>& LayerData);
+	bool BuildPagesForChunk(const TArray<FVTSourceTileEntry>& ActiveTileList);
+	void BuildTiles(const TArray<FVTSourceTileEntry>& TileList, uint32 layer, FLayerData& GeneratedData);
+	bool PushDataToChunk(const TArray<FVTSourceTileEntry>& Tiles, const TArray<FLayerData>& LayerData);
+
+	FVTBlockPayload* FindSourceBlock(int32 MipIndex, int32 BlockX, int32 BlockY);
 
-	int32 FindSourceBlockIndex(int32 MipIndex, int32 BlockX, int32 BlockY) const;
+	void BuildLayerBlocks(FSlowTask& BuildTask, uint32 LayerIndex, const FVirtualTextureSourceLayerData& LayerData, FTextureSourceData& SourceData, FTextureSourceData& CompositeSourceData, bool bAllowAsync);
+	void BuildBlockTiles(uint32 LayerIndex, uint32 BlockIndex, FVTBlockPayload& Block, const FVirtualTextureSourceLayerData& LayerData, bool bAllowAsync);
+	bool BuildChunks();
+
+	TArray<FVTLayerPayload> LayerPayload;
 
-	// Build the source data including mipmaps etc
-	void BuildSourcePixels(const FTextureSourceData& SourceData, const FTextureSourceData& CompositeSourceData);
-	// Release the source pixels
-	void FreeSourcePixels();
-	
 	// Cached inside this object
 	TArray<FTextureBuildSettings> SettingsPerLayer;
 	FVirtualTextureBuiltData &OutData;
@@ -105,14 +144,11 @@ private:
 	// Some convenience variables (mostly derived from the passed in build settings)
 	FVirtualTextureBuilderDerivedInfo DerivedInfo;
 
-	TArray<FVirtualTextureSourceLayerData> SourceLayers;
-	TArray<FTextureSourceBlockData> SourceBlocks;
-	//FTextureSourceBlockData SourceMiptailBlock;
-
 	ITextureCompressorModule *Compressor;
 	IImageWrapperModule *ImageWrapper;
 
 	const FString& DebugTexturePathName;
+	int32 ChunkDumpIndex = 0;
 
-	bool DetectAlphaChannel(const FImage &image);
+	static bool DetectAlphaChannel(const FImage &image);
 };
\ No newline at end of file
