[Optimization] - Only hash and calculate Material Debug Name once per material instead of once per-shader job.

- Pass FStrings by reference in PrepareMaterialShaderCompileJob instead of copying them.
- Remove dead commented out code.

diff --git a/Engine/Source/Runtime/Engine/Private/Materials/MaterialShader.cpp b/Engine/Source/Runtime/Engine/Private/Materials/MaterialShader.cpp
--- a/Engine/Source/Runtime/Engine/Private/Materials/MaterialShader.cpp
+++ b/Engine/Source/Runtime/Engine/Private/Materials/MaterialShader.cpp
@@ -1194,8 +1194,9 @@ static void PrepareMaterialShaderCompileJob(EShaderPlatform Platform,
 	const FMaterialShaderMapId& ShaderMapId,
 	FSharedShaderCompilerEnvironment* MaterialEnvironment,
 	const FShaderPipelineType* ShaderPipeline,
-	FString DebugDescription,
-	FString DebugExtension,
+	const FString& DebugGroupName,
+	const FString& DebugDescription,
+	const FString& DebugExtension,
 	FShaderCompileJob* NewJob)
 {
 	const FShaderCompileJobKey& Key = NewJob->Key;
@@ -1217,7 +1218,7 @@ static void PrepareMaterialShaderCompileJob(EShaderPlatform Platform,
 
 	// Compile the shader environment passed in with the shader type's source code.
 	::GlobalBeginCompileShader(
-		Material->GetUniqueAssetName(Platform, ShaderMapId) / LexToString(Material->GetQualityLevel()),
+		DebugGroupName,
 		nullptr,
 		ShaderType,
 		ShaderPipeline,
@@ -1246,6 +1247,7 @@ void FMaterialShaderType::BeginCompileShader(
 	EShaderPlatform Platform,
 	EShaderPermutationFlags PermutationFlags,
 	TArray<FShaderCommonCompileJobPtr>& NewJobs,
+	const FString& DebugGroupName,
 	const TCHAR* DebugDescription,
 	const TCHAR* DebugExtension
 	) const
@@ -1253,7 +1255,7 @@ void FMaterialShaderType::BeginCompileShader(
 	FShaderCompileJob* NewJob = GShaderCompilingManager->PrepareShaderCompileJob(ShaderMapJobId, FShaderCompileJobKey(this, nullptr, PermutationId), Priority);
 	if (NewJob)
 	{
-		PrepareMaterialShaderCompileJob(Platform, PermutationFlags, Material, ShaderMapId, MaterialEnvironment, nullptr, DebugDescription, DebugExtension, NewJob);
+		PrepareMaterialShaderCompileJob(Platform, PermutationFlags, Material, ShaderMapId, MaterialEnvironment, nullptr, DebugGroupName, DebugDescription, DebugExtension, NewJob);
 		NewJobs.Add(FShaderCommonCompileJobPtr(NewJob));
 	}
 }
@@ -1268,6 +1270,7 @@ void FMaterialShaderType::BeginCompileShaderPipeline(
 	FSharedShaderCompilerEnvironment* MaterialEnvironment,
 	const FShaderPipelineType* ShaderPipeline,
 	TArray<FShaderCommonCompileJobPtr>& NewJobs,
+	const FString& DebugGroupName,
 	const TCHAR* DebugDescription,
 	const TCHAR* DebugExtension)
 {
@@ -1280,7 +1283,7 @@ void FMaterialShaderType::BeginCompileShaderPipeline(
 	{
 		for (FShaderCompileJob* StageJob : NewPipelineJob->StageJobs)
 		{
-			PrepareMaterialShaderCompileJob(Platform, PermutationFlags, Material, ShaderMapId, MaterialEnvironment, ShaderPipeline, DebugDescription, DebugExtension, StageJob);
+			PrepareMaterialShaderCompileJob(Platform, PermutationFlags, Material, ShaderMapId, MaterialEnvironment, ShaderPipeline, DebugGroupName, DebugDescription, DebugExtension, StageJob);
 		}
 		NewJobs.Add(FShaderCommonCompileJobPtr(NewPipelineJob));
 	}
@@ -1848,6 +1851,8 @@ int32 FMaterialShaderMap::SubmitCompileJobs(uint32 CompilingShaderMapId,
 	const FMaterialShaderParameters MaterialParameters(Material);
 	const FMaterialShaderMapLayout& Layout = AcquireMaterialShaderMapLayout(ShaderPlatform, LocalPermutationFlags, MaterialParameters);
 
+	const FString DebugGroupName = Material->GetUniqueAssetName(ShaderPlatform, GetShaderMapId()) / LexToString(Material->GetQualityLevel());
+
 #if ALLOW_SHADERMAP_DEBUG_DATA
 	FString DebugExtensionStr(TEXT(""));
 	FString DebugDescriptionStr(TEXT(""));
@@ -1909,12 +1914,10 @@ int32 FMaterialShaderMap::SubmitCompileJobs(uint32 CompilingShaderMapId,
 					MaterialEnvironment,
 					MeshLayout.VertexFactoryType,
 					CompileJobs,
+					DebugGroupName,
 					DebugDescription,
 					DebugExtension
 				);
-				//TShaderTypePermutation<const FShaderType> ShaderTypePermutation(ShaderType, Shader.PermutationId);
-				//check(!SharedShaderJobs.Find(ShaderTypePermutation));
-				//SharedShaderJobs.Add(ShaderTypePermutation, Job);
 			}
 		}
 
@@ -1954,6 +1957,7 @@ int32 FMaterialShaderMap::SubmitCompileJobs(uint32 CompilingShaderMapId,
 					MeshLayout.VertexFactoryType,
 					Pipeline,
 					CompileJobs,
+					DebugGroupName,
 					DebugDescription,
 					DebugExtension);
 			}
@@ -2013,13 +2017,10 @@ int32 FMaterialShaderMap::SubmitCompileJobs(uint32 CompilingShaderMapId,
 				ShaderPlatform,
 				LocalPermutationFlags,
 				CompileJobs,
+				DebugGroupName,
 				DebugDescription,
 				DebugExtension
 			);
-
-			//TShaderTypePermutation<const FShaderType> ShaderTypePermutation(ShaderType, Shader.PermutationId);
-			//check(!SharedShaderJobs.Find(ShaderTypePermutation));
-			//SharedShaderJobs.Add(ShaderTypePermutation, Job);
 		}
 		NumShaders++;
 	}
@@ -2050,6 +2051,7 @@ int32 FMaterialShaderMap::SubmitCompileJobs(uint32 CompilingShaderMapId,
 					MaterialEnvironment,
 					Pipeline,
 					CompileJobs,
+					DebugGroupName,
 					DebugDescription,
 					DebugExtension);
 			}
diff --git a/Engine/Source/Runtime/Engine/Private/Materials/MaterialShared.cpp b/Engine/Source/Runtime/Engine/Private/Materials/MaterialShared.cpp
--- a/Engine/Source/Runtime/Engine/Private/Materials/MaterialShared.cpp
+++ b/Engine/Source/Runtime/Engine/Private/Materials/MaterialShared.cpp
@@ -3007,6 +3007,8 @@ void FMaterial::CacheGivenTypes(EShaderPlatform Platform, const TArray<const FVe
 		checkf(PipelineTypes.Num() == ShaderTypes.Num(), TEXT("The size of the pipeline type array and shader type array must match.  Pass in null entries if pipelines are not used."));
 		checkf(GetGameThreadCompilingShaderMapId() != 0, TEXT("Material is not prepared to compile yet.  Please call CacheShaders first."));
 
+		const FString DebugGroupName = GetUniqueAssetName(Platform, GameThreadShaderMap->GetShaderMapId()) / LexToString(GetQualityLevel());
+
 		TArray<FShaderCommonCompileJobPtr> CompileJobs;
 		for (int i = 0; i < VFTypes.Num(); ++i)
 		{
@@ -3028,6 +3030,7 @@ void FMaterial::CacheGivenTypes(EShaderPlatform Platform, const TArray<const FVe
 					VFType,
 					PipelineType,
 					CompileJobs,
+					DebugGroupName,
 					nullptr,
 					nullptr);
 			}
@@ -3057,6 +3060,7 @@ void FMaterial::CacheGivenTypes(EShaderPlatform Platform, const TArray<const FVe
 						Platform,
 						ShaderPermutation,
 						CompileJobs,
+						DebugGroupName,
 						nullptr,
 						nullptr);
 				}
@@ -3074,6 +3078,7 @@ void FMaterial::CacheGivenTypes(EShaderPlatform Platform, const TArray<const FVe
 					GameThreadPendingCompilerEnvironment,
 					VFType,
 					CompileJobs,
+					DebugGroupName,
 					nullptr,
 					nullptr);
 			}
@@ -3415,6 +3420,8 @@ bool FMaterial::TryGetShaders(const FMaterialShaderTypes& InTypes, const FVertex
 		return false;
 	};
 
+	const FString DebugGroupName = GetUniqueAssetName(ShaderPlatform, ShaderMap->GetShaderMapId()) / LexToString(GetQualityLevel());
+
 	if (InTypes.PipelineType &&
 		RHISupportsShaderPipelines(ShaderPlatform) &&
 		CVarShaderPipelines && CVarShaderPipelines->GetValueOnAnyThread(bIsInGameThread) != 0)
@@ -3490,7 +3497,8 @@ bool FMaterial::TryGetShaders(const FMaterialShaderTypes& InTypes, const FVertex
 									RenderingThreadPendingCompilerEnvironment, 
 									InVertexFactoryType, 
 									InTypes.PipelineType, 
-									CompileJobs, 
+									CompileJobs,
+									DebugGroupName,
 									nullptr, 
 									nullptr);
 							}
@@ -3505,7 +3513,8 @@ bool FMaterial::TryGetShaders(const FMaterialShaderTypes& InTypes, const FVertex
 									ShaderMap->GetShaderMapId(),
 									RenderingThreadPendingCompilerEnvironment, 
 									InTypes.PipelineType, 
-									CompileJobs, 
+									CompileJobs,
+									DebugGroupName,
 									nullptr, 
 									nullptr);
 							}
@@ -3573,6 +3582,7 @@ bool FMaterial::TryGetShaders(const FMaterialShaderTypes& InTypes, const FVertex
 										RenderingThreadPendingCompilerEnvironment,
 										InVertexFactoryType,
 										CompileJobs,
+										DebugGroupName,
 										nullptr,
 										nullptr);
 								}
@@ -3588,6 +3598,7 @@ bool FMaterial::TryGetShaders(const FMaterialShaderTypes& InTypes, const FVertex
 										ShaderPlatform,
 										PermutationFlags,
 										CompileJobs,
+										DebugGroupName,
 										nullptr,
 										nullptr);
 								}
diff --git a/Engine/Source/Runtime/Engine/Private/Materials/MeshMaterialShader.cpp b/Engine/Source/Runtime/Engine/Private/Materials/MeshMaterialShader.cpp
--- a/Engine/Source/Runtime/Engine/Private/Materials/MeshMaterialShader.cpp
+++ b/Engine/Source/Runtime/Engine/Private/Materials/MeshMaterialShader.cpp
@@ -29,6 +29,7 @@ static void PrepareMeshMaterialShaderCompileJob(EShaderPlatform Platform,
 	const FMaterialShaderMapId& MaterialShaderMapId,
 	FSharedShaderCompilerEnvironment* MaterialEnvironment,
 	const FShaderPipelineType* ShaderPipeline,
+	const FString& DebugGroupName,
 	const TCHAR* DebugDescription,
 	const TCHAR* DebugExtension,
 	FShaderCompileJob* NewJob)
@@ -64,7 +65,7 @@ static void PrepareMeshMaterialShaderCompileJob(EShaderPlatform Platform,
 
 	// Compile the shader environment passed in with the shader type's source code.
 	::GlobalBeginCompileShader(
-		Material->GetUniqueAssetName(Platform, MaterialShaderMapId) / LexToString(Material->GetQualityLevel()),
+		DebugGroupName,
 		VertexFactoryType,
 		ShaderType,
 		ShaderPipeline,
@@ -97,13 +98,14 @@ void FMeshMaterialShaderType::BeginCompileShader(
 	FSharedShaderCompilerEnvironment* MaterialEnvironment,
 	const FVertexFactoryType* VertexFactoryType,
 	TArray<FShaderCommonCompileJobPtr>& NewJobs,
+	const FString& DebugGroupName,
 	const TCHAR* DebugDescription,
 	const TCHAR* DebugExtension) const
 {
 	FShaderCompileJob* NewJob = GShaderCompilingManager->PrepareShaderCompileJob(ShaderMapJobId, FShaderCompileJobKey(this, VertexFactoryType, PermutationId), Priority);
 	if (NewJob)
 	{
-		PrepareMeshMaterialShaderCompileJob(Platform, PermutationFlags, Material, ShaderMapId, MaterialEnvironment, nullptr, DebugDescription, DebugExtension, NewJob);
+		PrepareMeshMaterialShaderCompileJob(Platform, PermutationFlags, Material, ShaderMapId, MaterialEnvironment, nullptr, DebugGroupName, DebugDescription, DebugExtension, NewJob);
 		NewJobs.Add(FShaderCommonCompileJobPtr(NewJob));
 	}
 }
@@ -120,6 +122,7 @@ void FMeshMaterialShaderType::BeginCompileShaderPipeline(
 	const FVertexFactoryType* VertexFactoryType,
 	const FShaderPipelineType* ShaderPipeline,
 	TArray<FShaderCommonCompileJobPtr>& NewJobs,
+	const FString& DebugGroupName,
 	const TCHAR* DebugDescription,
 	const TCHAR* DebugExtension)
 {
@@ -132,7 +135,7 @@ void FMeshMaterialShaderType::BeginCompileShaderPipeline(
 	{
 		for (FShaderCompileJob* StageJob : NewPipelineJob->StageJobs)
 		{
-			PrepareMeshMaterialShaderCompileJob(Platform, PermutationFlags, Material, ShaderMapId, MaterialEnvironment, ShaderPipeline, DebugDescription, DebugExtension, StageJob);
+			PrepareMeshMaterialShaderCompileJob(Platform, PermutationFlags, Material, ShaderMapId, MaterialEnvironment, ShaderPipeline, DebugGroupName, DebugDescription, DebugExtension, StageJob);
 		}
 		NewJobs.Add(FShaderCommonCompileJobPtr(NewPipelineJob));
 	}
diff --git a/Engine/Source/Runtime/Engine/Public/MaterialShaderType.h b/Engine/Source/Runtime/Engine/Public/MaterialShaderType.h
--- a/Engine/Source/Runtime/Engine/Public/MaterialShaderType.h
+++ b/Engine/Source/Runtime/Engine/Public/MaterialShaderType.h
@@ -147,6 +147,7 @@ public:
 		EShaderPlatform Platform,
 		EShaderPermutationFlags PermutationFlags,
 		TArray<TRefCountPtr<FShaderCommonCompileJob>>& NewJobs,
+		const FString& DebugGroupName,
 		const TCHAR* DebugDescription,
 		const TCHAR* DebugExtension
 	) const;
@@ -161,6 +162,7 @@ public:
 		FSharedShaderCompilerEnvironment* MaterialEnvironment,
 		const FShaderPipelineType* ShaderPipeline,
 		TArray<TRefCountPtr<FShaderCommonCompileJob>>& NewJobs,
+		const FString& DebugGroupName,
 		const TCHAR* DebugDescription,
 		const TCHAR* DebugExtension
 	);
diff --git a/Engine/Source/Runtime/Engine/Public/MeshMaterialShaderType.h b/Engine/Source/Runtime/Engine/Public/MeshMaterialShaderType.h
--- a/Engine/Source/Runtime/Engine/Public/MeshMaterialShaderType.h
+++ b/Engine/Source/Runtime/Engine/Public/MeshMaterialShaderType.h
@@ -97,6 +97,7 @@ public:
 		FSharedShaderCompilerEnvironment* MaterialEnvironment,
 		const FVertexFactoryType* VertexFactoryType,
 		TArray<TRefCountPtr<FShaderCommonCompileJob>>& NewJobs,
+		const FString& DebugGroupName,
 		const TCHAR* DebugDescription,
 		const TCHAR* DebugExtension
 		) const;
@@ -113,6 +114,7 @@ public:
 		const FVertexFactoryType* VertexFactoryType,
 		const FShaderPipelineType* ShaderPipeline,
 		TArray<TRefCountPtr<FShaderCommonCompileJob>>& NewJobs,
+		const FString& DebugGroupName,
 		const TCHAR* DebugDescription,
 		const TCHAR* DebugExtension
 		);
