World Partition: keep cells debug names in Test builds.

diff --git a/Engine/Source/Runtime/Engine/Private/StringDev.cpp b/Engine/Source/Runtime/Engine/Private/StringDev.cpp
deleted file mode 100644
diff --git a/Engine/Source/Runtime/Engine/Public/StringDev.h b/Engine/Source/Runtime/Engine/Public/StringDev.h
--- a/Engine/Source/Runtime/Engine/Public/StringDev.h
+++ b/Engine/Source/Runtime/Engine/Public/StringDev.h
@@ -4,27 +4,83 @@
 #include "CoreTypes.h"
 #include "Containers/UnrealString.h"
 
-/**
- * A string that is present in memory only in editor and development builds.
- */
-class ENGINE_API FStringDev
+class FOptionalStringNull
 {
 public:
-	FStringDev()
+	inline FOptionalStringNull()
+	{}
+
+	inline FOptionalStringNull(const FString& InString)
+	{}
+
+	inline  const FString& GetString() const
+	{
+		static FString String;
+		return String;
+	}
+	
+	inline void SetString(const FString& InString)
+	{}
+
+	friend inline FArchive& operator<<(FArchive& Ar, FOptionalStringNull& InString)
+	{
+		FString String;
+		return Ar << String;
+	}
+};
+
+class FOptionalStringNonNull
+{
+	FString String;
+
+public:
+	inline FOptionalStringNonNull()
 	{}
 
-	FStringDev(const FString& InString)
+	inline FOptionalStringNonNull(const FString& InString)
 	{
 		SetString(InString);
 	}
 
-	const FString& GetString() const;
-	void SetString(const FString& InString);
+	inline const FString& GetString() const
+	{
+		return String;
+	}
 
-	ENGINE_API friend FArchive& operator<<(FArchive& Ar, FStringDev& StringDev);
+	inline void SetString(const FString& InString)
+	{
+		String = InString;
+	}
 
-private:
-#if WITH_EDITOR || UE_BUILD_DEVELOPMENT
-	FString String;
+	friend inline FArchive& operator<<(FArchive& Ar, FOptionalStringNonNull& InString)
+	{
+		return Ar << InString.String;
+	}
+};
+
+/**
+ * A string that is present in memory only in editor and debug builds.
+ */
+#if WITH_EDITOR || UE_BUILD_DEBUG
+using FStringDebug = FOptionalStringNonNull;
+#else
+using FStringDebug = FOptionalStringNull;
+#endif
+
+/**
+ * A string that is present in memory only in editor, debug and development builds.
+ */
+#if WITH_EDITOR || UE_BUILD_DEBUG || UE_BUILD_DEVELOPMENT
+using FStringDev = FOptionalStringNonNull;
+#else
+using FStringDev = FOptionalStringNull;
 #endif
-};
\ No newline at end of file
+
+/**
+ * A string that is present in memory only in editor, debug, development and test builds.
+ */
+#if WITH_EDITOR || !UE_BUILD_SHIPPING
+using FStringTest = FOptionalStringNonNull;
+#else
+using FStringTest = FOptionalStringNull;
+#endif
\ No newline at end of file
diff --git a/Engine/Source/Runtime/Engine/Public/WorldPartition/WorldPartitionRuntimeCellDataSpatialHash.h b/Engine/Source/Runtime/Engine/Public/WorldPartition/WorldPartitionRuntimeCellDataSpatialHash.h
--- a/Engine/Source/Runtime/Engine/Public/WorldPartition/WorldPartitionRuntimeCellDataSpatialHash.h
+++ b/Engine/Source/Runtime/Engine/Public/WorldPartition/WorldPartitionRuntimeCellDataSpatialHash.h
@@ -43,7 +43,7 @@ class ENGINE_API UWorldPartitionRuntimeCellDataSpatialHash : public UWorldPartit
 	UPROPERTY()
 	FName GridName;
 
-	FStringDev DebugName;
+	FStringTest DebugName;
 
 private:
 	float ComputeSourceToCellAngleFactor(const FSphericalSector& SourceShape) const;
