World Partition: null streaming policy pointer checks for public functions.

diff --git a/Engine/Source/Runtime/Engine/Private/WorldPartition/WorldPartition.cpp b/Engine/Source/Runtime/Engine/Private/WorldPartition/WorldPartition.cpp
--- a/Engine/Source/Runtime/Engine/Private/WorldPartition/WorldPartition.cpp
+++ b/Engine/Source/Runtime/Engine/Private/WorldPartition/WorldPartition.cpp
@@ -1501,7 +1501,7 @@ bool UWorldPartition::GetIntersectingCells(const TArray<FWorldPartitionStreaming
 
 bool UWorldPartition::CanAddLoadedLevelToWorld(class ULevel* InLevel) const
 {
-	if (GetWorld()->IsGameWorld())
+	if (GetWorld()->IsGameWorld() && StreamingPolicy)
 	{
 		return StreamingPolicy->CanAddLoadedLevelToWorld(InLevel);
 	}
@@ -1510,7 +1510,7 @@ bool UWorldPartition::CanAddLoadedLevelToWorld(class ULevel* InLevel) const
 
 bool UWorldPartition::IsStreamingCompleted(const TArray<FWorldPartitionStreamingSource>* InStreamingSources) const
 {
-	if (GetWorld()->IsGameWorld())
+	if (GetWorld()->IsGameWorld() && StreamingPolicy)
 	{
 		return StreamingPolicy->IsStreamingCompleted(InStreamingSources);
 	}
@@ -1519,7 +1519,7 @@ bool UWorldPartition::IsStreamingCompleted(const TArray<FWorldPartitionStreaming
 
 bool UWorldPartition::IsStreamingCompleted(EWorldPartitionRuntimeCellState QueryState, const TArray<FWorldPartitionStreamingQuerySource>& QuerySources, bool bExactState) const
 {
-	if (GetWorld()->IsGameWorld())
+	if (GetWorld()->IsGameWorld() && StreamingPolicy)
 	{
 		return StreamingPolicy->IsStreamingCompleted(QueryState, QuerySources, bExactState);
 	}
