Fix crash on exit after running TestFramework.Validation.UTestNullNotNull
#rnx

diff --git a/Engine/Source/Runtime/Engine/Private/Tests/EngineAutomationTests.cpp b/Engine/Source/Runtime/Engine/Private/Tests/EngineAutomationTests.cpp
--- a/Engine/Source/Runtime/Engine/Private/Tests/EngineAutomationTests.cpp
+++ b/Engine/Source/Runtime/Engine/Private/Tests/EngineAutomationTests.cpp
@@ -917,6 +917,8 @@ bool FAutomationNullNotNullPtrEXPR::RunTest(const FString& Parameters)
 	UTEST_NULL_EXPR(nullptr);
 	UTEST_NOT_NULL_EXPR(World);
 
+	World->DestroyWorld(false);
+
 	return true;
 }
 
