DebugDraw: add a flag in FBatchedLine, FBatchedPoint and FBatchedMesh to hide the element
Add functions to hide/show elements per batchID
UControllerComponentDebugDraw: Add a static flag to specify if debug draw can be activated (true in editor and VK edit)
#rnx

diff --git a/Engine/Source/Runtime/Engine/Classes/Components/LineBatchComponent.h b/Engine/Source/Runtime/Engine/Classes/Components/LineBatchComponent.h
--- a/Engine/Source/Runtime/Engine/Classes/Components/LineBatchComponent.h
+++ b/Engine/Source/Runtime/Engine/Classes/Components/LineBatchComponent.h
@@ -37,6 +37,9 @@ struct FBatchedLine
 
 	UPROPERTY()
 	uint8 DepthPriority;
+	
+	UPROPERTY()
+	bool bIsVisible;
 
 	UPROPERTY()
 	uint32 BatchID;
@@ -48,15 +51,17 @@ struct FBatchedLine
 		, Thickness(0)
 		, RemainingLifeTime(0)
 		, DepthPriority(0)
+		, bIsVisible(true)
 		, BatchID(0)
 	{}
-	FBatchedLine(const FVector& InStart, const FVector& InEnd, const FLinearColor& InColor, float InLifeTime, float InThickness, uint8 InDepthPriority, uint32 InBatchID = 0)
+	FBatchedLine(const FVector& InStart, const FVector& InEnd, const FLinearColor& InColor, float InLifeTime, float InThickness, uint8 InDepthPriority, uint32 InBatchID = 0, bool bInVisible = true)
 		:	Start(InStart)
 		,	End(InEnd)
 		,	Color(InColor)
 		,	Thickness(InThickness)
 		,	RemainingLifeTime(InLifeTime)
 		,	DepthPriority(InDepthPriority)
+		,	bIsVisible(bInVisible)
 		,	BatchID(InBatchID)
 	{}
 };
@@ -80,6 +85,9 @@ struct FBatchedPoint
 
 	UPROPERTY()
 	uint8 DepthPriority;
+
+	UPROPERTY()
+	bool bIsVisible;
 	
 	UPROPERTY()
 	uint32 BatchID;
@@ -90,14 +98,16 @@ struct FBatchedPoint
 		, PointSize(0)
 		, RemainingLifeTime(0)
 		, DepthPriority(0)
+		, bIsVisible(true)
 		, BatchID(0)
 	{}
-	FBatchedPoint(const FVector& InPosition, const FLinearColor& InColor, float InPointSize, float InLifeTime, uint8 InDepthPriority, uint32 InBatchID = 0)
+	FBatchedPoint(const FVector& InPosition, const FLinearColor& InColor, float InPointSize, float InLifeTime, uint8 InDepthPriority, uint32 InBatchID = 0, bool bInVisible = true)
 		:	Position(InPosition)
 		,	Color(InColor)
 		,	PointSize(InPointSize)
 		,	RemainingLifeTime(InLifeTime)
 		,	DepthPriority(InDepthPriority)
+		,	bIsVisible(bInVisible)
 		,	BatchID(InBatchID)
 	{}
 };
@@ -110,9 +120,9 @@ struct FBatchedMesh
 	 * MeshVerts - linear array of world space vertex positions
 	 * MeshIndices - array of indices into MeshVerts.  Each triplet is a tri.  i.e. [0,1,2] is first tri, [3,4,5] is 2nd tri, etc
 	 */
-	FBatchedMesh(TArray<FVector> const& InMeshVerts, TArray<int32> const& InMeshIndices, FColor const& InColor, uint8 InDepthPriority, float LifeTime, uint32 InBatchID = 0)
+	FBatchedMesh(TArray<FVector> const& InMeshVerts, TArray<int32> const& InMeshIndices, FColor const& InColor, uint8 InDepthPriority, float LifeTime, uint32 InBatchID = 0, bool bInVisible = true)
 		: MeshVerts(InMeshVerts), MeshIndices(InMeshIndices), 
-		  Color(InColor), RemainingLifeTime(LifeTime), DepthPriority(InDepthPriority), BatchID(InBatchID)
+		  Color(InColor), RemainingLifeTime(LifeTime), DepthPriority(InDepthPriority), bIsVisible(bInVisible), BatchID(InBatchID)
 	{}
 
 	TArray<FVector> MeshVerts;
@@ -120,6 +130,7 @@ struct FBatchedMesh
 	FColor Color;
 	float RemainingLifeTime = 0.f;
 	uint8 DepthPriority = 0;
+	bool bIsVisible = true;
 	uint32 BatchID = 0;
 };
 
@@ -142,6 +153,8 @@ class ULineBatchComponent : public UPrimitiveComponent
 	TArray<FBatchedMesh> BatchedMeshes;
 	/** Whether to calculate a tight accurate bounds (encompassing all points), or use a giant bounds that is fast to compute. */
 	uint32 bCalculateAccurateBounds:1;
+	/** Whether this has some hidden lines */
+	uint32 bHasHiddenLines:1;
 
 	/** Defines the value for an invalid id */
 	static constexpr uint32 INVALID_ID = 0;
@@ -150,34 +163,34 @@ class ULineBatchComponent : public UPrimitiveComponent
 	ENGINE_API void DrawLines(TArrayView<FBatchedLine> InLines);
 
 	/** Draw a box */
-	ENGINE_API void DrawBox(const FBox& Box, const FMatrix& TM, FLinearColor Color, uint8 InDepthPriorityGroup, uint32 InBatchID = INVALID_ID);
+	ENGINE_API void DrawBox(const FBox& Box, const FMatrix& TM, FLinearColor Color, uint8 InDepthPriorityGroup, uint32 BatchID = INVALID_ID, bool bIsVisible = true);
 
 	/** Draw a box */
-	ENGINE_API void DrawBox(const FVector& Center, const FVector& Box, FLinearColor Color, float LifeTime, uint8 DepthPriority, float Thickness, uint32 InBatchID = INVALID_ID);
+	ENGINE_API void DrawBox(const FVector& Center, const FVector& Box, FLinearColor Color, float LifeTime, uint8 DepthPriority, float Thickness, uint32 BatchID = INVALID_ID, bool bIsVisible = true);
 	
 	/** Draw a box */
-	ENGINE_API void DrawBox(const FVector& Center, const FVector& Box, const FQuat& Rotation, FLinearColor Color, float LifeTime, uint8 DepthPriority, float Thickness, uint32 InBatchID = INVALID_ID);
+	ENGINE_API void DrawBox(const FVector& Center, const FVector& Box, const FQuat& Rotation, FLinearColor Color, float LifeTime, uint8 DepthPriority, float Thickness, uint32 BatchID = INVALID_ID, bool bIsVisible = true);
 
 	/** Draw an arrow */
-	ENGINE_API void DrawDirectionalArrow(const FMatrix& ArrowToWorld, FLinearColor InColor, float Length, float ArrowSize, uint8 DepthPriority, uint32 InBatchID = INVALID_ID);
+	ENGINE_API void DrawDirectionalArrow(const FMatrix& ArrowToWorld, FLinearColor InColor, float Length, float ArrowSize, uint8 DepthPriority, uint32 BatchID = INVALID_ID, bool bIsVisible = true);
 
 	/** Draw an arrow */
-	ENGINE_API void DrawDirectionalArrow(const FVector& LineStart, const FVector& LineEnd, float ArrowSize, FLinearColor Color, float LifeTime, uint8 DepthPriority, float Thickness, uint32 InBatchID = INVALID_ID);
+	ENGINE_API void DrawDirectionalArrow(const FVector& LineStart, const FVector& LineEnd, float ArrowSize, FLinearColor Color, float LifeTime, uint8 DepthPriority, float Thickness, uint32 BatchID = INVALID_ID, bool bIsVisible = true);
 
 	/** Draw a circle */
-	ENGINE_API void DrawCircle(const FVector& Base, const FVector& X, const FVector& Y, FLinearColor Color, float Radius, int32 NumSides, uint8 DepthPriority, uint32 InBatchID = INVALID_ID);
+	ENGINE_API void DrawCircle(const FVector& Base, const FVector& X, const FVector& Y, FLinearColor Color, float Radius, int32 NumSides, uint8 DepthPriority, uint32 BatchID = INVALID_ID, bool bIsVisible = true);
 
 	/** Draw a sphere */
-	ENGINE_API void DrawSphere(const FVector& Center, float Radius, int32 Segments, FLinearColor Color, float LifeTime, uint8 DepthPriority, float Thickness, uint32 InBatchID = INVALID_ID);
+	ENGINE_API void DrawSphere(const FVector& Center, float Radius, int32 Segments, FLinearColor Color, float LifeTime, uint8 DepthPriority, float Thickness, uint32 BatchID = INVALID_ID, bool bIsVisible = true);
 
 	/** Draw a cylinder */
-	ENGINE_API void DrawCylinder(const FVector& Start, const FVector& End, float Radius, int32 Segments, FLinearColor  Color, float LifeTime, uint8 DepthPriority, float Thickness, uint32 InBatchID = INVALID_ID);
+	ENGINE_API void DrawCylinder(const FVector& Start, const FVector& End, float Radius, int32 Segments, FLinearColor  Color, float LifeTime, uint8 DepthPriority, float Thickness, uint32 BatchID = INVALID_ID, bool bIsVisible = true);
 
 	/** Draw a cone */
-	ENGINE_API void DrawCone(const FVector& Origin, const FVector& Direction, float Length, float AngleWidth, float AngleHeight, int32 NumSides, FLinearColor DrawColor, float LifeTime, uint8 DepthPriority, float Thickness, uint32 InBatchID = INVALID_ID);
+	ENGINE_API void DrawCone(const FVector& Origin, const FVector& Direction, float Length, float AngleWidth, float AngleHeight, int32 NumSides, FLinearColor DrawColor, float LifeTime, uint8 DepthPriority, float Thickness, uint32 BatchID = INVALID_ID, bool bIsVisible = true);
 
 	/** Draw a cone */
-	ENGINE_API void DrawCapsule(const FVector& Center, float HalfHeight, float Radius, const FQuat& Rotation, FLinearColor Color, float LifeTime, uint8 DepthPriority, float Thickness, uint32 InBatchID = INVALID_ID);
+	ENGINE_API void DrawCapsule(const FVector& Center, float HalfHeight, float Radius, const FQuat& Rotation, FLinearColor Color, float LifeTime, uint8 DepthPriority, float Thickness, uint32 BatchID = INVALID_ID, bool bIsVisible = true);
 
 	ENGINE_API virtual void DrawLine(
 		const FVector& Start,
@@ -186,7 +199,8 @@ class ULineBatchComponent : public UPrimitiveComponent
 		uint8 DepthPriority,
 		float Thickness = 0.0f,
 		float LifeTime = 0.0f,
-		uint32 InBatchID = INVALID_ID
+		uint32 BatchID = INVALID_ID,
+		bool bIsVisible = true
 		);
 	ENGINE_API virtual void DrawPoint(
 		const FVector& Position,
@@ -194,13 +208,14 @@ class ULineBatchComponent : public UPrimitiveComponent
 		float PointSize,
 		uint8 DepthPriority,
 		float LifeTime = 0.0f,
-		uint32 InBatchID = INVALID_ID
+		uint32 BatchID = INVALID_ID,
+		bool bIsVisible = true
 		);
 
 	/** Draw a box */
-	ENGINE_API void DrawSolidBox(FBox const& Box, FTransform const& Xform, const FColor& Color, uint8 DepthPriority, float LifeTime, uint32 InBatchID = INVALID_ID);
+	ENGINE_API void DrawSolidBox(FBox const& Box, FTransform const& Xform, const FColor& Color, uint8 DepthPriority, float LifeTime, uint32 BatchID = INVALID_ID, bool bIsVisible = true);
 	/** Draw a mesh */
-	ENGINE_API void DrawMesh(TArray<FVector> const& Verts, TArray<int32> const& Indices, FColor const& Color, uint8 DepthPriority, float LifeTime, uint32 InBatchID = INVALID_ID);
+	ENGINE_API void DrawMesh(TArray<FVector> const& Verts, TArray<int32> const& Indices, FColor const& Color, uint8 DepthPriority, float LifeTime, uint32 BatchID = INVALID_ID, bool bIsVisible = true);
 
 	//~ Begin UPrimitiveComponent Interface.
 	ENGINE_API virtual FPrimitiveSceneProxy* CreateSceneProxy() override;
@@ -218,8 +233,14 @@ class ULineBatchComponent : public UPrimitiveComponent
 	
 	/** Remove batched lines, points and meshes with given ID */
 	ENGINE_API void ClearBatch(uint32 InBatchID);
+	
+	/** Set batched lines, points and meshes with given ID visible or hidden*/
+	ENGINE_API void SetBatchVisible(uint32 InBatchID, bool bInVisible);
+	
+	/** Set all batched lines, points and meshes visible or hidden*/
+	ENGINE_API void SetVisible(bool bInVisible);
 
 protected:
-	void AddHalfCircle(const FVector& Base, const FVector& X, const FVector& Y, const FLinearColor& Color, const float Radius, int32 NumSides, const float LifeTime, uint8 DepthPriority, const float Thickness, const uint32 InBatchID);
-	void AddCircle(const FVector& Base, const FVector& X, const FVector& Y, const FLinearColor& Color, const float Radius, int32 NumSides, const float LifeTime, uint8 DepthPriority, const float Thickness, const uint32 InBatchID);
+	void AddHalfCircle(const FVector& Base, const FVector& X, const FVector& Y, const FLinearColor& Color, const float Radius, int32 NumSides, const float LifeTime, uint8 DepthPriority, const float Thickness, const uint32 BatchID, bool bIsVisible);
+	void AddCircle(const FVector& Base, const FVector& X, const FVector& Y, const FLinearColor& Color, const float Radius, int32 NumSides, const float LifeTime, uint8 DepthPriority, const float Thickness, const uint32 BatchID, bool bIsVisible);
 };
diff --git a/Engine/Source/Runtime/Engine/Private/Components/LineBatchComponent.cpp b/Engine/Source/Runtime/Engine/Private/Components/LineBatchComponent.cpp
--- a/Engine/Source/Runtime/Engine/Private/Components/LineBatchComponent.cpp
+++ b/Engine/Source/Runtime/Engine/Private/Components/LineBatchComponent.cpp
@@ -27,9 +27,9 @@ public:
 	virtual void GetDynamicMeshElements(const TArray<const FSceneView*>& Views, const FSceneViewFamily& ViewFamily, uint32 VisibilityMap, FMeshElementCollector& Collector) const override;
 
 	/**
-	*  Returns a struct that describes to the renderer when to draw this proxy.
-	*	@param		Scene view to use to determine our relevence.
-	*  @return		View relevance struct
+	*	Returns a struct that describes to the renderer when to draw this proxy.
+	*	@param		View view to use to determine our relevance.
+	*	@return		View relevance struct
 	*/
 	virtual FPrimitiveViewRelevance GetViewRelevance(const FSceneView* View) const override;
 	virtual uint32 GetMemoryFootprint( void ) const override;
@@ -47,6 +47,12 @@ FLineBatcherSceneProxy::FLineBatcherSceneProxy(const ULineBatchComponent* InComp
 	Points(InComponent->BatchedPoints), Meshes(InComponent->BatchedMeshes)
 {
 	bWillEverBeLit = false;
+	if (InComponent->bHasHiddenLines)
+	{
+		Lines.RemoveAllSwap([](const FBatchedLine& Element) { return Element.bIsVisible == false; });
+		Points.RemoveAllSwap([](const FBatchedPoint& Pt) { return Pt.bIsVisible == false; });
+		Meshes.RemoveAllSwap([](const FBatchedMesh& Mesh) { return Mesh.bIsVisible == false; });
+	}
 }
 
 SIZE_T FLineBatcherSceneProxy::GetTypeHash() const
@@ -71,51 +77,49 @@ void FLineBatcherSceneProxy::GetDynamicMeshElements(const TArray<const FSceneVie
 #else
 			FPrimitiveDrawInterface* PDI = Collector.GetPDI(ViewIndex);
 #endif
-			for (int32 i = 0; i < Lines.Num(); i++)
+			for (const FBatchedLine& Line: Lines)
 			{
-				PDI->DrawLine(Lines[i].Start, Lines[i].End, Lines[i].Color, Lines[i].DepthPriority, Lines[i].Thickness);
+				PDI->DrawLine(Line.Start,Line.End, Line.Color, Line.DepthPriority, Line.Thickness);
 			}
 
-			for (int32 i = 0; i < Points.Num(); i++)
+			for (const FBatchedPoint& Point: Points)
 			{
-				PDI->DrawPoint(Points[i].Position, Points[i].Color, Points[i].PointSize, Points[i].DepthPriority);
+				PDI->DrawPoint(Point.Position, Point.Color, Point.PointSize, Point.DepthPriority);
 			}
 
-			for (int32 i = 0; i < Meshes.Num(); i++)
+			for (const FBatchedMesh& Mesh: Meshes)
 			{
 				static FVector const PosX(1.f,0,0);
 				static FVector const PosY(0,1.f,0);
 				static FVector const PosZ(0,0,1.f);
 
-				FBatchedMesh const& M = Meshes[i];
-
 				// this seems far from optimal in terms of perf, but it's for debugging
 				FDynamicMeshBuilder MeshBuilder(View->GetFeatureLevel());
 
 				// set up geometry
-				for (int32 VertIdx=0; VertIdx < M.MeshVerts.Num(); ++VertIdx)
+				for (const FVector& Vert: Mesh.MeshVerts)
 				{
-					MeshBuilder.AddVertex((FVector3f)M.MeshVerts[VertIdx], FVector2f::ZeroVector, (FVector3f)PosX, (FVector3f)PosY, (FVector3f)PosZ, FColor::White );
+					MeshBuilder.AddVertex((FVector3f)Vert, FVector2f::ZeroVector, (FVector3f)PosX, (FVector3f)PosY, (FVector3f)PosZ, FColor::White );
 				}
 				//MeshBuilder.AddTriangles(M.MeshIndices);
-				for (int32 Idx=0; Idx < M.MeshIndices.Num(); Idx+=3)
+				for (int32 Idx=0; Idx < Mesh.MeshIndices.Num(); Idx+=3)
 				{
-					MeshBuilder.AddTriangle( M.MeshIndices[Idx], M.MeshIndices[Idx+1], M.MeshIndices[Idx+2] );
+					MeshBuilder.AddTriangle( Mesh.MeshIndices[Idx], Mesh.MeshIndices[Idx+1], Mesh.MeshIndices[Idx+2] );
 				}
 
-				FMaterialRenderProxy* const MaterialRenderProxy = new FColoredMaterialRenderProxy(GEngine->DebugMeshMaterial->GetRenderProxy(), M.Color);
+				FMaterialRenderProxy* const MaterialRenderProxy = new FColoredMaterialRenderProxy(GEngine->DebugMeshMaterial->GetRenderProxy(), Mesh.Color);
 				Collector.RegisterOneFrameMaterialProxy(MaterialRenderProxy);
 
-				MeshBuilder.GetMesh(FMatrix::Identity, MaterialRenderProxy, M.DepthPriority, false, false, ViewIndex, Collector);
+				MeshBuilder.GetMesh(FMatrix::Identity, MaterialRenderProxy, Mesh.DepthPriority, false, false, ViewIndex, Collector);
 			}
 		}
 	}
 }
 
 /**
-*  Returns a struct that describes to the renderer when to draw this proxy.
-*	@param		Scene view to use to determine our relevence.
-*  @return		View relevance struct
+*	Returns a struct that describes to the renderer when to draw this proxy.
+*	@param		View view to use to determine our relevance.
+*	@return		View relevance struct
 */
 FPrimitiveViewRelevance FLineBatcherSceneProxy::GetViewRelevance(const FSceneView* View) const
 {
@@ -149,15 +153,18 @@ ULineBatchComponent::ULineBatchComponent(const FObjectInitializer& ObjectInitial
 	bUseEditorCompositing = true;
 	SetGenerateOverlapEvents(false);
 	bCalculateAccurateBounds = true;
+	bHasHiddenLines = false;
 	DefaultLifeTime = 1.0f;
 
 	// Ignore streaming updates since GetUsedMaterials() is not implemented.
 	bIgnoreStreamingManagerUpdate = true;
 }
 
-void ULineBatchComponent::DrawLine(const FVector& Start, const FVector& End, const FLinearColor& Color, uint8 DepthPriority, const float Thickness, const float LifeTime, uint32 InBatchID)
+void ULineBatchComponent::DrawLine(const FVector& Start, const FVector& End, const FLinearColor& Color, uint8 DepthPriority, const float Thickness, const float LifeTime, uint32 BatchID, bool bIsVisible)
 {
-	new(BatchedLines) FBatchedLine(Start, End, Color, LifeTime, Thickness, DepthPriority, InBatchID);
+	new(BatchedLines) FBatchedLine(Start, End, Color, LifeTime, Thickness, DepthPriority, BatchID, bIsVisible);
+	
+	bHasHiddenLines |= !bIsVisible;
 
 	// LineBatcher and PersistentLineBatcher components will be updated at the end of UWorld::Tick
 	MarkRenderStateDirty();
@@ -166,7 +173,7 @@ void ULineBatchComponent::DrawLine(const FVector& Start, const FVector& End, con
 void ULineBatchComponent::DrawLines(TArrayView<FBatchedLine> InLines)
 {
 	BatchedLines.Append(InLines);
-	
+
 	// LineBatcher and PersistentLineBatcher components will be updated at the end of UWorld::Tick
 	MarkRenderStateDirty();
 }
@@ -177,117 +184,123 @@ void ULineBatchComponent::DrawPoint(
 	float PointSize,
 	uint8 DepthPriority,
 	float LifeTime,
-	uint32 InBatchID
+	uint32 BatchID,
+	bool bIsVisible
 	)
 {
-	new(BatchedPoints) FBatchedPoint(Position, Color, PointSize, LifeTime, DepthPriority, InBatchID);
+	new(BatchedPoints) FBatchedPoint(Position, Color, PointSize, LifeTime, DepthPriority, BatchID, bIsVisible);
+	bHasHiddenLines |= !bIsVisible;
 	// LineBatcher and PersistentLineBatcher components will be updated at the end of UWorld::Tick
 	MarkRenderStateDirty();
 }
 
-void ULineBatchComponent::DrawBox(const FBox& Box, const FMatrix& TM, FLinearColor Color, uint8 InDepthPriorityGroup, uint32 InBatchID)
+void ULineBatchComponent::DrawBox(const FBox& Box, const FMatrix& TM, FLinearColor Color, uint8 InDepthPriorityGroup, uint32 BatchID, bool bIsVisible)
 {
 	FVector	B[2], P, Q;
-	int32 ai, aj;
 	B[0] = Box.Min;
 	B[1] = Box.Max;
 
-	for( ai=0; ai<2; ai++ ) for( aj=0; aj<2; aj++ )
+	for( int32 ai=0; ai<2; ai++ ) for( int32 aj=0; aj<2; aj++ )
 	{
 		P.X=B[ai].X; Q.X=B[ai].X;
 		P.Y=B[aj].Y; Q.Y=B[aj].Y;
 		P.Z=B[0].Z; Q.Z=B[1].Z;
-		new(BatchedLines) FBatchedLine(TM.TransformPosition(P), TM.TransformPosition(Q), Color, DefaultLifeTime, 0.0f, InDepthPriorityGroup, InBatchID);
+		new(BatchedLines) FBatchedLine(TM.TransformPosition(P), TM.TransformPosition(Q), Color, DefaultLifeTime, 0.0f, InDepthPriorityGroup, BatchID, bIsVisible);
 
 		P.Y=B[ai].Y; Q.Y=B[ai].Y;
 		P.Z=B[aj].Z; Q.Z=B[aj].Z;
 		P.X=B[0].X; Q.X=B[1].X;
-		new(BatchedLines) FBatchedLine(TM.TransformPosition(P), TM.TransformPosition(Q), Color, DefaultLifeTime, 0.0f, InDepthPriorityGroup, InBatchID);
+		new(BatchedLines) FBatchedLine(TM.TransformPosition(P), TM.TransformPosition(Q), Color, DefaultLifeTime, 0.0f, InDepthPriorityGroup, BatchID, bIsVisible);
 
 		P.Z=B[ai].Z; Q.Z=B[ai].Z;
 		P.X=B[aj].X; Q.X=B[aj].X;
 		P.Y=B[0].Y; Q.Y=B[1].Y;
-		new(BatchedLines) FBatchedLine(TM.TransformPosition(P), TM.TransformPosition(Q), Color, DefaultLifeTime, 0.0f, InDepthPriorityGroup, InBatchID);
+		new(BatchedLines) FBatchedLine(TM.TransformPosition(P), TM.TransformPosition(Q), Color, DefaultLifeTime, 0.0f, InDepthPriorityGroup, BatchID, bIsVisible);
 	}
+
+	bHasHiddenLines |= !bIsVisible;
 	// LineBatcher and PersistentLineBatcher components will be updated at the end of UWorld::Tick
 	MarkRenderStateDirty();
 }
 
-void ULineBatchComponent::DrawBox(FVector const& Center, FVector const& Box, FLinearColor Color, float LifeTime, uint8 DepthPriority, float Thickness, uint32 InBatchID)
+void ULineBatchComponent::DrawBox(FVector const& Center, FVector const& Box, FLinearColor Color, float LifeTime, uint8 DepthPriority, float Thickness, uint32 BatchID, bool bIsVisible)
 {
-	new(BatchedLines) FBatchedLine(Center + FVector( Box.X,  Box.Y,  Box.Z), Center + FVector( Box.X, -Box.Y, Box.Z), Color, LifeTime, Thickness, DepthPriority, InBatchID);
-	new(BatchedLines) FBatchedLine(Center + FVector( Box.X, -Box.Y,  Box.Z), Center + FVector(-Box.X, -Box.Y, Box.Z), Color, LifeTime, Thickness, DepthPriority, InBatchID);
-	new(BatchedLines) FBatchedLine(Center + FVector(-Box.X, -Box.Y,  Box.Z), Center + FVector(-Box.X,  Box.Y, Box.Z) ,Color, LifeTime, Thickness, DepthPriority, InBatchID);
-	new(BatchedLines) FBatchedLine(Center + FVector(-Box.X,  Box.Y,  Box.Z), Center + FVector( Box.X,  Box.Y, Box.Z), Color, LifeTime, Thickness, DepthPriority, InBatchID);
-
-	new(BatchedLines) FBatchedLine(Center + FVector( Box.X,  Box.Y, -Box.Z), Center + FVector( Box.X, -Box.Y, -Box.Z), Color, LifeTime, Thickness, DepthPriority, InBatchID);
-	new(BatchedLines) FBatchedLine(Center + FVector( Box.X, -Box.Y, -Box.Z), Center + FVector(-Box.X, -Box.Y, -Box.Z), Color, LifeTime, Thickness, DepthPriority, InBatchID);
-	new(BatchedLines) FBatchedLine(Center + FVector(-Box.X, -Box.Y, -Box.Z), Center + FVector(-Box.X,  Box.Y, -Box.Z), Color, LifeTime, Thickness, DepthPriority, InBatchID);
-	new(BatchedLines) FBatchedLine(Center + FVector(-Box.X,  Box.Y, -Box.Z), Center + FVector( Box.X,  Box.Y, -Box.Z), Color, LifeTime, Thickness, DepthPriority, InBatchID);
-
-	new(BatchedLines) FBatchedLine(Center + FVector( Box.X,  Box.Y,  Box.Z), Center + FVector( Box.X,  Box.Y, -Box.Z), Color, LifeTime, Thickness, DepthPriority, InBatchID);
-	new(BatchedLines) FBatchedLine(Center + FVector( Box.X, -Box.Y,  Box.Z), Center + FVector( Box.X, -Box.Y, -Box.Z), Color, LifeTime, Thickness, DepthPriority, InBatchID);
-	new(BatchedLines) FBatchedLine(Center + FVector(-Box.X, -Box.Y,  Box.Z), Center + FVector(-Box.X, -Box.Y, -Box.Z), Color, LifeTime, Thickness, DepthPriority, InBatchID);
-	new(BatchedLines) FBatchedLine(Center + FVector(-Box.X,  Box.Y,  Box.Z), Center + FVector(-Box.X,  Box.Y, -Box.Z), Color, LifeTime, Thickness, DepthPriority, InBatchID);
+	new(BatchedLines) FBatchedLine(Center + FVector( Box.X,  Box.Y,  Box.Z), Center + FVector( Box.X, -Box.Y, Box.Z), Color, LifeTime, Thickness, DepthPriority, BatchID, bIsVisible);
+	new(BatchedLines) FBatchedLine(Center + FVector( Box.X, -Box.Y,  Box.Z), Center + FVector(-Box.X, -Box.Y, Box.Z), Color, LifeTime, Thickness, DepthPriority, BatchID, bIsVisible);
+	new(BatchedLines) FBatchedLine(Center + FVector(-Box.X, -Box.Y,  Box.Z), Center + FVector(-Box.X,  Box.Y, Box.Z) ,Color, LifeTime, Thickness, DepthPriority, BatchID, bIsVisible);
+	new(BatchedLines) FBatchedLine(Center + FVector(-Box.X,  Box.Y,  Box.Z), Center + FVector( Box.X,  Box.Y, Box.Z), Color, LifeTime, Thickness, DepthPriority, BatchID, bIsVisible);
+
+	new(BatchedLines) FBatchedLine(Center + FVector( Box.X,  Box.Y, -Box.Z), Center + FVector( Box.X, -Box.Y, -Box.Z), Color, LifeTime, Thickness, DepthPriority, BatchID, bIsVisible);
+	new(BatchedLines) FBatchedLine(Center + FVector( Box.X, -Box.Y, -Box.Z), Center + FVector(-Box.X, -Box.Y, -Box.Z), Color, LifeTime, Thickness, DepthPriority, BatchID, bIsVisible);
+	new(BatchedLines) FBatchedLine(Center + FVector(-Box.X, -Box.Y, -Box.Z), Center + FVector(-Box.X,  Box.Y, -Box.Z), Color, LifeTime, Thickness, DepthPriority, BatchID, bIsVisible);
+	new(BatchedLines) FBatchedLine(Center + FVector(-Box.X,  Box.Y, -Box.Z), Center + FVector( Box.X,  Box.Y, -Box.Z), Color, LifeTime, Thickness, DepthPriority, BatchID, bIsVisible);
+
+	new(BatchedLines) FBatchedLine(Center + FVector( Box.X,  Box.Y,  Box.Z), Center + FVector( Box.X,  Box.Y, -Box.Z), Color, LifeTime, Thickness, DepthPriority, BatchID, bIsVisible);
+	new(BatchedLines) FBatchedLine(Center + FVector( Box.X, -Box.Y,  Box.Z), Center + FVector( Box.X, -Box.Y, -Box.Z), Color, LifeTime, Thickness, DepthPriority, BatchID, bIsVisible);
+	new(BatchedLines) FBatchedLine(Center + FVector(-Box.X, -Box.Y,  Box.Z), Center + FVector(-Box.X, -Box.Y, -Box.Z), Color, LifeTime, Thickness, DepthPriority, BatchID, bIsVisible);
+	new(BatchedLines) FBatchedLine(Center + FVector(-Box.X,  Box.Y,  Box.Z), Center + FVector(-Box.X,  Box.Y, -Box.Z), Color, LifeTime, Thickness, DepthPriority, BatchID, bIsVisible);
+
+	bHasHiddenLines |= !bIsVisible;
 	
 	MarkRenderStateDirty();
 }
 
-void ULineBatchComponent::DrawBox(FVector const& Center, FVector const& Box, const FQuat& Rotation, FLinearColor Color, float LifeTime, uint8 DepthPriority, float Thickness, uint32 InBatchID)
+void ULineBatchComponent::DrawBox(FVector const& Center, FVector const& Box, const FQuat& Rotation, FLinearColor Color, float LifeTime, uint8 DepthPriority, float Thickness, uint32 BatchID, bool bIsVisible)
 {
 	FTransform const Transform(Rotation);
 	FVector Start = Transform.TransformPosition(FVector( Box.X,  Box.Y,  Box.Z));
 	FVector End = Transform.TransformPosition(FVector( Box.X, -Box.Y, Box.Z));
-	new(BatchedLines) FBatchedLine(Center + Start, Center + End, Color, LifeTime, Thickness, DepthPriority, InBatchID);
+	new(BatchedLines) FBatchedLine(Center + Start, Center + End, Color, LifeTime, Thickness, DepthPriority, BatchID, bIsVisible);
 
 	Start = Transform.TransformPosition(FVector( Box.X, -Box.Y,  Box.Z));
 	End = Transform.TransformPosition(FVector(-Box.X, -Box.Y, Box.Z));
-	new(BatchedLines) FBatchedLine(Center + Start, Center + End, Color, LifeTime, Thickness, DepthPriority, InBatchID);
+	new(BatchedLines) FBatchedLine(Center + Start, Center + End, Color, LifeTime, Thickness, DepthPriority, BatchID, bIsVisible);
 
 	Start = Transform.TransformPosition(FVector(-Box.X, -Box.Y,  Box.Z));
 	End = Transform.TransformPosition(FVector(-Box.X,  Box.Y, Box.Z));
-	new(BatchedLines) FBatchedLine(Center + Start, Center + End, Color, LifeTime, Thickness, DepthPriority, InBatchID);
+	new(BatchedLines) FBatchedLine(Center + Start, Center + End, Color, LifeTime, Thickness, DepthPriority, BatchID, bIsVisible);
 
 	Start = Transform.TransformPosition(FVector(-Box.X,  Box.Y,  Box.Z));
 	End = Transform.TransformPosition(FVector( Box.X,  Box.Y, Box.Z));
-	new(BatchedLines) FBatchedLine(Center + Start, Center + End, Color, LifeTime, Thickness, DepthPriority, InBatchID);
+	new(BatchedLines) FBatchedLine(Center + Start, Center + End, Color, LifeTime, Thickness, DepthPriority, BatchID, bIsVisible);
 
 	Start = Transform.TransformPosition(FVector( Box.X,  Box.Y, -Box.Z));
 	End = Transform.TransformPosition(FVector( Box.X, -Box.Y, -Box.Z));
-	new(BatchedLines) FBatchedLine(Center + Start, Center + End, Color, LifeTime, Thickness, DepthPriority, InBatchID);
+	new(BatchedLines) FBatchedLine(Center + Start, Center + End, Color, LifeTime, Thickness, DepthPriority, BatchID, bIsVisible);
 
 	Start = Transform.TransformPosition(FVector( Box.X, -Box.Y, -Box.Z));
 	End = Transform.TransformPosition(FVector(-Box.X, -Box.Y, -Box.Z));
-	new(BatchedLines) FBatchedLine(Center + Start, Center + End, Color, LifeTime, Thickness, DepthPriority, InBatchID);
+	new(BatchedLines) FBatchedLine(Center + Start, Center + End, Color, LifeTime, Thickness, DepthPriority, BatchID, bIsVisible);
 
 	Start = Transform.TransformPosition(FVector(-Box.X, -Box.Y, -Box.Z));
 	End = Transform.TransformPosition(FVector(-Box.X,  Box.Y, -Box.Z));
-	new(BatchedLines) FBatchedLine(Center + Start, Center + End, Color, LifeTime, Thickness, DepthPriority, InBatchID);
+	new(BatchedLines) FBatchedLine(Center + Start, Center + End, Color, LifeTime, Thickness, DepthPriority, BatchID, bIsVisible);
 
 	Start = Transform.TransformPosition(FVector(-Box.X,  Box.Y, -Box.Z));
 	End = Transform.TransformPosition(FVector( Box.X,  Box.Y, -Box.Z));
-	new(BatchedLines) FBatchedLine(Center + Start, Center + End, Color, LifeTime, Thickness, DepthPriority, InBatchID);
+	new(BatchedLines) FBatchedLine(Center + Start, Center + End, Color, LifeTime, Thickness, DepthPriority, BatchID, bIsVisible);
 
 	Start = Transform.TransformPosition(FVector( Box.X,  Box.Y,  Box.Z));
 	End = Transform.TransformPosition(FVector( Box.X,  Box.Y, -Box.Z));
-	new(BatchedLines) FBatchedLine(Center + Start, Center + End, Color, LifeTime, Thickness, DepthPriority, InBatchID);
+	new(BatchedLines) FBatchedLine(Center + Start, Center + End, Color, LifeTime, Thickness, DepthPriority, BatchID, bIsVisible);
 
 	Start = Transform.TransformPosition(FVector( Box.X, -Box.Y,  Box.Z));
 	End = Transform.TransformPosition(FVector( Box.X, -Box.Y, -Box.Z));
-	new(BatchedLines) FBatchedLine(Center + Start, Center + End, Color, LifeTime, Thickness, DepthPriority, InBatchID);
+	new(BatchedLines) FBatchedLine(Center + Start, Center + End, Color, LifeTime, Thickness, DepthPriority, BatchID, bIsVisible);
 
 	Start = Transform.TransformPosition(FVector(-Box.X, -Box.Y,  Box.Z));
 	End = Transform.TransformPosition(FVector(-Box.X, -Box.Y, -Box.Z));
-	new(BatchedLines) FBatchedLine(Center + Start, Center + End, Color, LifeTime, Thickness, DepthPriority, InBatchID);
+	new(BatchedLines) FBatchedLine(Center + Start, Center + End, Color, LifeTime, Thickness, DepthPriority, BatchID, bIsVisible);
 
 	Start = Transform.TransformPosition(FVector(-Box.X,  Box.Y,  Box.Z));
 	End = Transform.TransformPosition(FVector(-Box.X,  Box.Y, -Box.Z));
-	new(BatchedLines) FBatchedLine(Center + Start, Center + End, Color, LifeTime, Thickness, DepthPriority, InBatchID);
+	new(BatchedLines) FBatchedLine(Center + Start, Center + End, Color, LifeTime, Thickness, DepthPriority, BatchID, bIsVisible);
 
+	bHasHiddenLines |= !bIsVisible;
 	MarkRenderStateDirty();
 }
 
-void ULineBatchComponent::DrawSolidBox(const FBox& Box, const FTransform& Xform, const FColor& Color, uint8 DepthPriority, float LifeTime, uint32 InBatchID)
+void ULineBatchComponent::DrawSolidBox(const FBox& Box, const FTransform& Xform, const FColor& Color, uint8 DepthPriority, float LifeTime, uint32 BatchID, bool bIsVisible)
 {
 	int32 const NewMeshIdx = BatchedMeshes.Add(FBatchedMesh());
 	FBatchedMesh& BM = BatchedMeshes[NewMeshIdx];
@@ -295,7 +308,8 @@ void ULineBatchComponent::DrawSolidBox(const FBox& Box, const FTransform& Xform,
 	BM.Color = Color;
 	BM.DepthPriority = DepthPriority;
 	BM.RemainingLifeTime = LifeTime;
-	BM.BatchID = InBatchID;
+	BM.BatchID = BatchID;
+	BM.bIsVisible = bIsVisible;
 
 	BM.MeshVerts.AddUninitialized(8);
 	BM.MeshVerts[0] = Xform.TransformPosition( FVector(Box.Min.X, Box.Min.Y, Box.Max.Z) );
@@ -327,10 +341,11 @@ void ULineBatchComponent::DrawSolidBox(const FBox& Box, const FTransform& Xform,
 		BM.MeshIndices[Idx] = Indices[Idx];
 	}
 
+	bHasHiddenLines |= !bIsVisible;
 	MarkRenderStateDirty();
 }
 
-void ULineBatchComponent::DrawMesh(TArray<FVector> const& Verts, TArray<int32> const& Indices, FColor const& Color, uint8 DepthPriority, float LifeTime, uint32 InBatchID)
+void ULineBatchComponent::DrawMesh(TArray<FVector> const& Verts, TArray<int32> const& Indices, FColor const& Color, uint8 DepthPriority, float LifeTime, uint32 BatchID, bool bIsVisible)
 {
 	// modifying array element directly to avoid copying arrays
 	int32 const NewMeshIdx = BatchedMeshes.Add(FBatchedMesh());
@@ -341,24 +356,27 @@ void ULineBatchComponent::DrawMesh(TArray<FVector> const& Verts, TArray<int32> c
 	BM.Color = Color;
 	BM.DepthPriority = DepthPriority;
 	BM.RemainingLifeTime = LifeTime;
-	BM.BatchID = InBatchID;
+	BM.BatchID = BatchID;
+	BM.bIsVisible = bIsVisible;
 
+	bHasHiddenLines |= !bIsVisible;
 	MarkRenderStateDirty();
 }
 
-void ULineBatchComponent::DrawDirectionalArrow(const FMatrix& ArrowToWorld, FLinearColor InColor, float Length, float ArrowSize, uint8 DepthPriority, uint32 InBatchID)
+void ULineBatchComponent::DrawDirectionalArrow(const FMatrix& ArrowToWorld, FLinearColor InColor, float Length, float ArrowSize, uint8 DepthPriority, uint32 BatchID, bool bIsVisible)
 {
 	const FVector Tip = ArrowToWorld.TransformPosition(FVector(Length,0,0));
-	new(BatchedLines) FBatchedLine(Tip,ArrowToWorld.TransformPosition(FVector::ZeroVector),InColor,DefaultLifeTime,0.0f,DepthPriority, InBatchID);
-	new(BatchedLines) FBatchedLine(Tip,ArrowToWorld.TransformPosition(FVector(Length-ArrowSize,+ArrowSize,+ArrowSize)),InColor,DefaultLifeTime,0.0f,DepthPriority, InBatchID);
-	new(BatchedLines) FBatchedLine(Tip,ArrowToWorld.TransformPosition(FVector(Length-ArrowSize,+ArrowSize,-ArrowSize)),InColor,DefaultLifeTime,0.0f,DepthPriority, InBatchID);
-	new(BatchedLines) FBatchedLine(Tip,ArrowToWorld.TransformPosition(FVector(Length-ArrowSize,-ArrowSize,+ArrowSize)),InColor,DefaultLifeTime,0.0f,DepthPriority, InBatchID);
-	new(BatchedLines) FBatchedLine(Tip,ArrowToWorld.TransformPosition(FVector(Length-ArrowSize,-ArrowSize,-ArrowSize)),InColor,DefaultLifeTime,0.0f,DepthPriority, InBatchID);
-	
+	new(BatchedLines) FBatchedLine(Tip,ArrowToWorld.TransformPosition(FVector::ZeroVector),InColor,DefaultLifeTime,0.0f,DepthPriority, BatchID, bIsVisible);
+	new(BatchedLines) FBatchedLine(Tip,ArrowToWorld.TransformPosition(FVector(Length-ArrowSize,+ArrowSize,+ArrowSize)),InColor,DefaultLifeTime,0.0f,DepthPriority, BatchID, bIsVisible);
+	new(BatchedLines) FBatchedLine(Tip,ArrowToWorld.TransformPosition(FVector(Length-ArrowSize,+ArrowSize,-ArrowSize)),InColor,DefaultLifeTime,0.0f,DepthPriority, BatchID, bIsVisible);
+	new(BatchedLines) FBatchedLine(Tip,ArrowToWorld.TransformPosition(FVector(Length-ArrowSize,-ArrowSize,+ArrowSize)),InColor,DefaultLifeTime,0.0f,DepthPriority, BatchID, bIsVisible);
+	new(BatchedLines) FBatchedLine(Tip,ArrowToWorld.TransformPosition(FVector(Length-ArrowSize,-ArrowSize,-ArrowSize)),InColor,DefaultLifeTime,0.0f,DepthPriority, BatchID, bIsVisible);
+
+	bHasHiddenLines |= !bIsVisible;
 	MarkRenderStateDirty();
 }
 
-void ULineBatchComponent::DrawDirectionalArrow(FVector const& LineStart, FVector const& LineEnd, float ArrowSize, FLinearColor Color, float LifeTime, uint8 DepthPriority, float Thickness, uint32 InBatchID)
+void ULineBatchComponent::DrawDirectionalArrow(FVector const& LineStart, FVector const& LineEnd, float ArrowSize, FLinearColor Color, float LifeTime, uint8 DepthPriority, float Thickness, uint32 BatchID, bool bIsVisible)
 {
 	FVector Dir = (LineEnd-LineStart);
 	Dir.Normalize();
@@ -376,15 +394,16 @@ void ULineBatchComponent::DrawDirectionalArrow(FVector const& LineStart, FVector
 	// since dir is x direction, my arrow will be pointing +y, -x and -y, -x
 	const float ArrowSqrt = FMath::Sqrt(ArrowSize);
 
-	new(BatchedLines) FBatchedLine(LineStart,LineEnd, Color, LifeTime, Thickness, DepthPriority, InBatchID);
-	new(BatchedLines) FBatchedLine(LineEnd,LineEnd + TM.TransformPosition(FVector(-ArrowSqrt, ArrowSqrt, 0)), Color, LifeTime, Thickness, DepthPriority, InBatchID);
-	new(BatchedLines) FBatchedLine(LineEnd,LineEnd + TM.TransformPosition(FVector(-ArrowSqrt, -ArrowSqrt, 0)), Color, LifeTime, Thickness, DepthPriority, InBatchID);
+	new(BatchedLines) FBatchedLine(LineStart,LineEnd, Color, LifeTime, Thickness, DepthPriority, BatchID, bIsVisible);
+	new(BatchedLines) FBatchedLine(LineEnd,LineEnd + TM.TransformPosition(FVector(-ArrowSqrt, ArrowSqrt, 0)), Color, LifeTime, Thickness, DepthPriority, BatchID, bIsVisible);
+	new(BatchedLines) FBatchedLine(LineEnd,LineEnd + TM.TransformPosition(FVector(-ArrowSqrt, -ArrowSqrt, 0)), Color, LifeTime, Thickness, DepthPriority, BatchID, bIsVisible);
 
+	bHasHiddenLines |= !bIsVisible;
 	MarkRenderStateDirty();
 }
 
 
-void ULineBatchComponent::AddHalfCircle(const FVector& Base, const FVector& X, const FVector& Y, const FLinearColor& Color, const float Radius, int32 NumSides, const float LifeTime, uint8 DepthPriority, const float Thickness, const uint32 InBatchID)
+void ULineBatchComponent::AddHalfCircle(const FVector& Base, const FVector& X, const FVector& Y, const FLinearColor& Color, const float Radius, int32 NumSides, const float LifeTime, uint8 DepthPriority, const float Thickness, const uint32 BatchID, bool bIsVisible)
 {
 	// Need at least 2 sides
 	NumSides = FMath::Max(NumSides, 2);
@@ -394,12 +413,12 @@ void ULineBatchComponent::AddHalfCircle(const FVector& Base, const FVector& X, c
 	for( int32 SideIndex = 0; SideIndex < (NumSides/2); SideIndex++)
 	{
 		FVector	Vertex = Base + (X * FMath::Cos(AngleDelta * (SideIndex + 1)) + Y * FMath::Sin(AngleDelta * (SideIndex + 1))) * Radius;
-		new(BatchedLines) FBatchedLine(LastVertex, Vertex, Color, LifeTime, Thickness, DepthPriority, InBatchID);
+		new(BatchedLines) FBatchedLine(LastVertex, Vertex, Color, LifeTime, Thickness, DepthPriority, BatchID, bIsVisible);
 		LastVertex = Vertex;
-	}	
+	}
 }
 
-void ULineBatchComponent::AddCircle(const FVector& Base, const FVector& X, const FVector& Y, const FLinearColor& Color, const float Radius, int32 NumSides, const float LifeTime, uint8 DepthPriority, const float Thickness, const uint32 InBatchID)
+void ULineBatchComponent::AddCircle(const FVector& Base, const FVector& X, const FVector& Y, const FLinearColor& Color, const float Radius, int32 NumSides, const float LifeTime, uint8 DepthPriority, const float Thickness, const uint32 BatchID, bool bInVisible)
 {
 	// Need at least 2 sides
 	NumSides = FMath::Max(NumSides, 2);
@@ -409,21 +428,22 @@ void ULineBatchComponent::AddCircle(const FVector& Base, const FVector& X, const
 	for (int32 SideIndex = 0; SideIndex < NumSides; SideIndex++)
 	{
 		const FVector Vertex = Base + (X * FMath::Cos(AngleDelta * (SideIndex + 1)) + Y * FMath::Sin(AngleDelta * (SideIndex + 1))) * Radius;
-		new(BatchedLines) FBatchedLine(LastVertex, Vertex, Color, LifeTime, Thickness, DepthPriority, InBatchID);
+		new(BatchedLines) FBatchedLine(LastVertex, Vertex, Color, LifeTime, Thickness, DepthPriority, BatchID, bInVisible);
 		LastVertex = Vertex;
 	}
 }
 
 /** Draw a circle */
-void ULineBatchComponent::DrawCircle(const FVector& Base, const FVector& X, const FVector& Y, FLinearColor Color, float Radius, int32 NumSides, uint8 DepthPriority, uint32 InBatchID)
+void ULineBatchComponent::DrawCircle(const FVector& Base, const FVector& X, const FVector& Y, FLinearColor Color, float Radius, int32 NumSides, uint8 DepthPriority, uint32 BatchID, bool bIsVisible)
 {
-	AddHalfCircle(Base, X, Y, Color, Radius, NumSides, DefaultLifeTime, DepthPriority, 0.f, InBatchID);
-
+	AddHalfCircle(Base, X, Y, Color, Radius, NumSides, DefaultLifeTime, DepthPriority, 0.f, BatchID, bIsVisible);
+		
+	bHasHiddenLines |= !bIsVisible;
 	MarkRenderStateDirty();
 }
 
 /** Draw a sphere */
-void ULineBatchComponent::DrawSphere(FVector const& Center, float Radius, int32 Segments, FLinearColor Color, float LifeTime, uint8 DepthPriority, float Thickness, uint32 InBatchID)
+void ULineBatchComponent::DrawSphere(FVector const& Center, float Radius, int32 Segments, FLinearColor Color, float LifeTime, uint8 DepthPriority, float Thickness, uint32 BatchID, bool bIsVisible)
 {
 	// Need at least 4 segments
 	Segments = FMath::Max(Segments, 4);
@@ -453,8 +473,8 @@ void ULineBatchComponent::DrawSphere(FVector const& Center, float Radius, int32
 			const FVector Vertex2 = FVector((CosX * SinY1), (SinX * SinY1), CosY1) * Radius + Center;
 			const FVector Vertex4 = FVector((CosX * SinY2), (SinX * SinY2), CosY2) * Radius + Center;
 
-			new(BatchedLines) FBatchedLine(Vertex1, Vertex2, Color, LifeTime, Thickness, DepthPriority, InBatchID);
-			new(BatchedLines) FBatchedLine(Vertex1, Vertex3, Color, LifeTime, Thickness, DepthPriority, InBatchID);
+			new(BatchedLines) FBatchedLine(Vertex1, Vertex2, Color, LifeTime, Thickness, DepthPriority, BatchID, bIsVisible);
+			new(BatchedLines) FBatchedLine(Vertex1, Vertex3, Color, LifeTime, Thickness, DepthPriority, BatchID, bIsVisible);
 
 			Vertex1 = Vertex2;
 			Vertex3 = Vertex4;
@@ -464,10 +484,12 @@ void ULineBatchComponent::DrawSphere(FVector const& Center, float Radius, int32
 		CosY1 = CosY2;
 		Latitude += AngleInc;
 	}
+	
+	bHasHiddenLines |= !bIsVisible;
 	MarkRenderStateDirty();
 }
 
-void ULineBatchComponent::DrawCylinder(FVector const& Start, FVector const& End, float Radius, int32 Segments, FLinearColor Color, float LifeTime, uint8 DepthPriority, float Thickness, uint32 InBatchID)
+void ULineBatchComponent::DrawCylinder(FVector const& Start, FVector const& End, float Radius, int32 Segments, FLinearColor Color, float LifeTime, uint8 DepthPriority, float Thickness, uint32 BatchID, bool bIsVisible)
 {
 	// Need at least 4 segments
 	Segments = FMath::Max(Segments, 4);
@@ -496,18 +518,20 @@ void ULineBatchComponent::DrawCylinder(FVector const& Start, FVector const& End,
 		FVector P2 = Segment + Start;
 		FVector P4 = Segment + End;
 
-		new(BatchedLines) FBatchedLine(P2, P4, Color, LifeTime, Thickness, DepthPriority, InBatchID);
-		new(BatchedLines) FBatchedLine(P1, P2, Color, LifeTime, Thickness, DepthPriority, InBatchID);
-		new(BatchedLines) FBatchedLine(P3, P4, Color, LifeTime, Thickness, DepthPriority, InBatchID);
+		new(BatchedLines) FBatchedLine(P2, P4, Color, LifeTime, Thickness, DepthPriority, BatchID, bIsVisible);
+		new(BatchedLines) FBatchedLine(P1, P2, Color, LifeTime, Thickness, DepthPriority, BatchID, bIsVisible);
+		new(BatchedLines) FBatchedLine(P3, P4, Color, LifeTime, Thickness, DepthPriority, BatchID, bIsVisible);
 
 		P1 = P2;
 		P3 = P4;
 		Angle += AngleInc;
 	}
+
+	bHasHiddenLines |= !bIsVisible;
 	MarkRenderStateDirty();
 }
 
-void ULineBatchComponent::DrawCone(FVector const& Origin, FVector const& Direction, float Length, float AngleWidth, float AngleHeight, int32 NumSides, FLinearColor DrawColor, float LifeTime, uint8 DepthPriority, float Thickness, uint32 InBatchID)
+void ULineBatchComponent::DrawCone(FVector const& Origin, FVector const& Direction, float Length, float AngleWidth, float AngleHeight, int32 NumSides, FLinearColor DrawColor, float LifeTime, uint8 DepthPriority, float Thickness, uint32 BatchID, bool bIsVisible)
 {
 	// Need at least 4 sides
 	NumSides = FMath::Max(NumSides, 4);
@@ -555,12 +579,12 @@ void ULineBatchComponent::DrawCone(FVector const& Origin, FVector const& Directi
 	for(int32 i = 0; i < NumSides; i++)
 	{
 		CurrentPoint = ConeToWorld.TransformPosition(ConeVerts[i]);
-		new(BatchedLines) FBatchedLine(ConeToWorld.GetOrigin(), CurrentPoint, DrawColor, LifeTime, Thickness, DepthPriority, InBatchID);
+		new(BatchedLines) FBatchedLine(ConeToWorld.GetOrigin(), CurrentPoint, DrawColor, LifeTime, Thickness, DepthPriority, BatchID, bIsVisible);
 
 		// PrevPoint must be defined to draw junctions
 		if( i > 0 )
 		{
-			new(BatchedLines) FBatchedLine(PrevPoint, CurrentPoint, DrawColor, LifeTime, Thickness, DepthPriority, InBatchID);
+			new(BatchedLines) FBatchedLine(PrevPoint, CurrentPoint, DrawColor, LifeTime, Thickness, DepthPriority, BatchID, bIsVisible);
 		}
 		else
 		{
@@ -570,12 +594,13 @@ void ULineBatchComponent::DrawCone(FVector const& Origin, FVector const& Directi
 		PrevPoint = CurrentPoint;
 	}
 	// Connect last junction to first
-	new(BatchedLines) FBatchedLine(CurrentPoint, FirstPoint, DrawColor, LifeTime, Thickness, DepthPriority, InBatchID);
+	new(BatchedLines) FBatchedLine(CurrentPoint, FirstPoint, DrawColor, LifeTime, Thickness, DepthPriority, BatchID, bIsVisible);
 
+	bHasHiddenLines |= !bIsVisible;	
 	MarkRenderStateDirty();
 }
 
-void ULineBatchComponent::DrawCapsule(FVector const& Center, float HalfHeight, float Radius, const FQuat& Rotation, FLinearColor Color, float LifeTime, uint8 DepthPriority, float Thickness, uint32 InBatchID)
+void ULineBatchComponent::DrawCapsule(FVector const& Center, float HalfHeight, float Radius, const FQuat& Rotation, FLinearColor Color, float LifeTime, uint8 DepthPriority, float Thickness, uint32 BatchID, bool bIsVisible)
 {
 	constexpr int32 DrawCollisionSides = 16;
 	const FVector Origin = Center;
@@ -589,30 +614,32 @@ void ULineBatchComponent::DrawCapsule(FVector const& Center, float HalfHeight, f
 	const FVector TopEnd = Origin + HalfAxis*ZAxis;
 	const FVector BottomEnd = Origin - HalfAxis*ZAxis;
 
-	AddCircle(TopEnd, XAxis, YAxis, Color, Radius, DrawCollisionSides, LifeTime, Thickness, DepthPriority,  InBatchID);
-	AddCircle( BottomEnd, XAxis, YAxis, Color, Radius, DrawCollisionSides, LifeTime, Thickness, DepthPriority, InBatchID);
+	AddCircle(TopEnd, XAxis, YAxis, Color, Radius, DrawCollisionSides, LifeTime, Thickness, DepthPriority,  BatchID, bIsVisible);
+	AddCircle( BottomEnd, XAxis, YAxis, Color, Radius, DrawCollisionSides, LifeTime, Thickness, DepthPriority, BatchID, bIsVisible);
 
 		// Draw domed caps
-	AddHalfCircle( TopEnd, YAxis, ZAxis, Color, Radius, DrawCollisionSides, LifeTime, Thickness, DepthPriority, InBatchID);
-	AddHalfCircle( TopEnd, XAxis, ZAxis, Color, Radius, DrawCollisionSides, LifeTime, Thickness, DepthPriority, InBatchID);
+	AddHalfCircle( TopEnd, YAxis, ZAxis, Color, Radius, DrawCollisionSides, LifeTime, Thickness, DepthPriority, BatchID, bIsVisible);
+	AddHalfCircle( TopEnd, XAxis, ZAxis, Color, Radius, DrawCollisionSides, LifeTime, Thickness, DepthPriority, BatchID, bIsVisible);
 
 	const FVector NegZAxis = -ZAxis;
 
-	AddHalfCircle( BottomEnd, YAxis, NegZAxis, Color, Radius, DrawCollisionSides, LifeTime, Thickness, DepthPriority, InBatchID);
-	AddHalfCircle( BottomEnd, XAxis, NegZAxis, Color, Radius, DrawCollisionSides, LifeTime, Thickness, DepthPriority, InBatchID);
+	AddHalfCircle( BottomEnd, YAxis, NegZAxis, Color, Radius, DrawCollisionSides, LifeTime, Thickness, DepthPriority, BatchID, bIsVisible);
+	AddHalfCircle( BottomEnd, XAxis, NegZAxis, Color, Radius, DrawCollisionSides, LifeTime, Thickness, DepthPriority, BatchID, bIsVisible);
 
 	// Draw connected lines
-	new(BatchedLines) FBatchedLine(TopEnd + Radius*XAxis, BottomEnd + Radius*XAxis, Color, LifeTime, Thickness, DepthPriority, InBatchID);
-	new(BatchedLines) FBatchedLine(TopEnd - Radius*XAxis, BottomEnd - Radius*XAxis, Color, LifeTime, Thickness, DepthPriority, InBatchID);
-	new(BatchedLines) FBatchedLine(TopEnd + Radius*XAxis, BottomEnd + Radius*XAxis, Color, LifeTime, Thickness, DepthPriority, InBatchID);
-	new(BatchedLines) FBatchedLine(TopEnd - Radius*XAxis, BottomEnd - Radius*XAxis, Color, LifeTime, Thickness, DepthPriority, InBatchID);
-	
+	new(BatchedLines) FBatchedLine(TopEnd + Radius*XAxis, BottomEnd + Radius*XAxis, Color, LifeTime, Thickness, DepthPriority, BatchID, bIsVisible);
+	new(BatchedLines) FBatchedLine(TopEnd - Radius*XAxis, BottomEnd - Radius*XAxis, Color, LifeTime, Thickness, DepthPriority, BatchID, bIsVisible);
+	new(BatchedLines) FBatchedLine(TopEnd + Radius*XAxis, BottomEnd + Radius*XAxis, Color, LifeTime, Thickness, DepthPriority, BatchID, bIsVisible);
+	new(BatchedLines) FBatchedLine(TopEnd - Radius*XAxis, BottomEnd - Radius*XAxis, Color, LifeTime, Thickness, DepthPriority, BatchID, bIsVisible);
+
+	bHasHiddenLines |= !bIsVisible;
 	MarkRenderStateDirty();
 }
 
 void ULineBatchComponent::TickComponent(float DeltaTime, enum ELevelTick TickType, FActorComponentTickFunction *ThisTickFunction)
 {
 	bool bDirty = false;
+	bHasHiddenLines = false;
 	// Update the life time of batched lines, removing the lines which have expired.
 	for(int32 LineIndex=0; LineIndex < BatchedLines.Num(); LineIndex++)
 	{
@@ -625,8 +652,14 @@ void ULineBatchComponent::TickComponent(float DeltaTime, enum ELevelTick TickTyp
 				// The line has expired, remove it.
 				BatchedLines.RemoveAtSwap(LineIndex--);
 				bDirty = true;
+				continue;
 			}
 		}
+		if (Line.bIsVisible == false && !bHasHiddenLines)
+		{
+			bHasHiddenLines = true;
+			bDirty = true;
+		}
 	}
 
 	// Update the life time of batched points, removing the points which have expired.
@@ -641,8 +674,14 @@ void ULineBatchComponent::TickComponent(float DeltaTime, enum ELevelTick TickTyp
 				// The point has expired, remove it.
 				BatchedPoints.RemoveAtSwap(PtIndex--);
 				bDirty = true;
+				continue;
 			}
 		}
+		if (Pt.bIsVisible == false && !bHasHiddenLines)
+		{
+			bHasHiddenLines = true;
+			bDirty = true;
+		}
 	}
 
 	// Update the life time of batched meshes, removing the meshes which have expired.
@@ -657,8 +696,14 @@ void ULineBatchComponent::TickComponent(float DeltaTime, enum ELevelTick TickTyp
 				// The mesh has expired, remove it.
 				BatchedMeshes.RemoveAtSwap(MeshIndex--);
 				bDirty = true;
+				continue;
 			}
 		}
+		if (Mesh.bIsVisible == false && !bHasHiddenLines)
+		{
+			bHasHiddenLines = true;
+			bDirty = true;
+		}
 	}
 
 	if(bDirty)
@@ -760,19 +805,89 @@ void ULineBatchComponent::Flush()
 	}
 }
 
-void ULineBatchComponent::ClearBatch(uint32 InBatchID)
+void ULineBatchComponent::ClearBatch(uint32 BatchID)
 {
-	if (InBatchID == INVALID_ID)
+	if (BatchID == INVALID_ID)
 	{
 		return;
 	}
 	
-	bool bDirty = BatchedLines.RemoveAllSwap([InBatchID](const FBatchedLine& Element) { return Element.BatchID == InBatchID; }) > 0;
-	bDirty |= BatchedPoints.RemoveAllSwap([InBatchID](const FBatchedPoint& Pt) { return Pt.BatchID == InBatchID; }) > 0;
-	bDirty |= BatchedMeshes.RemoveAllSwap([InBatchID](const FBatchedMesh& Mesh) { return Mesh.BatchID == InBatchID; }) > 0;
+	bool bDirty = BatchedLines.RemoveAllSwap([BatchID](const FBatchedLine& Element) { return Element.BatchID == BatchID; }) > 0;
+	bDirty |= BatchedPoints.RemoveAllSwap([BatchID](const FBatchedPoint& Pt) { return Pt.BatchID == BatchID; }) > 0;
+	bDirty |= BatchedMeshes.RemoveAllSwap([BatchID](const FBatchedMesh& Mesh) { return Mesh.BatchID == BatchID; }) > 0;
 
 	if(bDirty)
 	{
 		MarkRenderStateDirty();
 	}
 }
+
+void ULineBatchComponent::SetBatchVisible(uint32 BatchID, bool bIsVisible)
+{
+	if (BatchID == INVALID_ID)
+	{
+		return;
+	}
+	bool bDirty = false;
+	for (FBatchedLine& Line: BatchedLines)
+	{
+		if (Line.BatchID == BatchID)
+		{
+			Line.bIsVisible = bIsVisible;
+			bDirty = true;
+		}
+	}
+	for (FBatchedPoint& Point: BatchedPoints)
+	{
+		if (Point.BatchID == BatchID)
+		{
+			Point.bIsVisible = bIsVisible;
+			bDirty = true;
+		}
+	}
+	for (FBatchedMesh& Mesh: BatchedMeshes)
+	{
+		if (Mesh.BatchID == BatchID)
+		{
+			Mesh.bIsVisible = bIsVisible;
+			bDirty = true;
+		}
+	}
+	if(bDirty)
+	{
+		if (!bIsVisible)
+		{
+			bHasHiddenLines = true;	
+		}
+		MarkRenderStateDirty();
+	}
+}
+
+void ULineBatchComponent::SetVisible(bool bIsVisible)
+{
+	bool bDirty = false;
+	for (FBatchedLine& Line: BatchedLines)
+	{
+		Line.bIsVisible = bIsVisible;
+		bDirty = true;
+	}
+	for (FBatchedPoint& Point: BatchedPoints)
+	{
+		Point.bIsVisible = bIsVisible;
+		bDirty = true;
+	}
+	for (FBatchedMesh& Mesh: BatchedMeshes)
+	{
+		Mesh.bIsVisible = bIsVisible;
+		bDirty = true;
+	}
+	if(bDirty)
+	{
+		if (!bIsVisible)
+		{
+			bHasHiddenLines = true;	
+		}
+
+		MarkRenderStateDirty();
+	}
+}
