CIS fix audio mixer module not being loaded in device manager

diff --git a/Engine/Source/Runtime/Engine/Private/AudioDeviceManager.cpp b/Engine/Source/Runtime/Engine/Private/AudioDeviceManager.cpp
--- a/Engine/Source/Runtime/Engine/Private/AudioDeviceManager.cpp
+++ b/Engine/Source/Runtime/Engine/Private/AudioDeviceManager.cpp
@@ -389,6 +389,8 @@ bool FAudioDeviceManager::InitializeManager()
 		AudioFormatSettings = MakePimpl<Audio::FAudioFormatSettings>(GConfig, GEngineIni, FPlatformProperties::IniPlatformName());
 		RegisterAudioInfoFactories();
 
+		FModuleManager::Get().LoadModuleChecked(TEXT("AudioMixer"));
+
 #if WITH_EDITOR
 		IAudioEditorModule* AudioEditorModule = &FModuleManager::LoadModuleChecked<IAudioEditorModule>("AudioEditor");
 		AudioEditorModule->RegisterAudioMixerAssetActions();
