- Iterate UFXSystemAsset & UFXSystemComponent for 'ListParticleSystems'
- This allows NIagara to be picked up as part of the information

diff --git a/Engine/Source/Runtime/Engine/Private/UnrealEngine.cpp b/Engine/Source/Runtime/Engine/Private/UnrealEngine.cpp
--- a/Engine/Source/Runtime/Engine/Private/UnrealEngine.cpp
+++ b/Engine/Source/Runtime/Engine/Private/UnrealEngine.cpp
@@ -7094,9 +7094,9 @@ bool UEngine::HandleListParticleSystemsCommand( const TCHAR* Cmd, FOutputDevice&
 	TMap<UObject *,int32> SortMap;
 
 	FString Description;
-	for( TObjectIterator<UParticleSystem> SystemIt; SystemIt; ++SystemIt )
+	for( TObjectIterator<UFXSystemAsset> SystemIt; SystemIt; ++SystemIt )
 	{			
-		UParticleSystem* Tree = *SystemIt;
+		UFXSystemAsset* Tree = *SystemIt;
 		Description = FString::Printf(TEXT("%s"), *Tree->GetPathName());
 		FArchiveCountMem Count( Tree );
 		int32 RootSize = Count.GetMax();
@@ -7119,10 +7119,10 @@ bool UEngine::HandleListParticleSystemsCommand( const TCHAR* Cmd, FOutputDevice&
 		}
 	}
 
-	for( TObjectIterator<UParticleSystemComponent> It; It; ++It )
+	for( TObjectIterator<UFXSystemComponent> It; It; ++It )
 	{
-		UParticleSystemComponent* Comp = *It;
-		int32 *pIndex = SortMap.Find(Comp->Template);
+		UFXSystemComponent* Comp = *It;
+		int32 *pIndex = SortMap.Find(Comp->GetFXSystemAsset());
 
 		if (pIndex && SortedSets.IsValidIndex(*pIndex))
 		{				
@@ -7141,38 +7141,39 @@ bool UEngine::HandleListParticleSystemsCommand( const TCHAR* Cmd, FOutputDevice&
 			Set.Size += CompResSize.GetTotalMemoryBytes();
 			Set.ComponentCount++;
 
-			UParticleSystem* Tree = Comp->Template;
-			if (bDumpMesh && Tree != NULL)
+			if (UParticleSystemComponent* PSC = Cast<UParticleSystemComponent>(Comp))
 			{
-				for (int32 EmitterIdx = 0; (EmitterIdx < Tree->Emitters.Num()); EmitterIdx++)
+				UParticleSystem* Tree = PSC->Template;
+				if (bDumpMesh && Tree != NULL)
 				{
-					UParticleEmitter* Emitter = Tree->Emitters[EmitterIdx];
-					if (Emitter != NULL)
+					for (int32 EmitterIdx = 0; (EmitterIdx < Tree->Emitters.Num()); EmitterIdx++)
 					{
-						// Have to check each LOD level...
-						if (Emitter->LODLevels.Num() > 0)
+						if (UParticleEmitter* Emitter = Tree->Emitters[EmitterIdx])
 						{
-							UParticleLODLevel* LODLevel = Emitter->LODLevels[0];
-							if (LODLevel != NULL)
+							// Have to check each LOD level...
+							if (Emitter->LODLevels.Num() > 0)
 							{
-								if (LODLevel->RequiredModule->bUseLocalSpace == true)
+								if (UParticleLODLevel* LODLevel = Emitter->LODLevels[0])
 								{
-									UParticleModuleTypeDataMesh* MeshTD = Cast<UParticleModuleTypeDataMesh>(LODLevel->TypeDataModule);
-									if (MeshTD != NULL)
+									if (LODLevel->RequiredModule->bUseLocalSpace == true)
 									{
-										int32 InstCount = 0;
-										// MESH EMITTER
-										if (EmitterIdx < Comp->EmitterInstances.Num())
+										UParticleModuleTypeDataMesh* MeshTD = Cast<UParticleModuleTypeDataMesh>(LODLevel->TypeDataModule);
+										if (MeshTD != NULL)
 										{
-											FParticleEmitterInstance* Inst = Comp->EmitterInstances[EmitterIdx];
-											if (Inst != NULL)
+											int32 InstCount = 0;
+											// MESH EMITTER
+											if (EmitterIdx < PSC->EmitterInstances.Num())
 											{
-												InstCount = Inst->ActiveParticles;
+												FParticleEmitterInstance* Inst = PSC->EmitterInstances[EmitterIdx];
+												if (Inst != NULL)
+												{
+													InstCount = Inst->ActiveParticles;
+												}
+
+												UE_LOG(LogEngine, Warning, TEXT("---> PSys w/ mesh emitters: %2d %4d %s %s "), EmitterIdx, InstCount, 
+													PSC->SceneProxy ? TEXT("Y") : TEXT("N"),
+													*(Tree->GetPathName()));
 											}
-
-											UE_LOG(LogEngine, Warning, TEXT("---> PSys w/ mesh emitters: %2d %4d %s %s "), EmitterIdx, InstCount, 
-												Comp->SceneProxy ? TEXT("Y") : TEXT("N"),
-												*(Tree->GetPathName()));
 										}
 									}
 								}
@@ -7181,7 +7182,6 @@ bool UEngine::HandleListParticleSystemsCommand( const TCHAR* Cmd, FOutputDevice&
 					}
 				}
 			}
-
 		}
 	}
 
