Fix inverted parameters in DrawCapsule
#tests PIE, draw some capsules
#rnx

diff --git a/Engine/Source/Runtime/Engine/Private/Components/LineBatchComponent.cpp b/Engine/Source/Runtime/Engine/Private/Components/LineBatchComponent.cpp
--- a/Engine/Source/Runtime/Engine/Private/Components/LineBatchComponent.cpp
+++ b/Engine/Source/Runtime/Engine/Private/Components/LineBatchComponent.cpp
@@ -436,7 +436,7 @@ void ULineBatchComponent::AddCircle(const FVector& Base, const FVector& X, const
 /** Draw a circle */
 void ULineBatchComponent::DrawCircle(const FVector& Base, const FVector& X, const FVector& Y, FLinearColor Color, float Radius, int32 NumSides, uint8 DepthPriority, uint32 BatchID, bool bIsVisible)
 {
-	AddHalfCircle(Base, X, Y, Color, Radius, NumSides, DefaultLifeTime, DepthPriority, 0.f, BatchID, bIsVisible);
+	AddCircle(Base, X, Y, Color, Radius, NumSides, DefaultLifeTime, DepthPriority, 0.f, BatchID, bIsVisible);
 		
 	bHasHiddenLines |= !bIsVisible;
 	MarkRenderStateDirty();
@@ -614,17 +614,17 @@ void ULineBatchComponent::DrawCapsule(FVector const& Center, float HalfHeight, f
 	const FVector TopEnd = Origin + HalfAxis*ZAxis;
 	const FVector BottomEnd = Origin - HalfAxis*ZAxis;
 
-	AddCircle(TopEnd, XAxis, YAxis, Color, Radius, DrawCollisionSides, LifeTime, Thickness, DepthPriority,  BatchID, bIsVisible);
-	AddCircle( BottomEnd, XAxis, YAxis, Color, Radius, DrawCollisionSides, LifeTime, Thickness, DepthPriority, BatchID, bIsVisible);
+	AddCircle(TopEnd, XAxis, YAxis, Color, Radius, DrawCollisionSides, LifeTime, DepthPriority, Thickness, BatchID, bIsVisible);
+	AddCircle( BottomEnd, XAxis, YAxis, Color, Radius, DrawCollisionSides, LifeTime, DepthPriority, Thickness, BatchID, bIsVisible);
 
-		// Draw domed caps
-	AddHalfCircle( TopEnd, YAxis, ZAxis, Color, Radius, DrawCollisionSides, LifeTime, Thickness, DepthPriority, BatchID, bIsVisible);
-	AddHalfCircle( TopEnd, XAxis, ZAxis, Color, Radius, DrawCollisionSides, LifeTime, Thickness, DepthPriority, BatchID, bIsVisible);
+	// Draw domed caps
+	AddHalfCircle( TopEnd, YAxis, ZAxis, Color, Radius, DrawCollisionSides, LifeTime, DepthPriority, Thickness, BatchID, bIsVisible);
+	AddHalfCircle( TopEnd, XAxis, ZAxis, Color, Radius, DrawCollisionSides, LifeTime, DepthPriority, Thickness, BatchID, bIsVisible);
 
 	const FVector NegZAxis = -ZAxis;
 
-	AddHalfCircle( BottomEnd, YAxis, NegZAxis, Color, Radius, DrawCollisionSides, LifeTime, Thickness, DepthPriority, BatchID, bIsVisible);
-	AddHalfCircle( BottomEnd, XAxis, NegZAxis, Color, Radius, DrawCollisionSides, LifeTime, Thickness, DepthPriority, BatchID, bIsVisible);
+	AddHalfCircle( BottomEnd, YAxis, NegZAxis, Color, Radius, DrawCollisionSides, LifeTime, DepthPriority, Thickness, BatchID, bIsVisible);
+	AddHalfCircle( BottomEnd, XAxis, NegZAxis, Color, Radius, DrawCollisionSides, LifeTime, DepthPriority, Thickness, BatchID, bIsVisible);
 
 	// Draw connected lines
 	new(BatchedLines) FBatchedLine(TopEnd + Radius*XAxis, BottomEnd + Radius*XAxis, Color, LifeTime, Thickness, DepthPriority, BatchID, bIsVisible);
