Unregister on destroyed callback in LWI manager to account for actors potentially being reused by level streaming.

diff --git a/Engine/Source/Runtime/Engine/Private/LightWeightInstanceManager.cpp b/Engine/Source/Runtime/Engine/Private/LightWeightInstanceManager.cpp
--- a/Engine/Source/Runtime/Engine/Private/LightWeightInstanceManager.cpp
+++ b/Engine/Source/Runtime/Engine/Private/LightWeightInstanceManager.cpp
@@ -147,6 +147,7 @@ void ALightWeightInstanceManager::OnSpawnedActorDestroyed(AActor* DestroyedActor
 	const int32 DestroyedActorInstanceIndex = FindIndexForActor(DestroyedActor);
 	check(DestroyedActorInstanceIndex != INDEX_NONE);
 	Actors.Remove(DestroyedActorInstanceIndex);
+	DestroyedActor->OnDestroyed.RemoveAll(this);
 }
 
 int32 ALightWeightInstanceManager::FindIndexForActor(const AActor* InActor) const
