Motion Matching: Add in world debug text to make it easier to diagnose high level issues, such as what database we're searching.
Also removed bForceInterrupt from the Motion Matching node. That is best controlled via the SetDatabaseToSearch/SetDatabasesToSearch functions.

diff --git a/Engine/Plugins/Experimental/Animation/PoseSearch/Source/Runtime/Private/AnimNode_MotionMatching.cpp b/Engine/Plugins/Experimental/Animation/PoseSearch/Source/Runtime/Private/AnimNode_MotionMatching.cpp
--- a/Engine/Plugins/Experimental/Animation/PoseSearch/Source/Runtime/Private/AnimNode_MotionMatching.cpp
+++ b/Engine/Plugins/Experimental/Animation/PoseSearch/Source/Runtime/Private/AnimNode_MotionMatching.cpp
@@ -21,6 +21,8 @@
 #if ENABLE_ANIM_DEBUG
 static TAutoConsoleVariable<int32> CVarAnimNodeMotionMatchingDrawQuery(TEXT("a.AnimNode.MotionMatching.DebugDrawQuery"), 0, TEXT("Draw input query"));
 static TAutoConsoleVariable<int32> CVarAnimNodeMotionMatchingDrawCurResult(TEXT("a.AnimNode.MotionMatching.DebugDrawCurResult"), 0, TEXT("Draw current result"));
+static TAutoConsoleVariable<int32> CVarAnimNodeMotionMatchingDrawInfo(TEXT("a.AnimNode.MotionMatching.DebugDrawInfo"), 0, TEXT("Draw info like current databases to search"));
+static TAutoConsoleVariable<float> CVarAnimNodeMotionMatchingDrawInfoHeight(TEXT("a.AnimNode.MotionMatching.DebugDrawInfoHeight"), 50.f, TEXT("Vertical offset for DebugDrawInfo"));
 #endif
 
 /////////////////////////////////////////////////////
@@ -109,6 +111,20 @@ void FAnimNode_MotionMatching::UpdateAssetPlayer(const FAnimationUpdateContext&
 		DatabasesToSearch.Add(Database);
 	}
 
+#if ENABLE_ANIM_DEBUG
+	if (CVarAnimNodeMotionMatchingDrawInfo.GetValueOnAnyThread() > 0)
+	{
+		FString DebugInfo = FString::Printf(TEXT("bForceInterruptNextUpdate(%d)\n"), bForceInterruptNextUpdate);
+		DebugInfo += FString::Printf(TEXT("Current Database(%s)\n"), *GetNameSafe(MotionMatchingState.CurrentSearchResult.Database.Get()));
+		DebugInfo += FString::Printf(TEXT("Databases to search:\n"));
+		for (const UPoseSearchDatabase* DatabaseToSearch : DatabasesToSearch)
+		{
+			DebugInfo += FString::Printf(TEXT("  %s\n"), *GetNameSafe(DatabaseToSearch));
+		}
+		Context.AnimInstanceProxy->AnimDrawDebugInWorldMessage(DebugInfo, FVector::UpVector * CVarAnimNodeMotionMatchingDrawInfoHeight.GetValueOnAnyThread(), FColor::Yellow, 1.f /*TextScale*/);
+	}
+#endif // ENABLE_ANIM_DEBUG
+
 	// Execute core motion matching algorithm
 	UPoseSearchLibrary::UpdateMotionMatchingState(
 		Context,
@@ -122,12 +138,12 @@ void FAnimNode_MotionMatching::UpdateAssetPlayer(const FAnimationUpdateContext&
 		SearchThrottleTime,
 		PlayRate,
 		MotionMatchingState,
-		bForceInterrupt | bForceInterruptNextUpdate,
+		bForceInterruptNextUpdate,
 		bShouldSearch
 		#if ENABLE_ANIM_DEBUG
 		, CVarAnimNodeMotionMatchingDrawQuery.GetValueOnAnyThread() > 0
 		, CVarAnimNodeMotionMatchingDrawCurResult.GetValueOnAnyThread() > 0
-		#endif // WITH_EDITORONLY_DATA && ENABLE_ANIM_DEBUG
+		#endif // ENABLE_ANIM_DEBUG
 	);
 
 	// If a new pose is requested, blend into the new asset via BlendStackNode
diff --git a/Engine/Plugins/Experimental/Animation/PoseSearch/Source/Runtime/Public/PoseSearch/AnimNode_MotionMatching.h b/Engine/Plugins/Experimental/Animation/PoseSearch/Source/Runtime/Public/PoseSearch/AnimNode_MotionMatching.h
--- a/Engine/Plugins/Experimental/Animation/PoseSearch/Source/Runtime/Public/PoseSearch/AnimNode_MotionMatching.h
+++ b/Engine/Plugins/Experimental/Animation/PoseSearch/Source/Runtime/Public/PoseSearch/AnimNode_MotionMatching.h
@@ -68,10 +68,6 @@ private:
 	UPROPERTY(EditAnywhere, Category = Settings, meta = (PinHiddenByDefault))
 	bool bResetOnBecomingRelevant = true;
 
-	// If set to true, the continuing pose will be invalidated. This is useful if you want to force a re-selection of the animation segment instead of continuing with the previous segment, even if it has a better score.
-	UPROPERTY(EditAnywhere, Category = Settings, meta = (PinHiddenByDefault))
-	bool bForceInterrupt = false;
-
 	// If set to false, the motion matching node will perform a search only if the continuing pose is invalid. This is useful if you want to stagger searches of different nodes for performance reasons
 	UPROPERTY(EditAnywhere, Category = Settings, meta = (PinHiddenByDefault))
 	bool bShouldSearch = true;
diff --git a/Engine/Source/Runtime/Engine/Private/Animation/AnimInstanceProxy.cpp b/Engine/Source/Runtime/Engine/Private/Animation/AnimInstanceProxy.cpp
--- a/Engine/Source/Runtime/Engine/Private/Animation/AnimInstanceProxy.cpp
+++ b/Engine/Source/Runtime/Engine/Private/Animation/AnimInstanceProxy.cpp
@@ -643,6 +643,7 @@ void FAnimInstanceProxy::PostUpdate(UAnimInstance* InAnimInstance) const
 		switch (DebugItem.ItemType)
 		{
 			case EDrawDebugItemType::OnScreenMessage: GEngine->AddOnScreenDebugMessage(INDEX_NONE, 0.f, DebugItem.Color, DebugItem.Message, false, DebugItem.TextScale); break;
+			case EDrawDebugItemType::InWorldMessage: DrawDebugString(InAnimInstance->GetSkelMeshComponent()->GetWorld(), DebugItem.StartLoc, DebugItem.Message, InAnimInstance->GetSkelMeshComponent()->GetOwner(), DebugItem.Color, DebugItem.LifeTime, false /*bDrawShadow*/, DebugItem.TextScale.X); break;
 			case EDrawDebugItemType::DirectionalArrow: DrawDebugDirectionalArrow(InAnimInstance->GetSkelMeshComponent()->GetWorld(), DebugItem.StartLoc, DebugItem.EndLoc.Value, DebugItem.Size, DebugItem.Color, DebugItem.bPersistentLines, DebugItem.LifeTime, DebugItem.DepthPriority, DebugItem.Thickness); break;
 			case EDrawDebugItemType::Sphere: DrawDebugSphere(InAnimInstance->GetSkelMeshComponent()->GetWorld(), DebugItem.Center, DebugItem.Radius, DebugItem.Segments, DebugItem.Color, DebugItem.bPersistentLines, DebugItem.LifeTime, DebugItem.DepthPriority, DebugItem.Thickness); break;
 			case EDrawDebugItemType::Line: DrawDebugLine(InAnimInstance->GetSkelMeshComponent()->GetWorld(), DebugItem.StartLoc, DebugItem.EndLoc.Value, DebugItem.Color, DebugItem.bPersistentLines, DebugItem.LifeTime, DebugItem.DepthPriority, DebugItem.Thickness); break;
@@ -2121,6 +2122,19 @@ void FAnimInstanceProxy::AnimDrawDebugOnScreenMessage(const FString& DebugMessag
 	QueuedDrawDebugItems.Add(DrawDebugItem);
 }
 
+void FAnimInstanceProxy::AnimDrawDebugInWorldMessage(const FString& DebugMessage, const FVector& TextLocation, const FColor& Color, float TextScale)
+{
+	FQueuedDrawDebugItem DrawDebugItem;
+
+	DrawDebugItem.ItemType = EDrawDebugItemType::InWorldMessage;
+	DrawDebugItem.Message = DebugMessage;
+	DrawDebugItem.StartLoc = TextLocation;
+	DrawDebugItem.Color = Color;
+	DrawDebugItem.TextScale.X = TextScale;
+
+	QueuedDrawDebugItems.Add(DrawDebugItem);
+}
+
 void FAnimInstanceProxy::AnimDrawDebugDirectionalArrow(const FVector& LineStart, const FVector& LineEnd, float ArrowSize, const FColor& Color, bool bPersistentLines, float LifeTime, float Thickness, ESceneDepthPriorityGroup DepthPriority)
 {
 	FQueuedDrawDebugItem DrawDebugItem;
diff --git a/Engine/Source/Runtime/Engine/Public/Animation/AnimInstanceProxy.h b/Engine/Source/Runtime/Engine/Public/Animation/AnimInstanceProxy.h
--- a/Engine/Source/Runtime/Engine/Public/Animation/AnimInstanceProxy.h
+++ b/Engine/Source/Runtime/Engine/Public/Animation/AnimInstanceProxy.h
@@ -85,6 +85,7 @@ namespace EDrawDebugItemType
 		Point,
 		Circle,
 		Cone,
+		InWorldMessage
 	};
 }
 
@@ -459,6 +460,7 @@ public:
 	TArray<FQueuedDrawDebugItem> QueuedDrawDebugItems;
 
 	void AnimDrawDebugOnScreenMessage(const FString& DebugMessage, const FColor& Color, const FVector2D& TextScale = FVector2D::UnitVector, ESceneDepthPriorityGroup DepthPriority = SDPG_World);
+	void AnimDrawDebugInWorldMessage(const FString& DebugMessage, const FVector& TextLocation, const FColor& Color, float TextScale);
 	void AnimDrawDebugLine(const FVector& StartLoc, const FVector& EndLoc, const FColor& Color, bool bPersistentLines = false, float LifeTime = -1.f, float Thickness = 0.f, ESceneDepthPriorityGroup DepthPriority = SDPG_World);
 	void AnimDrawDebugDirectionalArrow(const FVector& LineStart, const FVector& LineEnd, float ArrowSize, const FColor& Color, bool bPersistentLines = false, float LifeTime = -1.f, float Thickness = 0.f, ESceneDepthPriorityGroup DepthPriority = SDPG_World);
 	void AnimDrawDebugSphere(const FVector& Center, float Radius, int32 Segments, const FColor& Color, bool bPersistentLines = false, float LifeTime = -1.f, float Thickness = 0.f, ESceneDepthPriorityGroup DepthPriority = SDPG_World);
@@ -469,6 +471,7 @@ public:
 	void AnimDrawDebugCone(const FVector& Center, float Radius, const FVector& Direction, float AngleWidth, float AngleHeight, int32 Segments, const FColor& Color, bool bPersistentLines = false, float LifeTime = -1.f, ESceneDepthPriorityGroup DepthPriority = SDPG_World, float Thickness = 0.f);
 #else
 	void AnimDrawDebugOnScreenMessage(const FString& DebugMessage, const FColor& Color, const FVector2D& TextScale = FVector2D::UnitVector, ESceneDepthPriorityGroup DepthPriority = SDPG_World) {}
+	void AnimDrawDebugInWorldMessage(const FString& DebugMessage, const FVector& TextLocation, const FColor& Color, float TextScale) {}
 	void AnimDrawDebugLine(const FVector& StartLoc, const FVector& EndLoc, const FColor& Color, bool bPersistentLines = false, float LifeTime = -1.f, float Thickness = 0.f, ESceneDepthPriorityGroup DepthPriority = SDPG_World) {}
 	void AnimDrawDebugDirectionalArrow(const FVector& LineStart, const FVector& LineEnd, float ArrowSize, const FColor& Color, bool bPersistentLines = false, float LifeTime = -1.f, float Thickness = 0.f, ESceneDepthPriorityGroup DepthPriority = SDPG_World) {}
 	void AnimDrawDebugSphere(const FVector& Center, float Radius, int32 Segments, const FColor& Color, bool bPersistentLines = false, float LifeTime = -1.f, float Thickness = 0.f, ESceneDepthPriorityGroup DepthPriority = SDPG_World) {}
