Fix r.SkinCache.RecomputeTangents=0 crashes when rendering skeletal mesh with the skin cache

diff --git a/Engine/Source/Runtime/Engine/Private/GPUSkinCache.cpp b/Engine/Source/Runtime/Engine/Private/GPUSkinCache.cpp
--- a/Engine/Source/Runtime/Engine/Private/GPUSkinCache.cpp
+++ b/Engine/Source/Runtime/Engine/Private/GPUSkinCache.cpp
@@ -1522,7 +1522,7 @@ bool FGPUSkinCache::ProcessEntry(
 	// Try to allocate a new entry
 	if (!InOutEntry)
 	{
-		bool WithTangents = RecomputeTangentsMode > 0;
+		const bool WithTangents = true;
 		int32 TotalNumVertices = VertexFactory->GetNumVertices();
 		
 		// IntermediateTangents buffer is needed if mesh has at least one section using vertex color as recompute tangents blending mask
