Fix transmission bug where transmission lut is incorrectly built based on unit scale.

diff --git a/Engine/Source/Runtime/Engine/Private/Rendering/BurleyNormalizedSSS.cpp b/Engine/Source/Runtime/Engine/Private/Rendering/BurleyNormalizedSSS.cpp
--- a/Engine/Source/Runtime/Engine/Private/Rendering/BurleyNormalizedSSS.cpp
+++ b/Engine/Source/Runtime/Engine/Private/Rendering/BurleyNormalizedSSS.cpp
@@ -239,12 +239,12 @@ void ComputeMirroredBSSSKernel(FLinearColor* TargetBuffer, uint32 TargetBufferSi
 void ComputeTransmissionProfileBurley(FLinearColor* TargetBuffer, uint32 TargetBufferSize, 
 									FLinearColor FalloffColor, float ExtinctionScale,
 									FLinearColor SurfaceAlbedo, FLinearColor DiffuseMeanFreePathInMm,
-									float InUnitToMm, FLinearColor TransmissionTintColor)
+									float WorldUnitScale, FLinearColor TransmissionTintColor)
 {
 	check(TargetBuffer);
 	check(TargetBufferSize > 0);
 
-
+	// Unit scale should be independent to the base unit.
 	// Example of scaling
 	// ----------------------------------------
 	// DistanceCm * UnitScale * CmToMm = Value (mm)
@@ -253,12 +253,11 @@ void ComputeTransmissionProfileBurley(FLinearColor* TargetBuffer, uint32 TargetB
 	//   1          1.0         10     =  10mm
 	//   1         10.0         10     = 100mm
 
-	// It seems there is an inconsistency here: either InUnitToMm should in fact be MmToUnit or the division need to be converted to a multiply
-	// This should be corrected as this makes the transmission scaling incorrect
-	const float UnitToMm = 1.f / InUnitToMm;
+	const float SubsurfaceScatteringUnitInCm = 0.1f;
+	const float UnitScale = WorldUnitScale / SubsurfaceScatteringUnitInCm;
 
 	static float MaxTransmissionProfileDistance = 5.0f; // See SSSS_MAX_TRANSMISSION_PROFILE_DISTANCE in TransmissionCommon.ush
-
+	static float CmToMm = 10.0f;
 	//assuming that the volume albedo is the same to the surface albedo for transmission.
 	FVector ScalingFactor = GetSearchLightDiffuseScalingFactor(SurfaceAlbedo);
 
@@ -266,8 +265,8 @@ void ComputeTransmissionProfileBurley(FLinearColor* TargetBuffer, uint32 TargetB
 
 	for (uint32 i = 0; i < TargetBufferSize; ++i)
 	{
-		const float DistanceInMm = i * InvSize * MaxTransmissionProfileDistance * UnitToMm;
-		const float OffsetInMm = ProfileRadiusOffset * UnitToMm;
+		const float DistanceInMm = i * InvSize * (MaxTransmissionProfileDistance * CmToMm) * UnitScale;
+		const float OffsetInMm = (ProfileRadiusOffset * CmToMm) * UnitScale;
 
 
 		FLinearColor TransmissionProfile = Burley_TransmissionProfile(DistanceInMm + OffsetInMm, SurfaceAlbedo, ScalingFactor, DiffuseMeanFreePathInMm);
diff --git a/Engine/Source/Runtime/Engine/Private/Rendering/SubsurfaceProfile.cpp b/Engine/Source/Runtime/Engine/Private/Rendering/SubsurfaceProfile.cpp
--- a/Engine/Source/Runtime/Engine/Private/Rendering/SubsurfaceProfile.cpp
+++ b/Engine/Source/Runtime/Engine/Private/Rendering/SubsurfaceProfile.cpp
@@ -489,8 +489,6 @@ void FSubsurfaceProfileTexture::CreateTexture(FRHICommandListImmediate& RHICmdLi
 		// to depricate {scatter radius, falloff} in favor of albedo/MFP, need the `subsurface color` as tint.
 		TextureRow[SSSS_TINT_SCALE_OFFSET] = Data.Tint;
 		TextureRow[SSSS_TINT_SCALE_OFFSET].A = EncodeWorldUnitScale(Data.WorldUnitScale);
-		
-		const float UnitToCm = Data.WorldUnitScale;
 
 		FLinearColor DifffuseMeanFreePathInMm;
 
@@ -544,7 +542,7 @@ void FSubsurfaceProfileTexture::CreateTexture(FRHICommandListImmediate& RHICmdLi
 		}
 
 		ComputeTransmissionProfileBurley(&TextureRow[BSSS_TRANSMISSION_PROFILE_OFFSET], BSSS_TRANSMISSION_PROFILE_SIZE, 
-			Data.FalloffColor, Data.ExtinctionScale, Data.SurfaceAlbedo, DifffuseMeanFreePathInMm, UnitToCm * CmToMm, Data.TransmissionTintColor);
+			Data.FalloffColor, Data.ExtinctionScale, Data.SurfaceAlbedo, DifffuseMeanFreePathInMm, Data.WorldUnitScale, Data.TransmissionTintColor);
 
 		// could be lower than 1 (but higher than 0) to range compress for better quality (for 8 bit)
 		const float TableMaxRGB = 1.0f;
