BlueprintGeneratedClass: When serializing into a cooked package, harvest SoftObjectPaths from the UBlueprint that are needed at runtime.
#rnx

diff --git a/Engine/Source/Runtime/Engine/Private/BlueprintGeneratedClass.cpp b/Engine/Source/Runtime/Engine/Private/BlueprintGeneratedClass.cpp
--- a/Engine/Source/Runtime/Engine/Private/BlueprintGeneratedClass.cpp
+++ b/Engine/Source/Runtime/Engine/Private/BlueprintGeneratedClass.cpp
@@ -1,6 +1,8 @@
 // Copyright Epic Games, Inc. All Rights Reserved.
 
 #include "Engine/BlueprintGeneratedClass.h"
+
+#include "Containers/RingBuffer.h"
 #include "Blueprint/BlueprintSupport.h"
 #include "HAL/IConsoleManager.h"
 #include "EngineLogs.h"
@@ -2160,6 +2162,47 @@ ENGINE_API int32 IncrementUberGraphSerialNumber()
 }
 #endif//VALIDATE_UBER_GRAPH_PERSISTENT_FRAME
 
+
+#if WITH_EDITORONLY_DATA
+/** An Archive that records all of the imported packages from a tree of exports. */
+class FImportExportCollector : public FArchiveUObject
+{
+public:
+	explicit FImportExportCollector(UPackage* InRootPackage);
+
+	/**
+	 * Mark that a given export (e.g. the export that is doing the collecting) should not be explored
+	 * if encountered again. Prevents infinite recursion when the collector is constructed and called during
+	 * Serialize.
+	 */
+	void AddExportToIgnore(UObject* Export);
+	/**
+	 * Serialize the given object, following its object references to find other imports and exports,
+	 * and recursively serialize any new exports that it references.
+	 */
+	void SerializeObjectAndReferencedExports(UObject* RootObject);
+	/** Restore the collector to empty. */
+	void Reset();
+	const TSet<UObject*>& GetExports() const { return Exports; }
+	const TMap<FSoftObjectPath, ESoftObjectPathCollectType>& GetImports() const { return Imports; }
+	const TMap<FName, ESoftObjectPathCollectType>& GetImportedPackages() const { return ImportedPackages; }
+
+	virtual FArchive& operator<<(UObject*& Obj) override;
+	virtual FArchive& operator<<(FSoftObjectPath& Value) override;
+
+private:
+	void AddImport(const FSoftObjectPath& Path, ESoftObjectPathCollectType CollectType);
+	ESoftObjectPathCollectType Union(ESoftObjectPathCollectType A, ESoftObjectPathCollectType B);
+
+	TSet<UObject*> Exports;
+	TRingBuffer<UObject*> ExportsExploreQueue;
+	TMap<FSoftObjectPath, ESoftObjectPathCollectType> Imports;
+	TMap<FName, ESoftObjectPathCollectType> ImportedPackages;
+	UPackage* RootPackage;
+	FName RootPackageName;
+};
+#endif
+
 void UBlueprintGeneratedClass::Serialize(FArchive& Ar)
 {
 #if VALIDATE_UBER_GRAPH_PERSISTENT_FRAME
@@ -2203,6 +2246,35 @@ void UBlueprintGeneratedClass::Serialize(FArchive& Ar)
 		UberGraphFramePointerProperty_DEPRECATED = nullptr;
 	}
 #endif
+
+#if WITH_EDITORONLY_DATA
+	if (Ar.IsSaving() && Ar.IsCooking() && Ar.IsObjectReferenceCollector())
+	{
+		// The UBlueprint class and its subobjects have imports that we need to include at runtime,
+		// but we exclude the Blueprint and its subobject from the saved cooked package.
+		// Find all imported packages from the Blueprint and its subobjects and declare them as
+		// used-in-game imports of the cooked package by serializing them as SoftObjectPaths.
+		FImportExportCollector Collector(this->GetPackage());
+		Collector.SetCookData(Ar.GetCookData());
+		Collector.AddExportToIgnore(this);
+		Collector.SerializeObjectAndReferencedExports(ClassGeneratedBy);
+		for (const TPair<FName, ESoftObjectPathCollectType>& Pair : Collector.GetImportedPackages())
+		{
+			if (Pair.Value != ESoftObjectPathCollectType::AlwaysCollect)
+			{
+				continue;
+			}
+			FName PackageName = Pair.Key;
+			if (FPackageName::IsScriptPackage(WriteToString<256>(PackageName)))
+			{
+				// Ignore native imports; we don't need to mark them for cooking
+				continue;
+			}
+			FSoftObjectPath PackageSoftPath(PackageName, NAME_None, FString());
+			Ar << PackageSoftPath;
+		}
+	}
+#endif
 }
 
 void UBlueprintGeneratedClass::GetLifetimeBlueprintReplicationList(TArray<FLifetimeProperty>& OutLifetimeProps) const
@@ -2570,5 +2642,100 @@ void UBlueprintGeneratedClass::PurgeCookedMetaData()
 		CookedMetaDataUtil::PurgeCookedMetaData<UClassCookedMetaData>(CachedCookedMetaDataPtr);
 	}
 }
+
+FImportExportCollector::FImportExportCollector(UPackage* InRootPackage)
+	: RootPackage(InRootPackage)
+	, RootPackageName(InRootPackage->GetFName())
+{
+	ArIsObjectReferenceCollector = true;
+	ArIsModifyingWeakAndStrongReferences = true;
+	SetIsSaving(true);
+	SetIsPersistent(true);
+}
+
+void FImportExportCollector::Reset()
+{
+	Exports.Reset();
+	Imports.Reset();
+}
+
+void FImportExportCollector::AddExportToIgnore(UObject* Export)
+{
+	Exports.Add(Export);
+}
+
+void FImportExportCollector::SerializeObjectAndReferencedExports(UObject* RootObject)
+{
+	*this << RootObject;
+	while (!ExportsExploreQueue.IsEmpty())
+	{
+		UObject* Export = ExportsExploreQueue.PopFrontValue();
+		Export->Serialize(*this);
+	}
+}
+
+FArchive& FImportExportCollector::operator<<(UObject*& Obj)
+{
+	if (!Obj)
+	{
+		return *this;
+	}
+	UPackage* Package = Obj->GetPackage();
+	if (!Package)
+	{
+		return *this;
+	}
+	if (Package != RootPackage)
+	{
+		AddImport(FSoftObjectPath(Obj), ESoftObjectPathCollectType::AlwaysCollect);
+		return *this;
+	}
+
+	bool bAlreadyExists;
+	Exports.Add(Obj, &bAlreadyExists);
+	if (bAlreadyExists)
+	{
+		return *this;
+	}
+	ExportsExploreQueue.Add(Obj);
+	return *this;
+}
+
+FArchive& FImportExportCollector::operator<<(FSoftObjectPath& Value)
+{
+	FName CurrentPackage;
+	FName PropertyName;
+	ESoftObjectPathCollectType CollectType;
+	ESoftObjectPathSerializeType SerializeType;
+	FSoftObjectPathThreadContext& ThreadContext = FSoftObjectPathThreadContext::Get();
+	ThreadContext.GetSerializationOptions(CurrentPackage, PropertyName, CollectType, SerializeType, this);
+
+	if (CollectType != ESoftObjectPathCollectType::NeverCollect && CollectType != ESoftObjectPathCollectType::NonPackage)
+	{
+		FName PackageName = Value.GetLongPackageFName();
+		if (PackageName != RootPackageName && !PackageName.IsNone())
+		{
+			AddImport(Value, CollectType);
+		}
+	}
+	return *this;
+}
+
+void FImportExportCollector::AddImport(const FSoftObjectPath& Path, ESoftObjectPathCollectType CollectType)
+{
+	ESoftObjectPathCollectType& ExistingImport = Imports.FindOrAdd(
+		Path, ESoftObjectPathCollectType::EditorOnlyCollect);
+	ExistingImport = Union(ExistingImport, CollectType);
+
+	ESoftObjectPathCollectType& ExistingPackage = ImportedPackages.FindOrAdd(
+		Path.GetLongPackageFName(), ESoftObjectPathCollectType::EditorOnlyCollect);
+	ExistingPackage = Union(ExistingPackage, CollectType);
+}
+
+ESoftObjectPathCollectType FImportExportCollector::Union(ESoftObjectPathCollectType A, ESoftObjectPathCollectType B)
+{
+	return static_cast<ESoftObjectPathCollectType>(FMath::Max(static_cast<int>(A), static_cast<int>(B)));
+}
+
 #endif //if WITH_EDITORONLY_DATA
 
