Network
* Added FNetContext for user code to tell if inside a real network RPC or not
* Added support to FNetContext::IsInsideNetRPC in both Iris and legacy

#rnx

diff --git a/Engine/Source/Runtime/Engine/Private/DataReplication.cpp b/Engine/Source/Runtime/Engine/Private/DataReplication.cpp
--- a/Engine/Source/Runtime/Engine/Private/DataReplication.cpp
+++ b/Engine/Source/Runtime/Engine/Private/DataReplication.cpp
@@ -11,6 +11,7 @@
 #include "Misc/MemStack.h"
 #include "Misc/ScopeExit.h"
 #include "Net/Core/Trace/Private/NetTraceInternal.h"
+#include "Net/Core/Misc/NetContext.h"
 #include "Net/NetworkProfiler.h"
 #include "Engine/PackageMapClient.h"
 #include "Net/RepLayout.h"
@@ -1221,22 +1222,22 @@ struct FScopedRPCTimingTracker
 		}
 	}
 
-	UNetConnection* Connection;
-	UFunction* Function;
+	UNetConnection* Connection = nullptr;
+	UFunction* Function = nullptr;
 	FRPCDoSDetection* RPCDoS = nullptr;
-	double StartTime;
+	double StartTime = 0.0;
 
 	static TArray<FScopedRPCTimingTracker*> ActiveTrackers;
 };
 
 TArray<FScopedRPCTimingTracker*> FScopedRPCTimingTracker::ActiveTrackers;
 
+/** Return the list of UFunctions currently tracked or  */
 ENGINE_API TArray<UFunction*> FindScopedRPCTrackers(UNetConnection* Connection = nullptr)
 {
 	TArray<UFunction*> FuncList;
-	for (int32 Idx = 0; Idx < FScopedRPCTimingTracker::ActiveTrackers.Num(); Idx++)
+	for (const FScopedRPCTimingTracker* TestTracker : FScopedRPCTimingTracker::ActiveTrackers)
 	{
-		const FScopedRPCTimingTracker* TestTracker = FScopedRPCTimingTracker::ActiveTrackers[Idx];
 		if (TestTracker && (Connection == nullptr || TestTracker->Connection == Connection))
 		{
 			FuncList.Add(TestTracker->Function);
@@ -1373,8 +1374,12 @@ bool FObjectReplicator::ReceivedRPC(FNetBitReader& Reader, const FReplicationFla
 			// Reset errors from replay driver
 			RPC_ResetLastFailedReason();
 
-			// Call the function.
-			Object->ProcessEvent(Function, Parms);
+
+			{
+				UE::Net::FScopedNetContextRPC CallingRPC;
+				// Call the function.
+				Object->ProcessEvent(Function, Parms);
+			}
 		}
 
 		// Destroy the parameters.
diff --git a/Engine/Source/Runtime/Experimental/Iris/Core/Private/Iris/ReplicationSystem/NetBlob/NetRPC.cpp b/Engine/Source/Runtime/Experimental/Iris/Core/Private/Iris/ReplicationSystem/NetBlob/NetRPC.cpp
--- a/Engine/Source/Runtime/Experimental/Iris/Core/Private/Iris/ReplicationSystem/NetBlob/NetRPC.cpp
+++ b/Engine/Source/Runtime/Experimental/Iris/Core/Private/Iris/ReplicationSystem/NetBlob/NetRPC.cpp
@@ -4,6 +4,7 @@
 
 #include "Iris/Core/BitTwiddling.h"
 #include "Net/Core/Trace/NetTrace.h"
+#include "Net/Core/Misc/NetContext.h"
 #include "Iris/ReplicationSystem/ObjectReferenceCache.h"
 #include "Iris/ReplicationSystem/ObjectReplicationBridge.h"
 #include "Iris/Serialization/NetReferenceCollector.h"
@@ -473,10 +474,10 @@ void FNetRPC::CallFunction(FNetSerializationContext& Context)
 			UE_LOG(LogIrisRpc, Verbose, TEXT("Calling %s RPC function %s for %s : %s"), (Function->FunctionFlags & FUNC_NetReliable ? TEXT("reliable") : TEXT("unreliable")), ToCStr(Function->GetName()), *NetObjectReference.ToString(), *Object->GetFullName());
 		}
 	}
-
 	// Call the function
 	if (Function->ParmsSize == 0)
 	{
+		UE::Net::FScopedNetContextRPC CallingRPC;
 		Object->ProcessEvent(const_cast<UFunction*>(Function), nullptr);
 	}
 	else
@@ -510,7 +511,10 @@ void FNetRPC::CallFunction(FNetSerializationContext& Context)
 		}
 		
 		FReplicationStateOperations::Dequantize(Context, FunctionParameters, QuantizedBlobState.Get(), BlobDescriptor);
-		Object->ProcessEvent(const_cast<UFunction*>(Function), FunctionParameters);
+		{
+			UE::Net::FScopedNetContextRPC CallingRPC;
+			Object->ProcessEvent(const_cast<UFunction*>(Function), FunctionParameters);
+		}
 
 		if (!EnumHasAnyFlags(BlobDescriptor->Traits, EReplicationStateTraits::IsSourceTriviallyDestructible))
 		{
diff --git a/Engine/Source/Runtime/Net/Core/Private/Net/Core/Misc/NetContext.cpp b/Engine/Source/Runtime/Net/Core/Private/Net/Core/Misc/NetContext.cpp
new file mode 100644
--- /dev/null
+++ b/Engine/Source/Runtime/Net/Core/Private/Net/Core/Misc/NetContext.cpp
@@ -0,0 +1,22 @@
+// Copyright Epic Games, Inc. All Rights Reserved.
+
+#include "Net/Core/Misc/NetContext.h"
+
+namespace UE::Net
+{
+
+bool FNetContext::bIsInRPCStack = false;
+
+FScopedNetContextRPC::FScopedNetContextRPC()
+{
+	FNetContext::bIsInRPCStack = true;
+}
+
+FScopedNetContextRPC::~FScopedNetContextRPC()
+{
+	checkSlow(FNetContext::bIsInRPCStack); // Cannot have two remote rpc stacked on top of each other
+	FNetContext::bIsInRPCStack = false;
+}
+
+} // end namespace UE::Net
+
diff --git a/Engine/Source/Runtime/Net/Core/Public/Net/Core/Misc/NetContext.h b/Engine/Source/Runtime/Net/Core/Public/Net/Core/Misc/NetContext.h
new file mode 100644
--- /dev/null
+++ b/Engine/Source/Runtime/Net/Core/Public/Net/Core/Misc/NetContext.h
@@ -0,0 +1,48 @@
+// Copyright Epic Games, Inc. All Rights Reserved.
+
+#pragma once
+
+namespace UE::Net::Private
+{
+	class FNetRPC;
+}
+
+namespace UE::Net
+{
+
+/** Stateless class that provides misc network context information */
+class FNetContext
+{
+public:
+	
+	/** 
+	* Tells true if we are inside code executed from a remote RPC.
+	* Will be false if the RPC is executed locally.
+	*/
+	static bool IsInsideNetRPC() { return bIsInRPCStack; }
+
+private:
+
+	friend class FScopedNetContextRPC;
+
+	FNetContext() = delete;
+	~FNetContext() = delete;
+
+	static bool bIsInRPCStack;
+};
+
+
+/** Used by friendly class to set the right network context */
+class NETCORE_API FScopedNetContextRPC
+{
+private:
+
+	// Only friend class are allowed to push the net context
+	friend Private::FNetRPC;
+	friend class FObjectReplicator;
+	
+	FScopedNetContextRPC();
+	~FScopedNetContextRPC();
+};
+
+} // end namespace UE::Net
