- Properly set some defines that weren't handled in the new HLSL generator code path;
- FExpressionQualitySwitch now properly sets the bits in FMaterialCachedExpressionData::QualityLevelUsed for used material quality levels;
- Add a GenerateHLSLExpression override to UMaterialExpressionChannelMaskParameter. The implementation of its parent class was used and led to incorrect preshader values. The override fixes this problem;
- Add a version number to FMaterialHLSLGenerator for invalidating materials using the new generator as it is being developed.

diff --git a/Engine/Source/Runtime/Engine/Classes/Materials/MaterialExpressionChannelMaskParameter.h b/Engine/Source/Runtime/Engine/Classes/Materials/MaterialExpressionChannelMaskParameter.h
--- a/Engine/Source/Runtime/Engine/Classes/Materials/MaterialExpressionChannelMaskParameter.h
+++ b/Engine/Source/Runtime/Engine/Classes/Materials/MaterialExpressionChannelMaskParameter.h
@@ -39,6 +39,8 @@ class UMaterialExpressionChannelMaskParameter : public UMaterialExpressionVector
 	
 	virtual bool IsInputConnectionRequired(int32 InputIndex) const override {return true;}
 	virtual uint32 GetInputType(int32 InputIndex) override {return MCT_Float4;}
+
+	virtual bool GenerateHLSLExpression(FMaterialHLSLGenerator& Generator, UE::HLSLTree::FScope& Scope, int32 OutputIndex, UE::HLSLTree::FExpression const*& OutExpression) const override;
 #endif
 
 	virtual bool IsUsedAsChannelMask() const override {return true;}
diff --git a/Engine/Source/Runtime/Engine/Private/HLSLTree/HLSLTreeCommon.cpp b/Engine/Source/Runtime/Engine/Private/HLSLTree/HLSLTreeCommon.cpp
--- a/Engine/Source/Runtime/Engine/Private/HLSLTree/HLSLTreeCommon.cpp
+++ b/Engine/Source/Runtime/Engine/Private/HLSLTree/HLSLTreeCommon.cpp
@@ -7,6 +7,8 @@
 #include "Shader/Preshader.h"
 #include "SceneTypes.h"
 #include "MaterialShared.h"
+#include "MaterialHLSLTree.h"
+#include "MaterialCachedData.h"
 
 namespace UE::HLSLTree
 {
@@ -1064,7 +1066,35 @@ FExpressionQualitySwitch::FExpressionQualitySwitch(TConstArrayView<const FExpres
 
 bool FExpressionQualitySwitch::IsInputActive(const FEmitContext& Context, int32 Index) const
 {
-	return Context.TargetParameters.IsGenericTarget() || Index == (int32)Context.Material->GetQualityLevel();
+	constexpr int32 DefaultInputIndex = (int32)EMaterialQualityLevel::Num;
+
+	if (Context.TargetParameters.IsGenericTarget())
+	{
+		return Index == DefaultInputIndex || Input[Index] != Input[DefaultInputIndex];
+	}
+	else
+	{
+		return Index == (int32)Context.Material->GetQualityLevel();
+	}
+}
+
+bool FExpressionQualitySwitch::PrepareValue(FEmitContext& Context, FEmitScope& Scope, const FRequestedType& RequestedType, FPrepareValueResult& OutResult) const
+{
+	FMaterialCachedExpressionData* CachedExpressionData = Context.FindData<Material::FEmitData>().CachedExpressionData;
+	if (CachedExpressionData)
+	{
+		constexpr int32 DefaultInputIndex = (int32)EMaterialQualityLevel::Num;
+
+		for (int32 Index = 0; Index < DefaultInputIndex; ++Index)
+		{
+			if (Input[Index] != Input[DefaultInputIndex])
+			{
+				CachedExpressionData->QualityLevelsUsed[Index] = true;
+			}
+		}
+	}
+
+	return FExpressionSwitchBase::PrepareValue(Context, Scope, RequestedType, OutResult);
 }
 
 FExpressionShaderStageSwitch::FExpressionShaderStageSwitch(TConstArrayView<const FExpression*> InInputs)
diff --git a/Engine/Source/Runtime/Engine/Private/Materials/MaterialExpressionHLSL.cpp b/Engine/Source/Runtime/Engine/Private/Materials/MaterialExpressionHLSL.cpp
--- a/Engine/Source/Runtime/Engine/Private/Materials/MaterialExpressionHLSL.cpp
+++ b/Engine/Source/Runtime/Engine/Private/Materials/MaterialExpressionHLSL.cpp
@@ -18,6 +18,7 @@
 #include "Materials/MaterialExpressionConstantBiasScale.h"
 #include "Materials/MaterialExpressionShadingModel.h"
 #include "Materials/MaterialExpressionParameter.h"
+#include "Materials/MaterialExpressionChannelMaskParameter.h"
 #include "Materials/MaterialExpressionCollectionParameter.h"
 #include "Materials/MaterialParameterCollection.h"
 #include "Materials/MaterialExpressionDynamicParameter.h"
@@ -337,7 +338,8 @@ bool UMaterialExpressionShadingPathSwitch::GenerateHLSLExpression(FMaterialHLSLG
 bool UMaterialExpressionQualitySwitch::GenerateHLSLExpression(FMaterialHLSLGenerator& Generator, UE::HLSLTree::FScope& Scope, int32 OutputIndex, UE::HLSLTree::FExpression const*& OutExpression) const
 {
 	using namespace UE::HLSLTree;
-	if (!Default.GetTracedInput().Expression)
+	const FExpressionInput DefaultTracedInput = Default.GetTracedInput();
+	if (!DefaultTracedInput.Expression)
 	{
 		return Generator.Error(TEXT("Missing default input"));
 	}
@@ -349,7 +351,8 @@ bool UMaterialExpressionQualitySwitch::GenerateHLSLExpression(FMaterialHLSLGener
 	{
 		const FExpression* Expression = nullptr;
 		const FExpressionInput& QualityInput = Inputs[Index];
-		if (QualityInput.GetTracedInput().Expression)
+		const FExpressionInput QualityTracedInput = QualityInput.GetTracedInput();
+		if (QualityTracedInput.Expression && (QualityTracedInput.Expression != DefaultTracedInput.Expression || QualityTracedInput.OutputIndex != DefaultTracedInput.OutputIndex))
 		{
 			Expression = QualityInput.AcquireHLSLExpression(Generator, Scope);
 		}
@@ -409,6 +412,31 @@ bool UMaterialExpressionParameter::GenerateHLSLExpression(FMaterialHLSLGenerator
 	return true;
 }
 
+bool UMaterialExpressionChannelMaskParameter::GenerateHLSLExpression(FMaterialHLSLGenerator& Generator, UE::HLSLTree::FScope& Scope, int32 OutputIndex, UE::HLSLTree::FExpression const*& OutExpression) const
+{
+	using namespace UE::HLSLTree;
+
+	if (!Input.GetTracedInput().Expression)
+	{
+		return Generator.Error(TEXT("Missing mask input"));
+	}
+
+	const FExpression* InputExpression = Input.AcquireHLSLExpression(Generator, Scope);
+	if (!InputExpression)
+	{
+		return Generator.Error(TEXT("Failed to generate expression for mask input"));
+	}
+
+	const FExpression* MaskExpression = nullptr;
+	if (!UMaterialExpressionParameter::GenerateHLSLExpression(Generator, Scope, OutputIndex, MaskExpression))
+	{
+		return false;
+	}
+
+	OutExpression = Generator.GetTree().NewDot(InputExpression, MaskExpression);
+	return true;
+}
+
 bool UMaterialExpressionCollectionParameter::GenerateHLSLExpression(FMaterialHLSLGenerator& Generator, UE::HLSLTree::FScope& Scope, int32 OutputIndex, UE::HLSLTree::FExpression const*& OutExpression) const
 {
 	int32 ParameterIndex = INDEX_NONE;
diff --git a/Engine/Source/Runtime/Engine/Private/Materials/MaterialHLSLEmitter.cpp b/Engine/Source/Runtime/Engine/Private/Materials/MaterialHLSLEmitter.cpp
--- a/Engine/Source/Runtime/Engine/Private/Materials/MaterialHLSLEmitter.cpp
+++ b/Engine/Source/Runtime/Engine/Private/Materials/MaterialHLSLEmitter.cpp
@@ -383,6 +383,9 @@ static void GetMaterialEnvironment(EShaderPlatform InPlatform,
 	const FMaterial& InMaterial,
 	const UE::HLSLTree::FEmitContext& EmitContext,
 	const FMaterialCompilationOutput& MaterialCompilationOutput,
+	bool bUsesEmissiveColor,
+	bool bUsesAnisotropy,
+	bool bIsFullyRough,
 	FShaderCompilerEnvironment& OutEnvironment)
 {
 	using namespace UE::HLSLTree;
@@ -495,10 +498,10 @@ static void GetMaterialEnvironment(EShaderPlatform InPlatform,
 	OutEnvironment.SetDefine(TEXT("NEEDS_PARTICLE_LOCAL_TO_WORLD"), false);// bUsesParticleLocalToWorld);
 	OutEnvironment.SetDefine(TEXT("NEEDS_PARTICLE_WORLD_TO_LOCAL"), false);// bUsesParticleWorldToLocal);
 	OutEnvironment.SetDefine(TEXT("USES_TRANSFORM_VECTOR"), false);// bUsesTransformVector);
-	OutEnvironment.SetDefine(TEXT("WANT_PIXEL_DEPTH_OFFSET"), false);// bUsesPixelDepthOffset);
-	OutEnvironment.SetDefineAndCompileArgument(TEXT("USES_WORLD_POSITION_OFFSET"), false);// bUsesWorldPositionOffset);
+	OutEnvironment.SetDefine(TEXT("WANT_PIXEL_DEPTH_OFFSET"), MaterialCompilationOutput.bUsesPixelDepthOffset);
+	OutEnvironment.SetDefineAndCompileArgument(TEXT("USES_WORLD_POSITION_OFFSET"), MaterialCompilationOutput.bUsesWorldPositionOffset);
 	OutEnvironment.SetDefineAndCompileArgument(TEXT("USES_DISPLACEMENT"), false);// bUsesDisplacement);
-	OutEnvironment.SetDefine(TEXT("USES_EMISSIVE_COLOR"), false);// bUsesEmissiveColor);
+	OutEnvironment.SetDefine(TEXT("USES_EMISSIVE_COLOR"), bUsesEmissiveColor);
 	// Distortion uses tangent space transform 
 	OutEnvironment.SetDefine(TEXT("USES_DISTORTION"), InMaterial.IsDistorted());
 
@@ -506,8 +509,8 @@ static void GetMaterialEnvironment(EShaderPlatform InPlatform,
 	OutEnvironment.SetDefine(TEXT("MATERIAL_ENABLE_TRANSLUCENCY_CLOUD_FOGGING"), InMaterial.ShouldApplyCloudFogging());
 	OutEnvironment.SetDefine(TEXT("MATERIAL_IS_SKY"), InMaterial.IsSky());
 	OutEnvironment.SetDefine(TEXT("MATERIAL_COMPUTE_FOG_PER_PIXEL"), InMaterial.ComputeFogPerPixel());
-	OutEnvironment.SetDefine(TEXT("MATERIAL_FULLY_ROUGH"), false);// bIsFullyRough || InMaterial.IsFullyRough());
-	OutEnvironment.SetDefine(TEXT("MATERIAL_USES_ANISOTROPY"), false);// bUsesAnisotropy);
+	OutEnvironment.SetDefine(TEXT("MATERIAL_FULLY_ROUGH"), bIsFullyRough);
+	OutEnvironment.SetDefine(TEXT("MATERIAL_USES_ANISOTROPY"), bUsesAnisotropy);
 
 	// Count the number of VTStacks (each stack will allocate a feedback slot)
 	OutEnvironment.SetDefine(TEXT("NUM_VIRTUALTEXTURE_SAMPLES"), EmitMaterialData.VTStacks.Num());
@@ -803,6 +806,9 @@ bool MaterialEmitHLSL(const FMaterialCompileTargetParameters& InCompilerTarget,
 
 	bool bUsesWorldPositionOffset = false;
 	bool bUsesPixelDepthOffset = false;
+	bool bUsesEmissiveColor = false;
+	bool bUsesAnisotropy = false;
+	bool bIsFullyRough = false;
 	bool bUsesDisplacement = false;
 
 	const FTargetParameters TargetParameters(InCompilerTarget.ShaderPlatform, InCompilerTarget.FeatureLevel, InCompilerTarget.TargetPlatform);
@@ -874,9 +880,28 @@ bool MaterialEmitHLSL(const FMaterialCompileTargetParameters& InCompilerTarget,
 			}
 		}
 
-		if (CachedTree->IsAttributeUsed(EmitContext, *EmitResultScope, PixelResultType0, MP_PixelDepthOffset))
+		bUsesPixelDepthOffset = bUsesPixelDepthOffset || CachedTree->IsAttributeUsed(EmitContext, *EmitResultScope, PixelResultType0, MP_PixelDepthOffset);
+		bUsesEmissiveColor = bUsesEmissiveColor || CachedTree->IsAttributeUsed(EmitContext, *EmitResultScope, PixelResultType0, MP_EmissiveColor);
+
+		bUsesAnisotropy = bUsesAnisotropy || (FDataDrivenShaderPlatformInfo::GetSupportsAnisotropicMaterials(InCompilerTarget.ShaderPlatform)
+			&& CachedTree->IsAttributeUsed(EmitContext, *EmitResultScope, PixelResultType0, MP_Anisotropy));
+
+		bIsFullyRough = bIsFullyRough || EmitContext.Material->IsFullyRough();
+		if (!bIsFullyRough)
 		{
-			bUsesPixelDepthOffset = true;
+			const FStructField* RoughnessField = CachedTree->GetMaterialAttributesType()->FindFieldByName(*FMaterialAttributeDefinitionMap::GetAttributeName(MP_Roughness));
+			check(RoughnessField && RoughnessField->GetNumComponents() == 1);
+
+			FRequestedType RequestedType(CachedTree->GetMaterialAttributesType(), false);
+			RequestedType.SetFieldRequested(RoughnessField);
+
+			const EExpressionEvaluation Evaluation = PixelResultType0.GetFieldEvaluation(*EmitResultScope, RequestedType, RoughnessField->ComponentIndex, 1);
+			if (Evaluation == EExpressionEvaluation::Constant)
+			{
+				const FValue ConstantValue = CachedTree->GetResultExpression()->GetValueConstant(EmitContext, *EmitResultScope, RequestedType);
+				check(ConstantValue.GetType().GetFlatFieldType(RoughnessField->FlatFieldIndex) == EValueType::Float1);
+				bIsFullyRough = ConstantValue.Component[RoughnessField->ComponentIndex].Float == 1.f;
+			}
 		}
 
 		if (CachedTree->IsAttributeUsed(EmitContext, *EmitResultScope, PixelResultType0, MP_Displacement))
@@ -964,6 +989,7 @@ bool MaterialEmitHLSL(const FMaterialCompileTargetParameters& InCompilerTarget,
 		EmitContext.EmitStatement(*EmitResultScope, TEXT("Parameters.MaterialAttributes = %;"), EmitResultExpression);
 		EmitMaterialData.EmitInterpolatorStatements(EmitContext, *EmitResultScope);
 
+		// TODO: Do we need to check PrevWorldPositionOffset?
 		bUsesWorldPositionOffset = CachedTree->IsAttributeUsed(EmitContext, *EmitResultScope, VertexResultType, MP_WorldPositionOffset);
 
 		CachedTree->GetTree().EmitShader(EmitContext, VertexCode);
@@ -976,7 +1002,7 @@ bool MaterialEmitHLSL(const FMaterialCompileTargetParameters& InCompilerTarget,
 
 	OutCompilationOutput.bModifiesMeshPosition = bUsesPixelDepthOffset || bUsesWorldPositionOffset || bUsesDisplacement;
 	OutCompilationOutput.bUsesWorldPositionOffset = bUsesWorldPositionOffset;
-	OutCompilationOutput.bUsesPixelDepthOffset = bUsesPixelDepthOffset;
+	OutCompilationOutput.bUsesPixelDepthOffset = bUsesPixelDepthOffset && AllowPixelDepthOffset(InCompilerTarget.ShaderPlatform);
 	OutCompilationOutput.bUsesDisplacement = bUsesDisplacement;
 	OutCompilationOutput.bUsesEyeAdaptation = EmitMaterialData.IsExternalInputUsed(SF_Pixel, Material::EExternalInput::EyeAdaptation);
 
@@ -1016,7 +1042,15 @@ bool MaterialEmitHLSL(const FMaterialCompileTargetParameters& InCompilerTarget,
 
 	OutMaterialEnvironment = new FSharedShaderCompilerEnvironment();
 	OutMaterialEnvironment->TargetPlatform = InCompilerTarget.TargetPlatform;
-	GetMaterialEnvironment(InCompilerTarget.ShaderPlatform, InOutMaterial, EmitContext, OutCompilationOutput, *OutMaterialEnvironment);
+	GetMaterialEnvironment(
+		InCompilerTarget.ShaderPlatform,
+		InOutMaterial,
+		EmitContext,
+		OutCompilationOutput,
+		bUsesEmissiveColor,
+		bUsesAnisotropy,
+		bIsFullyRough,
+		*OutMaterialEnvironment);
 	OutMaterialEnvironment->IncludeVirtualPathToContentsMap.Add(TEXT("/Engine/Generated/Material.ush"), MoveTemp(MaterialTemplateSource));
 
 	return true;
diff --git a/Engine/Source/Runtime/Engine/Private/Materials/MaterialShader.cpp b/Engine/Source/Runtime/Engine/Private/Materials/MaterialShader.cpp
--- a/Engine/Source/Runtime/Engine/Private/Materials/MaterialShader.cpp
+++ b/Engine/Source/Runtime/Engine/Private/Materials/MaterialShader.cpp
@@ -26,6 +26,7 @@
 #include "ShaderCodeLibrary.h"
 #include "UObject/UE5ReleaseStreamObjectVersion.h"
 #include "Materials/Material.h"
+#include "MaterialHLSLGenerator.h"
 
 int32 GMaterialExcludeNonPipelinedShaders = 1;
 static FAutoConsoleVariableRef CVarMaterialExcludeNonPipelinedShaders(
@@ -1146,7 +1147,7 @@ void FMaterialShaderMapId::AppendKeyString(FString& KeyString, bool bIncludeSour
 
 	if (bUsingNewHLSLGenerator)
 	{
-		KeyString += TEXT("_NewHLSL");
+		KeyString += FString::Printf(TEXT("_NewHLSL%d"), FMaterialHLSLGenerator::Version);
 	}
 
 	KeyString += StrataCompilationConfig.GetShaderMapKeyString();
diff --git a/Engine/Source/Runtime/Engine/Public/HLSLTree/HLSLTreeCommon.h b/Engine/Source/Runtime/Engine/Public/HLSLTree/HLSLTreeCommon.h
--- a/Engine/Source/Runtime/Engine/Public/HLSLTree/HLSLTreeCommon.h
+++ b/Engine/Source/Runtime/Engine/Public/HLSLTree/HLSLTreeCommon.h
@@ -268,6 +268,8 @@ public:
 
 	virtual const FExpression* NewSwitch(FTree& Tree, TConstArrayView<const FExpression*> InInputs) const override { return Tree.NewExpression<FExpressionQualitySwitch>(InInputs); }
 	virtual bool IsInputActive(const FEmitContext& Context, int32 Index) const override;
+
+	virtual bool PrepareValue(FEmitContext& Context, FEmitScope& Scope, const FRequestedType& RequestedType, FPrepareValueResult& OutResult) const override;
 };
 
 class FExpressionShaderStageSwitch : public FExpressionSwitchBase
diff --git a/Engine/Source/Runtime/Engine/Public/MaterialHLSLGenerator.h b/Engine/Source/Runtime/Engine/Public/MaterialHLSLGenerator.h
--- a/Engine/Source/Runtime/Engine/Public/MaterialHLSLGenerator.h
+++ b/Engine/Source/Runtime/Engine/Public/MaterialHLSLGenerator.h
@@ -65,6 +65,8 @@ struct TMaterialHLSLGeneratorType;
 class FMaterialHLSLGenerator
 {
 public:
+	static const uint32 Version = 1;
+
 	FMaterialHLSLGenerator(UMaterial* Material,
 		const FMaterialLayersFunctions* InLayerOverrides,
 		UMaterialExpression* InPreviewExpression,
