fix issue with a non-multi sweep against cluster unions which hit nothing (e.g. due to ignored actors) causing the entire trace to hit nothing - we must do another trace with the cluster union actor ignored to get the right result

diff --git a/Engine/Source/Runtime/Engine/Private/Collision/SceneQuery.cpp b/Engine/Source/Runtime/Engine/Private/Collision/SceneQuery.cpp
--- a/Engine/Source/Runtime/Engine/Private/Collision/SceneQuery.cpp
+++ b/Engine/Source/Runtime/Engine/Private/Collision/SceneQuery.cpp
@@ -480,12 +480,25 @@ bool TSceneCastCommonImp(const UWorld* World, typename Traits::TOutHits& OutHits
 						FClusterUnionHit ClusterUnionHit = DoClusterUnionTrace(OutHits, NewHit);
 						if (ClusterUnionHit.bIsClusterUnion)
 						{
-							Traits::ResetOutHits(OutHits, Start, End);
 							bBlockingHit = ClusterUnionHit.bHit;
 							if (ClusterUnionHit.bHit)
 							{
 								OutHits = NewHit;
 							}
+							else if (AActor* ClusterUnionActor = OutHits.GetActor())
+							{
+								// In the case where the trace hits a cluster union but after we subtrace the cluster union and we end up finding out we hit *nothing* then we actually need to
+								// redo the entire trace and force the SQ trace to ignore the cluster union actor. This is only relevant in the case where we aren't doing a multi-trace since in the
+								// case of a multi-trace, we would've found all the other things we hit as well. In the case of a non-multi (single) trace, the SQ trace determined that the cluster union
+								// is the best hit! But there could be other things we could've hit instead if we ignored the fact that we hit a cluster union.
+								FCollisionQueryParams NewParams = Params;
+								NewParams.AddIgnoredActor(ClusterUnionActor);
+								return TSceneCastCommonImp<Traits, TGeomInputs>(World, OutHits, GeomInputs, Start, End, TraceChannel, NewParams, ResponseParams, ObjectParams);
+							}
+							else
+							{
+								Traits::ResetOutHits(OutHits, Start, End);
+							}
 						}
 					}
 				}
