Variable Frame Stripping

Added functionality to strip at ratios higher 1 out of every 2. testing shows good savings at stripping 2/3 frames on switch and in editor of around 8mb with no noticeable visual degradation.

2 Cvars added "a.Compression.UseVaribleFrameStripping" & "a.Compression.VaribleFrameStrippingRatio" which enable and change the amount of variable stripping.

In addition added 3 variables to anim sequences bOverrideVariableFrameStripping, UseVariableFrameStripping& FrameStrippingRatio; which allow you to override any stripping options on a per asset basis. mainly for the purpose of the case that a particular animation is behaving badly under higher stripping you don’t have to lower the whole platforms stripping to accommodate.

diff --git a/Engine/Config/BaseEngine.ini b/Engine/Config/BaseEngine.ini
--- a/Engine/Config/BaseEngine.ini
+++ b/Engine/Config/BaseEngine.ini
@@ -3206,6 +3206,7 @@ BoneTimecodeCustomAttributeNameSettings=(HourAttributeName="TCHour",MinuteAttrib
 BoneCompressionSettings="/Engine/Animation/DefaultAnimBoneCompressionSettings"
 AnimationRecorderBoneCompressionSettings="/Engine/Animation/DefaultRecorderBoneCompression"
 CurveCompressionSettings="/Engine/Animation/DefaultAnimCurveCompressionSettings"
+VariableFrameStrippingSettings="/Engine/Animation/DefaultVariableFrameStrippingSettings"
 
 [/Script/Engine.MeshSimplificationSettings]
 r.MeshReductionModule="QuadricMeshReduction"
diff --git a/Engine/Source/Developer/AssetTools/Private/AssetTools.cpp b/Engine/Source/Developer/AssetTools/Private/AssetTools.cpp
--- a/Engine/Source/Developer/AssetTools/Private/AssetTools.cpp
+++ b/Engine/Source/Developer/AssetTools/Private/AssetTools.cpp
@@ -37,6 +37,7 @@
 #include "ClassViewerFilter.h"
 #include "AssetTypeActions/AssetTypeActions_AnimBoneCompressionSettings.h"
 #include "AssetTypeActions/AssetTypeActions_AnimCurveCompressionSettings.h"
+#include "AssetTypeActions/AssetTypeActions_VariableFrameStrippingSettings.h"
 #include "AssetTypeActions/AssetTypeActions_ForceFeedbackEffect.h"
 #include "AssetTypeActions/AssetTypeActions_ParticleSystem.h"
 #include "AssetTypeActions/AssetTypeActions_PhysicalMaterialMask.h"
@@ -1369,6 +1370,7 @@ UAssetToolsImpl::UAssetToolsImpl(const FObjectInitializer& ObjectInitializer)
 	// Register the built-in asset type actions
 	RegisterAssetTypeActions(MakeShareable(new FAssetTypeActions_AnimBoneCompressionSettings));
 	RegisterAssetTypeActions(MakeShareable(new FAssetTypeActions_AnimCurveCompressionSettings));
+	RegisterAssetTypeActions(MakeShareable(new FAssetTypeActions_VariableFrameStrippingSettings));
 	RegisterAssetTypeActions(MakeShareable(new FAssetTypeActions_ForceFeedbackEffect(InputCategoryBit)));
 	RegisterAssetTypeActions(MakeShareable(new FAssetTypeActions_ParticleSystem));
 	RegisterAssetTypeActions(MakeShareable(new FAssetTypeActions_PhysicalMaterialMask));
diff --git a/Engine/Source/Developer/AssetTools/Private/AssetTypeActions/AssetTypeActions_VariableFrameStrippingSettings.cpp b/Engine/Source/Developer/AssetTools/Private/AssetTypeActions/AssetTypeActions_VariableFrameStrippingSettings.cpp
new file mode 100644
--- /dev/null
+++ b/Engine/Source/Developer/AssetTools/Private/AssetTypeActions/AssetTypeActions_VariableFrameStrippingSettings.cpp
@@ -0,0 +1,21 @@
+// Copyright Epic Games, Inc. All Rights Reserved.
+
+#include "AssetTypeActions/AssetTypeActions_VariableFrameStrippingSettings.h"
+#include "Animation/AnimSequence.h"
+#include "Dialogs/Dialogs.h"
+#include "Styling/AppStyle.h"
+#include "Framework/MultiBox/MultiBoxBuilder.h"
+#include "Misc/MessageDialog.h"
+#include "Misc/ScopedSlowTask.h"
+#include "UObject/UObjectIterator.h"
+
+
+#define LOCTEXT_NAMESPACE "AssetTypeActions"
+
+void FAssetTypeActions_VariableFrameStrippingSettings::OpenAssetEditor(const TArray<UObject*>& InObjects, TSharedPtr<class IToolkitHost> EditWithinLevelEditor)
+{
+	FAssetTypeActions_Base::OpenAssetEditor(InObjects, EditWithinLevelEditor);
+}
+
+
+#undef LOCTEXT_NAMESPACE
\ No newline at end of file
diff --git a/Engine/Source/Developer/AssetTools/Private/AssetTypeActions/AssetTypeActions_VariableFrameStrippingSettings.h b/Engine/Source/Developer/AssetTools/Private/AssetTypeActions/AssetTypeActions_VariableFrameStrippingSettings.h
new file mode 100644
--- /dev/null
+++ b/Engine/Source/Developer/AssetTools/Private/AssetTypeActions/AssetTypeActions_VariableFrameStrippingSettings.h
@@ -0,0 +1,24 @@
+// Copyright Epic Games, Inc. All Rights Reserved.
+
+#pragma once
+
+#include "CoreMinimal.h"
+#include "AssetTypeActions_Base.h"
+#include "Animation/VariableFrameStrippingSettings.h"
+
+class FAssetTypeActions_VariableFrameStrippingSettings : public FAssetTypeActions_Base
+{
+public:
+	// IAssetTypeActions Implementation
+	virtual FText GetName() const override { return NSLOCTEXT("AssetTypeActions", "AssetTypeActions_VariableFrameStrippingSettings", "Variable Frame Stripping Settings"); }
+	virtual FColor GetTypeColor() const override { return FColor(255, 255, 0); }
+	virtual UClass* GetSupportedClass() const override { return UVariableFrameStrippingSettings::StaticClass(); }
+	virtual bool CanFilter() override { return true; }
+	virtual uint32 GetCategories() override { return EAssetTypeCategories::Animation; }
+
+	virtual void OpenAssetEditor(const TArray<UObject*>& InObjects, TSharedPtr<class IToolkitHost> EditWithinLevelEditor = TSharedPtr<IToolkitHost>()) override;
+
+
+
+
+};
diff --git a/Engine/Source/Editor/AnimationBlueprintLibrary/Private/AnimationBlueprintLibrary.cpp b/Engine/Source/Editor/AnimationBlueprintLibrary/Private/AnimationBlueprintLibrary.cpp
--- a/Engine/Source/Editor/AnimationBlueprintLibrary/Private/AnimationBlueprintLibrary.cpp
+++ b/Engine/Source/Editor/AnimationBlueprintLibrary/Private/AnimationBlueprintLibrary.cpp
@@ -9,6 +9,7 @@
 #include "Animation/AnimBlueprintGeneratedClass.h"
 #include "Animation/AnimBoneCompressionSettings.h"
 #include "Animation/AnimCurveCompressionSettings.h"
+#include "Animation/VariableFrameStrippingSettings.h"
 #include "Animation/AnimCurveTypes.h"
 #include "Animation/AnimData/AnimDataModel.h"
 #include "Animation/AnimData/AttributeIdentifier.h"
@@ -225,6 +226,34 @@ void UAnimationBlueprintLibrary::SetCurveCompressionSettings(UAnimSequence* Anim
 	AnimationSequence->CurveCompressionSettings = CompressionSettings;
 }
 
+void UAnimationBlueprintLibrary::GetVariableFrameStrippingSettings(const UAnimSequence* AnimationSequence, UVariableFrameStrippingSettings*& VariableFrameStrippingSettings)
+{
+	if (AnimationSequence == nullptr)
+	{
+		UE_LOG(LogAnimationBlueprintLibrary, Warning, TEXT("Invalid VariableFrameStrippingSettings supplied for GetVariableFrameStrippingSettings"));
+		return;
+	}
+
+	VariableFrameStrippingSettings = AnimationSequence->VariableFrameStrippingSettings;
+}
+
+void UAnimationBlueprintLibrary::SetVariableFrameStrippingSettings(UAnimSequence* AnimationSequence, UVariableFrameStrippingSettings* VariableFrameStrippingSettings)
+{
+	if (AnimationSequence == nullptr)
+	{
+		UE_LOG(LogAnimationBlueprintLibrary, Warning, TEXT("Invalid Animation Sequence supplied for SetVariableFrameStrippingSettings"));
+		return;
+	}
+
+	if (VariableFrameStrippingSettings == nullptr)
+	{
+		UE_LOG(LogAnimationBlueprintLibrary, Warning, TEXT("Invalid VariableFrameStrippingSettings supplied for SetVariableFrameStrippingSettings"));
+		return;
+	}
+
+	AnimationSequence->VariableFrameStrippingSettings = VariableFrameStrippingSettings;
+}
+
 void UAnimationBlueprintLibrary::GetAdditiveAnimationType(const UAnimSequence* AnimationSequence, TEnumAsByte<enum EAdditiveAnimationType>& AdditiveAnimationType)
 {
 	if (AnimationSequence)
diff --git a/Engine/Source/Editor/AnimationBlueprintLibrary/Public/AnimationBlueprintLibrary.h b/Engine/Source/Editor/AnimationBlueprintLibrary/Public/AnimationBlueprintLibrary.h
--- a/Engine/Source/Editor/AnimationBlueprintLibrary/Public/AnimationBlueprintLibrary.h
+++ b/Engine/Source/Editor/AnimationBlueprintLibrary/Public/AnimationBlueprintLibrary.h
@@ -28,6 +28,7 @@ class UAnimBlueprint;
 class UAnimBoneCompressionSettings;
 class UAnimCompress;
 class UAnimCurveCompressionSettings;
+class UVariableFrameStrippingSettings;
 class UAnimGraphNode_Base;
 class UAnimMetaData;
 class UAnimMontage;
@@ -130,6 +131,14 @@ public:
 	UFUNCTION(BlueprintCallable, Category = "AnimationBlueprintLibrary|Compression")
 	static void SetCurveCompressionSettings(UAnimSequence* AnimationSequence, UAnimCurveCompressionSettings* CompressionSettings);
 
+	/** Retrieves the Variable Frame Stripping Settings for the given Animation Sequence */
+	UFUNCTION(BlueprintPure, Category = "AnimationBlueprintLibrary|Compression")
+		static void GetVariableFrameStrippingSettings(const UAnimSequence* AnimationSequence, UVariableFrameStrippingSettings*& VariableFrameStrippingSettings);
+
+	/** Sets the Variable Frame Stripping Settings for the given Animation Sequence */
+	UFUNCTION(BlueprintCallable, Category = "AnimationBlueprintLibrary|Compression")
+		static void SetVariableFrameStrippingSettings(UAnimSequence* AnimationSequence, UVariableFrameStrippingSettings* VariableFrameStrippingSettings);
+
 	// Additive 
 	/** Retrieves the Additive Animation type for the given Animation Sequence */
 	UFUNCTION(BlueprintPure, Category = "AnimationBlueprintLibrary|Additive")
diff --git a/Engine/Source/Editor/UnrealEd/Classes/Factories/VariableFrameStrippingSettingsFactory.h b/Engine/Source/Editor/UnrealEd/Classes/Factories/VariableFrameStrippingSettingsFactory.h
new file mode 100644
--- /dev/null
+++ b/Engine/Source/Editor/UnrealEd/Classes/Factories/VariableFrameStrippingSettingsFactory.h
@@ -0,0 +1,18 @@
+// Copyright Epic Games, Inc. All Rights Reserved.
+
+#pragma once
+
+#include "CoreMinimal.h"
+#include "Factories/Factory.h"
+#include "VariableFrameStrippingSettingsFactory.generated.h"
+
+UCLASS(HideCategories = Object, MinimalAPI)
+class UVariableFrameStrippingSettingsFactory : public UFactory
+{
+	GENERATED_UCLASS_BODY()
+
+	//~ Begin UFactory Interface
+	virtual UObject* FactoryCreateNew(UClass* Class, UObject* InParent, FName Name, EObjectFlags Flags, UObject* Context, FFeedbackContext* Warn) override;
+	//~ Begin UFactory Interface
+
+};
diff --git a/Engine/Source/Editor/UnrealEd/Private/Factories/VariableFrameStrippingSettingsFactory.cpp b/Engine/Source/Editor/UnrealEd/Private/Factories/VariableFrameStrippingSettingsFactory.cpp
new file mode 100644
--- /dev/null
+++ b/Engine/Source/Editor/UnrealEd/Private/Factories/VariableFrameStrippingSettingsFactory.cpp
@@ -0,0 +1,25 @@
+// Copyright Epic Games, Inc. All Rights Reserved.
+
+#include "Factories/VariableFrameStrippingSettingsFactory.h"
+#include "Animation/VariableFrameStrippingSettings.h"
+#include "Templates/SubclassOf.h"
+
+class FFeedbackContext;
+class UClass;
+class UObject;
+
+UVariableFrameStrippingSettingsFactory::UVariableFrameStrippingSettingsFactory(const FObjectInitializer& ObjectInitializer)
+	: Super(ObjectInitializer)
+{
+	// Provide the factory with information about how to handle our asset
+	bCreateNew = true;
+	bEditAfterNew = true;
+	SupportedClass = UVariableFrameStrippingSettings::StaticClass();
+
+}
+
+UObject* UVariableFrameStrippingSettingsFactory::FactoryCreateNew(UClass* Class, UObject* InParent, FName Name, EObjectFlags Flags, UObject* Context, FFeedbackContext* Warn)
+{
+	// Create and return a new instance of VariableFrameStrippingSettings
+	return NewObject<UVariableFrameStrippingSettings>(InParent, Class, Name, Flags);
+}
\ No newline at end of file
diff --git a/Engine/Source/Runtime/Engine/Classes/Animation/AnimSequence.h b/Engine/Source/Runtime/Engine/Classes/Animation/AnimSequence.h
--- a/Engine/Source/Runtime/Engine/Classes/Animation/AnimSequence.h
+++ b/Engine/Source/Runtime/Engine/Classes/Animation/AnimSequence.h
@@ -22,6 +22,7 @@
 #include "Animation/AttributeCurve.h"
 #include "PerPlatformProperties.h"
 #include "IO/IoHash.h"
+#include "Animation/VariableFrameStrippingSettings.h"
 
 #if WITH_EDITOR
 #include "AnimData/IAnimationDataModel.h"
@@ -287,6 +288,9 @@ public:
 
 	FCompressedAnimSequence CompressedData;
 
+	UPROPERTY(Category = Compression, EditAnywhere, meta = (ForceShowEngineContent))
+	TObjectPtr<class UVariableFrameStrippingSettings> VariableFrameStrippingSettings;
+
 	/** Additive animation type. **/
 	UPROPERTY(EditAnywhere, Category=AdditiveSettings, AssetRegistrySearchable)
 	TEnumAsByte<enum EAdditiveAnimationType> AdditiveAnimType;
diff --git a/Engine/Source/Runtime/Engine/Classes/Animation/AnimStreamable.h b/Engine/Source/Runtime/Engine/Classes/Animation/AnimStreamable.h
--- a/Engine/Source/Runtime/Engine/Classes/Animation/AnimStreamable.h
+++ b/Engine/Source/Runtime/Engine/Classes/Animation/AnimStreamable.h
@@ -187,6 +187,10 @@ public:
 	UPROPERTY(Category = Compression, EditAnywhere)
 	TObjectPtr<class UAnimCurveCompressionSettings> CurveCompressionSettings;
 
+	/** The settings used to control whether or not to use variable frame stripping and its amount*/
+	UPROPERTY(Category = Compression, EditAnywhere)
+	TObjectPtr<class UVariableFrameStrippingSettings> VariableFrameStrippingSettings;
+
 	/** If this is on, it will allow extracting of root motion **/
 	UPROPERTY(EditAnywhere, AssetRegistrySearchable, Category = RootMotion, meta = (DisplayName = "EnableRootMotion"))
 	bool bEnableRootMotion;
diff --git a/Engine/Source/Runtime/Engine/Classes/Animation/VariableFrameStrippingSettings.h b/Engine/Source/Runtime/Engine/Classes/Animation/VariableFrameStrippingSettings.h
new file mode 100644
--- /dev/null
+++ b/Engine/Source/Runtime/Engine/Classes/Animation/VariableFrameStrippingSettings.h
@@ -0,0 +1,36 @@
+// Copyright Epic Games, Inc. All Rights Reserved.
+
+#pragma once
+
+#include "CoreMinimal.h"
+#include "VariableFrameStrippingSettings.generated.h"
+
+/*
+* This is a wrapper for the Varrible frame stripping Codec.
+* It allows for the mass changing of settings on animation sequences in an editor acessable way.
+*/
+UCLASS(hidecategories = Object)
+class ENGINE_API UVariableFrameStrippingSettings : public UObject
+{
+	GENERATED_UCLASS_BODY()
+
+	/**
+	* Enables the change from standard 1/2 frame stripping to stripping a higher amount of frames per frame kept
+	*/
+	UPROPERTY(Category = Compression, EditAnywhere)
+		FPerPlatformBool UseVariableFrameStripping;
+
+	/**
+	* The number of Frames to strip for every one you keep.
+	* Allows for overrides of that multiplier.
+	*/
+	UPROPERTY(Category = Compression, EditAnywhere, meta = (ClampMin = "1"))
+		FPerPlatformInt FrameStrippingRate;
+
+
+#if WITH_EDITORONLY_DATA
+	/** Generates a DDC key that takes into account the current settings, selected codec, input anim sequence and TargetPlatform */
+	void PopulateDDCKey(const UE::Anim::Compression::FAnimDDCKeyArgs& KeyArgs, FArchive& Ar);
+
+#endif
+};
diff --git a/Engine/Source/Runtime/Engine/Private/Animation/AnimCompressionTypes.cpp b/Engine/Source/Runtime/Engine/Private/Animation/AnimCompressionTypes.cpp
--- a/Engine/Source/Runtime/Engine/Private/Animation/AnimCompressionTypes.cpp
+++ b/Engine/Source/Runtime/Engine/Private/Animation/AnimCompressionTypes.cpp
@@ -22,6 +22,7 @@
 #include "Animation/SkeletonRemapping.h"
 #include "Animation/SkeletonRemappingRegistry.h"
 #include "Animation/AttributesRuntime.h"
+#include "Kismet/GameplayStatics.h"
 
 static FString GCompressionJsonOutput;
 static FAutoConsoleVariableRef CVarCompressionJsonOutput(
@@ -105,6 +106,42 @@ void StripFramesOdd(TArray<ArrayValue>& Keys, const int32 NumFrames)
 	}
 }
 
+template<typename ArrayValue>
+void StripFramesMultipler(TArray<ArrayValue>& Keys, const int32 NumFrames, int rate)
+{
+	if (Keys.Num() > rate)//make sure the animation has more frames than the amount to strip per kept frame
+	{
+		const int32 NewNumFrames = NumFrames / rate;
+
+		TArray<ArrayValue> NewKeys;
+		NewKeys.Reserve(NewNumFrames);
+
+		check(Keys.Num() == NumFrames);
+
+		NewKeys.Add(Keys[0]); //Always keep first 
+
+		//Always keep first and last
+		const int32 NumFramesToCalculate = NewNumFrames - 2;
+
+		// Frame increment is ratio of old frame spaces vs new frame spaces 
+		const double FrameIncrement = (double)(NumFrames - 1) / (double)(NewNumFrames - 1);
+
+		for (int32 Frame = 0; Frame < NumFramesToCalculate; ++Frame)
+		{
+			const double NextFramePosition = FrameIncrement * (Frame + 1);
+			const int32 Frame1 = (int32)NextFramePosition;
+			const float Alpha = (NextFramePosition - (double)Frame1);
+
+			NewKeys.Add(AnimationCompressionUtils::Interpolate(Keys[Frame1], Keys[Frame1 + 1], Alpha));
+
+		}
+		NewKeys.Add(Keys.Last()); // Always Keep Last
+
+		Keys = MoveTemp(NewKeys);
+
+	}
+}
+
 struct FByFramePoseEvalContext
 {
 public:
@@ -712,7 +749,6 @@ FCompressibleAnimData::FCompressibleAnimData(class UAnimSequence* InSeq, const b
 	, TargetPlatform(InTargetPlatform)
 {
 }
-
 void FCompressibleAnimData::FetchData(const ITargetPlatform* InPlatform)
 {
 #if WITH_EDITOR
@@ -862,35 +898,56 @@ void FCompressibleAnimData::FetchData(const ITargetPlatform* InPlatform)
 
 	if (bShouldPerformStripping)
 	{
+
+		FName TargetPlatformName = FName(UGameplayStatics::GetPlatformName());
+		TObjectPtr<class UVariableFrameStrippingSettings> VarFrameStrippingSettings = AnimSequence->VariableFrameStrippingSettings;
+		FPerPlatformBool platformBool = VarFrameStrippingSettings->UseVariableFrameStripping;
+		bool bUseMultiplier = platformBool.GetValueForPlatform(TargetPlatformName);
+
 		const int32 NumTracks = RawAnimationData.Num();
-		
-		// End frame does not count towards "Even framed" calculation
-		const bool bIsEvenFramed = ((NumberOfKeys - 1) % 2) == 0;
 
-		//Strip every other frame from tracks
-		if (bIsEvenFramed)
-		{
+		if (bUseMultiplier) {
+			int rate = AnimSequence->VariableFrameStrippingSettings->FrameStrippingRate.GetValueForPlatform(TargetPlatformName);
 			for (FRawAnimSequenceTrack& Track : RawAnimationData)
 			{
-				StripFramesEven(Track.PosKeys, NumberOfKeys);
-				StripFramesEven(Track.RotKeys, NumberOfKeys);
-				StripFramesEven(Track.ScaleKeys, NumberOfKeys);
+				StripFramesMultipler(Track.PosKeys, NumberOfKeys, rate);
+				StripFramesMultipler(Track.RotKeys, NumberOfKeys, rate);
+				StripFramesMultipler(Track.ScaleKeys, NumberOfKeys, rate);
 			}
 
 			const int32 ActualKeys = NumberOfKeys - 1; // strip bookmark end frame
-			NumberOfKeys = (ActualKeys / 2) + 1;
+
+			NumberOfKeys = (ActualKeys * rate) + 1;
 		}
-		else
-		{
-			for (FRawAnimSequenceTrack& Track : RawAnimationData)
+		else {
+			// End frame does not count towards "Even framed" calculation
+			const bool bIsEvenFramed = ((NumberOfKeys - 1) % 2) == 0;
+
+			//Strip every other frame from tracks
+			if (bIsEvenFramed)
 			{
-				StripFramesOdd(Track.PosKeys, NumberOfKeys);
-				StripFramesOdd(Track.RotKeys, NumberOfKeys);
-				StripFramesOdd(Track.ScaleKeys, NumberOfKeys);
+				for (FRawAnimSequenceTrack& Track : RawAnimationData)
+				{
+					StripFramesEven(Track.PosKeys, NumberOfKeys);
+					StripFramesEven(Track.RotKeys, NumberOfKeys);
+					StripFramesEven(Track.ScaleKeys, NumberOfKeys);
+				}
+
+				const int32 ActualKeys = NumberOfKeys - 1; // strip bookmark end frame
+				NumberOfKeys = (ActualKeys / 2) + 1;
 			}
+			else
+			{
+				for (FRawAnimSequenceTrack& Track : RawAnimationData)
+				{
+					StripFramesOdd(Track.PosKeys, NumberOfKeys);
+					StripFramesOdd(Track.RotKeys, NumberOfKeys);
+					StripFramesOdd(Track.ScaleKeys, NumberOfKeys);
+				}
 
-			const int32 ActualKeys = NumberOfKeys;
-			NumberOfKeys = (ActualKeys / 2);
+				const int32 ActualKeys = NumberOfKeys;
+				NumberOfKeys = (ActualKeys / 2);
+			}
 		}
 	}
 
@@ -904,7 +961,6 @@ void FCompressibleAnimData::FetchData(const ITargetPlatform* InPlatform)
 	}
 #endif
 }
-
 FCompressibleAnimData::FCompressibleAnimData(UAnimBoneCompressionSettings* InBoneCompressionSettings, UAnimCurveCompressionSettings* InCurveCompressionSettings, USkeleton* InSkeleton, EAnimInterpolationType InInterpolation, float InSequenceLength, int32 InNumberOfKeys, const ITargetPlatform* InTargetPlatform)
 	: CurveCompressionSettings(InCurveCompressionSettings)
 	, BoneCompressionSettings(InBoneCompressionSettings)
diff --git a/Engine/Source/Runtime/Engine/Private/Animation/AnimSequence.cpp b/Engine/Source/Runtime/Engine/Private/Animation/AnimSequence.cpp
--- a/Engine/Source/Runtime/Engine/Private/Animation/AnimSequence.cpp
+++ b/Engine/Source/Runtime/Engine/Private/Animation/AnimSequence.cpp
@@ -30,6 +30,7 @@
 #include "Animation/AnimBoneCompressionSettings.h"
 #include "Animation/AnimCurveCompressionCodec.h"
 #include "Animation/AnimCurveCompressionSettings.h"
+#include "Animation/VariableFrameStrippingSettings.h"
 #include "EditorFramework/AssetImportData.h"
 #include "Logging/MessageLog.h"
 #include "DerivedDataCacheInterface.h"
@@ -890,6 +891,10 @@ void UAnimSequence::GetPreloadDependencies(TArray<UObject*>& OutDeps)
 	{
 		OutDeps.Add(BoneCompressionSettings);
 	}
+	if (VariableFrameStrippingSettings != nullptr)
+	{
+		OutDeps.Add(VariableFrameStrippingSettings);
+	}
 }
 
 void UAnimSequence::PreSave(const class ITargetPlatform* TargetPlatform)
@@ -1117,7 +1122,8 @@ void UAnimSequence::PostEditChangeProperty(FPropertyChangedEvent& PropertyChange
 		bCompressionAffectingSettingsChanged =   PropertyChangedEvent.Property->GetFName() == GET_MEMBER_NAME_CHECKED(UAnimSequence, bAllowFrameStripping)
 											  || PropertyChangedEvent.Property->GetFName() == GET_MEMBER_NAME_CHECKED(UAnimSequence, CompressionErrorThresholdScale)
 											  || PropertyChangedEvent.Property->GetFName() == GET_MEMBER_NAME_CHECKED(UAnimSequence, BoneCompressionSettings)
-											  || PropertyChangedEvent.Property->GetFName() == GET_MEMBER_NAME_CHECKED(UAnimSequence, BoneCompressionSettings);
+											  || PropertyChangedEvent.Property->GetFName() == GET_MEMBER_NAME_CHECKED(UAnimSequence, BoneCompressionSettings)
+											  || PropertyChangedEvent.Property->GetFName() == GET_MEMBER_NAME_CHECKED(UAnimSequence, VariableFrameStrippingSettings);
 
 		bShouldResample = PropertyChangedEvent.GetMemberPropertyName() == GET_MEMBER_NAME_CHECKED(UAnimSequence, PlatformTargetFrameRate) || bChangedRefFrameIndex;
 	}
diff --git a/Engine/Source/Runtime/Engine/Private/Animation/AnimStreamable.cpp b/Engine/Source/Runtime/Engine/Private/Animation/AnimStreamable.cpp
--- a/Engine/Source/Runtime/Engine/Private/Animation/AnimStreamable.cpp
+++ b/Engine/Source/Runtime/Engine/Private/Animation/AnimStreamable.cpp
@@ -21,6 +21,7 @@
 #include "Animation/AnimBoneCompressionSettings.h"
 #include "Animation/AnimCurveCompressionSettings.h"
 #include "Animation/AnimCurveCompressionCodec.h"
+#include "Animation/VariableFrameStrippingSettings.h"
 #include "BonePose.h"
 #include "CommonFrameRates.h"
 #include "ContentStreaming.h"
@@ -480,6 +481,7 @@ void UAnimStreamable::InitFrom(const UAnimSequence* InSourceSequence)
 	CurveCompressionSettings = InSourceSequence->CurveCompressionSettings;
 	
 	DataModelInterface = StaticDuplicateObject(InSourceSequence->GetDataModelInterface().GetObject(), this);
+	VariableFrameStrippingSettings = InSourceSequence->VariableFrameStrippingSettings;
 
 	// Far from ideal (retrieving controller to ensure it matches the DataModelInterface type)
 	Controller = DataModelInterface->GetController();
@@ -563,6 +565,11 @@ void UAnimStreamable::RequestCompressedData(const ITargetPlatform* Platform)
 		CurveCompressionSettings = FAnimationUtils::GetDefaultAnimationCurveCompressionSettings();
 	}
 
+	if (VariableFrameStrippingSettings == nullptr)
+	{
+		VariableFrameStrippingSettings = FAnimationUtils::GetDefaultVariableFrameStrippingSettings();
+	}
+
 	checkf(Platform, TEXT("Failed to specify platform for streamable animation compression"));
 
 	FStreamableAnimPlatformData& PlatformData = GetStreamingAnimPlatformData(Platform);
@@ -783,12 +790,14 @@ FString UAnimStreamable::GetBaseDDCKey(uint32 NumChunks, const ITargetPlatform*
 	//  * Our skeletons virtual bone guid
 	//	* Compression Settings
 	//	* Curve compression settings
+	//  * Variable frame stripping settings
 
 	FArcToHexString ArcToHexString;
 
 	ArcToHexString.Ar << NumChunks;
 	BoneCompressionSettings->PopulateDDCKey(UE::Anim::Compression::FAnimDDCKeyArgs(*this, TargetPlatform), ArcToHexString.Ar);
 	CurveCompressionSettings->PopulateDDCKey(ArcToHexString.Ar);
+	VariableFrameStrippingSettings->PopulateDDCKey(UE::Anim::Compression::FAnimDDCKeyArgs(*this, TargetPlatform), ArcToHexString.Ar);
 
 	FString Ret = FString::Printf(TEXT("%s%s%s%s_%s"),
 		StreamingAnimChunkVersion,
diff --git a/Engine/Source/Runtime/Engine/Private/Animation/AnimationUtils.cpp b/Engine/Source/Runtime/Engine/Private/Animation/AnimationUtils.cpp
--- a/Engine/Source/Runtime/Engine/Private/Animation/AnimationUtils.cpp
+++ b/Engine/Source/Runtime/Engine/Private/Animation/AnimationUtils.cpp
@@ -16,6 +16,7 @@
 #include "Animation/AnimBoneCompressionCodec.h"
 #include "Animation/AnimBoneCompressionSettings.h"
 #include "Animation/AnimCurveCompressionSettings.h"
+#include "Animation/VariableFrameStrippingSettings.h"
 #include "AnimationCompression.h"
 #include "Engine/SkeletalMeshSocket.h"
 #include "UObject/LinkerLoad.h"
@@ -844,6 +845,7 @@ void FAnimationUtils::TallyErrorsFromPerturbation(
 static UAnimBoneCompressionSettings* DefaultBoneCompressionSettings = nullptr;
 static UAnimBoneCompressionSettings* DefaultRecorderBoneCompressionSettings = nullptr;
 static UAnimCurveCompressionSettings* DefaultCurveCompressionSettings = nullptr;
+static UVariableFrameStrippingSettings* DefaultVariableFrameStrippingSettings = nullptr;
 
 static void EnsureDependenciesAreLoaded(UObject* Object)
 {
@@ -915,6 +917,7 @@ void FAnimationUtils::PreloadCompressionSettings()
 	GetDefaultAnimationBoneCompressionSettings();
 	GetDefaultAnimationRecorderBoneCompressionSettings();
 	GetDefaultAnimationCurveCompressionSettings();
+	GetDefaultVariableFrameStrippingSettings();
 }
 
 #endif
@@ -949,6 +952,16 @@ UAnimCurveCompressionSettings* FAnimationUtils::GetDefaultAnimationCurveCompress
 	return DefaultCurveCompressionSettings;
 }
 
+UVariableFrameStrippingSettings* FAnimationUtils::GetDefaultVariableFrameStrippingSettings()
+{
+	if (DefaultVariableFrameStrippingSettings == nullptr)
+	{
+		DefaultVariableFrameStrippingSettings = Cast<UVariableFrameStrippingSettings>(GetDefaultAnimationCompressionSettings(TEXT("VariableFrameStrippingSettings")));
+	}
+
+	return DefaultVariableFrameStrippingSettings;
+}
+
 void FAnimationUtils::EnsureAnimSequenceLoaded(UAnimSequence& AnimSeq)
 {
 	EnsureDependenciesAreLoaded(&AnimSeq);
@@ -956,6 +969,7 @@ void FAnimationUtils::EnsureAnimSequenceLoaded(UAnimSequence& AnimSeq)
 	EnsureDependenciesAreLoaded(AnimSeq.BoneCompressionSettings);
 	EnsureDependenciesAreLoaded(AnimSeq.CurveCompressionSettings);
 	EnsureDependenciesAreLoaded(AnimSeq.RefPoseSeq);
+	EnsureDependenciesAreLoaded(AnimSeq.VariableFrameStrippingSettings);
 }
 
 void FAnimationUtils::ExtractTransformForFrameFromTrackSafe(const FRawAnimSequenceTrack& RawTrack, int32 Frame, FTransform& OutAtom)
diff --git a/Engine/Source/Runtime/Engine/Private/Animation/VariableFrameStrippingSettings.cpp b/Engine/Source/Runtime/Engine/Private/Animation/VariableFrameStrippingSettings.cpp
new file mode 100644
--- /dev/null
+++ b/Engine/Source/Runtime/Engine/Private/Animation/VariableFrameStrippingSettings.cpp
@@ -0,0 +1,28 @@
+// Copyright Epic Games, Inc. All Rights Reserved.
+
+#include "Animation/VariableFrameStrippingSettings.h"
+#include "Interfaces/ITargetPlatform.h"
+#include "PlatformInfo.h"
+#include "PerPlatformProperties.h"
+UVariableFrameStrippingSettings::UVariableFrameStrippingSettings(const FObjectInitializer& ObjectInitializer)
+	: Super(ObjectInitializer)
+{
+	UseVariableFrameStripping = FPerPlatformBool(false);
+	FrameStrippingRate = 3;
+}
+
+#if WITH_EDITORONLY_DATA
+/** Generates a DDC key that takes into account the current settings, selected codec, input anim sequence and TargetPlatform */
+void UVariableFrameStrippingSettings::PopulateDDCKey(const UE::Anim::Compression::FAnimDDCKeyArgs& KeyArgs, FArchive& Ar)
+{
+	const ITargetPlatform* TargetPlatform = KeyArgs.TargetPlatform;
+	if (TargetPlatform != nullptr)
+	{
+		FName TargetPlatformName = TargetPlatform->GetTargetPlatformInfo().Name;
+		bool perPlatformUse = UseVariableFrameStripping.GetValueForPlatform(TargetPlatformName);
+		Ar << perPlatformUse;
+		int perPlatformRate = FrameStrippingRate.GetValueForPlatform(TargetPlatformName);
+		Ar << perPlatformRate;
+	}
+}
+#endif
diff --git a/Engine/Source/Runtime/Engine/Public/AnimationUtils.h b/Engine/Source/Runtime/Engine/Public/AnimationUtils.h
--- a/Engine/Source/Runtime/Engine/Public/AnimationUtils.h
+++ b/Engine/Source/Runtime/Engine/Public/AnimationUtils.h
@@ -162,6 +162,9 @@ public:
 	/** Returns the default animation curve compression settings, can never by null. */
 	ENGINE_API static UAnimCurveCompressionSettings* GetDefaultAnimationCurveCompressionSettings();
 
+	/** Returns the default animation varriable frame stripping settings, can never by null. */
+	ENGINE_API static UVariableFrameStrippingSettings* GetDefaultVariableFrameStrippingSettings();
+
 	/** Makes sure that the animation sequence and all its dependencies required for compression are loaded. */
 	ENGINE_API static void EnsureAnimSequenceLoaded(UAnimSequence& AnimSeq);
 
