Fix a compile error

diff --git a/Engine/Source/Runtime/Engine/Private/Materials/MaterialHLSLEmitter.cpp b/Engine/Source/Runtime/Engine/Private/Materials/MaterialHLSLEmitter.cpp
--- a/Engine/Source/Runtime/Engine/Private/Materials/MaterialHLSLEmitter.cpp
+++ b/Engine/Source/Runtime/Engine/Private/Materials/MaterialHLSLEmitter.cpp
@@ -499,7 +499,7 @@ static void GetMaterialEnvironment(EShaderPlatform InPlatform,
 	OutEnvironment.SetDefine(TEXT("NEEDS_PARTICLE_WORLD_TO_LOCAL"), false);// bUsesParticleWorldToLocal);
 	OutEnvironment.SetDefine(TEXT("USES_TRANSFORM_VECTOR"), false);// bUsesTransformVector);
 	OutEnvironment.SetDefine(TEXT("WANT_PIXEL_DEPTH_OFFSET"), MaterialCompilationOutput.bUsesPixelDepthOffset);
-	OutEnvironment.SetDefineAndCompileArgument(TEXT("USES_WORLD_POSITION_OFFSET"), MaterialCompilationOutput.bUsesWorldPositionOffset);
+	OutEnvironment.SetDefineAndCompileArgument(TEXT("USES_WORLD_POSITION_OFFSET"), (bool)MaterialCompilationOutput.bUsesWorldPositionOffset);
 	OutEnvironment.SetDefineAndCompileArgument(TEXT("USES_DISPLACEMENT"), false);// bUsesDisplacement);
 	OutEnvironment.SetDefine(TEXT("USES_EMISSIVE_COLOR"), bUsesEmissiveColor);
 	// Distortion uses tangent space transform 
