Fix for handling multiple layers in VirtualTexture build

diff --git a/Engine/Source/Runtime/Engine/Private/VT/VirtualTextureDataBuilder.cpp b/Engine/Source/Runtime/Engine/Private/VT/VirtualTextureDataBuilder.cpp
--- a/Engine/Source/Runtime/Engine/Private/VT/VirtualTextureDataBuilder.cpp
+++ b/Engine/Source/Runtime/Engine/Private/VT/VirtualTextureDataBuilder.cpp
@@ -500,25 +500,17 @@ bool FVirtualTextureDataBuilder::BuildChunks()
 					const int32 BlockX = TileX / MipBlockSizeInTilesX;
 					const int32 BlockY = TileY / MipBlockSizeInTilesY;
 
-					FVTBlockPayload* Block = FindSourceBlock(Mip, BlockX, BlockY);
-					if (Block)
+					int32 BlockIndex = FindSourceBlockIndex(Mip, BlockX, BlockY);
+					if (BlockIndex != INDEX_NONE)
 					{
-						const int32 MipIndex = Mip - Block->MipBias;
-						const int32 BlockWidth = FMath::Max(Block->SizeX >> MipIndex, 1);
-						const int32 BlockWidthInTiles = FMath::DivideAndRoundUp(BlockWidth, TileSize);
-
+						FVTBlockPayload& Block = LayerPayload[0].Blocks[BlockIndex];
 						FVTSourceTileEntry* TileEntry = new(TilesInChunk) FVTSourceTileEntry;
-						TileEntry->Block = Block;
-						TileEntry->MipIndex = Mip;
-						TileEntry->MipIndexInBlock = MipIndex;
-						int32 TileInBlockX = TileX - Block->BlockX * MipBlockSizeInTilesX;
-						int32 TileInBlockY = TileY - Block->BlockY * MipBlockSizeInTilesY;
-
-						// TileIndex is destination index to use in output chunks
-						// TileIndexInBlock is source index into compressed data array of block/mip
-						TileEntry->TileIndexInBlock = TileInBlockY * BlockWidthInTiles + TileInBlockX;
+						TileEntry->BlockIndex = BlockIndex;
 						TileEntry->TileIndex = TileIndex;
-
+						TileEntry->MipIndex = Mip;
+						TileEntry->MipIndexInBlock = Mip - Block.MipBias;
+						TileEntry->TileX = TileX;
+						TileEntry->TileY = TileY;
 						OffsetData.AddTile(TileIndexInMip);
 					}
 				}
@@ -1119,6 +1111,9 @@ void FVirtualTextureDataBuilder::BuildTiles(const TArray<FVTSourceTileEntry>& Ti
 
 	FThreadSafeBool bCompressionError = false;
 	EPixelFormat CompressedFormat = PF_Unknown;
+	const int32 TileSize = SettingsPerLayer[0].VirtualTextureTileSize;
+	const int32 BlockSizeInTilesX = FMath::DivideAndRoundUp(DerivedInfo.BlockSizeX, TileSize);
+	const int32 BlockSizeInTilesY = FMath::DivideAndRoundUp(DerivedInfo.BlockSizeY, TileSize);
 
 	{
 		GeneratedData.TilePayload.AddDefaulted(TileList.Num());
@@ -1126,9 +1121,22 @@ void FVirtualTextureDataBuilder::BuildTiles(const TArray<FVTSourceTileEntry>& Ti
 		for (int32 TileIndex = 0; TileIndex < TileList.Num(); TileIndex++)
 		{
 			const FVTSourceTileEntry& Tile = TileList[TileIndex];
-			FVTBlockPayload* Block = Tile.Block;
-			GeneratedData.TilePayload[TileIndex] = MoveTemp(Block->Tiles[Tile.TileIndexInBlock].Mips[Tile.MipIndexInBlock].Payload);
-			CompressedFormat = Block->Tiles[Tile.TileIndexInBlock].Mips[Tile.MipIndexInBlock].CompressedFormat;
+			FVTBlockPayload& Block = LayerPayload[LayerIndex].Blocks[Tile.BlockIndex];
+
+			int32 MipIndex = Tile.MipIndex - Block.MipBias;
+			int32 BlockWidth = FMath::Max(Block.SizeX >> MipIndex, 1);
+			int32 BlockWidthInTiles = FMath::DivideAndRoundUp(BlockWidth, TileSize);
+
+			int32 MipBlockSizeInTilesX = FMath::Max(BlockSizeInTilesX >> Tile.MipIndex, 1);
+			int32 MipBlockSizeInTilesY = FMath::Max(BlockSizeInTilesY >> Tile.MipIndex, 1);
+			
+			int32 TileInBlockX = Tile.TileX - Block.BlockX * MipBlockSizeInTilesX;
+			int32 TileInBlockY = Tile.TileY - Block.BlockY * MipBlockSizeInTilesY;
+			int32 TileIndexInBlock = TileInBlockY * BlockWidthInTiles + TileInBlockX;
+
+			FVTTileMipPayload& MipPayload = Block.Tiles[TileIndexInBlock].Mips[Tile.MipIndexInBlock];
+			GeneratedData.TilePayload[TileIndex] = MoveTemp(MipPayload.Payload);
+			CompressedFormat = MipPayload.CompressedFormat;
 		}
 
 		GeneratedData.Codec = EVirtualTextureCodec::RawGPU;
@@ -1254,24 +1262,22 @@ bool FVirtualTextureDataBuilder::PushDataToChunk(const TArray<FVTSourceTileEntry
 	return true;
 }
 
-FVTBlockPayload* FVirtualTextureDataBuilder::FindSourceBlock(int32 MipIndex, int32 BlockX, int32 BlockY)
+int32 FVirtualTextureDataBuilder::FindSourceBlockIndex(int32 MipIndex, int32 BlockX, int32 BlockY)
 {
-	for (int32 LayerIndex = 0; LayerIndex < LayerPayload.Num(); ++LayerIndex)
+	// VT assumes that layer 0 is largest layer when assigning block to tiles
+	TArray<FVTBlockPayload>& Blocks = LayerPayload[0].Blocks;
+	for (int32 BlockIndex = 0; BlockIndex < Blocks.Num(); ++BlockIndex)
 	{
-		TArray<FVTBlockPayload>& Blocks = LayerPayload[LayerIndex].Blocks;
-		for (int32 BlockIndex = 0; BlockIndex < Blocks.Num(); ++BlockIndex)
+		FVTBlockPayload& Block = Blocks[BlockIndex];
+		if (BlockX >= Block.BlockX && BlockX < Block.BlockX + Block.SizeInBlocksX &&
+			BlockY >= Block.BlockY && BlockY < Block.BlockY + Block.SizeInBlocksY &&
+			MipIndex >= Block.MipBias &&
+			(MipIndex - Block.MipBias) < Block.NumMips)
 		{
-			FVTBlockPayload& Block = Blocks[BlockIndex];
-			if (BlockX >= Block.BlockX && BlockX < Block.BlockX + Block.SizeInBlocksX &&
-				BlockY >= Block.BlockY && BlockY < Block.BlockY + Block.SizeInBlocksY &&
-				MipIndex >= Block.MipBias &&
-				(MipIndex - Block.MipBias) < Block.NumMips)
-			{
-				return &Block;
-			}
+			return BlockIndex;
 		}
 	}
-	return nullptr;
+	return INDEX_NONE;
 }
 
 // Leaving this code here for now, in case we want to build a new/better system for creating/storing miptails
diff --git a/Engine/Source/Runtime/Engine/Private/VT/VirtualTextureDataBuilder.h b/Engine/Source/Runtime/Engine/Private/VT/VirtualTextureDataBuilder.h
--- a/Engine/Source/Runtime/Engine/Private/VT/VirtualTextureDataBuilder.h
+++ b/Engine/Source/Runtime/Engine/Private/VT/VirtualTextureDataBuilder.h
@@ -52,14 +52,14 @@ struct FVTLayerPayload
 	TArray<FVTBlockPayload> Blocks;
 };
 
-
 struct FVTSourceTileEntry
 {
-	FVTBlockPayload* Block = nullptr;
+	int32 BlockIndex = 0;
 	int32 TileIndex = 0;
-	int32 TileIndexInBlock = 0;
 	int32 MipIndex = 0;
 	int32 MipIndexInBlock = 0;
+	int32 TileX = 0;
+	int32 TileY = 0;
 };
 
 struct FLayerData
@@ -129,7 +129,7 @@ private:
 	void BuildTiles(const TArray<FVTSourceTileEntry>& TileList, uint32 layer, FLayerData& GeneratedData);
 	bool PushDataToChunk(const TArray<FVTSourceTileEntry>& Tiles, const TArray<FLayerData>& LayerData);
 
-	FVTBlockPayload* FindSourceBlock(int32 MipIndex, int32 BlockX, int32 BlockY);
+	int32 FindSourceBlockIndex(int32 MipIndex, int32 BlockX, int32 BlockY);
 
 	void BuildLayerBlocks(FSlowTask& BuildTask, uint32 LayerIndex, const FVirtualTextureSourceLayerData& LayerData, FTextureSourceData& SourceData, FTextureSourceData& CompositeSourceData, bool bAllowAsync);
 	void BuildBlockTiles(uint32 LayerIndex, uint32 BlockIndex, FVTBlockPayload& Block, const FVirtualTextureSourceLayerData& LayerData, bool bAllowAsync);
