Mitigate huge perf drop introduced in CL 25322360 by only computing the debug group name when we need to compile shaders.

#rnx

diff --git a/Engine/Source/Runtime/Engine/Private/Materials/MaterialShared.cpp b/Engine/Source/Runtime/Engine/Private/Materials/MaterialShared.cpp
--- a/Engine/Source/Runtime/Engine/Private/Materials/MaterialShared.cpp
+++ b/Engine/Source/Runtime/Engine/Private/Materials/MaterialShared.cpp
@@ -3420,8 +3420,6 @@ bool FMaterial::TryGetShaders(const FMaterialShaderTypes& InTypes, const FVertex
 		return false;
 	};
 
-	const FString DebugGroupName = GetUniqueAssetName(ShaderPlatform, ShaderMap->GetShaderMapId()) / LexToString(GetQualityLevel());
-
 	if (InTypes.PipelineType &&
 		RHISupportsShaderPipelines(ShaderPlatform) &&
 		CVarShaderPipelines && CVarShaderPipelines->GetValueOnAnyThread(bIsInGameThread) != 0)
@@ -3484,6 +3482,7 @@ bool FMaterial::TryGetShaders(const FMaterialShaderTypes& InTypes, const FVertex
 					{
 						if (!bShaderMapComplete)
 						{
+							const FString DebugGroupName = GetUniqueAssetName(ShaderPlatform, ShaderMap->GetShaderMapId()) / LexToString(GetQualityLevel());
 							if (InVertexFactoryType)
 							{
 								FMeshMaterialShaderType::BeginCompileShaderPipeline(
@@ -3569,6 +3568,7 @@ bool FMaterial::TryGetShaders(const FMaterialShaderTypes& InTypes, const FVertex
 						{
 							if (!bShaderMapComplete)
 							{
+								const FString DebugGroupName = GetUniqueAssetName(ShaderPlatform, ShaderMap->GetShaderMapId()) / LexToString(GetQualityLevel());
 								if (InVertexFactoryType)
 								{
 									ShaderType->AsMeshMaterialShaderType()->BeginCompileShader(
