Fix content bundle flushing their SubObject to cell mappings
Implement Add/RemoveExternalStreamingObject in WorldPartitionStreamingPolicy to keep inject/remove subobject mappings
Fix SoftObjectPath/Ptrs not resolving in ContentBundles

diff --git a/Engine/Source/Runtime/Engine/Private/WorldPartition/ContentBundle/ContentBundleEditor.cpp b/Engine/Source/Runtime/Engine/Private/WorldPartition/ContentBundle/ContentBundleEditor.cpp
--- a/Engine/Source/Runtime/Engine/Private/WorldPartition/ContentBundle/ContentBundleEditor.cpp
+++ b/Engine/Source/Runtime/Engine/Private/WorldPartition/ContentBundle/ContentBundleEditor.cpp
@@ -322,7 +322,7 @@ void FContentBundleEditor::GenerateStreaming(TArray<FString>* OutPackageToGenera
 	UWorldPartition* WorldPartition = GetInjectedWorld()->GetWorldPartition();
 	WorldPartition->GenerateContainerStreaming(Params, Context);
 
-	ExternalStreamingObject = WorldPartition->RuntimeHash->StoreToExternalStreamingObject(WorldPartition, *GetExternalStreamingObjectName());
+	ExternalStreamingObject = WorldPartition->FlushStreamingToExternalStreamingObject(GetExternalStreamingObjectName());
 
 	uint32 CellCount = 0;
 	ExternalStreamingObject->ForEachStreamingCells([&CellCount](const UWorldPartitionRuntimeCell& Cell)
@@ -344,8 +344,6 @@ void FContentBundleEditor::GenerateStreaming(TArray<FString>* OutPackageToGenera
 		// Clear streaming object. It will be kept alive & duplicated by UContentBundleDuplicateForPIEHelper.
 		ExternalStreamingObject = nullptr;
 	}
-
-	WorldPartition->FlushStreaming();
 }
 
 void FContentBundleEditor::OnBeginCook(IWorldPartitionCookPackageContext& CookContext)
diff --git a/Engine/Source/Runtime/Engine/Private/WorldPartition/WorldPartition.cpp b/Engine/Source/Runtime/Engine/Private/WorldPartition/WorldPartition.cpp
--- a/Engine/Source/Runtime/Engine/Private/WorldPartition/WorldPartition.cpp
+++ b/Engine/Source/Runtime/Engine/Private/WorldPartition/WorldPartition.cpp
@@ -1473,6 +1473,7 @@ bool UWorldPartition::InjectExternalStreamingObject(URuntimeHashExternalStreamin
 	bool bInjected = RuntimeHash->InjectExternalStreamingObject(ExternalStreamingObject);
 	if (bInjected)
 	{
+		StreamingPolicy->InjectExternalStreamingObject(ExternalStreamingObject);
 		GetWorld()->GetSubsystem<UHLODSubsystem>()->OnExternalStreamingObjectInjected(ExternalStreamingObject);
 	}
 
@@ -1484,6 +1485,11 @@ bool UWorldPartition::RemoveExternalStreamingObject(URuntimeHashExternalStreamin
 	bool bRemoved = RuntimeHash->RemoveExternalStreamingObject(ExternalStreamingObject);
 	if (bRemoved)
 	{
+		if (StreamingPolicy)
+		{
+			StreamingPolicy->RemoveExternalStreamingObject(ExternalStreamingObject);
+		}
+		
 		GetWorld()->GetSubsystem<UHLODSubsystem>()->OnExternalStreamingObjectRemoved(ExternalStreamingObject);
 	}
 
diff --git a/Engine/Source/Runtime/Engine/Private/WorldPartition/WorldPartitionLevelStreamingPolicy.cpp b/Engine/Source/Runtime/Engine/Private/WorldPartition/WorldPartitionLevelStreamingPolicy.cpp
--- a/Engine/Source/Runtime/Engine/Private/WorldPartition/WorldPartitionLevelStreamingPolicy.cpp
+++ b/Engine/Source/Runtime/Engine/Private/WorldPartition/WorldPartitionLevelStreamingPolicy.cpp
@@ -135,6 +135,18 @@ void UWorldPartitionLevelStreamingPolicy::RemapSoftObjectPath(FSoftObjectPath& O
 	const FSoftObjectPath SrcPath(ObjectPath);
 	ConvertEditorPathToRuntimePath(SrcPath, ObjectPath);
 }
+
+bool UWorldPartitionLevelStreamingPolicy::StoreToExternalStreamingObject(URuntimeHashExternalStreamingObjectBase& OutExternalStreamingObject)
+{
+	if (Super::StoreToExternalStreamingObject(OutExternalStreamingObject))
+	{
+		OutExternalStreamingObject.SubObjectsToCellRemapping = MoveTemp(SubObjectsToCellRemapping);
+		return true;
+	}
+
+	return false;
+}
+
 #endif
 
 bool UWorldPartitionLevelStreamingPolicy::ConvertEditorPathToRuntimePath(const FSoftObjectPath& InPath, FSoftObjectPath& OutPath) const
@@ -192,7 +204,7 @@ bool UWorldPartitionLevelStreamingPolicy::ConvertEditorPathToRuntimePath(const F
 			SubAssetName.Split(TEXT("."), &SubObjectContext, &SubObjectName);
 
 			// Try to find the corresponding streaming cell, if it doesn't exists the actor must be in the persistent level.
-			const FName* CellName = SubObjectsToCellRemapping.Find(*SubObjectContext);
+			const FName* CellName = FindCellNameForSubObject(*SubObjectContext);
 			if (!CellName)
 			{
 				OutPath = FSoftObjectPath(WorldAssetPath, InPath.GetSubPathString());
@@ -232,16 +244,13 @@ UObject* UWorldPartitionLevelStreamingPolicy::GetSubObject(const TCHAR* SubObjec
 	}
 
 	const FString SrcPath = UWorld::RemovePIEPrefix(*SubObjectContext);
-	if (FName* CellName = SubObjectsToCellRemapping.Find(FName(*SrcPath)))
+	if (const UWorldPartitionRuntimeLevelStreamingCell* Cell = FindCellForSubObject(*SrcPath))
 	{
-		if (UWorldPartitionRuntimeLevelStreamingCell* Cell = (UWorldPartitionRuntimeLevelStreamingCell*)StaticFindObject(UWorldPartitionRuntimeLevelStreamingCell::StaticClass(), GetOuterUWorldPartition()->RuntimeHash, *(CellName->ToString())))
+		if (UWorldPartitionLevelStreamingDynamic* LevelStreaming = Cell->GetLevelStreaming())
 		{
-			if (UWorldPartitionLevelStreamingDynamic* LevelStreaming = Cell->GetLevelStreaming())
+			if (LevelStreaming->GetLoadedLevel())
 			{
-				if (LevelStreaming->GetLoadedLevel())
-				{
-					return StaticFindObject(UObject::StaticClass(), LevelStreaming->GetLoadedLevel(), SubObjectPath);
-				}
+				return StaticFindObject(UObject::StaticClass(), LevelStreaming->GetLoadedLevel(), SubObjectPath);
 			}
 		}
 	}
@@ -249,6 +258,27 @@ UObject* UWorldPartitionLevelStreamingPolicy::GetSubObject(const TCHAR* SubObjec
 	return nullptr;
 }
 
+bool UWorldPartitionLevelStreamingPolicy::InjectExternalStreamingObject(URuntimeHashExternalStreamingObjectBase* ExternalStreamingObject)
+{
+	if (Super::InjectExternalStreamingObject(ExternalStreamingObject))
+	{
+		ExternalStreamingObjects.Add(ExternalStreamingObject);
+		return true;
+	}
+
+	return false;
+}
+
+bool UWorldPartitionLevelStreamingPolicy::RemoveExternalStreamingObject(URuntimeHashExternalStreamingObjectBase* ExternalStreamingObject)
+{
+	bool bSuccess = Super::RemoveExternalStreamingObject(ExternalStreamingObject);
+
+	ExternalStreamingObjects.RemoveSwap(ExternalStreamingObject);
+
+	return bSuccess;
+}
+
+
 void UWorldPartitionLevelStreamingPolicy::DrawRuntimeCellsDetails(UCanvas* Canvas, FVector2D& Offset)
 {
 	TRACE_CPUPROFILER_EVENT_SCOPE(UWorldPartitionLevelStreamingPolicy::DrawRuntimeCellsDetails);
@@ -314,4 +344,61 @@ void UWorldPartitionLevelStreamingPolicy::DrawRuntimeCellsDetails(UCanvas* Canva
 	}
 
 	Offset.Y = MaxPosY;
+}
+
+const FName* UWorldPartitionLevelStreamingPolicy::FindCellNameForSubObject(FName SubObjectName) const
+{
+	if (const FName* CellName = SubObjectsToCellRemapping.Find(SubObjectName))
+	{
+		return CellName;
+	}
+
+	for (const TWeakObjectPtr<URuntimeHashExternalStreamingObjectBase>& ExternalStreamingObject : ExternalStreamingObjects)
+	{
+		if (ExternalStreamingObject.IsValid())
+		{
+			if (const FName* CellName = ExternalStreamingObject.Get()->SubObjectsToCellRemapping.Find(SubObjectName))
+			{
+				return CellName;
+			}
+		}
+	}
+	
+	return nullptr;
+}
+
+const UWorldPartitionRuntimeLevelStreamingCell* UWorldPartitionLevelStreamingPolicy::FindCellForSubObject(FName SubObjectName) const
+{
+	FName CellName;
+	UObject* CellPackage = nullptr;
+	if (const FName* FoundCell = SubObjectsToCellRemapping.Find(SubObjectName))
+	{
+		CellName = *FoundCell;
+		CellPackage = GetOuterUWorldPartition()->RuntimeHash;
+	}
+
+	if (!CellPackage)
+	{
+		for (const TWeakObjectPtr<URuntimeHashExternalStreamingObjectBase>& ExternalStreamingObject : ExternalStreamingObjects)
+		{
+			if (ExternalStreamingObject.IsValid())
+			{
+				if (const FName* FoundCell = ExternalStreamingObject->SubObjectsToCellRemapping.Find(SubObjectName))
+				{
+					CellName = *FoundCell;
+					CellPackage = ExternalStreamingObject.Get();
+					break;
+				}
+			}
+		}
+	}
+	
+
+	if (CellPackage)
+	{
+		check(CellName.IsValid());
+		return (UWorldPartitionRuntimeLevelStreamingCell*)StaticFindObject(UWorldPartitionRuntimeLevelStreamingCell::StaticClass(), CellPackage, *(CellName.ToString()));
+	}
+
+	return nullptr;
 }
\ No newline at end of file
diff --git a/Engine/Source/Runtime/Engine/Private/WorldPartition/WorldPartitionStreamingGeneration.cpp b/Engine/Source/Runtime/Engine/Private/WorldPartition/WorldPartitionStreamingGeneration.cpp
--- a/Engine/Source/Runtime/Engine/Private/WorldPartition/WorldPartitionStreamingGeneration.cpp
+++ b/Engine/Source/Runtime/Engine/Private/WorldPartition/WorldPartitionStreamingGeneration.cpp
@@ -15,6 +15,7 @@
 #include "ProfilingDebugging/ScopedTimers.h"
 #include "WorldPartition/WorldPartitionRuntimeHash.h"
 #include "WorldPartition/WorldPartitionStreamingPolicy.h"
+#include "WorldPartition/WorldPartitionLevelStreamingPolicy.h"
 #include "WorldPartition/DataLayer/DataLayerManager.h"
 #include "WorldPartition/DataLayer/DataLayerUtils.h"
 #include "WorldPartition/ErrorHandling/WorldPartitionStreamingGenerationNullErrorHandler.h"
@@ -1193,6 +1194,15 @@ void UWorldPartition::FlushStreaming()
 	GeneratedStreamingPackageNames.Empty();
 }
 
+URuntimeHashExternalStreamingObjectBase* UWorldPartition::FlushStreamingToExternalStreamingObject(const FString& ExternalStreamingObjectName)
+{
+	URuntimeHashExternalStreamingObjectBase* ExternalStreamingObject = RuntimeHash->StoreToExternalStreamingObject(this, *ExternalStreamingObjectName);
+	StreamingPolicy->StoreToExternalStreamingObject(*ExternalStreamingObject);
+
+	FlushStreaming();
+	return ExternalStreamingObject;
+}
+
 void UWorldPartition::GenerateHLOD(ISourceControlHelper* SourceControlHelper, bool bCreateActorsOnly)
 {
 	ForEachActorDescContainer([this, &SourceControlHelper, bCreateActorsOnly](UActorDescContainer* InActorDescContainer)
diff --git a/Engine/Source/Runtime/Engine/Public/WorldPartition/WorldPartition.h b/Engine/Source/Runtime/Engine/Public/WorldPartition/WorldPartition.h
--- a/Engine/Source/Runtime/Engine/Public/WorldPartition/WorldPartition.h
+++ b/Engine/Source/Runtime/Engine/Public/WorldPartition/WorldPartition.h
@@ -209,6 +209,7 @@ public:
 	bool GenerateContainerStreaming(const FGenerateStreamingParams& InParams, FGenerateStreamingContext& InContext);
 
 	void FlushStreaming();
+	URuntimeHashExternalStreamingObjectBase* FlushStreamingToExternalStreamingObject(const FString& ExternalStreamingObjectName);
 
 	DECLARE_MULTICAST_DELEGATE_OneParam(FWorldPartitionGenerateStreamingDelegate, TArray<FString>*);
 	FWorldPartitionGenerateStreamingDelegate OnPreGenerateStreaming;
diff --git a/Engine/Source/Runtime/Engine/Public/WorldPartition/WorldPartitionLevelStreamingPolicy.h b/Engine/Source/Runtime/Engine/Public/WorldPartition/WorldPartitionLevelStreamingPolicy.h
--- a/Engine/Source/Runtime/Engine/Public/WorldPartition/WorldPartitionLevelStreamingPolicy.h
+++ b/Engine/Source/Runtime/Engine/Public/WorldPartition/WorldPartitionLevelStreamingPolicy.h
@@ -14,6 +14,7 @@
 #include "WorldPartition/WorldPartitionStreamingPolicy.h"
 #include "WorldPartitionLevelStreamingPolicy.generated.h"
 
+class UWorldPartitionRuntimeLevelStreamingCell;
 enum class EWorldPartitionRuntimeCellState : uint8;
 
 UCLASS()
@@ -30,20 +31,31 @@ public:
 	virtual void PrepareActorToCellRemapping() override;
 	virtual void RemapSoftObjectPath(FSoftObjectPath& ObjectPath) const override;
 	static FString GetCellPackagePath(const FName& InCellName, const UWorld* InWorld);
+
+	virtual bool StoreToExternalStreamingObject(URuntimeHashExternalStreamingObjectBase& OutExternalStreamingObject) override;
 #endif
 
 	virtual bool ConvertEditorPathToRuntimePath(const FSoftObjectPath& InPath, FSoftObjectPath& OutPath) const override;
 	virtual UObject* GetSubObject(const TCHAR* SubObjectPath) override;
 
+	virtual bool InjectExternalStreamingObject(URuntimeHashExternalStreamingObjectBase* ExternalStreamingObject) override;
+	virtual bool RemoveExternalStreamingObject(URuntimeHashExternalStreamingObjectBase* ExternalStreamingObject) override;
+
 protected:
 	virtual int32 GetCellLoadingCount() const override;
 
 	void ForEachActiveRuntimeCell(TFunctionRef<void(const UWorldPartitionRuntimeCell*)> Func) const;
 
 private:
+	const FName* FindCellNameForSubObject(FName SubObjectName) const;
+	const UWorldPartitionRuntimeLevelStreamingCell* FindCellForSubObject(FName SubObjectName) const;
+
 	UPROPERTY()
 	FTopLevelAssetPath SourceWorldAssetPath;
 
 	UPROPERTY()
 	TMap<FName, FName> SubObjectsToCellRemapping;
+
+	UPROPERTY(Transient)
+	TArray<TWeakObjectPtr<URuntimeHashExternalStreamingObjectBase>> ExternalStreamingObjects;
 };
\ No newline at end of file
diff --git a/Engine/Source/Runtime/Engine/Public/WorldPartition/WorldPartitionRuntimeHash.h b/Engine/Source/Runtime/Engine/Public/WorldPartition/WorldPartitionRuntimeHash.h
--- a/Engine/Source/Runtime/Engine/Public/WorldPartition/WorldPartitionRuntimeHash.h
+++ b/Engine/Source/Runtime/Engine/Public/WorldPartition/WorldPartitionRuntimeHash.h
@@ -52,6 +52,9 @@ public:
 	void PopulateGeneratorPackageForCook();
 #endif
 
+	UPROPERTY();
+	TMap<FName, FName> SubObjectsToCellRemapping;
+
 protected:
 	UPROPERTY();
 	TSoftObjectPtr<UWorld> OwningWorld;
diff --git a/Engine/Source/Runtime/Engine/Public/WorldPartition/WorldPartitionStreamingPolicy.h b/Engine/Source/Runtime/Engine/Public/WorldPartition/WorldPartitionStreamingPolicy.h
--- a/Engine/Source/Runtime/Engine/Public/WorldPartition/WorldPartitionStreamingPolicy.h
+++ b/Engine/Source/Runtime/Engine/Public/WorldPartition/WorldPartitionStreamingPolicy.h
@@ -44,6 +44,8 @@ public:
 	// PIE/Game methods
 	virtual void PrepareActorToCellRemapping() {}
 	virtual void RemapSoftObjectPath(FSoftObjectPath& ObjectPath) const {}
+
+	virtual bool StoreToExternalStreamingObject(URuntimeHashExternalStreamingObjectBase& OutExternalStreamingObject) { return true; }
 #endif
 
 #if !UE_BUILD_SHIPPING
@@ -60,6 +62,9 @@ public:
 
 	static bool IsUpdateStreamingOptimEnabled();
 
+	virtual bool InjectExternalStreamingObject(URuntimeHashExternalStreamingObjectBase* ExternalStreamingObject) { return true; }
+	virtual bool RemoveExternalStreamingObject(URuntimeHashExternalStreamingObjectBase* ExternalStreamingObject) { return true; }
+
 protected:
 	virtual int32 SetCellsStateToLoaded(const TArray<const UWorldPartitionRuntimeCell*>& ToLoadCells);
 	virtual int32 SetCellsStateToActivated(const TArray<const UWorldPartitionRuntimeCell*>& ToActivateCells);
