Fixed infinite loop in Level streaming caused by UWorld::RemoveFromWorld completing in non incremental removal move after first starting in incremental mode and not resetting CurrentLevelPendingInvisibility.
The code has never handled properly this case. It can now happen because UWorldPartitionLevelStreamingDynamic::ShouldBlockOnUnload overloads the default implementation.
UWorldPartitionLevelStreamingDynamic::ShouldBlockOnUnload can return false then true for the same level being removed.
The reason this can now happen is because it internally calls UWorldPartition::CanStream which recently changed and can now return true then false for a world partition cell level.

diff --git a/Engine/Source/Runtime/Engine/Private/World.cpp b/Engine/Source/Runtime/Engine/Private/World.cpp
--- a/Engine/Source/Runtime/Engine/Private/World.cpp
+++ b/Engine/Source/Runtime/Engine/Private/World.cpp
@@ -3524,6 +3524,14 @@ void UWorld::RemoveFromWorld( ULevel* Level, bool bAllowIncrementalRemoval, FNet
 				}
 			}
 		}
+		else if (Level == CurrentLevelPendingInvisibility)
+		{
+			// Finish current level pending invisibility without time limit
+			check(Level->bIsBeingRemoved);
+			while (!Level->IncrementalUnregisterComponents(MAX_int32)) {}
+			CurrentLevelPendingInvisibility = nullptr;
+			bFinishRemovingLevel = true;
+		}
 		else
 		{
 			BeginRemoval();
