Fix instanced static mesh loose params warnings in debug build. Create a loose parameter uniform buffer to capture those parameters.

diff --git a/Engine/Shaders/Private/LocalVertexFactory.ush b/Engine/Shaders/Private/LocalVertexFactory.ush
--- a/Engine/Shaders/Private/LocalVertexFactory.ush
+++ b/Engine/Shaders/Private/LocalVertexFactory.ush
@@ -12,17 +12,6 @@
 
 #include "/Engine/Generated/UniformBuffers/PrecomputedLightingBuffer.ush"
 
-#if USE_INSTANCE_CULLING || USE_INSTANCING
-	#if USE_DITHERED_LOD_TRANSITION
-		float4 InstancingViewZCompareZero;  // w contains random lod scale 
-		float4 InstancingViewZCompareOne;
-		float4 InstancingViewZConstant;
-		float4 InstancingTranslatedWorldViewOriginZero;
-		float4 InstancingTranslatedWorldViewOriginOne;
-	#endif
-	float4 InstancingFadeOutParams;
-#endif
-
 #if USE_INSTANCING
 float4 InstancingOffset;
 uint InstanceOffset;
@@ -895,22 +884,22 @@ FVertexFactoryIntermediates GetVertexFactoryIntermediates(FVertexFactoryInput In
 		
 		// PerInstanceParams.x: per-instance fade out amount
 		float3 InstanceTranslatedLocation = TransformLocalToTranslatedWorld(GetInstanceData(Intermediates).LocalBoundsCenter, GetInstanceData(Intermediates).LocalToWorld).xyz;
-		Intermediates.PerInstanceParams.x = 1.0 - saturate((length(InstanceTranslatedLocation) - InstancingFadeOutParams.x) * InstancingFadeOutParams.y);
+		Intermediates.PerInstanceParams.x = 1.0 - saturate((length(InstanceTranslatedLocation) - InstancedVFLooseParameters.InstancingFadeOutParams.x) * InstancedVFLooseParameters.InstancingFadeOutParams.y);
 		
 		// InstancingFadeOutParams.z,w are RenderSelected and RenderDeselected respectively.
 
 		// PerInstanceParams.y = hide / show flag, 
-		Intermediates.IsVisible = lerp(InstancingFadeOutParams.w, InstancingFadeOutParams.z, Intermediates.IsSelected);
+		Intermediates.IsVisible = lerp(InstancedVFLooseParameters.InstancingFadeOutParams.w, InstancedVFLooseParameters.InstancingFadeOutParams.z, Intermediates.IsSelected);
 		Intermediates.PerInstanceParams.y = Intermediates.IsVisible;
 
 		// PerInstanceParams.z dither fade cutoff
 	#if USE_DITHERED_LOD_TRANSITION
-		float RandomLOD = InstancingViewZCompareZero.w * Intermediates.SceneData.InstanceData.RandomID;
-		float ViewZZero = length(InstanceTranslatedLocation - InstancingTranslatedWorldViewOriginZero.xyz) + RandomLOD;
-		float ViewZOne = length(InstanceTranslatedLocation - InstancingTranslatedWorldViewOriginOne.xyz) + RandomLOD;
+		float RandomLOD = InstancedVFLooseParameters.InstancingViewZCompareZero.w * Intermediates.SceneData.InstanceData.RandomID; 
+		float ViewZZero = length(InstanceTranslatedLocation - InstancedVFLooseParameters.InstancingTranslatedWorldViewOriginZero.xyz) + RandomLOD;
+		float ViewZOne = length(InstanceTranslatedLocation - InstancedVFLooseParameters.InstancingTranslatedWorldViewOriginOne.xyz) + RandomLOD;
 		Intermediates.PerInstanceParams.z =
-			dot(float3(ViewZZero.xxx > InstancingViewZCompareZero.xyz), InstancingViewZConstant.xyz) * InstancingTranslatedWorldViewOriginZero.w +
-			dot(float3(ViewZOne.xxx > InstancingViewZCompareOne.xyz), InstancingViewZConstant.xyz) * InstancingTranslatedWorldViewOriginOne.w;
+			dot(float3(ViewZZero.xxx > InstancedVFLooseParameters.InstancingViewZCompareZero.xyz), InstancedVFLooseParameters.InstancingViewZConstant.xyz) * InstancedVFLooseParameters.InstancingTranslatedWorldViewOriginZero.w +
+			dot(float3(ViewZOne.xxx > InstancedVFLooseParameters.InstancingViewZCompareOne.xyz), InstancedVFLooseParameters.InstancingViewZConstant.xyz) * InstancedVFLooseParameters.InstancingTranslatedWorldViewOriginOne.w;
 		Intermediates.PerInstanceParams.y *= abs(Intermediates.PerInstanceParams.z) < .999;
 		Intermediates.IsVisible = Intermediates.PerInstanceParams.y;
 	#else
@@ -960,16 +949,16 @@ FVertexFactoryIntermediates GetVertexFactoryIntermediates(FVertexFactoryInput In
 	// GPUCULL_TODO: This can't be right
 	float3 InstanceTranslatedLocation = TransformLocalToTranslatedWorld(GetInstanceOrigin(Intermediates), PrimitiveData.LocalToWorld).xyz;
 	float InstanceViewDistance = length(InstanceTranslatedLocation);
-	Intermediates.PerInstanceParams.x = 1.0 - saturate((InstanceViewDistance - InstancingFadeOutParams.x) * InstancingFadeOutParams.y);
+	Intermediates.PerInstanceParams.x = 1.0 - saturate((InstanceViewDistance - InstancedVFLooseParameters.InstancingFadeOutParams.x) * InstancedVFLooseParameters.InstancingFadeOutParams.y);
 	// InstancingFadeOutParams.z,w are RenderSelected and RenderDeselected respectively.
-	Intermediates.PerInstanceParams.y = InstancingFadeOutParams.z * SelectedValue + InstancingFadeOutParams.w * (1-SelectedValue);
+	Intermediates.PerInstanceParams.y = InstancedVFLooseParameters.InstancingFadeOutParams.z * SelectedValue + InstancedVFLooseParameters.InstancingFadeOutParams.w * (1-SelectedValue);
 	#if USE_DITHERED_LOD_TRANSITION
-		float RandomLOD = InstancingViewZCompareZero.w * Intermediates.SceneData.InstanceData.RandomID;
-		float ViewZZero = length(InstanceTranslatedLocation - InstancingTranslatedWorldViewOriginZero.xyz) + RandomLOD;
-		float ViewZOne = length(InstanceTranslatedLocation - InstancingTranslatedWorldViewOriginOne.xyz) + RandomLOD;
+		float RandomLOD = InstancedVFLooseParameters.InstancingViewZCompareZero.w * Intermediates.SceneData.InstanceData.RandomID;
+		float ViewZZero = length(InstanceTranslatedLocation - InstancedVFLooseParameters.InstancingTranslatedWorldViewOriginZero.xyz) + RandomLOD;
+		float ViewZOne = length(InstanceTranslatedLocation - InstancedVFLooseParameters.InstancingTranslatedWorldViewOriginOne.xyz) + RandomLOD;
 		Intermediates.PerInstanceParams.z = 
-			dot(float3(ViewZZero.xxx > InstancingViewZCompareZero.xyz), InstancingViewZConstant.xyz) * InstancingTranslatedWorldViewOriginZero.w +
-			dot(float3(ViewZOne.xxx > InstancingViewZCompareOne.xyz), InstancingViewZConstant.xyz) * InstancingTranslatedWorldViewOriginOne.w;
+			dot(float3(ViewZZero.xxx > InstancedVFLooseParameters.InstancingViewZCompareZero.xyz), InstancedVFLooseParameters.InstancingViewZConstant.xyz) * InstancedVFLooseParameters.InstancingTranslatedWorldViewOriginZero.w +
+			dot(float3(ViewZOne.xxx > InstancedVFLooseParameters.InstancingViewZCompareOne.xyz), InstancedVFLooseParameters.InstancingViewZConstant.xyz) * InstancedVFLooseParameters.InstancingTranslatedWorldViewOriginOne.w;
 		Intermediates.PerInstanceParams.y *= abs(Intermediates.PerInstanceParams.z) < .999;
 	#else
 		Intermediates.PerInstanceParams.z = 0;
diff --git a/Engine/Source/Runtime/Engine/Classes/Engine/InstancedStaticMesh.h b/Engine/Source/Runtime/Engine/Classes/Engine/InstancedStaticMesh.h
--- a/Engine/Source/Runtime/Engine/Classes/Engine/InstancedStaticMesh.h
+++ b/Engine/Source/Runtime/Engine/Classes/Engine/InstancedStaticMesh.h
@@ -46,6 +46,17 @@ BEGIN_GLOBAL_SHADER_PARAMETER_STRUCT(FInstancedStaticMeshVertexFactoryUniformSha
 	SHADER_PARAMETER(int32, NumCustomDataFloats)
 END_GLOBAL_SHADER_PARAMETER_STRUCT()
 
+BEGIN_GLOBAL_SHADER_PARAMETER_STRUCT(FInstancedStaticMeshVFLooseUniformShaderParameters, ENGINE_API)
+	SHADER_PARAMETER(FVector4f, InstancingViewZCompareZero)
+	SHADER_PARAMETER(FVector4f, InstancingViewZCompareOne)
+	SHADER_PARAMETER(FVector4f, InstancingViewZConstant)
+	SHADER_PARAMETER(FVector4f, InstancingTranslatedWorldViewOriginZero)
+	SHADER_PARAMETER(FVector4f, InstancingTranslatedWorldViewOriginOne)
+	SHADER_PARAMETER(FVector4f, InstancingFadeOutParams)
+END_GLOBAL_SHADER_PARAMETER_STRUCT()
+
+typedef TUniformBufferRef<FInstancedStaticMeshVFLooseUniformShaderParameters> FInstancedStaticMeshVFLooseUniformShaderParametersRef;
+
 // This must match the maximum a user could specify in the material (see 
 // FHLSLMaterialTranslator::TextureCoordinate), otherwise the material will attempt 
 // to look up a texture coordinate we didn't provide an element for.
@@ -354,13 +365,7 @@ public:
 	{
 		FLocalVertexFactoryShaderParametersBase::Bind(ParameterMap);
 
-		InstancingFadeOutParamsParameter.Bind(ParameterMap, TEXT("InstancingFadeOutParams"));
-		InstancingViewZCompareZeroParameter.Bind(ParameterMap, TEXT("InstancingViewZCompareZero"));
-		InstancingViewZCompareOneParameter.Bind(ParameterMap, TEXT("InstancingViewZCompareOne"));
-		InstancingViewZConstantParameter.Bind(ParameterMap, TEXT("InstancingViewZConstant"));
 		InstancingOffsetParameter.Bind(ParameterMap, TEXT("InstancingOffset"));
-		InstancingWorldViewOriginZeroParameter.Bind(ParameterMap, TEXT("InstancingTranslatedWorldViewOriginZero"));
-		InstancingWorldViewOriginOneParameter.Bind(ParameterMap, TEXT("InstancingTranslatedWorldViewOriginOne"));
 		VertexFetch_InstanceOriginBufferParameter.Bind(ParameterMap, TEXT("VertexFetch_InstanceOriginBuffer"));
 		VertexFetch_InstanceTransformBufferParameter.Bind(ParameterMap, TEXT("VertexFetch_InstanceTransformBuffer"));
 		VertexFetch_InstanceLightmapBufferParameter.Bind(ParameterMap, TEXT("VertexFetch_InstanceLightmapBuffer"));
@@ -380,15 +385,7 @@ public:
 		) const;
 
 private:
-	
-	LAYOUT_FIELD(FShaderParameter, InstancingFadeOutParamsParameter)
-	LAYOUT_FIELD(FShaderParameter, InstancingViewZCompareZeroParameter)
-	LAYOUT_FIELD(FShaderParameter, InstancingViewZCompareOneParameter)
-	LAYOUT_FIELD(FShaderParameter, InstancingViewZConstantParameter)
 	LAYOUT_FIELD(FShaderParameter, InstancingOffsetParameter);
-	LAYOUT_FIELD(FShaderParameter, InstancingWorldViewOriginZeroParameter)
-	LAYOUT_FIELD(FShaderParameter, InstancingWorldViewOriginOneParameter)
-
 	LAYOUT_FIELD(FShaderResourceParameter, VertexFetch_InstanceOriginBufferParameter)
 	LAYOUT_FIELD(FShaderResourceParameter, VertexFetch_InstanceTransformBufferParameter)
 	LAYOUT_FIELD(FShaderResourceParameter, VertexFetch_InstanceLightmapBufferParameter)
@@ -600,6 +597,8 @@ public:
 	}
 
 protected:
+	FInstancedStaticMeshVFLooseUniformShaderParametersRef CreateLooseUniformBuffer(const FSceneView* View, const FInstancingUserData* InstancingUserData, uint32 InstancedLODRange, uint32 InstancedLODIndex, EUniformBufferUsage UniformBufferUsage) const;
+
 	/** Cache of the StaticMesh asset, needed to release SpeedTree resources*/
 	UStaticMesh* StaticMesh;
 
@@ -638,6 +637,9 @@ protected:
 private:
 
 	void SetupProxy(UInstancedStaticMeshComponent* InComponent);
+
+	/** Stores a loose uniform buffer per LOD, used for static view relevance. */
+	TMap<uint32, FInstancedStaticMeshVFLooseUniformShaderParametersRef> LODLooseUniformBuffers;
 };
 
 #if WITH_EDITOR
diff --git a/Engine/Source/Runtime/Engine/Private/HierarchicalInstancedStaticMesh.cpp b/Engine/Source/Runtime/Engine/Private/HierarchicalInstancedStaticMesh.cpp
--- a/Engine/Source/Runtime/Engine/Private/HierarchicalInstancedStaticMesh.cpp
+++ b/Engine/Source/Runtime/Engine/Private/HierarchicalInstancedStaticMesh.cpp
@@ -1272,7 +1272,7 @@ struct FFoliageElementParams
 	float FinalCullDistance;
 };
 
-void FHierarchicalStaticMeshSceneProxy::FillDynamicMeshElements(FMeshElementCollector& Collector, const FFoliageElementParams& ElementParams, const FFoliageRenderInstanceParams& Params) const
+void FHierarchicalStaticMeshSceneProxy::FillDynamicMeshElements(const FSceneView* View, FMeshElementCollector& Collector, const FFoliageElementParams& ElementParams, const FFoliageRenderInstanceParams& Params) const
 {
 	SCOPE_CYCLE_COUNTER(STAT_FoliageBatchTime);
 	int64 TotalTriangles = 0;
@@ -1290,11 +1290,14 @@ void FHierarchicalStaticMeshSceneProxy::FillDynamicMeshElements(FMeshElementColl
 
 		for (int32 SelectionGroupIndex = 0; SelectionGroupIndex < ElementParams.NumSelectionGroups; SelectionGroupIndex++)
 		{
+			const bool bDitherLODEnabled = ElementParams.bBlendLODs;
+			const uint32 InstancedLODRange = bDitherLODEnabled ? 1 : 0;
+			FInstancedStaticMeshVFLooseUniformShaderParametersRef LooseUniformBuffer = CreateLooseUniformBuffer(View, ElementParams.PassUserData[SelectionGroupIndex], InstancedLODRange, LODIndex, EUniformBufferUsage::UniformBuffer_SingleFrame);
+
 			for (int32 SectionIndex = 0; SectionIndex < TotalNumSections; SectionIndex++)
 			{
 				const FLODInfo& ProxyLODInfo = LODs[LODIndex];
 				UMaterialInterface* Material = ProxyLODInfo.Sections[SectionIndex].Material;
-				const bool bDitherLODEnabled = ElementParams.bBlendLODs;
 
 				TArray<uint32, SceneRenderingAllocator>& RunArray = bDitherLODEnabled ? Params.MultipleLODRuns[LODIndex] : Params.SingleLODRuns[LODIndex];
 
@@ -1326,8 +1329,9 @@ void FHierarchicalStaticMeshSceneProxy::FillDynamicMeshElements(FMeshElementColl
 					MeshBatchElement.MaxScreenSize = 1.0;
 					MeshBatchElement.MinScreenSize = 0.0;
 					MeshBatchElement.InstancedLODIndex = LODIndex;
-					MeshBatchElement.InstancedLODRange = bDitherLODEnabled ? 1 : 0;
+					MeshBatchElement.InstancedLODRange = InstancedLODRange;
 					MeshBatchElement.PrimitiveUniformBuffer = GetUniformBuffer();
+					MeshBatchElement.LooseParametersUniformBuffer = LooseUniformBuffer;
 
 					int32 TotalInstances = bDitherLODEnabled ? Params.TotalMultipleLODInstances[LODIndex] : Params.TotalSingleLODInstances[LODIndex];
 					{
@@ -1422,8 +1426,9 @@ void FHierarchicalStaticMeshSceneProxy::FillDynamicMeshElements(FMeshElementColl
 						BatchElement0.MaxScreenSize = 1.0;
 						BatchElement0.MinScreenSize = 0.0;
 						BatchElement0.InstancedLODIndex = LODIndex;
-						BatchElement0.InstancedLODRange = bDitherLODEnabled ? 1 : 0;
+						BatchElement0.InstancedLODRange = InstancedLODRange;
 						BatchElement0.PrimitiveUniformBuffer = GetUniformBuffer();
+						BatchElement0.LooseParametersUniformBuffer = LooseUniformBuffer;
 						MeshElement.bCanApplyViewModeOverrides = true;
 						MeshElement.bUseSelectionOutline = ElementParams.BatchRenderSelection[SelectionGroupIndex];
 						MeshElement.bUseWireframeSelectionColoring = ElementParams.BatchRenderSelection[SelectionGroupIndex];
@@ -1855,7 +1860,7 @@ void FHierarchicalStaticMeshSceneProxy::GetDynamicMeshElements(const TArray<cons
 #endif
 				}
 
-				FillDynamicMeshElements(Collector, ElementParams, InstanceParams);
+				FillDynamicMeshElements(View, Collector, ElementParams, InstanceParams);
 			}
 
 			int32 UnbuiltInstanceCount = InstanceCountToRender - FirstUnbuiltIndex;
@@ -1968,7 +1973,7 @@ void FHierarchicalStaticMeshSceneProxy::GetDynamicMeshElements(const TArray<cons
 					const int8 LowestLOD = (RenderData->LODResources.Num() - 1);
 					InstanceParams.AddRun(LowestLOD, LowestLOD, FirstUnbuiltIndex, FirstUnbuiltIndex + UnbuiltInstanceCount - 1);
 				}
-				FillDynamicMeshElements(Collector, ElementParams, InstanceParams);
+				FillDynamicMeshElements(View, Collector, ElementParams, InstanceParams);
 			}
 
 			if (View->Family->EngineShowFlags.HISMCOcclusionBounds)
diff --git a/Engine/Source/Runtime/Engine/Private/InstancedStaticMesh.cpp b/Engine/Source/Runtime/Engine/Private/InstancedStaticMesh.cpp
--- a/Engine/Source/Runtime/Engine/Private/InstancedStaticMesh.cpp
+++ b/Engine/Source/Runtime/Engine/Private/InstancedStaticMesh.cpp
@@ -69,6 +69,7 @@
 IMPLEMENT_HIT_PROXY(HInstancedStaticMeshInstance, HHitProxy);
 
 IMPLEMENT_GLOBAL_SHADER_PARAMETER_STRUCT(FInstancedStaticMeshVertexFactoryUniformShaderParameters, "InstanceVF");
+IMPLEMENT_GLOBAL_SHADER_PARAMETER_STRUCT(FInstancedStaticMeshVFLooseUniformShaderParameters, "InstancedVFLooseParameters");
 
 TAutoConsoleVariable<int32> CVarMinLOD(
 	TEXT("foliage.MinLOD"),
@@ -76,6 +77,11 @@
 	TEXT("Used to discard the top LODs for performance evaluation. -1: Disable all effects of this cvar."),
 	ECVF_Scalability | ECVF_Default);
 
+static TAutoConsoleVariable<int32> CVarCullAllInVertexShader(
+	TEXT("foliage.CullAllInVertexShader"),
+	0,
+	TEXT("Debugging, if this is greater than 0, cull all instances in the vertex shader."));
+
 static TAutoConsoleVariable<int32> CVarRayTracingRenderInstances(
 	TEXT("r.RayTracing.Geometry.InstancedStaticMeshes"),
 	1,
@@ -1478,6 +1484,8 @@
 				const int32 LODIndex = GetLOD(View);
 				const FStaticMeshLODResources& LODModel = StaticMesh->GetRenderData()->LODResources[LODIndex];
 
+				FInstancedStaticMeshVFLooseUniformShaderParametersRef LooseUniformBuffer = CreateLooseUniformBuffer(View, PassUserData[SelectionGroupIndex], /*InstancedLODRange=*/0, LODIndex, EUniformBufferUsage::UniformBuffer_SingleFrame);
+
 				for (int32 SectionIndex = 0; SectionIndex < LODModel.Sections.Num(); SectionIndex++)
 				{
 					const int32 NumBatches = GetNumMeshBatches();
@@ -1491,6 +1499,7 @@
 							//@todo-rco this is only supporting selection on the first element
 							MeshElement.Elements[0].UserData = PassUserData[SelectionGroupIndex];
 							MeshElement.Elements[0].bUserDataIsColorVertexBuffer = false;
+							MeshElement.Elements[0].LooseParametersUniformBuffer = LooseUniformBuffer;
 							MeshElement.bCanApplyViewModeOverrides = true;
 							MeshElement.bUseSelectionOutline = BatchRenderSelection[SelectionGroupIndex];
 							MeshElement.bUseWireframeSelectionColoring = BatchRenderSelection[SelectionGroupIndex];
@@ -1547,6 +1556,157 @@
 	}
 }
 
+FInstancedStaticMeshVFLooseUniformShaderParametersRef FInstancedStaticMeshSceneProxy::CreateLooseUniformBuffer(const FSceneView* View, const FInstancingUserData* InstancingUserData, uint32 InstancedLODRange, uint32 InstancedLODIndex, EUniformBufferUsage UniformBufferUsage) const
+{
+	FInstancedStaticMeshVFLooseUniformShaderParameters LooseParameters;
+
+	{
+		FVector4f InstancingViewZCompareZero(MIN_flt, MIN_flt, MAX_flt, 1.0f);
+		FVector4f InstancingViewZCompareOne(MIN_flt, MIN_flt, MAX_flt, 0.0f);
+		FVector4f InstancingViewZConstant(ForceInit);
+		FVector4f InstancingTranslatedWorldViewOriginZero(ForceInit);
+		FVector4f InstancingTranslatedWorldViewOriginOne(ForceInit);
+		InstancingTranslatedWorldViewOriginOne.W = 1.0f;
+
+		// InstancedLODRange is only set for HierarchicalInstancedStaticMeshes
+		if (InstancingUserData && InstancedLODRange)
+		{
+			int32 FirstLOD = InstancingUserData->MinLOD;
+
+			int32 DebugMin = FMath::Min(CVarMinLOD.GetValueOnRenderThread(), InstancingUserData->MeshRenderData->LODResources.Num() - 1);
+			if (DebugMin >= 0)
+			{
+				FirstLOD = FMath::Max(FirstLOD, DebugMin);
+			}
+
+			FBoxSphereBounds ScaledBounds = InstancingUserData->MeshRenderData->Bounds.TransformBy(FTransform(FRotator::ZeroRotator, FVector::ZeroVector, InstancingUserData->AverageInstancesScale));
+			float SphereRadius = ScaledBounds.SphereRadius;
+			float MinSize = View->ViewMatrices.IsPerspectiveProjection() ? CVarFoliageMinimumScreenSize.GetValueOnRenderThread() : 0.0f;
+			float LODScale = InstancingUserData->LODDistanceScale;
+			float LODRandom = CVarRandomLODRange.GetValueOnRenderThread();
+			float MaxDrawDistanceScale = GetCachedScalabilityCVars().ViewDistanceScale;
+
+			if (InstancedLODIndex)
+			{
+				InstancingViewZConstant.X = -1.0f;
+			}
+			else
+			{
+				// this is the first LOD, so we don't have a fade-in region
+				InstancingViewZConstant.X = 0.0f;
+			}
+			InstancingViewZConstant.Y = 0.0f;
+			InstancingViewZConstant.Z = 1.0f;
+
+			// now we subtract off the lower segments, since they will be incorporated 
+			InstancingViewZConstant.Y -= InstancingViewZConstant.X;
+			InstancingViewZConstant.Z -= InstancingViewZConstant.X + InstancingViewZConstant.Y;
+			//not using W InstancingViewZConstant.W -= InstancingViewZConstant.X + InstancingViewZConstant.Y + InstancingViewZConstant.Z;
+
+			for (int32 SampleIndex = 0; SampleIndex < 2; SampleIndex++)
+			{
+				FVector4f& InstancingViewZCompare(SampleIndex ? InstancingViewZCompareOne : InstancingViewZCompareZero);
+
+				float FinalCull = MAX_flt;
+				if (MinSize > 0.0)
+				{
+					FinalCull = ComputeBoundsDrawDistance(MinSize, SphereRadius, View->ViewMatrices.GetProjectionMatrix()) * LODScale;
+				}
+				if (InstancingUserData->EndCullDistance > 0.0f)
+				{
+					FinalCull = FMath::Min(FinalCull, InstancingUserData->EndCullDistance * MaxDrawDistanceScale);
+				}
+
+				InstancingViewZCompare.Z = FinalCull;
+				if (int(InstancedLODIndex) < InstancingUserData->MeshRenderData->LODResources.Num() - 1)
+				{
+					float NextCut = ComputeBoundsDrawDistance(InstancingUserData->MeshRenderData->ScreenSize[InstancedLODIndex + 1].GetValue(), SphereRadius, View->ViewMatrices.GetProjectionMatrix()) * LODScale;
+					InstancingViewZCompare.Z = FMath::Min(NextCut, FinalCull);
+				}
+
+				InstancingViewZCompare.X = MIN_flt;
+				if (int(InstancedLODIndex) > FirstLOD)
+				{
+					float CurCut = ComputeBoundsDrawDistance(InstancingUserData->MeshRenderData->ScreenSize[InstancedLODIndex].GetValue(), SphereRadius, View->ViewMatrices.GetProjectionMatrix()) * LODScale;
+					if (CurCut < FinalCull)
+					{
+						InstancingViewZCompare.Y = CurCut;
+					}
+					else
+					{
+						// this LOD is completely removed by one of the other two factors
+						InstancingViewZCompare.Y = MIN_flt;
+						InstancingViewZCompare.Z = MIN_flt;
+					}
+				}
+				else
+				{
+					// this is the first LOD, so we don't have a fade-in region
+					InstancingViewZCompare.Y = MIN_flt;
+				}
+			}
+
+			const FVector PreViewTranslation = View->ViewMatrices.GetPreViewTranslation();
+			InstancingTranslatedWorldViewOriginZero = (FVector3f)(View->GetTemporalLODOrigin(0) + PreViewTranslation); // LWC_TODO: precision loss
+			InstancingTranslatedWorldViewOriginOne = (FVector3f)(View->GetTemporalLODOrigin(1) + PreViewTranslation); // LWC_TODO: precision loss
+
+			float Alpha = View->GetTemporalLODTransition();
+			InstancingTranslatedWorldViewOriginZero.W = 1.0f - Alpha;
+			InstancingTranslatedWorldViewOriginOne.W = Alpha;
+
+			InstancingViewZCompareZero.W = LODRandom;
+		}
+
+		LooseParameters.InstancingViewZCompareZero = InstancingViewZCompareZero;
+		LooseParameters.InstancingViewZCompareOne = InstancingViewZCompareOne;
+		LooseParameters.InstancingViewZConstant = InstancingViewZConstant;
+		LooseParameters.InstancingTranslatedWorldViewOriginZero = InstancingTranslatedWorldViewOriginZero;
+		LooseParameters.InstancingTranslatedWorldViewOriginOne = InstancingTranslatedWorldViewOriginOne;
+	}
+
+	{
+		FVector4f InstancingFadeOutParams(MAX_flt, 0.f, 1.f, 1.f);
+		if (InstancingUserData)
+		{
+			const float MaxDrawDistanceScale = GetCachedScalabilityCVars().ViewDistanceScale;
+			const float StartDistance = InstancingUserData->StartCullDistance * MaxDrawDistanceScale;
+			const float EndDistance = InstancingUserData->EndCullDistance * MaxDrawDistanceScale;
+
+			InstancingFadeOutParams.X = StartDistance;
+			if (EndDistance > 0)
+			{
+				if (EndDistance > StartDistance)
+				{
+					InstancingFadeOutParams.Y = 1.f / (float)(EndDistance - StartDistance);
+				}
+				else
+				{
+					InstancingFadeOutParams.Y = 1.f;
+				}
+			}
+			else
+			{
+				InstancingFadeOutParams.Y = 0.f;
+			}
+			if (CVarCullAllInVertexShader.GetValueOnRenderThread() > 0)
+			{
+				InstancingFadeOutParams.Z = 0.0f;
+				InstancingFadeOutParams.W = 0.0f;
+			}
+			else
+			{
+				InstancingFadeOutParams.Z = InstancingUserData->bRenderSelected ? 1.f : 0.f;
+				InstancingFadeOutParams.W = InstancingUserData->bRenderUnselected ? 1.f : 0.f;
+			}
+		}
+
+		LooseParameters.InstancingFadeOutParams = InstancingFadeOutParams;
+
+	}
+
+	return CreateUniformBufferImmediate(LooseParameters, UniformBufferUsage);
+}
+
 void FInstancedStaticMeshSceneProxy::SetupProxy(UInstancedStaticMeshComponent* InComponent)
 {
 #if WITH_EDITOR
@@ -1800,6 +1960,12 @@
 			}
 		}
 	}
+
+	for (int32 i = 0; i < LODs.Num(); ++i)
+	{
+		FInstancedStaticMeshVFLooseUniformShaderParametersRef LooseUniformBuffer = CreateLooseUniformBuffer(nullptr, &UserData_AllInstances, 0, i, EUniformBufferUsage::UniformBuffer_MultiFrame);
+		LODLooseUniformBuffers.Add(i, LooseUniformBuffer);
+	}
 }
 
 void FInstancedStaticMeshSceneProxy::DestroyRenderThreadResources()
@@ -1852,6 +2018,7 @@
 	BatchElement0.InstancedLODIndex = LODIndex;
 	BatchElement0.UserIndex = 0;
 	BatchElement0.PrimitiveUniformBuffer = GetUniformBuffer();
+	BatchElement0.LooseParametersUniformBuffer = LODLooseUniformBuffers[LODIndex];
 	
 	if (OutMeshBatch.MaterialRenderProxy)
 	{
@@ -5318,11 +5485,6 @@
 	return true;
 }
 
-static TAutoConsoleVariable<int32> CVarCullAllInVertexShader(
-	TEXT("foliage.CullAllInVertexShader"),
-	0,
-	TEXT("Debugging, if this is greater than 0, cull all instances in the vertex shader."));
-
 void FInstancedStaticMeshVertexFactoryShaderParameters::GetElementShaderBindings(
 	const class FSceneInterface* Scene,
 	const FSceneView* View,
@@ -5363,153 +5525,13 @@
 		}
 	}
 
-	if( InstancingWorldViewOriginOneParameter.IsBound() )
+	FVector4f InstancingOffset(ForceInit);
+	// InstancedLODRange is only set for HierarchicalInstancedStaticMeshes
+	if (InstancingUserData && BatchElement.InstancedLODRange)
 	{
-		FVector4f InstancingViewZCompareZero(MIN_flt, MIN_flt, MAX_flt, 1.0f);
-		FVector4f InstancingViewZCompareOne(MIN_flt, MIN_flt, MAX_flt, 0.0f);
-		FVector4f InstancingViewZConstant(ForceInit);
-		FVector4f InstancingOffset(ForceInit);
-		FVector4f InstancingTranslatedWorldViewOriginZero(ForceInit);
-		FVector4f InstancingTranslatedWorldViewOriginOne(ForceInit);
-		InstancingTranslatedWorldViewOriginOne.W = 1.0f;
+		InstancingOffset = (FVector3f)InstancingUserData->InstancingOffset; // LWC_TODO: precision loss
+	}
+	ShaderBindings.Add(InstancingOffsetParameter, InstancingOffset);
 
-		// InstancedLODRange is only set for HierarchicalInstancedStaticMeshes
-		if (InstancingUserData && BatchElement.InstancedLODRange)
-		{
-			int32 FirstLOD = InstancingUserData->MinLOD;
-
-			int32 DebugMin = FMath::Min(CVarMinLOD.GetValueOnRenderThread(), InstancingUserData->MeshRenderData->LODResources.Num() - 1);
-			if (DebugMin >= 0)
-			{
-				FirstLOD = FMath::Max(FirstLOD, DebugMin);
-			}
-
-			FBoxSphereBounds ScaledBounds = InstancingUserData->MeshRenderData->Bounds.TransformBy(FTransform(FRotator::ZeroRotator, FVector::ZeroVector, InstancingUserData->AverageInstancesScale));
-			float SphereRadius = ScaledBounds.SphereRadius;
-			float MinSize = View->ViewMatrices.IsPerspectiveProjection() ? CVarFoliageMinimumScreenSize.GetValueOnRenderThread() : 0.0f;
-			float LODScale = InstancingUserData->LODDistanceScale;
-			float LODRandom = CVarRandomLODRange.GetValueOnRenderThread();
-			float MaxDrawDistanceScale = GetCachedScalabilityCVars().ViewDistanceScale;
-
-			if (BatchElement.InstancedLODIndex)
-			{
-				InstancingViewZConstant.X = -1.0f;
-			}
-			else
-			{
-				// this is the first LOD, so we don't have a fade-in region
-				InstancingViewZConstant.X = 0.0f;
-			}
-			InstancingViewZConstant.Y = 0.0f;
-			InstancingViewZConstant.Z = 1.0f;
-
-			// now we subtract off the lower segments, since they will be incorporated 
-			InstancingViewZConstant.Y -= InstancingViewZConstant.X;
-			InstancingViewZConstant.Z -= InstancingViewZConstant.X + InstancingViewZConstant.Y;
-			//not using W InstancingViewZConstant.W -= InstancingViewZConstant.X + InstancingViewZConstant.Y + InstancingViewZConstant.Z;
-
-			for (int32 SampleIndex = 0; SampleIndex < 2; SampleIndex++)
-			{
-				FVector4f& InstancingViewZCompare(SampleIndex ? InstancingViewZCompareOne : InstancingViewZCompareZero);
-
-				float FinalCull = MAX_flt;
-				if (MinSize > 0.0)
-				{
-					FinalCull = ComputeBoundsDrawDistance(MinSize, SphereRadius, View->ViewMatrices.GetProjectionMatrix()) * LODScale;
-				}
-				if (InstancingUserData->EndCullDistance > 0.0f)
-				{
-					FinalCull = FMath::Min(FinalCull, InstancingUserData->EndCullDistance * MaxDrawDistanceScale);
-				}
-
-				InstancingViewZCompare.Z = FinalCull;
-				if (int(BatchElement.InstancedLODIndex) < InstancingUserData->MeshRenderData->LODResources.Num() - 1)
-				{
-					float NextCut = ComputeBoundsDrawDistance(InstancingUserData->MeshRenderData->ScreenSize[BatchElement.InstancedLODIndex + 1].GetValue(), SphereRadius, View->ViewMatrices.GetProjectionMatrix()) * LODScale;
-					InstancingViewZCompare.Z = FMath::Min(NextCut, FinalCull);
-				}
-
-				InstancingViewZCompare.X = MIN_flt;
-				if (int(BatchElement.InstancedLODIndex) > FirstLOD)
-				{
-					float CurCut = ComputeBoundsDrawDistance(InstancingUserData->MeshRenderData->ScreenSize[BatchElement.InstancedLODIndex].GetValue(), SphereRadius, View->ViewMatrices.GetProjectionMatrix()) * LODScale;
-					if (CurCut < FinalCull)
-					{
-						InstancingViewZCompare.Y = CurCut;
-					}
-					else
-					{
-						// this LOD is completely removed by one of the other two factors
-						InstancingViewZCompare.Y = MIN_flt;
-						InstancingViewZCompare.Z = MIN_flt;
-					}
-				}
-				else
-				{
-					// this is the first LOD, so we don't have a fade-in region
-					InstancingViewZCompare.Y = MIN_flt;
-				}
-			}
-
-			InstancingOffset = (FVector3f)InstancingUserData->InstancingOffset; // LWC_TODO: precision loss
-
-			const FVector PreViewTranslation = View->ViewMatrices.GetPreViewTranslation();
-			InstancingTranslatedWorldViewOriginZero = (FVector3f)(View->GetTemporalLODOrigin(0) + PreViewTranslation); // LWC_TODO: precision loss
-			InstancingTranslatedWorldViewOriginOne = (FVector3f)(View->GetTemporalLODOrigin(1) + PreViewTranslation); // LWC_TODO: precision loss
-
-			float Alpha = View->GetTemporalLODTransition();
-			InstancingTranslatedWorldViewOriginZero.W = 1.0f - Alpha;
-			InstancingTranslatedWorldViewOriginOne.W = Alpha;
-
-			InstancingViewZCompareZero.W = LODRandom;
-		}
-
-		ShaderBindings.Add(InstancingViewZCompareZeroParameter, InstancingViewZCompareZero);
-		ShaderBindings.Add(InstancingViewZCompareOneParameter, InstancingViewZCompareOne);
-		ShaderBindings.Add(InstancingViewZConstantParameter, InstancingViewZConstant);
-		ShaderBindings.Add(InstancingOffsetParameter, InstancingOffset);
-		ShaderBindings.Add(InstancingWorldViewOriginZeroParameter, InstancingTranslatedWorldViewOriginZero);
-		ShaderBindings.Add(InstancingWorldViewOriginOneParameter, InstancingTranslatedWorldViewOriginOne);
-	}
-
-	if( InstancingFadeOutParamsParameter.IsBound() )
-	{
-		FVector4f InstancingFadeOutParams(MAX_flt,0.f,1.f,1.f);
-		if (InstancingUserData)
-		{
-			const float MaxDrawDistanceScale = GetCachedScalabilityCVars().ViewDistanceScale;
-			const float StartDistance = InstancingUserData->StartCullDistance * MaxDrawDistanceScale;
-			const float EndDistance = InstancingUserData->EndCullDistance * MaxDrawDistanceScale;
-
-			InstancingFadeOutParams.X = StartDistance;
-			if( EndDistance > 0 )
-			{
-				if( EndDistance > StartDistance )
-				{
-					InstancingFadeOutParams.Y = 1.f / (float)(EndDistance - StartDistance);
-				}
-				else
-				{
-					InstancingFadeOutParams.Y = 1.f;
-				}
-			}
-			else
-			{
-				InstancingFadeOutParams.Y = 0.f;
-			}
-			if (CVarCullAllInVertexShader.GetValueOnRenderThread() > 0)
-			{
-				InstancingFadeOutParams.Z = 0.0f;
-				InstancingFadeOutParams.W = 0.0f;
-			}
-			else
-			{
-				InstancingFadeOutParams.Z = InstancingUserData->bRenderSelected ? 1.f : 0.f;
-				InstancingFadeOutParams.W = InstancingUserData->bRenderUnselected ? 1.f : 0.f;
-			}
-		}
-
-		ShaderBindings.Add(InstancingFadeOutParamsParameter, InstancingFadeOutParams);
-
-	}
+	ShaderBindings.Add(Shader->GetUniformBufferParameter<FInstancedStaticMeshVFLooseUniformShaderParameters>(), BatchElement.LooseParametersUniformBuffer);
 }
diff --git a/Engine/Source/Runtime/Engine/Public/HierarchicalStaticMeshSceneProxy.h b/Engine/Source/Runtime/Engine/Public/HierarchicalStaticMeshSceneProxy.h
--- a/Engine/Source/Runtime/Engine/Public/HierarchicalStaticMeshSceneProxy.h
+++ b/Engine/Source/Runtime/Engine/Public/HierarchicalStaticMeshSceneProxy.h
@@ -93,7 +93,7 @@ public:
 
 	virtual void ApplyWorldOffset(FVector InOffset) override;
 
-	void FillDynamicMeshElements(FMeshElementCollector& Collector, const FFoliageElementParams& ElementParams, const FFoliageRenderInstanceParams& Instances) const;
+	void FillDynamicMeshElements(const FSceneView* View, FMeshElementCollector& Collector, const FFoliageElementParams& ElementParams, const FFoliageRenderInstanceParams& Instances) const;
 
 	template<bool TUseVector, bool THasWPODisplacement>
 	void Traverse(const FFoliageCullInstanceParams& Params, int32 Index, int32 MinLOD, int32 MaxLOD, bool bFullyContained = false) const;
diff --git a/Engine/Source/Runtime/Engine/Public/MeshBatch.h b/Engine/Source/Runtime/Engine/Public/MeshBatch.h
--- a/Engine/Source/Runtime/Engine/Public/MeshBatch.h
+++ b/Engine/Source/Runtime/Engine/Public/MeshBatch.h
@@ -201,6 +201,10 @@ struct FMeshBatchElement
 	 */
 	const TUniformBuffer<FPrimitiveUniformShaderParameters>* PrimitiveUniformBufferResource;
 
+	/** Uniform buffer containing the "loose" parameters that aren't wrapped in other uniform buffers. Those parameters can be unique per mesh batch, e.g. view dependent. */
+	FUniformBufferRHIRef LooseParametersUniformBuffer;
+
+	/** The index buffer to draw the mesh batch with. */
 	const FIndexBuffer* IndexBuffer;
 
 	/**
