Making fx.DetailedCSVStats render thread safe.

diff --git a/Engine/Source/Runtime/Engine/Private/Particles/ParticlePerfStats.cpp b/Engine/Source/Runtime/Engine/Private/Particles/ParticlePerfStats.cpp
--- a/Engine/Source/Runtime/Engine/Private/Particles/ParticlePerfStats.cpp
+++ b/Engine/Source/Runtime/Engine/Private/Particles/ParticlePerfStats.cpp
@@ -1133,7 +1133,7 @@ static FAutoConsoleVariableRef CVarWriteDetailedCSVStats(
 	bDetailedCSVStats,
 	TEXT("If true, we write detailed partilce stats to the CSV profiler. \n"),
 	FConsoleVariableDelegate::CreateStatic(&OnDetailedCSVStatsEnabledChanged),
-	ECVF_Default);
+	ECVF_Default, ECVF_RenderThreadSafe);
 
 
 FParticlePerfStatsListenerPtr FParticlePerfStatsListener_CSVProfiler::CSVListener;
