Added missing const

#rnx

diff --git a/Engine/Source/Runtime/Engine/Classes/Engine/MemberReference.h b/Engine/Source/Runtime/Engine/Classes/Engine/MemberReference.h
--- a/Engine/Source/Runtime/Engine/Classes/Engine/MemberReference.h
+++ b/Engine/Source/Runtime/Engine/Classes/Engine/MemberReference.h
@@ -175,7 +175,7 @@ public:
 
 	/** Set to a non-'self' member, so must include reference to class owning the member. */
 	ENGINE_API void SetExternalMember(FName InMemberName, TSubclassOf<class UObject> InMemberParentClass);
-	ENGINE_API void SetExternalMember(FName InMemberName, TSubclassOf<class UObject> InMemberParentClass, FGuid& InMemberGuid);
+	ENGINE_API void SetExternalMember(FName InMemberName, TSubclassOf<class UObject> InMemberParentClass, const FGuid& InMemberGuid);
 
 	/** Set to reference a global field (intended for things like natively defined delegate signatures) */
 	ENGINE_API void SetGlobalField(FName InFieldName, UPackage* InParentPackage);
diff --git a/Engine/Source/Runtime/Engine/Private/MemberReference.cpp b/Engine/Source/Runtime/Engine/Private/MemberReference.cpp
--- a/Engine/Source/Runtime/Engine/Private/MemberReference.cpp
+++ b/Engine/Source/Runtime/Engine/Private/MemberReference.cpp
@@ -29,7 +29,7 @@ void FMemberReference::SetExternalMember(FName InMemberName, TSubclassOf<class U
 	bWasDeprecated = false;
 }
 
-void FMemberReference::SetExternalMember(FName InMemberName, TSubclassOf<class UObject> InMemberParentClass, FGuid& InMemberGuid)
+void FMemberReference::SetExternalMember(FName InMemberName, TSubclassOf<class UObject> InMemberParentClass, const FGuid& InMemberGuid)
 {
 	SetExternalMember(InMemberName, InMemberParentClass);
 	MemberGuid = InMemberGuid;
