Replaced ContentBundleEpoch with UWorldPartition::StreamingStateEpoch, moved logic inside WorldPartition's InjectExternalStreamingObject/RemoveExternalStreamingObject.

diff --git a/Engine/Source/Runtime/Engine/Private/WorldPartition/ContentBundle/ContentBundle.cpp b/Engine/Source/Runtime/Engine/Private/WorldPartition/ContentBundle/ContentBundle.cpp
--- a/Engine/Source/Runtime/Engine/Private/WorldPartition/ContentBundle/ContentBundle.cpp
+++ b/Engine/Source/Runtime/Engine/Private/WorldPartition/ContentBundle/ContentBundle.cpp
@@ -12,8 +12,6 @@
 #include "WorldPartition/ContentBundle/ContentBundleWorldSubsystem.h"
 #endif
 
-int32 FContentBundle::ContentBundlesEpoch = 0;
-
 FContentBundle::FContentBundle(TSharedPtr<FContentBundleClient>& InClient, UWorld* InWorld)
 	: FContentBundleBase(InClient, InWorld)
 	, ExternalStreamingObjectPackage(nullptr)
@@ -65,8 +63,6 @@ void FContentBundle::DoInjectContent()
 		{
 			UE_LOG(LogContentBundle, Log, TEXT("%s Streaming Object Injected."), *ContentBundle::Log::MakeDebugInfoString(*this));
 			SetStatus(EContentBundleStatus::ContentInjected);
-
-			ContentBundlesEpoch++;
 		}
 		else
 		{
@@ -85,11 +81,7 @@ void FContentBundle::DoRemoveContent()
 {
 	if (ExternalStreamingObject != nullptr)
 	{
-		if (GetInjectedWorld()->GetWorldPartition()->RemoveExternalStreamingObject(ExternalStreamingObject))
-		{
-			ContentBundlesEpoch++;
-		}
-		else
+		if (!GetInjectedWorld()->GetWorldPartition()->RemoveExternalStreamingObject(ExternalStreamingObject))
 		{
 			UE_LOG(LogContentBundle, Error, TEXT("%s Error while removing streaming object."), *ContentBundle::Log::MakeDebugInfoString(*this));
 		}
diff --git a/Engine/Source/Runtime/Engine/Private/WorldPartition/WorldPartition.cpp b/Engine/Source/Runtime/Engine/Private/WorldPartition/WorldPartition.cpp
--- a/Engine/Source/Runtime/Engine/Private/WorldPartition/WorldPartition.cpp
+++ b/Engine/Source/Runtime/Engine/Private/WorldPartition/WorldPartition.cpp
@@ -285,6 +285,7 @@ UWorldPartition::UWorldPartition(const FObjectInitializer& ObjectInitializer)
 	bEnableStreaming = true;
 	ServerStreamingMode = EWorldPartitionServerStreamingMode::ProjectDefault;
 	ServerStreamingOutMode = EWorldPartitionServerStreamingOutMode::ProjectDefault;
+	StreamingStateEpoch = 0;
 
 #if WITH_EDITOR
 	WorldPartitionStreamingPolicyClass = UWorldPartitionLevelStreamingPolicy::StaticClass();
@@ -1475,6 +1476,7 @@ bool UWorldPartition::InjectExternalStreamingObject(URuntimeHashExternalStreamin
 	{
 		StreamingPolicy->InjectExternalStreamingObject(ExternalStreamingObject);
 		GetWorld()->GetSubsystem<UHLODSubsystem>()->OnExternalStreamingObjectInjected(ExternalStreamingObject);
+		++StreamingStateEpoch;
 	}
 
 	return bInjected;
@@ -1491,6 +1493,7 @@ bool UWorldPartition::RemoveExternalStreamingObject(URuntimeHashExternalStreamin
 		}
 		
 		GetWorld()->GetSubsystem<UHLODSubsystem>()->OnExternalStreamingObjectRemoved(ExternalStreamingObject);
+		++StreamingStateEpoch;
 	}
 
 	return bRemoved;
diff --git a/Engine/Source/Runtime/Engine/Private/WorldPartition/WorldPartitionStreamingPolicy.cpp b/Engine/Source/Runtime/Engine/Private/WorldPartition/WorldPartitionStreamingPolicy.cpp
--- a/Engine/Source/Runtime/Engine/Private/WorldPartition/WorldPartitionStreamingPolicy.cpp
+++ b/Engine/Source/Runtime/Engine/Private/WorldPartition/WorldPartitionStreamingPolicy.cpp
@@ -73,8 +73,8 @@ UWorldPartitionStreamingPolicy::UWorldPartitionStreamingPolicy(const FObjectInit
 	, bLastUpdateCompletedLoadingAndActivation(false)
 	, bCriticalPerformanceRequestedBlockTillOnWorld(false)
 	, CriticalPerformanceBlockTillLevelStreamingCompletedEpoch(0)
-	, DataLayersStatesServerEpoch(INT_MIN)
-	, ContentBundleServerEpoch(INT_MIN)
+	, ServerDataLayersStatesEpoch(INT_MIN)
+	, ServerStreamingStateEpoch(INT_MIN)
 	, ServerStreamingEnabledEpoch(INT_MIN)
 	, UpdateStreamingHash(0)
 	, UpdateStreamingSourcesHash(0)
@@ -163,7 +163,7 @@ uint32 UWorldPartitionStreamingPolicy::ComputeUpdateStreamingHash(bool bCanOptim
 		// Build hash that will be used to detect relevant changes
 		FHashBuilder HashBuilder;
 		HashBuilder << ComputeServerStreamingEnabledEpoch();
-		HashBuilder << FContentBundle::GetContentBundleEpoch();
+		HashBuilder << GetOuterUWorldPartition()->GetStreamingStateEpoch();
 		HashBuilder << AWorldDataLayers::GetDataLayersStateEpoch();
 		HashBuilder << bIsStreaming3D;
 		for (const FWorldPartitionStreamingSource& Source : StreamingSources)
@@ -352,8 +352,8 @@ void UWorldPartitionStreamingPolicy::UpdateStreamingState()
 			// Server will activate all non data layer cells at first and then load/activate/unload data layer cells only when the data layer states change
 			if (!bIsServerStreamingEnabled && 
 				(ServerStreamingEnabledEpoch == NewServerStreamingEnabledEpoch) &&
-				(ContentBundleServerEpoch == FContentBundle::GetContentBundleEpoch()) &&
-				(DataLayersStatesServerEpoch == AWorldDataLayers::GetDataLayersStateEpoch()))
+				(ServerStreamingStateEpoch == WorldPartition->GetStreamingStateEpoch()) &&
+				(ServerDataLayersStatesEpoch == AWorldDataLayers::GetDataLayersStateEpoch()))
 			{
 				// Server as nothing to do early out
 				return; 
@@ -593,8 +593,8 @@ void UWorldPartitionStreamingPolicy::UpdateStreamingState()
 	// Update Epoch if we aren't waiting for clients anymore
 	if (bUpdateEpoch)
 	{
-		DataLayersStatesServerEpoch = AWorldDataLayers::GetDataLayersStateEpoch();
-		ContentBundleServerEpoch = FContentBundle::GetContentBundleEpoch();
+		ServerDataLayersStatesEpoch = AWorldDataLayers::GetDataLayersStateEpoch();
+		ServerStreamingStateEpoch = WorldPartition->GetStreamingStateEpoch();
 		ServerStreamingEnabledEpoch = NewServerStreamingEnabledEpoch;
 		UE_LOG(LogWorldPartition, Verbose, TEXT("Server epoch updated"));
 	}
diff --git a/Engine/Source/Runtime/Engine/Public/WorldPartition/ContentBundle/ContentBundle.h b/Engine/Source/Runtime/Engine/Public/WorldPartition/ContentBundle/ContentBundle.h
--- a/Engine/Source/Runtime/Engine/Public/WorldPartition/ContentBundle/ContentBundle.h
+++ b/Engine/Source/Runtime/Engine/Public/WorldPartition/ContentBundle/ContentBundle.h
@@ -23,8 +23,6 @@ public:
 	virtual bool IsValid() const override;
 	//~ End IContentBundle Interface
 
-	static int32 GetContentBundleEpoch() { return ContentBundlesEpoch; }
-
 protected:
 	//~ Begin IContentBundle Interface
 	virtual void DoInitialize() override;
@@ -38,8 +36,6 @@ private:
 	void InitializeForPIE();
 #endif
 
-	static int32 ContentBundlesEpoch;
-
 	UPackage* ExternalStreamingObjectPackage;
 	URuntimeHashExternalStreamingObjectBase* ExternalStreamingObject;
 };
\ No newline at end of file
diff --git a/Engine/Source/Runtime/Engine/Public/WorldPartition/WorldPartition.h b/Engine/Source/Runtime/Engine/Public/WorldPartition/WorldPartition.h
--- a/Engine/Source/Runtime/Engine/Public/WorldPartition/WorldPartition.h
+++ b/Engine/Source/Runtime/Engine/Public/WorldPartition/WorldPartition.h
@@ -310,6 +310,7 @@ public:
 
 public:
 	static bool IsSimulating(bool bIncludeTestEnableSimulationStreamingSource = true);
+	int32 GetStreamingStateEpoch() const { return StreamingStateEpoch; }
 
 	void Initialize(UWorld* World, const FTransform& InTransform);
 	bool IsInitialized() const;
@@ -460,6 +461,7 @@ private:
 	static FAutoConsoleVariableRef CVarDebugDedicatedServerStreaming;
 #endif
 
+	int32 StreamingStateEpoch;
 	static int32 GlobalEnableServerStreaming;
 	static bool bGlobalEnableServerStreamingOut;
 	static bool bUseMakingVisibleTransactionRequests;
@@ -504,7 +506,6 @@ public:
 	{
 		return RegisteredEditorLoaderAdapters;
 	}
-
 private:
 #endif
 
diff --git a/Engine/Source/Runtime/Engine/Public/WorldPartition/WorldPartitionStreamingPolicy.h b/Engine/Source/Runtime/Engine/Public/WorldPartition/WorldPartitionStreamingPolicy.h
--- a/Engine/Source/Runtime/Engine/Public/WorldPartition/WorldPartitionStreamingPolicy.h
+++ b/Engine/Source/Runtime/Engine/Public/WorldPartition/WorldPartitionStreamingPolicy.h
@@ -121,8 +121,8 @@ private:
 	bool bLastUpdateCompletedLoadingAndActivation;
 	bool bCriticalPerformanceRequestedBlockTillOnWorld;
 	int32 CriticalPerformanceBlockTillLevelStreamingCompletedEpoch;
-	int32 DataLayersStatesServerEpoch;
-	int32 ContentBundleServerEpoch;
+	int32 ServerDataLayersStatesEpoch;
+	int32 ServerStreamingStateEpoch;
 	int32 ServerStreamingEnabledEpoch;
 	uint32 UpdateStreamingHash;
 	uint32 UpdateStreamingSourcesHash;
