[NO NON-PII COMMIT MESSAGE]

diff --git a/Engine/Source/Runtime/Engine/Classes/PhysicsField/PhysicsFieldComponent.h b/Engine/Source/Runtime/Engine/Classes/PhysicsField/PhysicsFieldComponent.h
--- a/Engine/Source/Runtime/Engine/Classes/PhysicsField/PhysicsFieldComponent.h
+++ b/Engine/Source/Runtime/Engine/Classes/PhysicsField/PhysicsFieldComponent.h
@@ -6,6 +6,7 @@
 #include "Field/FieldSystem.h"
 #include "UObject/ObjectMacros.h"
 #include "RHI.h"
+#include "RHIUtilities.h"
 #include "RenderResource.h"
 #include "PrimitiveSceneProxy.h"
 #include "Kismet/BlueprintFunctionLibrary.h"
