Exposed FSquare2DGridHelper::FGridLevel in engine.

diff --git a/Engine/Source/Runtime/Engine/Private/WorldPartition/RuntimeSpatialHash/RuntimeSpatialHashGridHelper.cpp b/Engine/Source/Runtime/Engine/Private/WorldPartition/RuntimeSpatialHash/RuntimeSpatialHashGridHelper.cpp
--- a/Engine/Source/Runtime/Engine/Private/WorldPartition/RuntimeSpatialHash/RuntimeSpatialHashGridHelper.cpp
+++ b/Engine/Source/Runtime/Engine/Private/WorldPartition/RuntimeSpatialHash/RuntimeSpatialHashGridHelper.cpp
@@ -81,13 +81,10 @@ FSquare2DGridHelper::FSquare2DGridHelper(const FBox& InWorldBounds, const FVecto
 		CurrentGridSize >>= 1;
 	}
 
-#if WITH_EDITOR
 	// Make sure the always loaded cell exists
 	GetAlwaysLoadedCell();
-#endif
 }
 
-#if WITH_EDITOR
 void FSquare2DGridHelper::ForEachCells(TFunctionRef<void(const FSquare2DGridHelper::FGridLevel::FGridCell&)> InOperation) const
 {
 	for (int32 Level = 0; Level < Levels.Num(); Level++)
@@ -98,7 +95,6 @@ void FSquare2DGridHelper::ForEachCells(TFunctionRef<void(const FSquare2DGridHelp
 		}
 	}
 }
-#endif
 
 int32 FSquare2DGridHelper::ForEachIntersectingCells(const FBox& InBox, TFunctionRef<void(const FGridCellCoord&)> InOperation, int32 InStartLevel) const
 {
diff --git a/Engine/Source/Runtime/Engine/Public/WorldPartition/RuntimeSpatialHash/RuntimeSpatialHashGridHelper.h b/Engine/Source/Runtime/Engine/Public/WorldPartition/RuntimeSpatialHash/RuntimeSpatialHashGridHelper.h
--- a/Engine/Source/Runtime/Engine/Public/WorldPartition/RuntimeSpatialHash/RuntimeSpatialHashGridHelper.h
+++ b/Engine/Source/Runtime/Engine/Public/WorldPartition/RuntimeSpatialHash/RuntimeSpatialHashGridHelper.h
@@ -285,6 +285,7 @@ struct ENGINE_API FSquare2DGridHelper
 			FDataLayersID DataLayersID;
 			FGuid ContentBundleID;
 		};
+#endif
 
 		struct FGridCell
 		{
@@ -292,6 +293,7 @@ struct ENGINE_API FSquare2DGridHelper
 				: Coords(InCoords)
 			{}
 
+#if WITH_EDITOR
 			void AddActorSetInstance(const IStreamingGenerationContext::FActorSetInstance* ActorSetInstance)
 			{
 				const FDataLayersID DataLayersID = FDataLayersID(ActorSetInstance->DataLayers);
@@ -315,6 +317,7 @@ struct ENGINE_API FSquare2DGridHelper
 				}
 				return nullptr;
 			}
+#endif
 
 			FGridCellCoord GetCoords() const
 			{
@@ -323,18 +326,16 @@ struct ENGINE_API FSquare2DGridHelper
 
 		private:
 			FGridCellCoord Coords;
+#if WITH_EDITOR
 			TSet<FGridCellDataChunk> DataChunks;
-		};
 #endif
+		};
 
 		inline FGridLevel(const FVector2D& InOrigin, int64 InCellSize, int64 InGridSize, int32 InLevel)
 			: FGrid2D(InOrigin, InCellSize, InGridSize)
-#if WITH_EDITOR
 			, Level(InLevel)
-#endif
 		{}
 
-#if WITH_EDITOR
 		/**
 		 * Returns the cell at the specified grid coordinate
 		 *
@@ -384,12 +385,10 @@ struct ENGINE_API FSquare2DGridHelper
 		int32 Level;
 		TArray<FGridCell> Cells;
 		TMap<int64, int64> CellsMapping;
-#endif
 	};
 
 	FSquare2DGridHelper(const FBox& InWorldBounds, const FVector& InOrigin, int64 InCellSize);
 
-#if WITH_EDITOR
 	// Returns the lowest grid level
 	inline FGridLevel& GetLowestLevel() { return Levels[0]; }
 
@@ -401,7 +400,6 @@ struct ENGINE_API FSquare2DGridHelper
 
 	// Returns the cell at the given coord
 	inline const FGridLevel::FGridCell& GetCell(const FGridCellCoord& InCoords) const { return Levels[InCoords.Z].GetCell(FGridCellCoord2(InCoords.X, InCoords.Y)); }
-#endif
 
 	/**
 	 * Returns the cell bounds
@@ -439,10 +437,8 @@ struct ENGINE_API FSquare2DGridHelper
 		return false;
 	}
 
-#if WITH_EDITOR
 	// Runs a function on all cells
 	void ForEachCells(TFunctionRef<void(const FSquare2DGridHelper::FGridLevel::FGridCell&)> InOperation) const;
-#endif
 
 	/**
 	 * Runs a function on all intersecting cells for the provided box
