World Partition Actor Descriptors: actor referenced through editor-only properties will now be discarded during streaming generation.

diff --git a/Engine/Source/Editor/UnrealEd/Private/Commandlets/WorldPartitionConvertCommandlet.cpp b/Engine/Source/Editor/UnrealEd/Private/Commandlets/WorldPartitionConvertCommandlet.cpp
--- a/Engine/Source/Editor/UnrealEd/Private/Commandlets/WorldPartitionConvertCommandlet.cpp
+++ b/Engine/Source/Editor/UnrealEd/Private/Commandlets/WorldPartitionConvertCommandlet.cpp
@@ -1129,7 +1129,9 @@ int32 UWorldPartitionConvertCommandlet::Main(const FString& Params)
 			if (LevelHasLevelScriptBlueprint(Level))
 			{
 				ULevelScriptBlueprint* LevelScriptBlueprint = Level->GetLevelScriptBlueprint(true);
-				TArray<AActor*> LevelScriptActorReferences = ActorsReferencesUtils::GetActorReferences(LevelScriptBlueprint);
+				const ActorsReferencesUtils::FGetActorReferencesParams Params(LevelScriptBlueprint);
+				TArray<AActor*> LevelScriptActorReferences;
+				Algo::Transform(ActorsReferencesUtils::GetActorReferences(Params), LevelScriptActorReferences, [](const ActorsReferencesUtils::FActorReference& ActorReference) { return ActorReference.Actor; });
 
 				for (AActor* LevelScriptActorReference : LevelScriptActorReferences)
 				{
@@ -1250,7 +1252,9 @@ int32 UWorldPartitionConvertCommandlet::Main(const FString& Params)
 				LevelScriptActorReferences.Add(LevelScriptActor);
 
 				ULevelScriptBlueprint* LevelScriptBlueprint = SubLevel->GetLevelScriptBlueprint(true);
-				LevelScriptActorReferences.Append(ActorsReferencesUtils::GetActorReferences(LevelScriptBlueprint, RF_NoFlags, true));
+				const ActorsReferencesUtils::FGetActorReferencesParams LevelScriptBlueprintReferencesParams = ActorsReferencesUtils::FGetActorReferencesParams(LevelScriptBlueprint)
+					.SetRecursive(true);
+				Algo::Transform(ActorsReferencesUtils::GetActorReferences(LevelScriptBlueprintReferencesParams), LevelScriptActorReferences, [](const ActorsReferencesUtils::FActorReference& ActorReference) { return ActorReference.Actor; });
 
 				for(AActor* Actor: SubLevel->Actors)
 				{
@@ -1265,7 +1269,9 @@ int32 UWorldPartitionConvertCommandlet::Main(const FString& Params)
 						else
 						{
 							TSet<AActor*> ActorReferences;
-							ActorReferences.Append(ActorsReferencesUtils::GetActorReferences(Actor, RF_NoFlags, true));
+							const ActorsReferencesUtils::FGetActorReferencesParams ActorReferencesParams = ActorsReferencesUtils::FGetActorReferencesParams(Actor)
+								.SetRecursive(true);
+							Algo::Transform(ActorsReferencesUtils::GetActorReferences(ActorReferencesParams), ActorReferences, [](const ActorsReferencesUtils::FActorReference& ActorReference) { return ActorReference.Actor; });
 
 							for (AActor* ActorReference : ActorReferences)
 							{
diff --git a/Engine/Source/Runtime/Core/Public/UObject/FortniteMainBranchObjectVersion.h b/Engine/Source/Runtime/Core/Public/UObject/FortniteMainBranchObjectVersion.h
--- a/Engine/Source/Runtime/Core/Public/UObject/FortniteMainBranchObjectVersion.h
+++ b/Engine/Source/Runtime/Core/Public/UObject/FortniteMainBranchObjectVersion.h
@@ -307,6 +307,9 @@ struct CORE_API FFortniteMainBranchObjectVersion
 		// Upgrade Navigation Links to use 64 bits for the ID
 		NavigationLinkID32To64,
 
+		// Serialize editor only references in world partition actor descriptors
+		WorldPartitionActorDescSerializeEditorOnlyReferences,
+
 		// -----<new versions can be added above this line>-------------------------------------------------
 		VersionPlusOne,
 		LatestVersion = VersionPlusOne - 1
diff --git a/Engine/Source/Runtime/Engine/Private/ActorReferencesUtils.cpp b/Engine/Source/Runtime/Engine/Private/ActorReferencesUtils.cpp
--- a/Engine/Source/Runtime/Engine/Private/ActorReferencesUtils.cpp
+++ b/Engine/Source/Runtime/Engine/Private/ActorReferencesUtils.cpp
@@ -2,13 +2,32 @@
 
 #include "ActorReferencesUtils.h"
 #include "GameFramework/Actor.h"
+#include "Algo/Transform.h"
 
 class FArchiveGatherExternalActorRefs : public FArchiveUObject
 {
 public:
-	FArchiveGatherExternalActorRefs(UObject* InRoot, TSet<AActor*>& InActorReferences, EObjectFlags InRequiredFlags, bool bInRecursive)
+	using FActorReference = ActorsReferencesUtils::FActorReference;
+
+	struct FActorReferenceKeyFuncs : DefaultKeyFuncs<FActorReference>
+	{
+		template<typename ComparableKey>
+		static FORCEINLINE bool Matches(KeyInitType A, ComparableKey B)
+		{
+			return A.Actor == B.Actor;
+		}
+
+		static FORCEINLINE uint32 GetKeyHash(KeyInitType Key)
+		{
+			return GetTypeHash(Key.Actor);
+		}
+	};
+
+	using FActorReferenceSet = TSet<FActorReference, FActorReferenceKeyFuncs>;
+
+	FArchiveGatherExternalActorRefs(UObject* InRoot, EObjectFlags InRequiredFlags, bool bInRecursive, FActorReferenceSet& InActorReference)
 		: Root(InRoot)
-		, ActorReferences(InActorReferences)
+		, ActorReferences(InActorReference)
 		, RequiredFlags(InRequiredFlags)
 		, bRecursive(bInRecursive)
 	{
@@ -57,6 +76,19 @@ public:
 	}
 
 private:
+	virtual void PushSerializedProperty(FProperty* InProperty, const bool bIsEditorOnlyProperty) override
+	{
+		FArchiveUObject::PushSerializedProperty(InProperty, bIsEditorOnlyProperty);
+		EditorOnlyStack += InProperty->IsEditorOnlyProperty() ? 1 : 0;
+	}
+
+	virtual void PopSerializedProperty(FProperty* InProperty, const bool bIsEditorOnlyProperty) override
+	{
+		FArchiveUObject::PopSerializedProperty(InProperty, bIsEditorOnlyProperty);
+		EditorOnlyStack -= InProperty->IsEditorOnlyProperty() ? 1 : 0;
+		check(EditorOnlyStack >= 0);
+	}
+
 	void HandleObjectReference(UObject* Obj)
 	{
 		AActor* Actor = Cast<AActor>(Obj);
@@ -78,26 +110,54 @@ private:
 
 			if (((RequiredFlags == RF_NoFlags) || TopParentActor->HasAnyFlags(RequiredFlags)) && (TopParentActor != Root))
 			{
-				ActorReferences.Add(TopParentActor);
+				bool bWasAlreadyInSetPtr;
+				FActorReference TopParentActorReference;
+				TopParentActorReference.Actor = TopParentActor;
+				FActorReference& ActorReference = ActorReferences.FindOrAdd(TopParentActorReference, &bWasAlreadyInSetPtr);
+
+				const bool bIsEditorOnly = EditorOnlyStack > 0;
+				if (!bWasAlreadyInSetPtr)
+				{
+					ActorReference.bIsEditorOnly = bIsEditorOnly;
+				}
+				else
+				{
+					ActorReference.bIsEditorOnly &= bIsEditorOnly;
+				}
 			}
 		}
 	}
 
 	UObject* Root;
-	TSet<AActor*>& ActorReferences;
+	FActorReferenceSet& ActorReferences;
 	TSet<UObject*> SubObjects;
+	int32 EditorOnlyStack;
 	EObjectFlags RequiredFlags;
 	bool bRecursive;
 };
 
 TArray<AActor*> ActorsReferencesUtils::GetExternalActorReferences(UObject* Root, bool bRecursive)
 {
+PRAGMA_DISABLE_DEPRECATION_WARNINGS
 	return GetActorReferences(Root, RF_HasExternalPackage, bRecursive);
+PRAGMA_ENABLE_DEPRECATION_WARNINGS
 }
 
 TArray<AActor*> ActorsReferencesUtils::GetActorReferences(UObject* Root, EObjectFlags RequiredFlags, bool bRecursive)
 {
-	TSet<AActor*> Result;
-	FArchiveGatherExternalActorRefs Ar(Root, Result, RequiredFlags, bRecursive);
-	return Result.Array();
+	FGetActorReferencesParams Params = FGetActorReferencesParams(Root)
+		.SetRequiredFlags(RequiredFlags)
+		.SetRecursive(bRecursive);
+
+	TArray<AActor*> Result;
+	const TArray<FActorReference> ActorReferences = GetActorReferences(Params);
+	Algo::Transform(ActorReferences, Result, [](const FActorReference& ActorReference) { return ActorReference.Actor; });
+	return Result;
 }
+
+TArray<ActorsReferencesUtils::FActorReference> ActorsReferencesUtils::GetActorReferences(const FGetActorReferencesParams& InParams)
+{
+	FArchiveGatherExternalActorRefs::FActorReferenceSet Result;
+	FArchiveGatherExternalActorRefs Ar(InParams.Root, InParams.RequiredFlags, InParams.bRecursive, Result);
+	return Result.Array();
+}
\ No newline at end of file
diff --git a/Engine/Source/Runtime/Engine/Private/WorldPartition/WorldPartition.cpp b/Engine/Source/Runtime/Engine/Private/WorldPartition/WorldPartition.cpp
--- a/Engine/Source/Runtime/Engine/Private/WorldPartition/WorldPartition.cpp
+++ b/Engine/Source/Runtime/Engine/Private/WorldPartition/WorldPartition.cpp
@@ -1756,14 +1756,16 @@ void UWorldPartition::AppendAssetRegistryTags(TArray<FAssetRegistryTag>& OutTags
 	// Append level script references so we can perform changelists validations without loading the world
 	if (const ULevelScriptBlueprint* LevelScriptBlueprint = GetWorld()->PersistentLevel->GetLevelScriptBlueprint(true))
 	{
-		TArray<AActor*> LevelScriptExternalActorReferences = ActorsReferencesUtils::GetExternalActorReferences((UObject*)LevelScriptBlueprint);
+		const ActorsReferencesUtils::FGetActorReferencesParams Params = ActorsReferencesUtils::FGetActorReferencesParams((UObject*)LevelScriptBlueprint)
+			.SetRequiredFlags(RF_HasExternalPackage);
+		TArray<ActorsReferencesUtils::FActorReference> LevelScriptExternalActorReferences = ActorsReferencesUtils::GetActorReferences(Params);
 		
 		if (LevelScriptExternalActorReferences.Num())
 		{
 			FStringBuilderBase StringBuilder;
-			for (const AActor* Actor : LevelScriptExternalActorReferences)
+			for (const ActorsReferencesUtils::FActorReference& ActorReference : LevelScriptExternalActorReferences)
 			{
-				StringBuilder.Append(Actor->GetActorGuid().ToString(EGuidFormats::Short));
+				StringBuilder.Append(ActorReference.Actor->GetActorGuid().ToString(EGuidFormats::Short));
 				StringBuilder.AppendChar(TEXT(','));
 			}
 			StringBuilder.RemoveSuffix(1);
diff --git a/Engine/Source/Runtime/Engine/Private/WorldPartition/WorldPartitionActorDesc.cpp b/Engine/Source/Runtime/Engine/Private/WorldPartition/WorldPartitionActorDesc.cpp
--- a/Engine/Source/Runtime/Engine/Private/WorldPartition/WorldPartitionActorDesc.cpp
+++ b/Engine/Source/Runtime/Engine/Private/WorldPartition/WorldPartitionActorDesc.cpp
@@ -181,14 +181,23 @@ void FWorldPartitionActorDesc::Init(const AActor* InActor)
 			ParentActor = AttachParentActor->GetActorGuid();
 		}
 
-		TArray<AActor*> ActorReferences = ActorsReferencesUtils::GetExternalActorReferences(const_cast<AActor*>(InActor));
+		const ActorsReferencesUtils::FGetActorReferencesParams Params = ActorsReferencesUtils::FGetActorReferencesParams(const_cast<AActor*>(InActor))
+			.SetRequiredFlags(RF_HasExternalPackage);
+		TArray<ActorsReferencesUtils::FActorReference> ActorReferences = ActorsReferencesUtils::GetActorReferences(Params);
 
 		if (ActorReferences.Num())
 		{
-			References.Empty(ActorReferences.Num());
-			for(AActor* ActorReference: ActorReferences)
+			References.Reserve(ActorReferences.Num());
+			for (const ActorsReferencesUtils::FActorReference& ActorReference : ActorReferences)
 			{
-				References.Add(ActorReference->GetActorGuid());
+				const FGuid ActorReferenceGuid = ActorReference.Actor->GetActorGuid();
+
+				References.Add(ActorReferenceGuid);
+
+				if (ActorReference.bIsEditorOnly)
+				{
+					EditorOnlyReferences.Add(ActorReferenceGuid);
+				}
 			}
 		}
 
@@ -258,6 +267,7 @@ bool FWorldPartitionActorDesc::Equals(const FWorldPartitionActorDesc* Other) con
 		ContentBundleGuid == Other->ContentBundleGuid &&
 		CompareUnsortedArrays(DataLayers, Other->DataLayers) &&
 		CompareUnsortedArrays(References, Other->References) &&
+		CompareUnsortedArrays(EditorOnlyReferences, Other->EditorOnlyReferences) &&
 		CompareUnsortedArrays(Tags, Other->Tags) &&
 		Properties == Other->Properties;
 }
@@ -277,7 +287,8 @@ bool FWorldPartitionActorDesc::ShouldResave(const FWorldPartitionActorDesc* Othe
 		ParentActor != Other->ParentActor ||
 		ContentBundleGuid != Other->ContentBundleGuid ||
 		!CompareUnsortedArrays(DataLayers, Other->DataLayers) ||
-		!CompareUnsortedArrays(References, Other->References))
+		!CompareUnsortedArrays(References, Other->References) ||
+		!CompareUnsortedArrays(EditorOnlyReferences, Other->EditorOnlyReferences))
 	{
 		return true;
 	}
@@ -425,6 +436,11 @@ FString FWorldPartitionActorDesc::ToString(EToStringMode Mode) const
 				Result.Appendf(TEXT(" References:%s"), *FString::JoinBy(References, TEXT(","), [&](const FGuid& ReferenceGuid) { return ReferenceGuid.ToString(); }));
 			}
 
+			if (EditorOnlyReferences.Num())
+			{
+				Result.Appendf(TEXT(" EditorOnlyReferences:%s"), *FString::JoinBy(EditorOnlyReferences, TEXT(","), [&](const FGuid& ReferenceGuid) { return ReferenceGuid.ToString(); }));
+			}
+
 			if (Tags.Num())
 			{
 				Result.Appendf(TEXT(" Tags:%s"), *FString::JoinBy(Tags, TEXT(","), [&](const FName& TagName) { return TagName.ToString(); }));
@@ -520,6 +536,11 @@ void FWorldPartitionActorDesc::Serialize(FArchive& Ar)
 
 	Ar << References;
 
+	if (Ar.CustomVer(FFortniteMainBranchObjectVersion::GUID) >= FFortniteMainBranchObjectVersion::WorldPartitionActorDescSerializeEditorOnlyReferences)
+	{
+		Ar << EditorOnlyReferences;
+	}
+
 	if (Ar.CustomVer(FFortniteNCBranchObjectVersion::GUID) >= FFortniteNCBranchObjectVersion::WorldPartitionActorDescTagsSerialization)
 	{
 		Ar << TDeltaSerialize<TArray<FName>>(Tags);
diff --git a/Engine/Source/Runtime/Engine/Private/WorldPartition/WorldPartitionActorDescView.cpp b/Engine/Source/Runtime/Engine/Private/WorldPartition/WorldPartitionActorDescView.cpp
--- a/Engine/Source/Runtime/Engine/Private/WorldPartition/WorldPartitionActorDescView.cpp
+++ b/Engine/Source/Runtime/Engine/Private/WorldPartition/WorldPartitionActorDescView.cpp
@@ -265,6 +265,11 @@ AActor* FWorldPartitionActorDescView::GetActor() const
 	return ActorDesc->GetActor();
 }
 
+bool FWorldPartitionActorDescView::IsEditorOnlyReference(const FGuid& ReferenceGuid) const
+{
+	return ActorDesc->IsEditorOnlyReference(ReferenceGuid);
+}
+
 bool FWorldPartitionActorDescView::GetProperty(FName PropertyName, FName* PropertyValue) const
 {
 	return ActorDesc->GetProperty(PropertyName, PropertyValue);
diff --git a/Engine/Source/Runtime/Engine/Private/WorldPartition/WorldPartitionStreamingGeneration.cpp b/Engine/Source/Runtime/Engine/Private/WorldPartition/WorldPartitionStreamingGeneration.cpp
--- a/Engine/Source/Runtime/Engine/Private/WorldPartition/WorldPartitionStreamingGeneration.cpp
+++ b/Engine/Source/Runtime/Engine/Private/WorldPartition/WorldPartitionStreamingGeneration.cpp
@@ -606,8 +606,10 @@ class FWorldPartitionStreamingGenerator
 				// Gather all references to external actors from the level script and make them always loaded
 				if (ULevelScriptBlueprint* LevelScriptBlueprint = WorldPartitionContext->GetTypedOuter<UWorld>()->PersistentLevel->GetLevelScriptBlueprint(true))
 				{
-					TArray<AActor*> LevelScriptExternalActorReferences = ActorsReferencesUtils::GetExternalActorReferences(LevelScriptBlueprint);
-					Algo::Transform(LevelScriptExternalActorReferences, LevelScriptReferences, [](const AActor* Actor) { return Actor->GetActorGuid(); });
+					const ActorsReferencesUtils::FGetActorReferencesParams Params = ActorsReferencesUtils::FGetActorReferencesParams(LevelScriptBlueprint)
+						.SetRequiredFlags(RF_HasExternalPackage);
+					TArray<ActorsReferencesUtils::FActorReference> LevelScriptExternalActorReferences = ActorsReferencesUtils::GetActorReferences(Params);
+					Algo::Transform(LevelScriptExternalActorReferences, LevelScriptReferences, [](const ActorsReferencesUtils::FActorReference& ActorReference) { return ActorReference.Actor->GetActorGuid(); });
 				}
 
 				// Validate data layers
@@ -803,55 +805,68 @@ PRAGMA_ENABLE_DEPRECATION_WARNINGS
 
 					if (ReferenceActorDescView)
 					{
-						// Validate grid placement
-						if (!IsReferenceGridPlacementValid(*RefererActorDescView, *ReferenceActorDescView))
+						// The actor reference is not editor-only, but we are referencing it through an editor-only property
+						if (RefererActorDescView->IsEditorOnlyReference(ReferenceActorDescView->GetGuid()))
 						{
-							if (PassType == EPassType::ErrorReporting)
-							{
-								ErrorHandler->OnInvalidReferenceGridPlacement(*RefererActorDescView, *ReferenceActorDescView);									
-							}
-							else
+							if (PassType == EPassType::Fixup)
 							{
-								RefererActorDescView->SetForcedNonSpatiallyLoaded();
-								ReferenceActorDescView->SetForcedNonSpatiallyLoaded();
+								EditorReferences.Add(Info.ReferenceGuid);
 							}
 
 							NbErrorsDetected++;
 						}
-
-						if (!IsReferenceDataLayersValid(*RefererActorDescView, *ReferenceActorDescView))
+						else
 						{
-							if (PassType == EPassType::ErrorReporting)
+							// Validate grid placement
+							if (!IsReferenceGridPlacementValid(*RefererActorDescView, *ReferenceActorDescView))
 							{
-								ErrorHandler->OnInvalidReferenceDataLayers(*RefererActorDescView, *ReferenceActorDescView);									
+								if (PassType == EPassType::ErrorReporting)
+								{
+									ErrorHandler->OnInvalidReferenceGridPlacement(*RefererActorDescView, *ReferenceActorDescView);									
+								}
+								else
+								{
+									RefererActorDescView->SetForcedNonSpatiallyLoaded();
+									ReferenceActorDescView->SetForcedNonSpatiallyLoaded();
+								}
+
+								NbErrorsDetected++;
 							}
-							else
+
+							if (!IsReferenceDataLayersValid(*RefererActorDescView, *ReferenceActorDescView))
 							{
-								RefererActorDescView->SetInvalidDataLayers();
-								ReferenceActorDescView->SetInvalidDataLayers();
-							}
+								if (PassType == EPassType::ErrorReporting)
+								{
+									ErrorHandler->OnInvalidReferenceDataLayers(*RefererActorDescView, *ReferenceActorDescView);									
+								}
+								else
+								{
+									RefererActorDescView->SetInvalidDataLayers();
+									ReferenceActorDescView->SetInvalidDataLayers();
+								}
 
-							NbErrorsDetected++;
-						}
+								NbErrorsDetected++;
+							}
 
-						if (!IsReferenceRuntimeGridValid(*RefererActorDescView, *ReferenceActorDescView))
-						{
-							if (PassType == EPassType::ErrorReporting)
+							if (!IsReferenceRuntimeGridValid(*RefererActorDescView, *ReferenceActorDescView))
 							{
-								ErrorHandler->OnInvalidReferenceRuntimeGrid(*RefererActorDescView, *ReferenceActorDescView);
+								if (PassType == EPassType::ErrorReporting)
+								{
+									ErrorHandler->OnInvalidReferenceRuntimeGrid(*RefererActorDescView, *ReferenceActorDescView);
+								}
+								else
+								{
+									RefererActorDescView->SetForcedNoRuntimeGrid();
+									ReferenceActorDescView->SetForcedNoRuntimeGrid();
+								}
+
+								NbErrorsDetected++;
 							}
-							else
+
+							if (PassType == EPassType::Fixup)
 							{
-								RefererActorDescView->SetForcedNoRuntimeGrid();
-								ReferenceActorDescView->SetForcedNoRuntimeGrid();
+								RuntimeReferences.Add(Info.ReferenceGuid);
 							}
-
-							NbErrorsDetected++;
-						}
-
-						if (PassType == EPassType::Fixup)
-						{
-							RuntimeReferences.Add(Info.ReferenceGuid);
 						}
 					}
 					else
diff --git a/Engine/Source/Runtime/Engine/Public/ActorReferencesUtils.h b/Engine/Source/Runtime/Engine/Public/ActorReferencesUtils.h
--- a/Engine/Source/Runtime/Engine/Public/ActorReferencesUtils.h
+++ b/Engine/Source/Runtime/Engine/Public/ActorReferencesUtils.h
@@ -12,10 +12,37 @@ namespace ActorsReferencesUtils
 	/**
 	 * Gather direct references to external actors from the root object.
 	 */
+	UE_DEPRECATED(5.3, "GetExternalActorReferences is deprecated. Use GetActorReferences with params instead.")
 	ENGINE_API TArray<AActor*> GetExternalActorReferences(UObject* Root, bool bRecursive = false);
 
 	/**
 	 * Gather direct references to actors from the root object.
 	 */
+	UE_DEPRECATED(5.3, "GetActorReferences is deprecated. Use GetActorReferences with params instead.")
 	ENGINE_API TArray<AActor*> GetActorReferences(UObject* Root, EObjectFlags RequiredFlags = RF_NoFlags, bool bRecursive = false);
+
+	struct FGetActorReferencesParams
+	{
+		FGetActorReferencesParams(UObject* InRoot)
+			: Root(InRoot)
+		{}
+
+		UObject* Root = nullptr;
+		EObjectFlags RequiredFlags = RF_NoFlags;
+		bool bRecursive = false;
+
+		FGetActorReferencesParams& SetRequiredFlags(EObjectFlags InRequiredFlags) { RequiredFlags = InRequiredFlags; return *this; }
+		FGetActorReferencesParams& SetRecursive(bool bInRecursive) { bRecursive = bInRecursive; return *this; }
+	};
+
+	struct FActorReference
+	{
+		AActor* Actor;
+		bool bIsEditorOnly;
+	};
+
+	/**
+	 * Gather references to actors from the root object.
+	 */
+	ENGINE_API TArray<FActorReference> GetActorReferences(const FGetActorReferencesParams& InParams);
 }
\ No newline at end of file
diff --git a/Engine/Source/Runtime/Engine/Public/WorldPartition/WorldPartitionActorDesc.h b/Engine/Source/Runtime/Engine/Public/WorldPartition/WorldPartitionActorDesc.h
--- a/Engine/Source/Runtime/Engine/Public/WorldPartition/WorldPartitionActorDesc.h
+++ b/Engine/Source/Runtime/Engine/Public/WorldPartition/WorldPartitionActorDesc.h
@@ -220,6 +220,11 @@ public:
 		return References;
 	}
 
+	bool IsEditorOnlyReference(const FGuid& ReferenceGuid) const
+	{
+		return EditorOnlyReferences.Contains(ReferenceGuid);
+	}
+
 	UActorDescContainer* GetContainer() const
 	{
 		return Container;
@@ -295,6 +300,7 @@ protected:
 	FName							HLODLayer;
 	TArray<FName>					DataLayers;
 	TArray<FGuid>					References;
+	TArray<FGuid>					EditorOnlyReferences; // References that aren't necessarily editor only but referenced through an editor only property.
 	TArray<FName>					Tags;
 	FPropertyPairsMap				Properties;
 	FName							FolderPath;
diff --git a/Engine/Source/Runtime/Engine/Public/WorldPartition/WorldPartitionActorDescView.h b/Engine/Source/Runtime/Engine/Public/WorldPartition/WorldPartitionActorDescView.h
--- a/Engine/Source/Runtime/Engine/Public/WorldPartition/WorldPartitionActorDescView.h
+++ b/Engine/Source/Runtime/Engine/Public/WorldPartition/WorldPartitionActorDescView.h
@@ -98,6 +98,8 @@ public:
 
 	const FWorldPartitionActorDesc* GetActorDesc() const { return ActorDesc; }
 
+	bool IsEditorOnlyReference(const FGuid& ReferenceGuid) const;
+
 	bool GetProperty(FName PropertyName, FName* PropertyValue) const;
 	bool HasProperty(FName PropertyName) const;
 
