Allow streamable manager to flush only the relevant async loading requests for a given streamable handle.

diff --git a/Engine/Source/Runtime/CoreUObject/Private/Serialization/AsyncPackageLoader.cpp b/Engine/Source/Runtime/CoreUObject/Private/Serialization/AsyncPackageLoader.cpp
--- a/Engine/Source/Runtime/CoreUObject/Private/Serialization/AsyncPackageLoader.cpp
+++ b/Engine/Source/Runtime/CoreUObject/Private/Serialization/AsyncPackageLoader.cpp
@@ -251,6 +251,22 @@ bool IsInAsyncLoadingThreadCoreUObjectInternal()
 	}
 }
 
+void FlushAsyncLoading(TConstArrayView<int32> RequestIds)
+{
+	if (RequestIds.Num() == 0)
+	{
+		FlushAsyncLoading(INDEX_NONE);
+	}
+	else
+	{
+		for (int32 Id : RequestIds)
+		{
+			UE_CLOG(Id == INDEX_NONE, LogStreaming, Warning, TEXT("FlushAsyncLoading called with a list including request id -1, this will flush ALL async loading. If this is intentional, pass an empty list instead."));
+			FlushAsyncLoading(Id);
+		}
+	}
+}
+
 void FlushAsyncLoading(int32 RequestId /* = INDEX_NONE */)
 {
 	TRACE_CPUPROFILER_EVENT_SCOPE(FlushAsyncLoading);
diff --git a/Engine/Source/Runtime/CoreUObject/Private/UObject/UObjectGlobals.cpp b/Engine/Source/Runtime/CoreUObject/Private/UObject/UObjectGlobals.cpp
--- a/Engine/Source/Runtime/CoreUObject/Private/UObject/UObjectGlobals.cpp
+++ b/Engine/Source/Runtime/CoreUObject/Private/UObject/UObjectGlobals.cpp
@@ -1258,9 +1258,15 @@ bool ParseObject( const TCHAR* Stream, const TCHAR* Match, UClass* Class, UObjec
 	}
 }
 
+UE_TRACE_EVENT_BEGIN(Cpu, LoadObject, NoSync)
+	UE_TRACE_EVENT_FIELD(UE::Trace::WideString, AssetPath)
+UE_TRACE_EVENT_END()
+
 UObject* StaticLoadObjectInternal(UClass* ObjectClass, UObject* InOuter, const TCHAR* InName, const TCHAR* Filename, uint32 LoadFlags, UPackageMap* Sandbox, bool bAllowObjectReconciliation, const FLinkerInstancingContext* InstancingContext)
 {
-	SCOPE_CYCLE_COUNTER(STAT_LoadObject);
+	UE_TRACE_LOG_SCOPED_T(Cpu, LoadObject, CpuChannel)
+		<< LoadObject.AssetPath(InName);
+	SCOPED_NAMED_EVENT(LoadObject, FColor::Red);
 	check(InName);
 
 	FScopedLoadingState ScopedLoadingState(InName);
diff --git a/Engine/Source/Runtime/CoreUObject/Public/UObject/UObjectGlobals.h b/Engine/Source/Runtime/CoreUObject/Public/UObject/UObjectGlobals.h
--- a/Engine/Source/Runtime/CoreUObject/Public/UObject/UObjectGlobals.h
+++ b/Engine/Source/Runtime/CoreUObject/Public/UObject/UObjectGlobals.h
@@ -858,6 +858,13 @@ COREUOBJECT_API bool IsReferenced( UObject*& Res, EObjectFlags KeepFlags, EInter
  */
 COREUOBJECT_API void FlushAsyncLoading(int32 PackageID = INDEX_NONE);
 
+/**
+ * Blocks till a set of pending async load requests are complete.
+ *
+ * @param RequestIds list of return values from LoadPackageAsync to wait for. 
+ */
+COREUOBJECT_API void FlushAsyncLoading(TConstArrayView<int32> RequestIds);
+
 /**
  * Return number of active async load package requests
  */
diff --git a/Engine/Source/Runtime/Engine/Classes/Engine/StreamableManager.h b/Engine/Source/Runtime/Engine/Classes/Engine/StreamableManager.h
--- a/Engine/Source/Runtime/Engine/Classes/Engine/StreamableManager.h
+++ b/Engine/Source/Runtime/Engine/Classes/Engine/StreamableManager.h
@@ -587,6 +587,7 @@ private:
 
 	void RemoveReferencedAsset(const FSoftObjectPath& Target, TSharedRef<FStreamableHandle> Handle);
 	void StartHandleRequests(TSharedRef<FStreamableHandle> Handle);
+	TArray<int32> GetAsyncLoadRequestIds(TSharedRef<FStreamableHandle> Handle);
 	void FindInMemory(FSoftObjectPath& InOutTarget, struct FStreamable* Existing, UPackage* Package = nullptr);
 	FSoftObjectPath HandleLoadedRedirector(UObjectRedirector* LoadedRedirector, FSoftObjectPath RequestedPath, struct FStreamable* RequestedStreamable);
 	struct FStreamable* FindStreamable(const FSoftObjectPath& Target) const;
diff --git a/Engine/Source/Runtime/Engine/Private/StreamableManager.cpp b/Engine/Source/Runtime/Engine/Private/StreamableManager.cpp
--- a/Engine/Source/Runtime/Engine/Private/StreamableManager.cpp
+++ b/Engine/Source/Runtime/Engine/Private/StreamableManager.cpp
@@ -10,6 +10,9 @@
 #include "HAL/IConsoleManager.h"
 #include "Tickable.h"
 #include "Serialization/LoadTimeTrace.h"
+#include "ProfilingDebugging/CpuProfilerTrace.h"
+#include "Trace/Trace.h"
+#include "Trace/Trace.inl"
 
 DEFINE_LOG_CATEGORY_STATIC(LogStreamableManager, Log, All);
 
@@ -22,6 +25,16 @@ static FAutoConsoleVariableRef CVarStreamableDelegateDelayFrames(
 	ECVF_Default
 );
 
+// CVar to switch back to legacy behavior of non-specifically flushing async loading when waiting on a request handle
+static bool GStreamableFlushAllAsyncLoadRequestsOnWait = 0;
+static FAutoConsoleVariableRef CVarStreamableFlushAllAsyncLoadRequestsOnWait(
+	TEXT("s.StreamableFlushAllAsyncLoadRequestsOnWait"),
+	GStreamableFlushAllAsyncLoadRequestsOnWait,
+	TEXT("Flush async loading without a specific request ID when waiting on streamable handles."),
+	ECVF_Default
+);
+
+
 const FString FStreamableHandle::HandleDebugName_Preloading = FString(TEXT("Preloading"));
 const FString FStreamableHandle::HandleDebugName_AssetList = FString(TEXT("LoadAssetList"));
 const FString FStreamableHandle::HandleDebugName_CombinedHandle = FString(TEXT("CreateCombinedHandle"));
@@ -408,6 +421,18 @@ EAsyncPackageState::Type FStreamableHandle::WaitUntilComplete(float Timeout, boo
 		}
 	}
 
+	// If we have have a timeout we can't call FlushAsyncLoading so use ProcessAsyncLoadingUntilComplete 
+	if (Timeout == 0.0f && !GStreamableFlushAllAsyncLoadRequestsOnWait)
+	{
+		TArray<int32> RequestIds = OwningManager->GetAsyncLoadRequestIds(AsShared());
+		FlushAsyncLoading(RequestIds);	
+		if (ensureMsgf(HasLoadCompletedOrStalled(), TEXT("Flushing async loading by request id did not complete loading for streamable handle %s"), *GetDebugName()))
+		{
+			return EAsyncPackageState::Complete;
+		}
+		// If for some reason the streamables don't consider themselves complete, fall back to old codepath which flushes asyncing loading without specific IDs
+	}
+	 
 	// Finish when all handles are completed or stalled. If we started stalled above then there will be no stalled handles
 	EAsyncPackageState::Type State = ProcessAsyncLoadingUntilComplete([this]() { return HasLoadCompletedOrStalled(); }, Timeout);
 	
@@ -415,7 +440,6 @@ EAsyncPackageState::Type FStreamableHandle::WaitUntilComplete(float Timeout, boo
 	{
 		ensureMsgf(HasLoadCompletedOrStalled() || WasCanceled(), TEXT("WaitUntilComplete failed for streamable handle %s, async loading is done but handle is not complete"), *GetDebugName());
 	}
-
 	return State;
 }
 
@@ -1016,6 +1040,9 @@ struct FStreamable
 	/** Handles waiting for this to load. Handles may be duplicated with redirectors. */
 	TArray< TSharedRef< FStreamableHandle> > LoadingHandles;
 
+	/** Id for async load request with async loading system */
+	int32 RequestId = INDEX_NONE;
+
 	/** If this object is currently being loaded */
 	bool bAsyncLoadRequestOutstanding = false;
 
@@ -1242,6 +1269,7 @@ FStreamable* FStreamableManager::StreamInternal(const FSoftObjectPath& InTargetN
 		{
 			UE_LOG(LogStreamableManager, Verbose, TEXT("     Already in progress %s"), *TargetName.ToString());
 			check(!Existing->Target); // should not be a load request unless the target is invalid
+			check(Existing->RequestId != INDEX_NONE);
 			ensure(IsAsyncLoading()); // Nothing should be pending if there is no async loading happening
 
 			// Don't return as we potentially want to sync load it
@@ -1321,7 +1349,8 @@ FStreamable* FStreamableManager::StreamInternal(const FSoftObjectPath& InTargetN
 #if WITH_EDITOR
 				FCookLoadScope CookLoadScope(Handle->GetCookLoadType());
 #endif
-				LoadPackageAsync(PackagePath,
+				// This may overwrite an existing request id, this is intentional - see comment above re: cancellation
+				Existing->RequestId = LoadPackageAsync(PackagePath,
 					NAME_None /* PackageNameToCreate */,
 					FLoadPackageAsyncDelegate::CreateSP(Handle, &FStreamableHandle::AsyncLoadCallbackWrapper, TargetName),
 					PKG_None /* InPackageFlags */,
@@ -1521,8 +1550,43 @@ void FStreamableManager::StartHandleRequests(TSharedRef<FStreamableHandle> Handl
 	}
 }
 
+TArray<int32> FStreamableManager::GetAsyncLoadRequestIds(TSharedRef<FStreamableHandle> InitialHandle)
+{
+	TArray<TSharedPtr<FStreamableHandle>> Queue;
+	Queue.Reserve(1 + InitialHandle->ChildHandles.Num());
+	Queue.Add(InitialHandle.ToSharedPtr());	
+	
+	TArray<int32> RequestIds;
+	for (int32 i=0; i < Queue.Num(); ++i)
+	{
+		TSharedPtr<FStreamableHandle> Handle = Queue[i];
+		if (!Handle.IsValid()) 
+		{
+			 continue; 
+		}
+
+		Queue.Append(Handle->ChildHandles);	
+		
+		for (const FSoftObjectPath& Path : Handle->RequestedAssets)
+		{
+			FStreamable* Streamable = StreamableItems.FindRef(Path);
+			if (Streamable && Streamable->RequestId != INDEX_NONE)
+			{
+				RequestIds.Add(Streamable->RequestId);				
+			}
+		}
+	}
+	return RequestIds;
+}
+
+UE_TRACE_EVENT_BEGIN(Cpu, StreamableManager_LoadSynchronous, NoSync)
+	UE_TRACE_EVENT_FIELD(UE::Trace::WideString, AssetPath)
+UE_TRACE_EVENT_END()
+
 UObject* FStreamableManager::LoadSynchronous(const FSoftObjectPath& Target, bool bManageActiveHandle, TSharedPtr<FStreamableHandle>* RequestHandlePointer)
 {
+	UE_TRACE_LOG_SCOPED_T(Cpu, StreamableManager_LoadSynchronous, CpuChannel)
+		<< StreamableManager_LoadSynchronous.AssetPath(*WriteToWideString<FName::StringBufferSize>(Target));
 	TSharedPtr<FStreamableHandle> Request = RequestSyncLoad(Target, bManageActiveHandle, FString::Printf(TEXT("LoadSynchronous of %s"), *Target.ToString()));
 
 	if (RequestHandlePointer)
