- Added GetWorldPartitionCell() that returns, if it exists, an IWorldPartitionCell*, which simplifies usage.
- When the level of a dynamically created world partition cell is streamed-in, set it's WorldPartition cell.

diff --git a/Engine/Source/Runtime/Engine/Classes/Engine/LevelStreaming.h b/Engine/Source/Runtime/Engine/Classes/Engine/LevelStreaming.h
--- a/Engine/Source/Runtime/Engine/Classes/Engine/LevelStreaming.h
+++ b/Engine/Source/Runtime/Engine/Classes/Engine/LevelStreaming.h
@@ -20,6 +20,7 @@ class ALevelScriptActor;
 class ALevelStreamingVolume;
 class ULevel;
 class ULevelStreaming;
+class IWorldPartitionCell;
 struct FNetLevelVisibilityTransactionId;
 
 enum class ENetLevelVisibilityRequest
@@ -332,6 +333,7 @@ public:
 	//~ Begin UObject Interface
 	virtual void PostLoad() override;
 	virtual void Serialize( FArchive& Ar ) override;
+	virtual const IWorldPartitionCell* GetWorldPartitionCell() const { return nullptr; }
 #if WITH_EDITOR
 	virtual void PostEditChangeProperty(FPropertyChangedEvent& PropertyChangedEvent) override;
 	
diff --git a/Engine/Source/Runtime/Engine/Private/EngineUtils.cpp b/Engine/Source/Runtime/Engine/Private/EngineUtils.cpp
--- a/Engine/Source/Runtime/Engine/Private/EngineUtils.cpp
+++ b/Engine/Source/Runtime/Engine/Private/EngineUtils.cpp
@@ -455,7 +455,7 @@ TArray<FSubLevelStatus> GetSubLevelsStatus( UWorld* World, bool SortByActorCount
 
 			if (LevelStreaming->GetLoadedLevel())
 			{
-				if (const IWorldPartitionCell* WorldPartitionCell = LevelStreaming->GetLoadedLevel()->GetWorldPartitionRuntimeCell())
+				if (const IWorldPartitionCell* WorldPartitionCell = LevelStreaming->GetWorldPartitionCell())
 				{
 					LevelStatus.LevelLabel = WorldPartitionCell->GetDebugName();
 				}
diff --git a/Engine/Source/Runtime/Engine/Private/Level.cpp b/Engine/Source/Runtime/Engine/Private/Level.cpp
--- a/Engine/Source/Runtime/Engine/Private/Level.cpp
+++ b/Engine/Source/Runtime/Engine/Private/Level.cpp
@@ -2995,6 +2995,15 @@ void ULevel::OnLevelLoaded()
 #endif
 	};
 
+	// Find associated level streaming for this level
+	const ULevelStreaming* LevelStreaming = FLevelUtils::FindStreamingLevel(this);
+
+	// Set level's associated WorldPartitionRuntimeCell for dynamically injected cells
+	if (LevelStreaming && !WorldPartitionRuntimeCell.GetUniqueID().IsValid())
+	{
+		WorldPartitionRuntimeCell = Cast<const UWorldPartitionRuntimeCell>(LevelStreaming->GetWorldPartitionCell());
+	}
+
 	// 1. Cook commandlet does it's own UWorldPartition::Initialize call in FWorldPartitionCookPackageSplitter::GetGenerateList
 	// 2. Do not Initialize if World doesn't have a UWorldPartitionSubsystem (Known case is when WorldType == EWorldType::Inactive)
 	if (!IsRunningCookCommandlet() && OwningWorld->HasSubsystem<UWorldPartitionSubsystem>())
@@ -3020,11 +3029,7 @@ void ULevel::OnLevelLoaded()
 
 			if (bIsMainWorldLevel || bInitializeForEditor || bInitializeForGame)
 			{
-				FTransform Transform = FTransform::Identity;
-				if (ULevelStreaming* LevelStreaming = FLevelUtils::FindStreamingLevel(this))
-				{
-					Transform = LevelStreaming->LevelTransform;
-				}
+				FTransform Transform = LevelStreaming ? LevelStreaming->LevelTransform : FTransform::Identity;
 
 				// It is allowed for a WorldPartition to already be initialized in the case where a partitioned sub-level 
 				// was streamed-out and is streamed-in again, without a CleanupLevel (GC).
diff --git a/Engine/Source/Runtime/Engine/Private/ProfilingDebugging/LevelStreamingProfilingSubsystem.cpp b/Engine/Source/Runtime/Engine/Private/ProfilingDebugging/LevelStreamingProfilingSubsystem.cpp
--- a/Engine/Source/Runtime/Engine/Private/ProfilingDebugging/LevelStreamingProfilingSubsystem.cpp
+++ b/Engine/Source/Runtime/Engine/Private/ProfilingDebugging/LevelStreamingProfilingSubsystem.cpp
@@ -174,14 +174,11 @@ TUniquePtr<ULevelStreamingProfilingSubsystem::FActiveLevel> ULevelStreamingProfi
 			break;
 		}
 	}
-	if (const UWorldPartitionLevelStreamingDynamic* WPStreamingLevel = Cast<UWorldPartitionLevelStreamingDynamic>(StreamingLevel))
+	if (const UWorldPartitionRuntimeCell* Cell = Cast<const UWorldPartitionRuntimeCell>(StreamingLevel->GetWorldPartitionCell()))
 	{
-		if (const UWorldPartitionRuntimeCell* Cell = WPStreamingLevel->GetWorldPartitionRuntimeCell())
-		{
-			LevelStats[Level->StatsIndex].CellBounds = Cell->GetCellBounds();
-			LevelStats[Level->StatsIndex].ContentBounds = Cell->GetContentBounds();
-			LevelStats[Level->StatsIndex].bIsHLOD = Cell->GetIsHLOD();
-		}
+		LevelStats[Level->StatsIndex].CellBounds = Cell->GetCellBounds();
+		LevelStats[Level->StatsIndex].ContentBounds = Cell->GetContentBounds();
+		LevelStats[Level->StatsIndex].bIsHLOD = Cell->GetIsHLOD();
 	}
 	UpdateLevelState(*Level, StreamingLevel, InitialState, Level->StateStartTime);
 	return MoveTemp(Level);
diff --git a/Engine/Source/Runtime/Engine/Private/WorldPartition/HLOD/HLODDestruction.cpp b/Engine/Source/Runtime/Engine/Private/WorldPartition/HLOD/HLODDestruction.cpp
--- a/Engine/Source/Runtime/Engine/Private/WorldPartition/HLOD/HLODDestruction.cpp
+++ b/Engine/Source/Runtime/Engine/Private/WorldPartition/HLOD/HLODDestruction.cpp
@@ -3,27 +3,22 @@
 #include "WorldPartition/HLOD/HLODDestruction.h"
 
 #include "Engine/LevelStreaming.h"
+#include "Engine/World.h"
 #include "LevelUtils.h"
 #include "UObject/ScriptInterface.h"
 #include "WorldPartition/HLOD/DestructibleHLODComponent.h"
 #include "WorldPartition/HLOD/HLODActor.h"
 #include "WorldPartition/HLOD/HLODSubsystem.h"
-#include "WorldPartition/WorldPartitionLevelStreamingDynamic.h"
 
 #include UE_INLINE_GENERATED_CPP_BY_NAME(HLODDestruction)
 
-#if !WITH_EDITOR
-#include "Engine/World.h"
-#endif
-
 namespace
 {
 	static const UWorldPartitionRuntimeCell* GetActorRuntimeCell(AActor* InActor)
 	{
 		const ULevel* ActorLevel = InActor->GetLevel();
 		const ULevelStreaming* LevelStreaming = FLevelUtils::FindStreamingLevel(ActorLevel);
-		const UWorldPartitionLevelStreamingDynamic* LevelStreamingDynamic = Cast<UWorldPartitionLevelStreamingDynamic>(LevelStreaming);
-		return LevelStreamingDynamic ? LevelStreamingDynamic->GetWorldPartitionRuntimeCell() : nullptr;
+		return LevelStreaming ? Cast<UWorldPartitionRuntimeCell>(LevelStreaming->GetWorldPartitionCell()) : nullptr;
 	}
 
 	static FWorldPartitionHLODDestructionTag ResolveHLODDestructionTagForActor(AActor* InActor)
diff --git a/Engine/Source/Runtime/Engine/Private/WorldPartition/HLOD/HLODSubsystem.cpp b/Engine/Source/Runtime/Engine/Private/WorldPartition/HLOD/HLODSubsystem.cpp
--- a/Engine/Source/Runtime/Engine/Private/WorldPartition/HLOD/HLODSubsystem.cpp
+++ b/Engine/Source/Runtime/Engine/Private/WorldPartition/HLOD/HLODSubsystem.cpp
@@ -16,7 +16,6 @@
 #include "SceneView.h"
 #include "WorldPartition/WorldPartitionRuntimeHash.h"
 #include "SceneManagement.h"
-#include "WorldPartition/WorldPartitionLevelStreamingDynamic.h"
 #include "Engine/Engine.h"
 #include "Engine/StaticMesh.h"
 #include "EngineModule.h"
@@ -103,8 +102,7 @@ namespace FHLODSubsystem
     {
 	    const ULevel* ActorLevel = InActor->GetLevel();
 	    const ULevelStreaming* LevelStreaming = FLevelUtils::FindStreamingLevel(ActorLevel);
-	    const UWorldPartitionLevelStreamingDynamic* LevelStreamingDynamic = Cast<UWorldPartitionLevelStreamingDynamic>(LevelStreaming);
-	    return LevelStreamingDynamic ? LevelStreamingDynamic->GetWorldPartitionRuntimeCell() : nullptr;
+	    return LevelStreaming ? Cast<const UWorldPartitionRuntimeCell>(LevelStreaming->GetWorldPartitionCell()) : nullptr;
     }
     
     static UWorldPartition* GetWorldPartition(AWorldPartitionHLOD* InWorldPartitionHLOD)
diff --git a/Engine/Source/Runtime/Engine/Private/WorldPartition/WorldPartitionLevelStreamingPolicy.cpp b/Engine/Source/Runtime/Engine/Private/WorldPartition/WorldPartitionLevelStreamingPolicy.cpp
--- a/Engine/Source/Runtime/Engine/Private/WorldPartition/WorldPartitionLevelStreamingPolicy.cpp
+++ b/Engine/Source/Runtime/Engine/Private/WorldPartition/WorldPartitionLevelStreamingPolicy.cpp
@@ -38,12 +38,9 @@ void UWorldPartitionLevelStreamingPolicy::ForEachActiveRuntimeCell(TFunctionRef<
 	UWorld* World = WorldPartition->GetWorld();
 	for (ULevelStreaming* LevelStreaming : World->GetStreamingLevels())
 	{
-		if (UWorldPartitionLevelStreamingDynamic* WorldPartitionLevelStreaming = Cast<UWorldPartitionLevelStreamingDynamic>(LevelStreaming))
+		if (const UWorldPartitionRuntimeCell* Cell = Cast<const UWorldPartitionRuntimeCell>(LevelStreaming->GetWorldPartitionCell()))
 		{
-			if (const UWorldPartitionRuntimeCell* Cell = WorldPartitionLevelStreaming->GetWorldPartitionRuntimeCell())
-			{
-				Func(Cell);
-			}
+			Func(Cell);
 		}
 	}
 }
diff --git a/Engine/Source/Runtime/Engine/Private/WorldPartition/WorldPartitionSubsystem.cpp b/Engine/Source/Runtime/Engine/Private/WorldPartition/WorldPartitionSubsystem.cpp
--- a/Engine/Source/Runtime/Engine/Private/WorldPartition/WorldPartitionSubsystem.cpp
+++ b/Engine/Source/Runtime/Engine/Private/WorldPartition/WorldPartitionSubsystem.cpp
@@ -1040,12 +1040,9 @@ void UWorldPartitionSubsystem::DrawStreamingStatusLegend(class UCanvas* Canvas,
 	{
 		if (!StreamingWorld || (StreamingWorld == LevelStreaming->GetStreamingWorld()))
 		{
-			if (UWorldPartitionLevelStreamingDynamic* WorldPartitionLevelStreaming = Cast<UWorldPartitionLevelStreamingDynamic>(LevelStreaming))
+			if (const UWorldPartitionRuntimeCell* Cell = Cast<const UWorldPartitionRuntimeCell>(LevelStreaming->GetWorldPartitionCell()))
 			{
-				if (const UWorldPartitionRuntimeCell* Cell = WorldPartitionLevelStreaming->GetWorldPartitionRuntimeCell())
-				{
-					StreamingStatus[(int32)Cell->GetStreamingStatus()]++;
-				}
+				StreamingStatus[(int32)Cell->GetStreamingStatus()]++;
 			}
 		}
 	}
diff --git a/Engine/Source/Runtime/Engine/Public/WorldPartition/WorldPartitionLevelStreamingDynamic.h b/Engine/Source/Runtime/Engine/Public/WorldPartition/WorldPartitionLevelStreamingDynamic.h
--- a/Engine/Source/Runtime/Engine/Public/WorldPartition/WorldPartitionLevelStreamingDynamic.h
+++ b/Engine/Source/Runtime/Engine/Public/WorldPartition/WorldPartitionLevelStreamingDynamic.h
@@ -36,8 +36,8 @@ class ENGINE_API UWorldPartitionLevelStreamingDynamic : public ULevelStreamingDy
 	virtual bool ShouldBeAlwaysLoaded() const override { return bShouldBeAlwaysLoaded; }
 	virtual bool ShouldBlockOnUnload() const override;
 	virtual bool ShouldRequireFullVisibilityToRender() const override { return true; }
-
 	virtual bool RequestVisibilityChange(bool bVisible) override;
+	virtual const IWorldPartitionCell* GetWorldPartitionCell() const override { return GetWorldPartitionRuntimeCell(); }
 
 #if !WITH_EDITOR
 	virtual void PostLoad();
