World Partition - Avoid all processing related to the streaming generation log if we aren't generating streaming for the main WP

diff --git a/Engine/Source/Runtime/Engine/Private/WorldPartition/WorldPartitionStreamingGeneration.cpp b/Engine/Source/Runtime/Engine/Private/WorldPartition/WorldPartitionStreamingGeneration.cpp
--- a/Engine/Source/Runtime/Engine/Private/WorldPartition/WorldPartitionStreamingGeneration.cpp
+++ b/Engine/Source/Runtime/Engine/Private/WorldPartition/WorldPartitionStreamingGeneration.cpp
@@ -1155,6 +1155,7 @@ bool UWorldPartition::GenerateContainerStreaming(const FGenerateStreamingParams&
 
 	// Dump state log
 	TUniquePtr<FArchive> LogFileAr;
+	TUniquePtr<FHierarchicalLogArchive> HierarchicalLogAr;
 
 	if (IsMainWorldPartition() && (!GIsBuildMachine || IsRunningCookCommandlet()))
 	{
@@ -1165,14 +1166,9 @@ bool UWorldPartition::GenerateContainerStreaming(const FGenerateStreamingParams&
 		LogFileAr = FWorldPartitionStreamingGenerator::CreateDumpStateLogArchive(*StateLogSuffix, !InParams.OutputLogPath);
 		
 		InContext.OutputLogFilename = LogFileAr->GetArchiveName();
-	}
-	else
-	{
-		LogFileAr = MakeUnique<FArchive>();
+		HierarchicalLogAr = MakeUnique<FHierarchicalLogArchive>(*LogFileAr);
 	}
 
-	FHierarchicalLogArchive HierarchicalLogAr(*LogFileAr);
-
 	FWorldPartitionStreamingGenerator::FWorldPartitionStreamingGeneratorParams StreamingGeneratorParams;
 	StreamingGeneratorParams.WorldPartitionContext = this;
 	StreamingGeneratorParams.ModifiedActorsDescList = ModifiedActorsDescList;
@@ -1185,7 +1181,10 @@ bool UWorldPartition::GenerateContainerStreaming(const FGenerateStreamingParams&
 	// Preparation Phase
 	StreamingGenerator.PreparationPhase(InParams.ActorDescContainer);
 
-	StreamingGenerator.DumpStateLog(HierarchicalLogAr);
+	if (HierarchicalLogAr.IsValid())
+	{
+		StreamingGenerator.DumpStateLog(*HierarchicalLogAr);
+	}
 
 	// Generate streaming
 	check(!StreamingPolicy);
@@ -1194,7 +1193,11 @@ bool UWorldPartition::GenerateContainerStreaming(const FGenerateStreamingParams&
 	check(RuntimeHash);
 	if (RuntimeHash->GenerateStreaming(StreamingPolicy, StreamingGenerator.GetStreamingGenerationContext(), InContext.PackagesToGenerate))
 	{
-		RuntimeHash->DumpStateLog(HierarchicalLogAr);
+		if (HierarchicalLogAr.IsValid())
+		{
+			RuntimeHash->DumpStateLog(*HierarchicalLogAr);
+		}
+
 		StreamingPolicy->PrepareActorToCellRemapping();
 		return true;
 	}
