Added optional SourceCellGuid to a world partition cell which is necessary for any HLOD actor part of a level that is instanced multiple times (shared amongst multiple cells)

diff --git a/Engine/Source/Runtime/Engine/Private/WorldPartition/HLOD/HLODActor.cpp b/Engine/Source/Runtime/Engine/Private/WorldPartition/HLOD/HLODActor.cpp
--- a/Engine/Source/Runtime/Engine/Private/WorldPartition/HLOD/HLODActor.cpp
+++ b/Engine/Source/Runtime/Engine/Private/WorldPartition/HLOD/HLODActor.cpp
@@ -74,6 +74,21 @@ AWorldPartitionHLOD::AWorldPartitionHLOD(const FObjectInitializer& ObjectInitial
 #endif
 }
 
+const FGuid& AWorldPartitionHLOD::GetSourceCellGuid() const
+{
+	// When no source cell guid was set, try resolving it through its associated world partition runtime cell
+	// This is necessary for any HLOD actor part of a level that is instanced multiple times (shared amongst multiple cells)
+	if (!SourceCellGuid.IsValid())
+	{
+		const UWorldPartitionRuntimeCell* Cell = Cast<UWorldPartitionRuntimeCell>(GetLevel()->GetWorldPartitionRuntimeCell());
+		if (Cell && Cell->GetIsHLOD())
+		{
+			const_cast<AWorldPartitionHLOD*>(this)->SourceCellGuid = Cell->GetSourceCellGuid();
+		}
+	}
+	return SourceCellGuid;
+}
+
 void AWorldPartitionHLOD::SetVisibility(bool bInVisible)
 {
 	ForEachComponent<USceneComponent>(false, [bInVisible](USceneComponent* SceneComponent)
diff --git a/Engine/Source/Runtime/Engine/Public/WorldPartition/HLOD/HLODActor.h b/Engine/Source/Runtime/Engine/Public/WorldPartition/HLOD/HLODActor.h
--- a/Engine/Source/Runtime/Engine/Public/WorldPartition/HLOD/HLODActor.h
+++ b/Engine/Source/Runtime/Engine/Public/WorldPartition/HLOD/HLODActor.h
@@ -23,7 +23,7 @@ class ENGINE_API AWorldPartitionHLOD : public AActor
 public:
 	void SetVisibility(bool bInVisible);
 
-	inline FGuid GetSourceCellGuid() const { return SourceCellGuid; }
+	const FGuid& GetSourceCellGuid() const;
 	inline uint32 GetLODLevel() const { return LODLevel; }
 
 	virtual bool IsHLODRelevant() const override { return true; }
diff --git a/Engine/Source/Runtime/Engine/Public/WorldPartition/WorldPartitionRuntimeCell.h b/Engine/Source/Runtime/Engine/Public/WorldPartition/WorldPartitionRuntimeCell.h
--- a/Engine/Source/Runtime/Engine/Public/WorldPartition/WorldPartitionRuntimeCell.h
+++ b/Engine/Source/Runtime/Engine/Public/WorldPartition/WorldPartitionRuntimeCell.h
@@ -241,16 +241,18 @@ class ENGINE_API UWorldPartitionRuntimeCell : public UObject, public IWorldParti
 	virtual bool PopulateGeneratedPackageForCook(UPackage* InPackage, TArray<UPackage*>& OutModifiedPackages) PURE_VIRTUAL(UWorldPartitionRuntimeCell::PopulateGeneratedPackageForCook, return false;);
 	virtual FString GetPackageNameToCreate() const PURE_VIRTUAL(UWorldPartitionRuntimeCell::GetPackageNameToCreate, return FString(""););
 
-	void SetIsHLOD(bool bInIsHLOD) { bIsHLOD = bInIsHLOD; }
-
 	virtual void DumpStateLog(FHierarchicalLogArchive& Ar);
 #endif
-	
+
+	void SetIsHLOD(bool bInIsHLOD) { bIsHLOD = bInIsHLOD; }
 	bool GetIsHLOD() const { return bIsHLOD; }
 	
-	FGuid GetGuid() const { return CellGuid; }
+	const FGuid& GetGuid() const { return CellGuid; }
 	void SetGuid(const FGuid& InCellGuid) { CellGuid = InCellGuid; }
 
+	const FGuid& GetSourceCellGuid() const { return SourceCellGuid; }
+	void SetSourceCellGuid(const FGuid& InSourceCellGuid) { SourceCellGuid = InSourceCellGuid; }
+
 #if !UE_BUILD_SHIPPING
 	void SetDebugStreamingPriority(float InDebugStreamingPriority) { DebugStreamingPriority = InDebugStreamingPriority; }
 #endif
@@ -290,6 +292,10 @@ protected:
 	UPROPERTY()
 	FGuid CellGuid;
 
+	// Used by injected HLOD cells
+	UPROPERTY()
+	FGuid SourceCellGuid;
+
 #if !UE_BUILD_SHIPPING
 	// Represents the streaming priority relative to other cells
 	float DebugStreamingPriority;
