Fix missing member initialization from previous cl.

diff --git a/Engine/Source/Runtime/Engine/Private/ActorReferencesUtils.cpp b/Engine/Source/Runtime/Engine/Private/ActorReferencesUtils.cpp
--- a/Engine/Source/Runtime/Engine/Private/ActorReferencesUtils.cpp
+++ b/Engine/Source/Runtime/Engine/Private/ActorReferencesUtils.cpp
@@ -28,6 +28,7 @@ public:
 	FArchiveGatherExternalActorRefs(UObject* InRoot, EObjectFlags InRequiredFlags, bool bInRecursive, FActorReferenceSet& InActorReference)
 		: Root(InRoot)
 		, ActorReferences(InActorReference)
+		, EditorOnlyStack(0)
 		, RequiredFlags(InRequiredFlags)
 		, bRecursive(bInRecursive)
 	{
