World Partition: remove assert assuming that editor-only actor references are editor-only actors now that we properly handle editor-only properties (they can point to non-editor-only actors).

diff --git a/Engine/Source/Runtime/Engine/Private/WorldPartition/WorldPartitionRuntimeLevelStreamingCell.cpp b/Engine/Source/Runtime/Engine/Private/WorldPartition/WorldPartitionRuntimeLevelStreamingCell.cpp
--- a/Engine/Source/Runtime/Engine/Private/WorldPartition/WorldPartitionRuntimeLevelStreamingCell.cpp
+++ b/Engine/Source/Runtime/Engine/Private/WorldPartition/WorldPartitionRuntimeLevelStreamingCell.cpp
@@ -210,7 +210,6 @@ void UWorldPartitionRuntimeLevelStreamingCell::AddActorToCell(const FWorldPartit
 		for (const FGuid& EditorReferenceGuid : ActorDescView.GetEditorReferences())
 		{
 			const FWorldPartitionActorDesc& ReferenceActorDesc = InContainer->GetActorDescChecked(EditorReferenceGuid);
-			check(ReferenceActorDesc.GetActorIsEditorOnly());
 
 			Packages.Emplace(
 				ReferenceActorDesc.GetActorPackage(), 
