Added missing const

#rnx

diff --git a/Engine/Source/Runtime/Engine/Classes/Engine/MemberReference.h b/Engine/Source/Runtime/Engine/Classes/Engine/MemberReference.h
--- a/Engine/Source/Runtime/Engine/Classes/Engine/MemberReference.h
+++ b/Engine/Source/Runtime/Engine/Classes/Engine/MemberReference.h
@@ -185,7 +185,7 @@ public:
 
 	/** Set up this reference to a 'self' member name */
 	ENGINE_API void SetSelfMember(FName InMemberName);
-	ENGINE_API void SetSelfMember(FName InMemberName, FGuid& InMemberGuid);
+	ENGINE_API void SetSelfMember(FName InMemberName, const FGuid& InMemberGuid);
 
 	/** Set up this reference to a 'self' member name, scoped to a struct */
 	ENGINE_API void SetLocalMember(FName InMemberName, UStruct* InScope, const FGuid InMemberGuid);
diff --git a/Engine/Source/Runtime/Engine/Private/MemberReference.cpp b/Engine/Source/Runtime/Engine/Private/MemberReference.cpp
--- a/Engine/Source/Runtime/Engine/Private/MemberReference.cpp
+++ b/Engine/Source/Runtime/Engine/Private/MemberReference.cpp
@@ -58,7 +58,7 @@ void FMemberReference::SetSelfMember(FName InMemberName)
 	bWasDeprecated = false;
 }
 
-void FMemberReference::SetSelfMember(FName InMemberName, FGuid& InMemberGuid)
+void FMemberReference::SetSelfMember(FName InMemberName, const FGuid& InMemberGuid)
 {
 	SetSelfMember(InMemberName);
 	MemberGuid = InMemberGuid;
