Verse debug draw: support Clear and ClearChannel
Remove visible flag per line, use one flag per LineBatchComponent
#rnx

diff --git a/Engine/Source/Runtime/Engine/Classes/Components/LineBatchComponent.h b/Engine/Source/Runtime/Engine/Classes/Components/LineBatchComponent.h
--- a/Engine/Source/Runtime/Engine/Classes/Components/LineBatchComponent.h
+++ b/Engine/Source/Runtime/Engine/Classes/Components/LineBatchComponent.h
@@ -37,9 +37,6 @@ struct FBatchedLine
 
 	UPROPERTY()
 	uint8 DepthPriority;
-	
-	UPROPERTY()
-	bool bIsVisible;
 
 	UPROPERTY()
 	uint32 BatchID;
@@ -51,17 +48,15 @@ struct FBatchedLine
 		, Thickness(0)
 		, RemainingLifeTime(0)
 		, DepthPriority(0)
-		, bIsVisible(true)
 		, BatchID(0)
 	{}
-	FBatchedLine(const FVector& InStart, const FVector& InEnd, const FLinearColor& InColor, float InLifeTime, float InThickness, uint8 InDepthPriority, uint32 InBatchID = 0, bool bInVisible = true)
+	FBatchedLine(const FVector& InStart, const FVector& InEnd, const FLinearColor& InColor, float InLifeTime, float InThickness, uint8 InDepthPriority, uint32 InBatchID = 0)
 		:	Start(InStart)
 		,	End(InEnd)
 		,	Color(InColor)
 		,	Thickness(InThickness)
 		,	RemainingLifeTime(InLifeTime)
 		,	DepthPriority(InDepthPriority)
-		,	bIsVisible(bInVisible)
 		,	BatchID(InBatchID)
 	{}
 };
@@ -85,9 +80,6 @@ struct FBatchedPoint
 
 	UPROPERTY()
 	uint8 DepthPriority;
-
-	UPROPERTY()
-	bool bIsVisible;
 	
 	UPROPERTY()
 	uint32 BatchID;
@@ -98,16 +90,14 @@ struct FBatchedPoint
 		, PointSize(0)
 		, RemainingLifeTime(0)
 		, DepthPriority(0)
-		, bIsVisible(true)
 		, BatchID(0)
 	{}
-	FBatchedPoint(const FVector& InPosition, const FLinearColor& InColor, float InPointSize, float InLifeTime, uint8 InDepthPriority, uint32 InBatchID = 0, bool bInVisible = true)
+	FBatchedPoint(const FVector& InPosition, const FLinearColor& InColor, float InPointSize, float InLifeTime, uint8 InDepthPriority, uint32 InBatchID = 0)
 		:	Position(InPosition)
 		,	Color(InColor)
 		,	PointSize(InPointSize)
 		,	RemainingLifeTime(InLifeTime)
 		,	DepthPriority(InDepthPriority)
-		,	bIsVisible(bInVisible)
 		,	BatchID(InBatchID)
 	{}
 };
@@ -120,9 +110,9 @@ struct FBatchedMesh
 	 * MeshVerts - linear array of world space vertex positions
 	 * MeshIndices - array of indices into MeshVerts.  Each triplet is a tri.  i.e. [0,1,2] is first tri, [3,4,5] is 2nd tri, etc
 	 */
-	FBatchedMesh(TArray<FVector> const& InMeshVerts, TArray<int32> const& InMeshIndices, FColor const& InColor, uint8 InDepthPriority, float LifeTime, uint32 InBatchID = 0, bool bInVisible = true)
+	FBatchedMesh(TArray<FVector> const& InMeshVerts, TArray<int32> const& InMeshIndices, FColor const& InColor, uint8 InDepthPriority, float LifeTime, uint32 InBatchID = 0)
 		: MeshVerts(InMeshVerts), MeshIndices(InMeshIndices), 
-		  Color(InColor), RemainingLifeTime(LifeTime), DepthPriority(InDepthPriority), bIsVisible(bInVisible), BatchID(InBatchID)
+		  Color(InColor), RemainingLifeTime(LifeTime), DepthPriority(InDepthPriority), BatchID(InBatchID)
 	{}
 
 	TArray<FVector> MeshVerts;
@@ -130,7 +120,6 @@ struct FBatchedMesh
 	FColor Color;
 	float RemainingLifeTime = 0.f;
 	uint8 DepthPriority = 0;
-	bool bIsVisible = true;
 	uint32 BatchID = 0;
 };
 
@@ -153,8 +142,6 @@ class ULineBatchComponent : public UPrimitiveComponent
 	TArray<FBatchedMesh> BatchedMeshes;
 	/** Whether to calculate a tight accurate bounds (encompassing all points), or use a giant bounds that is fast to compute. */
 	uint32 bCalculateAccurateBounds:1;
-	/** Whether this has some hidden lines */
-	uint32 bHasHiddenLines:1;
 
 	/** Defines the value for an invalid id */
 	static constexpr uint32 INVALID_ID = 0;
@@ -163,34 +150,34 @@ class ULineBatchComponent : public UPrimitiveComponent
 	ENGINE_API void DrawLines(TArrayView<FBatchedLine> InLines);
 
 	/** Draw a box */
-	ENGINE_API void DrawBox(const FBox& Box, const FMatrix& TM, FLinearColor Color, uint8 InDepthPriorityGroup, uint32 BatchID = INVALID_ID, bool bIsVisible = true);
+	ENGINE_API void DrawBox(const FBox& Box, const FMatrix& TM, FLinearColor Color, uint8 InDepthPriorityGroup, uint32 BatchID = INVALID_ID);
 
 	/** Draw a box */
-	ENGINE_API void DrawBox(const FVector& Center, const FVector& Box, FLinearColor Color, float LifeTime, uint8 DepthPriority, float Thickness, uint32 BatchID = INVALID_ID, bool bIsVisible = true);
+	ENGINE_API void DrawBox(const FVector& Center, const FVector& Box, FLinearColor Color, float LifeTime, uint8 DepthPriority, float Thickness, uint32 BatchID = INVALID_ID);
 	
 	/** Draw a box */
-	ENGINE_API void DrawBox(const FVector& Center, const FVector& Box, const FQuat& Rotation, FLinearColor Color, float LifeTime, uint8 DepthPriority, float Thickness, uint32 BatchID = INVALID_ID, bool bIsVisible = true);
+	ENGINE_API void DrawBox(const FVector& Center, const FVector& Box, const FQuat& Rotation, FLinearColor Color, float LifeTime, uint8 DepthPriority, float Thickness, uint32 BatchID = INVALID_ID);
 
 	/** Draw an arrow */
-	ENGINE_API void DrawDirectionalArrow(const FMatrix& ArrowToWorld, FLinearColor InColor, float Length, float ArrowSize, uint8 DepthPriority, uint32 BatchID = INVALID_ID, bool bIsVisible = true);
+	ENGINE_API void DrawDirectionalArrow(const FMatrix& ArrowToWorld, FLinearColor InColor, float Length, float ArrowSize, uint8 DepthPriority, uint32 BatchID = INVALID_ID);
 
 	/** Draw an arrow */
-	ENGINE_API void DrawDirectionalArrow(const FVector& LineStart, const FVector& LineEnd, float ArrowSize, FLinearColor Color, float LifeTime, uint8 DepthPriority, float Thickness, uint32 BatchID = INVALID_ID, bool bIsVisible = true);
+	ENGINE_API void DrawDirectionalArrow(const FVector& LineStart, const FVector& LineEnd, float ArrowSize, FLinearColor Color, float LifeTime, uint8 DepthPriority, float Thickness, uint32 BatchID = INVALID_ID);
 
 	/** Draw a circle */
-	ENGINE_API void DrawCircle(const FVector& Base, const FVector& X, const FVector& Y, FLinearColor Color, float Radius, int32 NumSides, uint8 DepthPriority, uint32 BatchID = INVALID_ID, bool bIsVisible = true);
+	ENGINE_API void DrawCircle(const FVector& Base, const FVector& X, const FVector& Y, FLinearColor Color, float Radius, int32 NumSides, uint8 DepthPriority, uint32 BatchID = INVALID_ID);
 
 	/** Draw a sphere */
-	ENGINE_API void DrawSphere(const FVector& Center, float Radius, int32 Segments, FLinearColor Color, float LifeTime, uint8 DepthPriority, float Thickness, uint32 BatchID = INVALID_ID, bool bIsVisible = true);
+	ENGINE_API void DrawSphere(const FVector& Center, float Radius, int32 Segments, FLinearColor Color, float LifeTime, uint8 DepthPriority, float Thickness, uint32 BatchID = INVALID_ID);
 
 	/** Draw a cylinder */
-	ENGINE_API void DrawCylinder(const FVector& Start, const FVector& End, float Radius, int32 Segments, FLinearColor  Color, float LifeTime, uint8 DepthPriority, float Thickness, uint32 BatchID = INVALID_ID, bool bIsVisible = true);
+	ENGINE_API void DrawCylinder(const FVector& Start, const FVector& End, float Radius, int32 Segments, FLinearColor  Color, float LifeTime, uint8 DepthPriority, float Thickness, uint32 BatchID = INVALID_ID);
 
 	/** Draw a cone */
-	ENGINE_API void DrawCone(const FVector& Origin, const FVector& Direction, float Length, float AngleWidth, float AngleHeight, int32 NumSides, FLinearColor DrawColor, float LifeTime, uint8 DepthPriority, float Thickness, uint32 BatchID = INVALID_ID, bool bIsVisible = true);
+	ENGINE_API void DrawCone(const FVector& Origin, const FVector& Direction, float Length, float AngleWidth, float AngleHeight, int32 NumSides, FLinearColor DrawColor, float LifeTime, uint8 DepthPriority, float Thickness, uint32 BatchID = INVALID_ID);
 
 	/** Draw a cone */
-	ENGINE_API void DrawCapsule(const FVector& Center, float HalfHeight, float Radius, const FQuat& Rotation, FLinearColor Color, float LifeTime, uint8 DepthPriority, float Thickness, uint32 BatchID = INVALID_ID, bool bIsVisible = true);
+	ENGINE_API void DrawCapsule(const FVector& Center, float HalfHeight, float Radius, const FQuat& Rotation, FLinearColor Color, float LifeTime, uint8 DepthPriority, float Thickness, uint32 BatchID = INVALID_ID);
 
 	ENGINE_API virtual void DrawLine(
 		const FVector& Start,
@@ -199,8 +186,7 @@ class ULineBatchComponent : public UPrimitiveComponent
 		uint8 DepthPriority,
 		float Thickness = 0.0f,
 		float LifeTime = 0.0f,
-		uint32 BatchID = INVALID_ID,
-		bool bIsVisible = true
+		uint32 BatchID = INVALID_ID
 		);
 	ENGINE_API virtual void DrawPoint(
 		const FVector& Position,
@@ -208,14 +194,13 @@ class ULineBatchComponent : public UPrimitiveComponent
 		float PointSize,
 		uint8 DepthPriority,
 		float LifeTime = 0.0f,
-		uint32 BatchID = INVALID_ID,
-		bool bIsVisible = true
+		uint32 BatchID = INVALID_ID
 		);
 
 	/** Draw a box */
-	ENGINE_API void DrawSolidBox(FBox const& Box, FTransform const& Xform, const FColor& Color, uint8 DepthPriority, float LifeTime, uint32 BatchID = INVALID_ID, bool bIsVisible = true);
+	ENGINE_API void DrawSolidBox(FBox const& Box, FTransform const& Xform, const FColor& Color, uint8 DepthPriority, float LifeTime, uint32 BatchID = INVALID_ID);
 	/** Draw a mesh */
-	ENGINE_API void DrawMesh(TArray<FVector> const& Verts, TArray<int32> const& Indices, FColor const& Color, uint8 DepthPriority, float LifeTime, uint32 BatchID = INVALID_ID, bool bIsVisible = true);
+	ENGINE_API void DrawMesh(TArray<FVector> const& Verts, TArray<int32> const& Indices, FColor const& Color, uint8 DepthPriority, float LifeTime, uint32 BatchID = INVALID_ID);
 
 	//~ Begin UPrimitiveComponent Interface.
 	ENGINE_API virtual FPrimitiveSceneProxy* CreateSceneProxy() override;
@@ -234,13 +219,7 @@ class ULineBatchComponent : public UPrimitiveComponent
 	/** Remove batched lines, points and meshes with given ID */
 	ENGINE_API void ClearBatch(uint32 InBatchID);
 	
-	/** Set batched lines, points and meshes with given ID visible or hidden*/
-	ENGINE_API void SetBatchVisible(uint32 InBatchID, bool bInVisible);
-	
-	/** Set all batched lines, points and meshes visible or hidden*/
-	ENGINE_API void SetVisible(bool bInVisible);
-
 protected:
-	void AddHalfCircle(const FVector& Base, const FVector& X, const FVector& Y, const FLinearColor& Color, const float Radius, int32 NumSides, const float LifeTime, uint8 DepthPriority, const float Thickness, const uint32 BatchID, bool bIsVisible);
-	void AddCircle(const FVector& Base, const FVector& X, const FVector& Y, const FLinearColor& Color, const float Radius, int32 NumSides, const float LifeTime, uint8 DepthPriority, const float Thickness, const uint32 BatchID, bool bIsVisible);
+	void AddHalfCircle(const FVector& Base, const FVector& X, const FVector& Y, const FLinearColor& Color, const float Radius, int32 NumSides, const float LifeTime, uint8 DepthPriority, const float Thickness, const uint32 BatchID);
+	void AddCircle(const FVector& Base, const FVector& X, const FVector& Y, const FLinearColor& Color, const float Radius, int32 NumSides, const float LifeTime, uint8 DepthPriority, const float Thickness, const uint32 BatchID);
 };
diff --git a/Engine/Source/Runtime/Engine/Private/Components/LineBatchComponent.cpp b/Engine/Source/Runtime/Engine/Private/Components/LineBatchComponent.cpp
--- a/Engine/Source/Runtime/Engine/Private/Components/LineBatchComponent.cpp
+++ b/Engine/Source/Runtime/Engine/Private/Components/LineBatchComponent.cpp
@@ -47,12 +47,6 @@ FLineBatcherSceneProxy::FLineBatcherSceneProxy(const ULineBatchComponent* InComp
 	Points(InComponent->BatchedPoints), Meshes(InComponent->BatchedMeshes)
 {
 	bWillEverBeLit = false;
-	if (InComponent->bHasHiddenLines)
-	{
-		Lines.RemoveAllSwap([](const FBatchedLine& Element) { return Element.bIsVisible == false; });
-		Points.RemoveAllSwap([](const FBatchedPoint& Pt) { return Pt.bIsVisible == false; });
-		Meshes.RemoveAllSwap([](const FBatchedMesh& Mesh) { return Mesh.bIsVisible == false; });
-	}
 }
 
 SIZE_T FLineBatcherSceneProxy::GetTypeHash() const
@@ -153,18 +147,15 @@ ULineBatchComponent::ULineBatchComponent(const FObjectInitializer& ObjectInitial
 	bUseEditorCompositing = true;
 	SetGenerateOverlapEvents(false);
 	bCalculateAccurateBounds = true;
-	bHasHiddenLines = false;
 	DefaultLifeTime = 1.0f;
 
 	// Ignore streaming updates since GetUsedMaterials() is not implemented.
 	bIgnoreStreamingManagerUpdate = true;
 }
 
-void ULineBatchComponent::DrawLine(const FVector& Start, const FVector& End, const FLinearColor& Color, uint8 DepthPriority, const float Thickness, const float LifeTime, uint32 BatchID, bool bIsVisible)
+void ULineBatchComponent::DrawLine(const FVector& Start, const FVector& End, const FLinearColor& Color, uint8 DepthPriority, const float Thickness, const float LifeTime, uint32 BatchID)
 {
-	new(BatchedLines) FBatchedLine(Start, End, Color, LifeTime, Thickness, DepthPriority, BatchID, bIsVisible);
-	
-	bHasHiddenLines |= !bIsVisible;
+	BatchedLines.Emplace(Start, End, Color, LifeTime, Thickness, DepthPriority, BatchID);
 
 	// LineBatcher and PersistentLineBatcher components will be updated at the end of UWorld::Tick
 	MarkRenderStateDirty();
@@ -184,17 +175,15 @@ void ULineBatchComponent::DrawPoint(
 	float PointSize,
 	uint8 DepthPriority,
 	float LifeTime,
-	uint32 BatchID,
-	bool bIsVisible
+	uint32 BatchID
 	)
 {
-	new(BatchedPoints) FBatchedPoint(Position, Color, PointSize, LifeTime, DepthPriority, BatchID, bIsVisible);
-	bHasHiddenLines |= !bIsVisible;
+	BatchedPoints.Emplace(Position, Color, PointSize, LifeTime, DepthPriority, BatchID);
 	// LineBatcher and PersistentLineBatcher components will be updated at the end of UWorld::Tick
 	MarkRenderStateDirty();
 }
 
-void ULineBatchComponent::DrawBox(const FBox& Box, const FMatrix& TM, FLinearColor Color, uint8 InDepthPriorityGroup, uint32 BatchID, bool bIsVisible)
+void ULineBatchComponent::DrawBox(const FBox& Box, const FMatrix& TM, FLinearColor Color, uint8 InDepthPriorityGroup, uint32 BatchID)
 {
 	FVector	B[2], P, Q;
 	B[0] = Box.Min;
@@ -205,102 +194,98 @@ void ULineBatchComponent::DrawBox(const FBox& Box, const FMatrix& TM, FLinearCol
 		P.X=B[ai].X; Q.X=B[ai].X;
 		P.Y=B[aj].Y; Q.Y=B[aj].Y;
 		P.Z=B[0].Z; Q.Z=B[1].Z;
-		new(BatchedLines) FBatchedLine(TM.TransformPosition(P), TM.TransformPosition(Q), Color, DefaultLifeTime, 0.0f, InDepthPriorityGroup, BatchID, bIsVisible);
+		BatchedLines.Emplace(TM.TransformPosition(P), TM.TransformPosition(Q), Color, DefaultLifeTime, 0.0f, InDepthPriorityGroup, BatchID);
 
 		P.Y=B[ai].Y; Q.Y=B[ai].Y;
 		P.Z=B[aj].Z; Q.Z=B[aj].Z;
 		P.X=B[0].X; Q.X=B[1].X;
-		new(BatchedLines) FBatchedLine(TM.TransformPosition(P), TM.TransformPosition(Q), Color, DefaultLifeTime, 0.0f, InDepthPriorityGroup, BatchID, bIsVisible);
+		BatchedLines.Emplace(TM.TransformPosition(P), TM.TransformPosition(Q), Color, DefaultLifeTime, 0.0f, InDepthPriorityGroup, BatchID);
 
 		P.Z=B[ai].Z; Q.Z=B[ai].Z;
 		P.X=B[aj].X; Q.X=B[aj].X;
 		P.Y=B[0].Y; Q.Y=B[1].Y;
-		new(BatchedLines) FBatchedLine(TM.TransformPosition(P), TM.TransformPosition(Q), Color, DefaultLifeTime, 0.0f, InDepthPriorityGroup, BatchID, bIsVisible);
+		BatchedLines.Emplace(TM.TransformPosition(P), TM.TransformPosition(Q), Color, DefaultLifeTime, 0.0f, InDepthPriorityGroup, BatchID);
 	}
 
-	bHasHiddenLines |= !bIsVisible;
 	// LineBatcher and PersistentLineBatcher components will be updated at the end of UWorld::Tick
 	MarkRenderStateDirty();
 }
 
-void ULineBatchComponent::DrawBox(FVector const& Center, FVector const& Box, FLinearColor Color, float LifeTime, uint8 DepthPriority, float Thickness, uint32 BatchID, bool bIsVisible)
+void ULineBatchComponent::DrawBox(FVector const& Center, FVector const& Box, FLinearColor Color, float LifeTime, uint8 DepthPriority, float Thickness, uint32 BatchID)
 {
-	new(BatchedLines) FBatchedLine(Center + FVector( Box.X,  Box.Y,  Box.Z), Center + FVector( Box.X, -Box.Y, Box.Z), Color, LifeTime, Thickness, DepthPriority, BatchID, bIsVisible);
-	new(BatchedLines) FBatchedLine(Center + FVector( Box.X, -Box.Y,  Box.Z), Center + FVector(-Box.X, -Box.Y, Box.Z), Color, LifeTime, Thickness, DepthPriority, BatchID, bIsVisible);
-	new(BatchedLines) FBatchedLine(Center + FVector(-Box.X, -Box.Y,  Box.Z), Center + FVector(-Box.X,  Box.Y, Box.Z) ,Color, LifeTime, Thickness, DepthPriority, BatchID, bIsVisible);
-	new(BatchedLines) FBatchedLine(Center + FVector(-Box.X,  Box.Y,  Box.Z), Center + FVector( Box.X,  Box.Y, Box.Z), Color, LifeTime, Thickness, DepthPriority, BatchID, bIsVisible);
-
-	new(BatchedLines) FBatchedLine(Center + FVector( Box.X,  Box.Y, -Box.Z), Center + FVector( Box.X, -Box.Y, -Box.Z), Color, LifeTime, Thickness, DepthPriority, BatchID, bIsVisible);
-	new(BatchedLines) FBatchedLine(Center + FVector( Box.X, -Box.Y, -Box.Z), Center + FVector(-Box.X, -Box.Y, -Box.Z), Color, LifeTime, Thickness, DepthPriority, BatchID, bIsVisible);
-	new(BatchedLines) FBatchedLine(Center + FVector(-Box.X, -Box.Y, -Box.Z), Center + FVector(-Box.X,  Box.Y, -Box.Z), Color, LifeTime, Thickness, DepthPriority, BatchID, bIsVisible);
-	new(BatchedLines) FBatchedLine(Center + FVector(-Box.X,  Box.Y, -Box.Z), Center + FVector( Box.X,  Box.Y, -Box.Z), Color, LifeTime, Thickness, DepthPriority, BatchID, bIsVisible);
-
-	new(BatchedLines) FBatchedLine(Center + FVector( Box.X,  Box.Y,  Box.Z), Center + FVector( Box.X,  Box.Y, -Box.Z), Color, LifeTime, Thickness, DepthPriority, BatchID, bIsVisible);
-	new(BatchedLines) FBatchedLine(Center + FVector( Box.X, -Box.Y,  Box.Z), Center + FVector( Box.X, -Box.Y, -Box.Z), Color, LifeTime, Thickness, DepthPriority, BatchID, bIsVisible);
-	new(BatchedLines) FBatchedLine(Center + FVector(-Box.X, -Box.Y,  Box.Z), Center + FVector(-Box.X, -Box.Y, -Box.Z), Color, LifeTime, Thickness, DepthPriority, BatchID, bIsVisible);
-	new(BatchedLines) FBatchedLine(Center + FVector(-Box.X,  Box.Y,  Box.Z), Center + FVector(-Box.X,  Box.Y, -Box.Z), Color, LifeTime, Thickness, DepthPriority, BatchID, bIsVisible);
-
-	bHasHiddenLines |= !bIsVisible;
-	
+	BatchedLines.Emplace(Center + FVector( Box.X,  Box.Y,  Box.Z), Center + FVector( Box.X, -Box.Y, Box.Z), Color, LifeTime, Thickness, DepthPriority, BatchID);
+	BatchedLines.Emplace(Center + FVector( Box.X, -Box.Y,  Box.Z), Center + FVector(-Box.X, -Box.Y, Box.Z), Color, LifeTime, Thickness, DepthPriority, BatchID);
+	BatchedLines.Emplace(Center + FVector(-Box.X, -Box.Y,  Box.Z), Center + FVector(-Box.X,  Box.Y, Box.Z) ,Color, LifeTime, Thickness, DepthPriority, BatchID);
+	BatchedLines.Emplace(Center + FVector(-Box.X,  Box.Y,  Box.Z), Center + FVector( Box.X,  Box.Y, Box.Z), Color, LifeTime, Thickness, DepthPriority, BatchID);
+
+	BatchedLines.Emplace(Center + FVector( Box.X,  Box.Y, -Box.Z), Center + FVector( Box.X, -Box.Y, -Box.Z), Color, LifeTime, Thickness, DepthPriority, BatchID);
+	BatchedLines.Emplace(Center + FVector( Box.X, -Box.Y, -Box.Z), Center + FVector(-Box.X, -Box.Y, -Box.Z), Color, LifeTime, Thickness, DepthPriority, BatchID);
+	BatchedLines.Emplace(Center + FVector(-Box.X, -Box.Y, -Box.Z), Center + FVector(-Box.X,  Box.Y, -Box.Z), Color, LifeTime, Thickness, DepthPriority, BatchID);
+	BatchedLines.Emplace(Center + FVector(-Box.X,  Box.Y, -Box.Z), Center + FVector( Box.X,  Box.Y, -Box.Z), Color, LifeTime, Thickness, DepthPriority, BatchID);
+
+	BatchedLines.Emplace(Center + FVector( Box.X,  Box.Y,  Box.Z), Center + FVector( Box.X,  Box.Y, -Box.Z), Color, LifeTime, Thickness, DepthPriority, BatchID);
+	BatchedLines.Emplace(Center + FVector( Box.X, -Box.Y,  Box.Z), Center + FVector( Box.X, -Box.Y, -Box.Z), Color, LifeTime, Thickness, DepthPriority, BatchID);
+	BatchedLines.Emplace(Center + FVector(-Box.X, -Box.Y,  Box.Z), Center + FVector(-Box.X, -Box.Y, -Box.Z), Color, LifeTime, Thickness, DepthPriority, BatchID);
+	BatchedLines.Emplace(Center + FVector(-Box.X,  Box.Y,  Box.Z), Center + FVector(-Box.X,  Box.Y, -Box.Z), Color, LifeTime, Thickness, DepthPriority, BatchID);
+
 	MarkRenderStateDirty();
 }
 
-void ULineBatchComponent::DrawBox(FVector const& Center, FVector const& Box, const FQuat& Rotation, FLinearColor Color, float LifeTime, uint8 DepthPriority, float Thickness, uint32 BatchID, bool bIsVisible)
+void ULineBatchComponent::DrawBox(FVector const& Center, FVector const& Box, const FQuat& Rotation, FLinearColor Color, float LifeTime, uint8 DepthPriority, float Thickness, uint32 BatchID)
 {
 	FTransform const Transform(Rotation);
 	FVector Start = Transform.TransformPosition(FVector( Box.X,  Box.Y,  Box.Z));
 	FVector End = Transform.TransformPosition(FVector( Box.X, -Box.Y, Box.Z));
-	new(BatchedLines) FBatchedLine(Center + Start, Center + End, Color, LifeTime, Thickness, DepthPriority, BatchID, bIsVisible);
+	BatchedLines.Emplace(Center + Start, Center + End, Color, LifeTime, Thickness, DepthPriority, BatchID);
 
 	Start = Transform.TransformPosition(FVector( Box.X, -Box.Y,  Box.Z));
 	End = Transform.TransformPosition(FVector(-Box.X, -Box.Y, Box.Z));
-	new(BatchedLines) FBatchedLine(Center + Start, Center + End, Color, LifeTime, Thickness, DepthPriority, BatchID, bIsVisible);
+	BatchedLines.Emplace(Center + Start, Center + End, Color, LifeTime, Thickness, DepthPriority, BatchID);
 
 	Start = Transform.TransformPosition(FVector(-Box.X, -Box.Y,  Box.Z));
 	End = Transform.TransformPosition(FVector(-Box.X,  Box.Y, Box.Z));
-	new(BatchedLines) FBatchedLine(Center + Start, Center + End, Color, LifeTime, Thickness, DepthPriority, BatchID, bIsVisible);
+	BatchedLines.Emplace(Center + Start, Center + End, Color, LifeTime, Thickness, DepthPriority, BatchID);
 
 	Start = Transform.TransformPosition(FVector(-Box.X,  Box.Y,  Box.Z));
 	End = Transform.TransformPosition(FVector( Box.X,  Box.Y, Box.Z));
-	new(BatchedLines) FBatchedLine(Center + Start, Center + End, Color, LifeTime, Thickness, DepthPriority, BatchID, bIsVisible);
+	BatchedLines.Emplace(Center + Start, Center + End, Color, LifeTime, Thickness, DepthPriority, BatchID);
 
 	Start = Transform.TransformPosition(FVector( Box.X,  Box.Y, -Box.Z));
 	End = Transform.TransformPosition(FVector( Box.X, -Box.Y, -Box.Z));
-	new(BatchedLines) FBatchedLine(Center + Start, Center + End, Color, LifeTime, Thickness, DepthPriority, BatchID, bIsVisible);
+	BatchedLines.Emplace(Center + Start, Center + End, Color, LifeTime, Thickness, DepthPriority, BatchID);
 
 	Start = Transform.TransformPosition(FVector( Box.X, -Box.Y, -Box.Z));
 	End = Transform.TransformPosition(FVector(-Box.X, -Box.Y, -Box.Z));
-	new(BatchedLines) FBatchedLine(Center + Start, Center + End, Color, LifeTime, Thickness, DepthPriority, BatchID, bIsVisible);
+	BatchedLines.Emplace(Center + Start, Center + End, Color, LifeTime, Thickness, DepthPriority, BatchID);
 
 	Start = Transform.TransformPosition(FVector(-Box.X, -Box.Y, -Box.Z));
 	End = Transform.TransformPosition(FVector(-Box.X,  Box.Y, -Box.Z));
-	new(BatchedLines) FBatchedLine(Center + Start, Center + End, Color, LifeTime, Thickness, DepthPriority, BatchID, bIsVisible);
+	BatchedLines.Emplace(Center + Start, Center + End, Color, LifeTime, Thickness, DepthPriority, BatchID);
 
 	Start = Transform.TransformPosition(FVector(-Box.X,  Box.Y, -Box.Z));
 	End = Transform.TransformPosition(FVector( Box.X,  Box.Y, -Box.Z));
-	new(BatchedLines) FBatchedLine(Center + Start, Center + End, Color, LifeTime, Thickness, DepthPriority, BatchID, bIsVisible);
+	BatchedLines.Emplace(Center + Start, Center + End, Color, LifeTime, Thickness, DepthPriority, BatchID);
 
 	Start = Transform.TransformPosition(FVector( Box.X,  Box.Y,  Box.Z));
 	End = Transform.TransformPosition(FVector( Box.X,  Box.Y, -Box.Z));
-	new(BatchedLines) FBatchedLine(Center + Start, Center + End, Color, LifeTime, Thickness, DepthPriority, BatchID, bIsVisible);
+	BatchedLines.Emplace(Center + Start, Center + End, Color, LifeTime, Thickness, DepthPriority, BatchID);
 
 	Start = Transform.TransformPosition(FVector( Box.X, -Box.Y,  Box.Z));
 	End = Transform.TransformPosition(FVector( Box.X, -Box.Y, -Box.Z));
-	new(BatchedLines) FBatchedLine(Center + Start, Center + End, Color, LifeTime, Thickness, DepthPriority, BatchID, bIsVisible);
+	BatchedLines.Emplace(Center + Start, Center + End, Color, LifeTime, Thickness, DepthPriority, BatchID);
 
 	Start = Transform.TransformPosition(FVector(-Box.X, -Box.Y,  Box.Z));
 	End = Transform.TransformPosition(FVector(-Box.X, -Box.Y, -Box.Z));
-	new(BatchedLines) FBatchedLine(Center + Start, Center + End, Color, LifeTime, Thickness, DepthPriority, BatchID, bIsVisible);
+	BatchedLines.Emplace(Center + Start, Center + End, Color, LifeTime, Thickness, DepthPriority, BatchID);
 
 	Start = Transform.TransformPosition(FVector(-Box.X,  Box.Y,  Box.Z));
 	End = Transform.TransformPosition(FVector(-Box.X,  Box.Y, -Box.Z));
-	new(BatchedLines) FBatchedLine(Center + Start, Center + End, Color, LifeTime, Thickness, DepthPriority, BatchID, bIsVisible);
+	BatchedLines.Emplace(Center + Start, Center + End, Color, LifeTime, Thickness, DepthPriority, BatchID);
 
-	bHasHiddenLines |= !bIsVisible;
 	MarkRenderStateDirty();
 }
 
-void ULineBatchComponent::DrawSolidBox(const FBox& Box, const FTransform& Xform, const FColor& Color, uint8 DepthPriority, float LifeTime, uint32 BatchID, bool bIsVisible)
+void ULineBatchComponent::DrawSolidBox(const FBox& Box, const FTransform& Xform, const FColor& Color, uint8 DepthPriority, float LifeTime, uint32 BatchID)
 {
 	int32 const NewMeshIdx = BatchedMeshes.Add(FBatchedMesh());
 	FBatchedMesh& BM = BatchedMeshes[NewMeshIdx];
@@ -309,7 +294,6 @@ void ULineBatchComponent::DrawSolidBox(const FBox& Box, const FTransform& Xform,
 	BM.DepthPriority = DepthPriority;
 	BM.RemainingLifeTime = LifeTime;
 	BM.BatchID = BatchID;
-	BM.bIsVisible = bIsVisible;
 
 	BM.MeshVerts.AddUninitialized(8);
 	BM.MeshVerts[0] = Xform.TransformPosition( FVector(Box.Min.X, Box.Min.Y, Box.Max.Z) );
@@ -341,11 +325,10 @@ void ULineBatchComponent::DrawSolidBox(const FBox& Box, const FTransform& Xform,
 		BM.MeshIndices[Idx] = Indices[Idx];
 	}
 
-	bHasHiddenLines |= !bIsVisible;
 	MarkRenderStateDirty();
 }
 
-void ULineBatchComponent::DrawMesh(TArray<FVector> const& Verts, TArray<int32> const& Indices, FColor const& Color, uint8 DepthPriority, float LifeTime, uint32 BatchID, bool bIsVisible)
+void ULineBatchComponent::DrawMesh(TArray<FVector> const& Verts, TArray<int32> const& Indices, FColor const& Color, uint8 DepthPriority, float LifeTime, uint32 BatchID)
 {
 	// modifying array element directly to avoid copying arrays
 	int32 const NewMeshIdx = BatchedMeshes.Add(FBatchedMesh());
@@ -357,26 +340,23 @@ void ULineBatchComponent::DrawMesh(TArray<FVector> const& Verts, TArray<int32> c
 	BM.DepthPriority = DepthPriority;
 	BM.RemainingLifeTime = LifeTime;
 	BM.BatchID = BatchID;
-	BM.bIsVisible = bIsVisible;
 
-	bHasHiddenLines |= !bIsVisible;
 	MarkRenderStateDirty();
 }
 
-void ULineBatchComponent::DrawDirectionalArrow(const FMatrix& ArrowToWorld, FLinearColor InColor, float Length, float ArrowSize, uint8 DepthPriority, uint32 BatchID, bool bIsVisible)
+void ULineBatchComponent::DrawDirectionalArrow(const FMatrix& ArrowToWorld, FLinearColor InColor, float Length, float ArrowSize, uint8 DepthPriority, uint32 BatchID)
 {
 	const FVector Tip = ArrowToWorld.TransformPosition(FVector(Length,0,0));
-	new(BatchedLines) FBatchedLine(Tip,ArrowToWorld.TransformPosition(FVector::ZeroVector),InColor,DefaultLifeTime,0.0f,DepthPriority, BatchID, bIsVisible);
-	new(BatchedLines) FBatchedLine(Tip,ArrowToWorld.TransformPosition(FVector(Length-ArrowSize,+ArrowSize,+ArrowSize)),InColor,DefaultLifeTime,0.0f,DepthPriority, BatchID, bIsVisible);
-	new(BatchedLines) FBatchedLine(Tip,ArrowToWorld.TransformPosition(FVector(Length-ArrowSize,+ArrowSize,-ArrowSize)),InColor,DefaultLifeTime,0.0f,DepthPriority, BatchID, bIsVisible);
-	new(BatchedLines) FBatchedLine(Tip,ArrowToWorld.TransformPosition(FVector(Length-ArrowSize,-ArrowSize,+ArrowSize)),InColor,DefaultLifeTime,0.0f,DepthPriority, BatchID, bIsVisible);
-	new(BatchedLines) FBatchedLine(Tip,ArrowToWorld.TransformPosition(FVector(Length-ArrowSize,-ArrowSize,-ArrowSize)),InColor,DefaultLifeTime,0.0f,DepthPriority, BatchID, bIsVisible);
+	BatchedLines.Emplace(Tip,ArrowToWorld.TransformPosition(FVector::ZeroVector),InColor,DefaultLifeTime,0.0f,DepthPriority, BatchID);
+	BatchedLines.Emplace(Tip,ArrowToWorld.TransformPosition(FVector(Length-ArrowSize,+ArrowSize,+ArrowSize)),InColor,DefaultLifeTime,0.0f,DepthPriority, BatchID);
+	BatchedLines.Emplace(Tip,ArrowToWorld.TransformPosition(FVector(Length-ArrowSize,+ArrowSize,-ArrowSize)),InColor,DefaultLifeTime,0.0f,DepthPriority, BatchID);
+	BatchedLines.Emplace(Tip,ArrowToWorld.TransformPosition(FVector(Length-ArrowSize,-ArrowSize,+ArrowSize)),InColor,DefaultLifeTime,0.0f,DepthPriority, BatchID);
+	BatchedLines.Emplace(Tip,ArrowToWorld.TransformPosition(FVector(Length-ArrowSize,-ArrowSize,-ArrowSize)),InColor,DefaultLifeTime,0.0f,DepthPriority, BatchID);
 
-	bHasHiddenLines |= !bIsVisible;
 	MarkRenderStateDirty();
 }
 
-void ULineBatchComponent::DrawDirectionalArrow(FVector const& LineStart, FVector const& LineEnd, float ArrowSize, FLinearColor Color, float LifeTime, uint8 DepthPriority, float Thickness, uint32 BatchID, bool bIsVisible)
+void ULineBatchComponent::DrawDirectionalArrow(FVector const& LineStart, FVector const& LineEnd, float ArrowSize, FLinearColor Color, float LifeTime, uint8 DepthPriority, float Thickness, uint32 BatchID)
 {
 	FVector Dir = (LineEnd-LineStart);
 	Dir.Normalize();
@@ -394,16 +374,15 @@ void ULineBatchComponent::DrawDirectionalArrow(FVector const& LineStart, FVector
 	// since dir is x direction, my arrow will be pointing +y, -x and -y, -x
 	const float ArrowSqrt = FMath::Sqrt(ArrowSize);
 
-	new(BatchedLines) FBatchedLine(LineStart,LineEnd, Color, LifeTime, Thickness, DepthPriority, BatchID, bIsVisible);
-	new(BatchedLines) FBatchedLine(LineEnd,LineEnd + TM.TransformPosition(FVector(-ArrowSqrt, ArrowSqrt, 0)), Color, LifeTime, Thickness, DepthPriority, BatchID, bIsVisible);
-	new(BatchedLines) FBatchedLine(LineEnd,LineEnd + TM.TransformPosition(FVector(-ArrowSqrt, -ArrowSqrt, 0)), Color, LifeTime, Thickness, DepthPriority, BatchID, bIsVisible);
+	BatchedLines.Emplace(LineStart,LineEnd, Color, LifeTime, Thickness, DepthPriority, BatchID);
+	BatchedLines.Emplace(LineEnd,LineEnd + TM.TransformPosition(FVector(-ArrowSqrt, ArrowSqrt, 0)), Color, LifeTime, Thickness, DepthPriority, BatchID);
+	BatchedLines.Emplace(LineEnd,LineEnd + TM.TransformPosition(FVector(-ArrowSqrt, -ArrowSqrt, 0)), Color, LifeTime, Thickness, DepthPriority, BatchID);
 
-	bHasHiddenLines |= !bIsVisible;
 	MarkRenderStateDirty();
 }
 
 
-void ULineBatchComponent::AddHalfCircle(const FVector& Base, const FVector& X, const FVector& Y, const FLinearColor& Color, const float Radius, int32 NumSides, const float LifeTime, uint8 DepthPriority, const float Thickness, const uint32 BatchID, bool bIsVisible)
+void ULineBatchComponent::AddHalfCircle(const FVector& Base, const FVector& X, const FVector& Y, const FLinearColor& Color, const float Radius, int32 NumSides, const float LifeTime, uint8 DepthPriority, const float Thickness, const uint32 BatchID)
 {
 	// Need at least 2 sides
 	NumSides = FMath::Max(NumSides, 2);
@@ -413,12 +392,12 @@ void ULineBatchComponent::AddHalfCircle(const FVector& Base, const FVector& X, c
 	for( int32 SideIndex = 0; SideIndex < (NumSides/2); SideIndex++)
 	{
 		FVector	Vertex = Base + (X * FMath::Cos(AngleDelta * (SideIndex + 1)) + Y * FMath::Sin(AngleDelta * (SideIndex + 1))) * Radius;
-		new(BatchedLines) FBatchedLine(LastVertex, Vertex, Color, LifeTime, Thickness, DepthPriority, BatchID, bIsVisible);
+		BatchedLines.Emplace(LastVertex, Vertex, Color, LifeTime, Thickness, DepthPriority, BatchID);
 		LastVertex = Vertex;
 	}
 }
 
-void ULineBatchComponent::AddCircle(const FVector& Base, const FVector& X, const FVector& Y, const FLinearColor& Color, const float Radius, int32 NumSides, const float LifeTime, uint8 DepthPriority, const float Thickness, const uint32 BatchID, bool bInVisible)
+void ULineBatchComponent::AddCircle(const FVector& Base, const FVector& X, const FVector& Y, const FLinearColor& Color, const float Radius, int32 NumSides, const float LifeTime, uint8 DepthPriority, const float Thickness, const uint32 BatchID)
 {
 	// Need at least 2 sides
 	NumSides = FMath::Max(NumSides, 2);
@@ -428,22 +407,21 @@ void ULineBatchComponent::AddCircle(const FVector& Base, const FVector& X, const
 	for (int32 SideIndex = 0; SideIndex < NumSides; SideIndex++)
 	{
 		const FVector Vertex = Base + (X * FMath::Cos(AngleDelta * (SideIndex + 1)) + Y * FMath::Sin(AngleDelta * (SideIndex + 1))) * Radius;
-		new(BatchedLines) FBatchedLine(LastVertex, Vertex, Color, LifeTime, Thickness, DepthPriority, BatchID, bInVisible);
+		BatchedLines.Emplace(LastVertex, Vertex, Color, LifeTime, Thickness, DepthPriority, BatchID);
 		LastVertex = Vertex;
 	}
 }
 
 /** Draw a circle */
-void ULineBatchComponent::DrawCircle(const FVector& Base, const FVector& X, const FVector& Y, FLinearColor Color, float Radius, int32 NumSides, uint8 DepthPriority, uint32 BatchID, bool bIsVisible)
+void ULineBatchComponent::DrawCircle(const FVector& Base, const FVector& X, const FVector& Y, FLinearColor Color, float Radius, int32 NumSides, uint8 DepthPriority, uint32 BatchID)
 {
-	AddCircle(Base, X, Y, Color, Radius, NumSides, DefaultLifeTime, DepthPriority, 0.f, BatchID, bIsVisible);
+	AddCircle(Base, X, Y, Color, Radius, NumSides, DefaultLifeTime, DepthPriority, 0.f, BatchID);
 		
-	bHasHiddenLines |= !bIsVisible;
 	MarkRenderStateDirty();
 }
 
 /** Draw a sphere */
-void ULineBatchComponent::DrawSphere(FVector const& Center, float Radius, int32 Segments, FLinearColor Color, float LifeTime, uint8 DepthPriority, float Thickness, uint32 BatchID, bool bIsVisible)
+void ULineBatchComponent::DrawSphere(FVector const& Center, float Radius, int32 Segments, FLinearColor Color, float LifeTime, uint8 DepthPriority, float Thickness, uint32 BatchID)
 {
 	// Need at least 4 segments
 	Segments = FMath::Max(Segments, 4);
@@ -473,8 +451,8 @@ void ULineBatchComponent::DrawSphere(FVector const& Center, float Radius, int32
 			const FVector Vertex2 = FVector((CosX * SinY1), (SinX * SinY1), CosY1) * Radius + Center;
 			const FVector Vertex4 = FVector((CosX * SinY2), (SinX * SinY2), CosY2) * Radius + Center;
 
-			new(BatchedLines) FBatchedLine(Vertex1, Vertex2, Color, LifeTime, Thickness, DepthPriority, BatchID, bIsVisible);
-			new(BatchedLines) FBatchedLine(Vertex1, Vertex3, Color, LifeTime, Thickness, DepthPriority, BatchID, bIsVisible);
+			BatchedLines.Emplace(Vertex1, Vertex2, Color, LifeTime, Thickness, DepthPriority, BatchID);
+			BatchedLines.Emplace(Vertex1, Vertex3, Color, LifeTime, Thickness, DepthPriority, BatchID);
 
 			Vertex1 = Vertex2;
 			Vertex3 = Vertex4;
@@ -484,12 +462,11 @@ void ULineBatchComponent::DrawSphere(FVector const& Center, float Radius, int32
 		CosY1 = CosY2;
 		Latitude += AngleInc;
 	}
-	
-	bHasHiddenLines |= !bIsVisible;
+
 	MarkRenderStateDirty();
 }
 
-void ULineBatchComponent::DrawCylinder(FVector const& Start, FVector const& End, float Radius, int32 Segments, FLinearColor Color, float LifeTime, uint8 DepthPriority, float Thickness, uint32 BatchID, bool bIsVisible)
+void ULineBatchComponent::DrawCylinder(FVector const& Start, FVector const& End, float Radius, int32 Segments, FLinearColor Color, float LifeTime, uint8 DepthPriority, float Thickness, uint32 BatchID)
 {
 	// Need at least 4 segments
 	Segments = FMath::Max(Segments, 4);
@@ -518,20 +495,19 @@ void ULineBatchComponent::DrawCylinder(FVector const& Start, FVector const& End,
 		FVector P2 = Segment + Start;
 		FVector P4 = Segment + End;
 
-		new(BatchedLines) FBatchedLine(P2, P4, Color, LifeTime, Thickness, DepthPriority, BatchID, bIsVisible);
-		new(BatchedLines) FBatchedLine(P1, P2, Color, LifeTime, Thickness, DepthPriority, BatchID, bIsVisible);
-		new(BatchedLines) FBatchedLine(P3, P4, Color, LifeTime, Thickness, DepthPriority, BatchID, bIsVisible);
+		BatchedLines.Emplace(P2, P4, Color, LifeTime, Thickness, DepthPriority, BatchID);
+		BatchedLines.Emplace(P1, P2, Color, LifeTime, Thickness, DepthPriority, BatchID);
+		BatchedLines.Emplace(P3, P4, Color, LifeTime, Thickness, DepthPriority, BatchID);
 
 		P1 = P2;
 		P3 = P4;
 		Angle += AngleInc;
 	}
 
-	bHasHiddenLines |= !bIsVisible;
 	MarkRenderStateDirty();
 }
 
-void ULineBatchComponent::DrawCone(FVector const& Origin, FVector const& Direction, float Length, float AngleWidth, float AngleHeight, int32 NumSides, FLinearColor DrawColor, float LifeTime, uint8 DepthPriority, float Thickness, uint32 BatchID, bool bIsVisible)
+void ULineBatchComponent::DrawCone(FVector const& Origin, FVector const& Direction, float Length, float AngleWidth, float AngleHeight, int32 NumSides, FLinearColor DrawColor, float LifeTime, uint8 DepthPriority, float Thickness, uint32 BatchID)
 {
 	// Need at least 4 sides
 	NumSides = FMath::Max(NumSides, 4);
@@ -579,12 +555,12 @@ void ULineBatchComponent::DrawCone(FVector const& Origin, FVector const& Directi
 	for(int32 i = 0; i < NumSides; i++)
 	{
 		CurrentPoint = ConeToWorld.TransformPosition(ConeVerts[i]);
-		new(BatchedLines) FBatchedLine(ConeToWorld.GetOrigin(), CurrentPoint, DrawColor, LifeTime, Thickness, DepthPriority, BatchID, bIsVisible);
+		BatchedLines.Emplace(ConeToWorld.GetOrigin(), CurrentPoint, DrawColor, LifeTime, Thickness, DepthPriority, BatchID);
 
 		// PrevPoint must be defined to draw junctions
 		if( i > 0 )
 		{
-			new(BatchedLines) FBatchedLine(PrevPoint, CurrentPoint, DrawColor, LifeTime, Thickness, DepthPriority, BatchID, bIsVisible);
+			BatchedLines.Emplace(PrevPoint, CurrentPoint, DrawColor, LifeTime, Thickness, DepthPriority, BatchID);
 		}
 		else
 		{
@@ -594,13 +570,12 @@ void ULineBatchComponent::DrawCone(FVector const& Origin, FVector const& Directi
 		PrevPoint = CurrentPoint;
 	}
 	// Connect last junction to first
-	new(BatchedLines) FBatchedLine(CurrentPoint, FirstPoint, DrawColor, LifeTime, Thickness, DepthPriority, BatchID, bIsVisible);
+	BatchedLines.Emplace(CurrentPoint, FirstPoint, DrawColor, LifeTime, Thickness, DepthPriority, BatchID);
 
-	bHasHiddenLines |= !bIsVisible;	
 	MarkRenderStateDirty();
 }
 
-void ULineBatchComponent::DrawCapsule(FVector const& Center, float HalfHeight, float Radius, const FQuat& Rotation, FLinearColor Color, float LifeTime, uint8 DepthPriority, float Thickness, uint32 BatchID, bool bIsVisible)
+void ULineBatchComponent::DrawCapsule(FVector const& Center, float HalfHeight, float Radius, const FQuat& Rotation, FLinearColor Color, float LifeTime, uint8 DepthPriority, float Thickness, uint32 BatchID)
 {
 	constexpr int32 DrawCollisionSides = 16;
 	const FVector Origin = Center;
@@ -614,32 +589,30 @@ void ULineBatchComponent::DrawCapsule(FVector const& Center, float HalfHeight, f
 	const FVector TopEnd = Origin + HalfAxis*ZAxis;
 	const FVector BottomEnd = Origin - HalfAxis*ZAxis;
 
-	AddCircle(TopEnd, XAxis, YAxis, Color, Radius, DrawCollisionSides, LifeTime, DepthPriority, Thickness, BatchID, bIsVisible);
-	AddCircle( BottomEnd, XAxis, YAxis, Color, Radius, DrawCollisionSides, LifeTime, DepthPriority, Thickness, BatchID, bIsVisible);
+	AddCircle(TopEnd, XAxis, YAxis, Color, Radius, DrawCollisionSides, LifeTime, DepthPriority, Thickness, BatchID);
+	AddCircle( BottomEnd, XAxis, YAxis, Color, Radius, DrawCollisionSides, LifeTime, DepthPriority, Thickness, BatchID);
 
 	// Draw domed caps
-	AddHalfCircle( TopEnd, YAxis, ZAxis, Color, Radius, DrawCollisionSides, LifeTime, DepthPriority, Thickness, BatchID, bIsVisible);
-	AddHalfCircle( TopEnd, XAxis, ZAxis, Color, Radius, DrawCollisionSides, LifeTime, DepthPriority, Thickness, BatchID, bIsVisible);
+	AddHalfCircle( TopEnd, YAxis, ZAxis, Color, Radius, DrawCollisionSides, LifeTime, DepthPriority, Thickness, BatchID);
+	AddHalfCircle( TopEnd, XAxis, ZAxis, Color, Radius, DrawCollisionSides, LifeTime, DepthPriority, Thickness, BatchID);
 
 	const FVector NegZAxis = -ZAxis;
 
-	AddHalfCircle( BottomEnd, YAxis, NegZAxis, Color, Radius, DrawCollisionSides, LifeTime, DepthPriority, Thickness, BatchID, bIsVisible);
-	AddHalfCircle( BottomEnd, XAxis, NegZAxis, Color, Radius, DrawCollisionSides, LifeTime, DepthPriority, Thickness, BatchID, bIsVisible);
+	AddHalfCircle( BottomEnd, YAxis, NegZAxis, Color, Radius, DrawCollisionSides, LifeTime, DepthPriority, Thickness, BatchID);
+	AddHalfCircle( BottomEnd, XAxis, NegZAxis, Color, Radius, DrawCollisionSides, LifeTime, DepthPriority, Thickness, BatchID);
 
 	// Draw connected lines
-	new(BatchedLines) FBatchedLine(TopEnd + Radius*XAxis, BottomEnd + Radius*XAxis, Color, LifeTime, Thickness, DepthPriority, BatchID, bIsVisible);
-	new(BatchedLines) FBatchedLine(TopEnd - Radius*XAxis, BottomEnd - Radius*XAxis, Color, LifeTime, Thickness, DepthPriority, BatchID, bIsVisible);
-	new(BatchedLines) FBatchedLine(TopEnd + Radius*XAxis, BottomEnd + Radius*XAxis, Color, LifeTime, Thickness, DepthPriority, BatchID, bIsVisible);
-	new(BatchedLines) FBatchedLine(TopEnd - Radius*XAxis, BottomEnd - Radius*XAxis, Color, LifeTime, Thickness, DepthPriority, BatchID, bIsVisible);
+	BatchedLines.Emplace(TopEnd + Radius*XAxis, BottomEnd + Radius*XAxis, Color, LifeTime, Thickness, DepthPriority, BatchID);
+	BatchedLines.Emplace(TopEnd - Radius*XAxis, BottomEnd - Radius*XAxis, Color, LifeTime, Thickness, DepthPriority, BatchID);
+	BatchedLines.Emplace(TopEnd + Radius*XAxis, BottomEnd + Radius*XAxis, Color, LifeTime, Thickness, DepthPriority, BatchID);
+	BatchedLines.Emplace(TopEnd - Radius*XAxis, BottomEnd - Radius*XAxis, Color, LifeTime, Thickness, DepthPriority, BatchID);
 
-	bHasHiddenLines |= !bIsVisible;
 	MarkRenderStateDirty();
 }
 
 void ULineBatchComponent::TickComponent(float DeltaTime, enum ELevelTick TickType, FActorComponentTickFunction *ThisTickFunction)
 {
 	bool bDirty = false;
-	bHasHiddenLines = false;
 	// Update the life time of batched lines, removing the lines which have expired.
 	for(int32 LineIndex=0; LineIndex < BatchedLines.Num(); LineIndex++)
 	{
@@ -652,14 +625,8 @@ void ULineBatchComponent::TickComponent(float DeltaTime, enum ELevelTick TickTyp
 				// The line has expired, remove it.
 				BatchedLines.RemoveAtSwap(LineIndex--);
 				bDirty = true;
-				continue;
 			}
 		}
-		if (Line.bIsVisible == false && !bHasHiddenLines)
-		{
-			bHasHiddenLines = true;
-			bDirty = true;
-		}
 	}
 
 	// Update the life time of batched points, removing the points which have expired.
@@ -674,14 +641,8 @@ void ULineBatchComponent::TickComponent(float DeltaTime, enum ELevelTick TickTyp
 				// The point has expired, remove it.
 				BatchedPoints.RemoveAtSwap(PtIndex--);
 				bDirty = true;
-				continue;
 			}
 		}
-		if (Pt.bIsVisible == false && !bHasHiddenLines)
-		{
-			bHasHiddenLines = true;
-			bDirty = true;
-		}
 	}
 
 	// Update the life time of batched meshes, removing the meshes which have expired.
@@ -696,14 +657,8 @@ void ULineBatchComponent::TickComponent(float DeltaTime, enum ELevelTick TickTyp
 				// The mesh has expired, remove it.
 				BatchedMeshes.RemoveAtSwap(MeshIndex--);
 				bDirty = true;
-				continue;
 			}
 		}
-		if (Mesh.bIsVisible == false && !bHasHiddenLines)
-		{
-			bHasHiddenLines = true;
-			bDirty = true;
-		}
 	}
 
 	if(bDirty)
@@ -822,72 +777,4 @@ void ULineBatchComponent::ClearBatch(uint32 BatchID)
 	}
 }
 
-void ULineBatchComponent::SetBatchVisible(uint32 BatchID, bool bIsVisible)
-{
-	if (BatchID == INVALID_ID)
-	{
-		return;
-	}
-	bool bDirty = false;
-	for (FBatchedLine& Line: BatchedLines)
-	{
-		if (Line.BatchID == BatchID)
-		{
-			Line.bIsVisible = bIsVisible;
-			bDirty = true;
-		}
-	}
-	for (FBatchedPoint& Point: BatchedPoints)
-	{
-		if (Point.BatchID == BatchID)
-		{
-			Point.bIsVisible = bIsVisible;
-			bDirty = true;
-		}
-	}
-	for (FBatchedMesh& Mesh: BatchedMeshes)
-	{
-		if (Mesh.BatchID == BatchID)
-		{
-			Mesh.bIsVisible = bIsVisible;
-			bDirty = true;
-		}
-	}
-	if(bDirty)
-	{
-		if (!bIsVisible)
-		{
-			bHasHiddenLines = true;	
-		}
-		MarkRenderStateDirty();
-	}
-}
-
-void ULineBatchComponent::SetVisible(bool bIsVisible)
-{
-	bool bDirty = false;
-	for (FBatchedLine& Line: BatchedLines)
-	{
-		Line.bIsVisible = bIsVisible;
-		bDirty = true;
-	}
-	for (FBatchedPoint& Point: BatchedPoints)
-	{
-		Point.bIsVisible = bIsVisible;
-		bDirty = true;
-	}
-	for (FBatchedMesh& Mesh: BatchedMeshes)
-	{
-		Mesh.bIsVisible = bIsVisible;
-		bDirty = true;
-	}
-	if(bDirty)
-	{
-		if (!bIsVisible)
-		{
-			bHasHiddenLines = true;	
-		}
 
-		MarkRenderStateDirty();
-	}
-}
