Substrate: Remove STRATA_DIFFUSE_CHAN in favor of using MATERIAL_ROUGHDIFFUSE

Having two shader defines that did the same thing was a bit confusing. Instead use MATERIAL_ROUGHDIFFUSE for both substrate and non-substrate cases, and handle the logic purely on the C++ side.

diff --git a/Engine/Shaders/Private/PathTracing/Material/PathTracingMaterialCommon.ush b/Engine/Shaders/Private/PathTracing/Material/PathTracingMaterialCommon.ush
--- a/Engine/Shaders/Private/PathTracing/Material/PathTracingMaterialCommon.ush
+++ b/Engine/Shaders/Private/PathTracing/Material/PathTracingMaterialCommon.ush
@@ -212,8 +212,7 @@ float ShadowTerminatorTerm(float3 L, float3 N, float3 Ns)
 
 float3 GetPathTracingDiffuseModel(float3 DiffuseColor, float Roughness, float NoV, float NoL, float VoH, float NoH)
 {
-	// TODO: Always use MATERIAL_ROUGHDIFFUSE instead of have two different defines for Substrate vs. Legacy?
-#if (!STRATA_ENABLED && MATERIAL_ROUGHDIFFUSE) || (STRATA_ENABLED && STRATA_DIFFUSE_CHAN)
+#if MATERIAL_ROUGHDIFFUSE
 	// Clamp to [0,1] as the model is ever so slightly non-energy conserving at grazing angles which can cause unwanted noise
 	return saturate(PI * Diffuse_Chan(DiffuseColor, Pow4(Roughness), NoV, NoL, VoH, NoH, 1.0f));
 #else
diff --git a/Engine/Shaders/Private/Strata/StrataEvaluation.ush b/Engine/Shaders/Private/Strata/StrataEvaluation.ush
--- a/Engine/Shaders/Private/Strata/StrataEvaluation.ush
+++ b/Engine/Shaders/Private/Strata/StrataEvaluation.ush
@@ -333,8 +333,8 @@ FStrataEvaluateResult StrataEvaluateBSDFCommon(FStrataBSDFContext BSDFContext, F
 			////
 			//// Evaluate the diffuse component.
 			////
-
-			if (STRATA_DIFFUSE_CHAN && Settings.bRoughDiffuseEnabled && any(DiffuseColor > 0))
+#if MATERIAL_ROUGHDIFFUSE
+			if (Settings.bRoughDiffuseEnabled && any(DiffuseColor > 0))
 			{
 				// * If the specular layer is anisotropic, the diffuse takes the 'main' roughness value rather than the tangent/bitangent value
 				// * The Chan model bakes transmittance specular directional albedo assuming F=0.04. In previous code we reapplied 
@@ -343,6 +343,7 @@ FStrataEvaluateResult StrataEvaluateBSDFCommon(FStrataBSDFContext BSDFContext, F
 				Sample.DiffusePathValue = Diffuse_Chan(DiffuseColor, Alpha2, NoV, AreaLightContext.NoL, VoH, NoH, GetAreaLightDiffuseMicroReflWeight(AreaLightContext.AreaLight));
 			}
 			else
+#endif
 			{
 				Sample.DiffusePathValue = Diffuse_Lambert(DiffuseColor);
 			}
diff --git a/Engine/Source/Runtime/Engine/Private/ShaderCompiler/ShaderCompiler.cpp b/Engine/Source/Runtime/Engine/Private/ShaderCompiler/ShaderCompiler.cpp
--- a/Engine/Source/Runtime/Engine/Private/ShaderCompiler/ShaderCompiler.cpp
+++ b/Engine/Source/Runtime/Engine/Private/ShaderCompiler/ShaderCompiler.cpp
@@ -6628,15 +6628,12 @@ void GlobalBeginCompileShader(
 		Input.Environment.SetDefine(TEXT("SUPPORT_VSM_FOWARD_QUALITY"), bHighQualityShadow ? 1 : 0);
 	}
 
+	const bool bStrata = Strata::IsStrataEnabled();
 	{
-		const bool bStrata = Strata::IsStrataEnabled();
 		Input.Environment.SetDefine(TEXT("STRATA_ENABLED"), bStrata ? 1 : 0);
 
 		if (bStrata)
 		{
-			const bool bStrataRoughDiffuse = Strata::IsRoughDiffuseEnabled() && !Strata::IsBackCompatibilityEnabled();
-			Input.Environment.SetDefine(TEXT("STRATA_DIFFUSE_CHAN"), bStrataRoughDiffuse ? 1 : 0);
-
 			const uint32 StrataShadingQuality = Strata::GetShadingQuality(Target.GetPlatform());
 			Input.Environment.SetDefine(TEXT("STRATA_SHADING_QUALITY"), StrataShadingQuality);
 
@@ -6669,7 +6666,8 @@ void GlobalBeginCompileShader(
 	{
 		static IConsoleVariable* CVar = IConsoleManager::Get().FindConsoleVariable(TEXT("r.Material.RoughDiffuse"));
 		const bool bMaterialRoughDiffuse = CVar && CVar->GetInt() != 0;
-		Input.Environment.SetDefine(TEXT("MATERIAL_ROUGHDIFFUSE"), bMaterialRoughDiffuse ? 1 : 0);
+		const bool bStrataRoughDiffuse = Strata::IsRoughDiffuseEnabled() && !Strata::IsBackCompatibilityEnabled();
+		Input.Environment.SetDefine(TEXT("MATERIAL_ROUGHDIFFUSE"), (bStrata ? bStrataRoughDiffuse : bMaterialRoughDiffuse) ? 1 : 0);
 	}
 
 	{
