Supports ISM with mesh paint vertex color on mobile.

diff --git a/Engine/Source/Runtime/Engine/Classes/Components/StaticMeshComponent.h b/Engine/Source/Runtime/Engine/Classes/Components/StaticMeshComponent.h
--- a/Engine/Source/Runtime/Engine/Classes/Components/StaticMeshComponent.h
+++ b/Engine/Source/Runtime/Engine/Classes/Components/StaticMeshComponent.h
@@ -694,7 +694,7 @@ protected:
 	/** Collect all the PSO precache data used by the static mesh component */
 	virtual void CollectPSOPrecacheData(const FPSOPrecacheParams& BasePrecachePSOParams, FComponentPSOPrecacheParamsList& OutParams) override;
 	/** Shared implementation for all StaticMesh derived components */
-	using GetPSOVertexElementsFn = TFunctionRef<void(const FStaticMeshLODResources& LODRenderData, bool bSupportsManualVertexFetch, FVertexDeclarationElementList& Elements)>;
+	using GetPSOVertexElementsFn = TFunctionRef<void(const FStaticMeshLODResources& LODRenderData, int32 LODIndex, bool bSupportsManualVertexFetch, FVertexDeclarationElementList& Elements)>;
 	void CollectPSOPrecacheDataImpl(const FVertexFactoryType* VFType, const FPSOPrecacheParams& BasePrecachePSOParams, GetPSOVertexElementsFn GetVertexElements, FComponentPSOPrecacheParamsList& OutParams) const;
 		
 	/** Whether the component type supports static lighting. */
diff --git a/Engine/Source/Runtime/Engine/Private/Components/SplineMeshComponent.cpp b/Engine/Source/Runtime/Engine/Private/Components/SplineMeshComponent.cpp
--- a/Engine/Source/Runtime/Engine/Private/Components/SplineMeshComponent.cpp
+++ b/Engine/Source/Runtime/Engine/Private/Components/SplineMeshComponent.cpp
@@ -642,7 +642,7 @@ void USplineMeshComponent::CollectPSOPrecacheData(const FPSOPrecacheParams& Base
 	FVertexFactoryType* VertexFactoryType = &FSplineMeshVertexFactory::StaticType;
 	int32 LightMapCoordinateIndex = GetStaticMesh()->GetLightMapCoordinateIndex();
 
-	auto SMC_GetElements = [LightMapCoordinateIndex](const FStaticMeshLODResources& LODRenderData, bool bSupportsManualVertexFetch, FVertexDeclarationElementList& Elements)
+	auto SMC_GetElements = [LightMapCoordinateIndex](const FStaticMeshLODResources& LODRenderData, int32 LODIndex, bool bSupportsManualVertexFetch, FVertexDeclarationElementList& Elements)
 	{
 		// FIXME: This will miss when SM component overrides vertex colors and source StaticMesh does not have vertex colors
 		bool bOverrideColorVertexBuffer = false;
diff --git a/Engine/Source/Runtime/Engine/Private/Components/StaticMeshComponent.cpp b/Engine/Source/Runtime/Engine/Private/Components/StaticMeshComponent.cpp
--- a/Engine/Source/Runtime/Engine/Private/Components/StaticMeshComponent.cpp
+++ b/Engine/Source/Runtime/Engine/Private/Components/StaticMeshComponent.cpp
@@ -1557,12 +1557,14 @@ void UStaticMeshComponent::CollectPSOPrecacheDataImpl(
 	bool bAnySectionCastsShadows = false;
 
 	FPSOPrecacheVertexFactoryDataPerMaterialIndexList VFTypesPerMaterialIndex;
-	for (FStaticMeshLODResources& LODRenderData : GetStaticMesh()->GetRenderData()->LODResources)
+	FStaticMeshLODResourcesArray& LODResources = GetStaticMesh()->GetRenderData()->LODResources;
+	for (int32 LODIndex = 0; LODIndex < LODResources.Num(); ++LODIndex)
 	{
+		FStaticMeshLODResources& LODRenderData = LODResources[LODIndex];
 		FVertexDeclarationElementList VertexElements;
 		if (!bSupportsManualVertexFetch)
 		{
-			GetVertexElements(LODRenderData, bSupportsManualVertexFetch, VertexElements);
+			GetVertexElements(LODRenderData, LODIndex, bSupportsManualVertexFetch, VertexElements);
 		}
 
 		for (FStaticMeshSection& RenderSection : LODRenderData.Sections)
@@ -1634,7 +1636,7 @@ void UStaticMeshComponent::CollectPSOPrecacheData(const FPSOPrecacheParams& Base
 	// FIXME: Need a precise per-LOD test
 	bool bOverrideColorVertexBuffer = LODData.Num() != 0 && LODData[0].OverrideVertexColors != nullptr;
 	
-	auto SMC_GetElements = [LightMapCoordinateIndex, bOverrideColorVertexBuffer](const FStaticMeshLODResources& LODRenderData, bool bSupportsManualVertexFetch, FVertexDeclarationElementList& Elements)
+	auto SMC_GetElements = [LightMapCoordinateIndex, bOverrideColorVertexBuffer](const FStaticMeshLODResources& LODRenderData, int32 LODIndex, bool bSupportsManualVertexFetch, FVertexDeclarationElementList& Elements)
 	{
 		int32 NumTexCoords = (int32)LODRenderData.VertexBuffers.StaticMeshVertexBuffer.GetNumTexCoords();
 		int32 LODLightMapCoordinateIndex = LightMapCoordinateIndex < NumTexCoords ? LightMapCoordinateIndex : NumTexCoords - 1;
diff --git a/Engine/Source/Runtime/Engine/Private/InstancedStaticMesh.cpp b/Engine/Source/Runtime/Engine/Private/InstancedStaticMesh.cpp
--- a/Engine/Source/Runtime/Engine/Private/InstancedStaticMesh.cpp
+++ b/Engine/Source/Runtime/Engine/Private/InstancedStaticMesh.cpp
@@ -1169,10 +1169,10 @@ void FInstancedStaticMeshRenderData::ReleaseResources(FSceneInterface* Scene, co
 
 void InitInstancedStaticMeshVertexFactoryComponents(
 	const FStaticMeshVertexBuffers& VertexBuffers,
+	const FColorVertexBuffer* ColorVertexBuffer,
 	const FStaticMeshInstanceBuffer* InstanceBuffer,
 	const FInstancedStaticMeshVertexFactory* VertexFactory,
 	int32 LightMapCoordinateIndex,
-	bool bHasColorVertexData,
 	bool bRHISupportsManualVertexFetch,
 	FInstancedStaticMeshVertexFactory::FDataType& OutData)
 {
@@ -1185,9 +1185,9 @@ void InitInstancedStaticMeshVertexFactoryComponents(
 		VertexBuffers.StaticMeshVertexBuffer.BindLightMapVertexBuffer(VertexFactory, OutData, LightMapCoordinateIndex);
 	}
 
-	if (bHasColorVertexData)
+	if (ColorVertexBuffer != nullptr)
 	{
-		VertexBuffers.ColorVertexBuffer.BindColorVertexBuffer(VertexFactory, OutData);
+		ColorVertexBuffer->BindColorVertexBuffer(VertexFactory, OutData);
 	}
 	else
 	{
@@ -1223,7 +1223,12 @@ void FInstancedStaticMeshRenderData::BindBuffersToVertexFactories()
 		FInstancedStaticMeshVertexFactory::FDataType Data;
 		// Assign to the vertex factory for this LOD.
 		FInstancedStaticMeshVertexFactory& VertexFactory = VertexFactories[LODIndex];
-		InitInstancedStaticMeshVertexFactoryComponents(RenderData->VertexBuffers, InstanceBuffer, &VertexFactory, LightMapCoordinateIndex, RenderData->bHasColorVertexData, bRHISupportsManualVertexFetch, Data);
+		const FColorVertexBuffer* ColorVertexBuffer = RenderData->bHasColorVertexData ? &(RenderData->VertexBuffers.ColorVertexBuffer) : nullptr;
+		if (Component->LODData.IsValidIndex(LODIndex) && Component->LODData[LODIndex].OverrideVertexColors)
+		{
+			ColorVertexBuffer = Component->LODData[LODIndex].OverrideVertexColors;
+		}
+		InitInstancedStaticMeshVertexFactoryComponents(RenderData->VertexBuffers, ColorVertexBuffer, InstanceBuffer, &VertexFactory, LightMapCoordinateIndex, bRHISupportsManualVertexFetch, Data);
 		VertexFactory.SetData(Data);
 		VertexFactory.InitResource();
 	}
@@ -4950,10 +4955,15 @@ void UInstancedStaticMeshComponent::CollectPSOPrecacheData(const FPSOPrecachePar
 	const FVertexFactoryType* VFType = ShouldCreateNaniteProxy() ? &Nanite::FVertexFactory::StaticType : &FInstancedStaticMeshVertexFactory::StaticType;
 	int32 LightMapCoordinateIndex = GetStaticMesh()->GetLightMapCoordinateIndex();
 	
-	auto ISMC_GetElements = [LightMapCoordinateIndex, InstanceBuffer](const FStaticMeshLODResources& LODRenderData, bool bSupportsManualVertexFetch, FVertexDeclarationElementList& Elements)
+	auto ISMC_GetElements = [LightMapCoordinateIndex, InstanceBuffer, this](const FStaticMeshLODResources& LODRenderData, int32 LODIndex, bool bSupportsManualVertexFetch, FVertexDeclarationElementList& Elements)
 	{
 		FInstancedStaticMeshVertexFactory::FDataType Data;
-		InitInstancedStaticMeshVertexFactoryComponents(LODRenderData.VertexBuffers, InstanceBuffer, nullptr /*VertexFactory*/, LightMapCoordinateIndex, LODRenderData.bHasColorVertexData, bSupportsManualVertexFetch, Data);
+		const FColorVertexBuffer* ColorVertexBuffer = LODRenderData.bHasColorVertexData ? &(LODRenderData.VertexBuffers.ColorVertexBuffer) : nullptr;
+		if (LODData.IsValidIndex(LODIndex) && LODData[LODIndex].OverrideVertexColors)
+		{
+			ColorVertexBuffer = LODData[LODIndex].OverrideVertexColors;
+		}
+		InitInstancedStaticMeshVertexFactoryComponents(LODRenderData.VertexBuffers, ColorVertexBuffer, InstanceBuffer, nullptr /*VertexFactory*/, LightMapCoordinateIndex, bSupportsManualVertexFetch, Data);
 		FInstancedStaticMeshVertexFactory::GetVertexElements(GMaxRHIFeatureLevel, EVertexInputStreamType::Default, bSupportsManualVertexFetch, Data, Elements);
 	};
 
