Small refactor : Use FArchiveMD5::GetGuidFromHash to simplify code

diff --git a/Engine/Source/Editor/DataLayerEditor/Private/DataLayer/DataLayerActorTreeItem.h b/Engine/Source/Editor/DataLayerEditor/Private/DataLayer/DataLayerActorTreeItem.h
--- a/Engine/Source/Editor/DataLayerEditor/Private/DataLayer/DataLayerActorTreeItem.h
+++ b/Engine/Source/Editor/DataLayerEditor/Private/DataLayer/DataLayerActorTreeItem.h
@@ -50,17 +50,13 @@ public:
 
 	static FSceneOutlinerTreeItemID ComputeTreeItemID(AActor* InActor, const UDataLayerInstance* InDataLayerInstance)
 	{
-		FArchiveMD5 Ar;
 		FObjectKey ActorKey(InActor);
-		Ar << ActorKey;
-
 		FObjectKey DataLayerInstanceKey(InDataLayerInstance);
-		Ar << DataLayerInstanceKey;
 
-		FMD5Hash MD5Hash;
-		Ar.GetHash(MD5Hash);
+		FArchiveMD5 Ar;
+		Ar << ActorKey << DataLayerInstanceKey;
 
-		return FSceneOutlinerTreeItemID(MD5HashToGuid(MD5Hash));
+		return FSceneOutlinerTreeItemID(Ar.GetGuidFromHash());
 	}
 
 	bool Filter(FFilterPredicate Pred) const
diff --git a/Engine/Source/Editor/DataLayerEditor/Private/DataLayer/DataLayersActorDescTreeItem.h b/Engine/Source/Editor/DataLayerEditor/Private/DataLayer/DataLayersActorDescTreeItem.h
--- a/Engine/Source/Editor/DataLayerEditor/Private/DataLayer/DataLayersActorDescTreeItem.h
+++ b/Engine/Source/Editor/DataLayerEditor/Private/DataLayer/DataLayersActorDescTreeItem.h
@@ -64,19 +64,13 @@ public:
 	
 	static FSceneOutlinerTreeItemID ComputeTreeItemID(FGuid InActorGuid, UActorDescContainer* InContainer, const UDataLayerInstance* InDataLayer)
 	{
-		FArchiveMD5 Ar;
-		Ar << InActorGuid;
-
 		FObjectKey ContainerKey(InContainer);
-		Ar << ContainerKey;
-
 		FObjectKey DataLayerInstanceKey(InDataLayer);
-		Ar << DataLayerInstanceKey;
 
-		FMD5Hash MD5Hash;
-		Ar.GetHash(MD5Hash);
+		FArchiveMD5 Ar;
+		Ar << InActorGuid << ContainerKey << DataLayerInstanceKey;
 
-		return FSceneOutlinerTreeItemID(MD5HashToGuid(MD5Hash));
+		return FSceneOutlinerTreeItemID(Ar.GetGuidFromHash());
 	}
 
 	static TArray<AActor*> GetParentActors(UActorDescContainer* InContainer)
diff --git a/Engine/Source/Editor/SceneOutliner/Private/ActorDescTreeItem.cpp b/Engine/Source/Editor/SceneOutliner/Private/ActorDescTreeItem.cpp
--- a/Engine/Source/Editor/SceneOutliner/Private/ActorDescTreeItem.cpp
+++ b/Engine/Source/Editor/SceneOutliner/Private/ActorDescTreeItem.cpp
@@ -251,10 +251,7 @@ FSceneOutlinerTreeItemID FActorDescTreeItem::ComputeTreeItemID(FGuid InActorGuid
 	FObjectKey ContainerKey(InContainer);
 	Ar << ContainerKey;
 
-	FMD5Hash MD5Hash;
-	Ar.GetHash(MD5Hash);
-	
-	return FSceneOutlinerTreeItemID(MD5HashToGuid(MD5Hash));
+	return FSceneOutlinerTreeItemID(Ar.GetGuidFromHash());
 }
 
 bool FActorDescTreeItem::ShouldDisplayInOutliner(const FWorldPartitionActorDesc* ActorDesc)
diff --git a/Engine/Source/Runtime/Engine/Private/ExternalPackageHelper.cpp b/Engine/Source/Runtime/Engine/Private/ExternalPackageHelper.cpp
--- a/Engine/Source/Runtime/Engine/Private/ExternalPackageHelper.cpp
+++ b/Engine/Source/Runtime/Engine/Private/ExternalPackageHelper.cpp
@@ -98,10 +98,7 @@ FString FExternalPackageHelper::GetExternalPackageName(const FString& InOuterPac
 	FArchiveMD5 ArMD5;
 	ArMD5 << ObjectPath;
 
-	FMD5Hash MD5Hash;
-	ArMD5.GetHash(MD5Hash);
-
-	FGuid PackageGuid = MD5HashToGuid(MD5Hash);
+	FGuid PackageGuid = ArMD5.GetGuidFromHash();
 	check(PackageGuid.IsValid());
 
 	FString GuidBase36 = PackageGuid.ToString(EGuidFormats::Base36Encoded);
diff --git a/Engine/Source/Runtime/Engine/Private/Level.cpp b/Engine/Source/Runtime/Engine/Private/Level.cpp
--- a/Engine/Source/Runtime/Engine/Private/Level.cpp
+++ b/Engine/Source/Runtime/Engine/Private/Level.cpp
@@ -3508,10 +3508,7 @@ FString ULevel::GetActorPackageName(const FString& InBaseDir, EActorPackagingSch
 	FArchiveMD5 ArMD5;
 	ArMD5 << ActorPath;
 
-	FMD5Hash MD5Hash;
-	ArMD5.GetHash(MD5Hash);
-
-	FGuid PackageGuid = MD5HashToGuid(MD5Hash);
+	FGuid PackageGuid = ArMD5.GetGuidFromHash();
 	check(PackageGuid.IsValid());
 
 	FString GuidBase36 = PackageGuid.ToString(EGuidFormats::Base36Encoded);
diff --git a/Engine/Source/Runtime/Engine/Private/WorldPartition/HLOD/HLODActor.cpp b/Engine/Source/Runtime/Engine/Private/WorldPartition/HLOD/HLODActor.cpp
--- a/Engine/Source/Runtime/Engine/Private/WorldPartition/HLOD/HLODActor.cpp
+++ b/Engine/Source/Runtime/Engine/Private/WorldPartition/HLOD/HLODActor.cpp
@@ -221,10 +221,7 @@ void AWorldPartitionHLOD::PostLoad()
 		FArchiveMD5 ArMD5;
 		ArMD5 << GridName << CellGlobalCoord[0] << CellGlobalCoord[1] << CellGlobalCoord[2] << DataLayerID << ContentBundleID;
 
-		FMD5Hash MD5Hash;
-		ArMD5.GetHash(MD5Hash);
-
-		SourceCellGuid = MD5HashToGuid(MD5Hash);
+		SourceCellGuid = ArMD5.GetGuidFromHash();
 		check(SourceCellGuid.IsValid());
 	}
 
diff --git a/Engine/Source/Runtime/Engine/Private/WorldPartition/WorldPartitionActorDesc.cpp b/Engine/Source/Runtime/Engine/Private/WorldPartition/WorldPartitionActorDesc.cpp
--- a/Engine/Source/Runtime/Engine/Private/WorldPartition/WorldPartitionActorDesc.cpp
+++ b/Engine/Source/Runtime/Engine/Private/WorldPartition/WorldPartitionActorDesc.cpp
@@ -63,14 +63,11 @@ void FWorldPartitionActorDesc::Init(const AActor* InActor)
 		check(!InActor->IsPackageExternal());
 		check(!InActor->GetActorGuid().IsValid());
 
-		FArchiveMD5 MD5Ar;
+		FArchiveMD5 ArMD5;
 		FString ClassPath = BaseClass.IsValid() ? BaseClass.ToString() : NativeClass.ToString();
-		MD5Ar << ClassPath;
+		ArMD5 << ClassPath;
 		
-		FMD5Hash MD5Hash;
-		MD5Ar.GetHash(MD5Hash);
-
-		Guid = MD5HashToGuid(MD5Hash);
+		Guid = ArMD5.GetGuidFromHash();
 
 		bIsDefaultActorDesc = true;
 	}
diff --git a/Engine/Source/Runtime/Engine/Private/WorldPartition/WorldPartitionRuntimeSpatialHash.cpp b/Engine/Source/Runtime/Engine/Private/WorldPartition/WorldPartitionRuntimeSpatialHash.cpp
--- a/Engine/Source/Runtime/Engine/Private/WorldPartition/WorldPartitionRuntimeSpatialHash.cpp
+++ b/Engine/Source/Runtime/Engine/Private/WorldPartition/WorldPartitionRuntimeSpatialHash.cpp
@@ -1284,11 +1284,7 @@ FGuid UWorldPartitionRuntimeSpatialHash::GetCellGuid(FName InGridName, const FGr
 
 	FArchiveMD5 ArMD5;
 	ArMD5 << GridName << CellGlobalCoord << DataLayerID << ContentBundleID;
-
-	FMD5Hash MD5Hash;
-	ArMD5.GetHash(MD5Hash);
-
-	FGuid CellGuid = MD5HashToGuid(MD5Hash);
+	FGuid CellGuid = ArMD5.GetGuidFromHash();
 	check(CellGuid.IsValid());
 
 	return CellGuid;
@@ -1674,9 +1670,6 @@ EWorldPartitionStreamingPerformance UWorldPartitionRuntimeSpatialHash::GetStream
 	check(Cell->GetBlockOnSlowLoading());
 	const double BlockOnSlowStreamingWarningRatio = GBlockOnSlowStreamingRatio * GBlockOnSlowStreamingWarningFactor;
 	
-	const UWorldPartitionRuntimeCell* StreamingCell = Cast<const UWorldPartitionRuntimeCell>(Cell);
-	check(StreamingCell);
-
 	const UWorldPartitionRuntimeCellDataSpatialHash* CellDataSpatialHash = CastChecked<UWorldPartitionRuntimeCellDataSpatialHash>(Cell->RuntimeCellData);
 	const FSpatialHashStreamingGrid* StreamingGrid = GetStreamingGridByName(CellDataSpatialHash->GridName);
 
diff --git a/Engine/Source/Runtime/Landscape/Private/Landscape.cpp b/Engine/Source/Runtime/Landscape/Private/Landscape.cpp
--- a/Engine/Source/Runtime/Landscape/Private/Landscape.cpp
+++ b/Engine/Source/Runtime/Landscape/Private/Landscape.cpp
@@ -3456,11 +3456,7 @@ void ALandscapeProxy::PostLoadFixupLandscapeGuidsIfInstanced()
 			Ar << PackageName;
 		}
 
-		FMD5Hash MD5Hash;
-		Ar.GetHash(MD5Hash);
-
-		FGuid NewLandscapeGuid = MD5HashToGuid(MD5Hash);
-		this->LandscapeGuid = NewLandscapeGuid;
+		this->LandscapeGuid = Ar.GetGuidFromHash();
 	}
 }
 
