World Partition: adjust comment as to why we need to rename sub worlds level streaming objects to the owning world.

diff --git a/Engine/Source/Runtime/Engine/Private/WorldPartition/WorldPartitionRuntimeLevelStreamingCell.cpp b/Engine/Source/Runtime/Engine/Private/WorldPartition/WorldPartitionRuntimeLevelStreamingCell.cpp
--- a/Engine/Source/Runtime/Engine/Private/WorldPartition/WorldPartitionRuntimeLevelStreamingCell.cpp
+++ b/Engine/Source/Runtime/Engine/Private/WorldPartition/WorldPartitionRuntimeLevelStreamingCell.cpp
@@ -372,8 +372,10 @@ UWorldPartitionLevelStreamingDynamic* UWorldPartitionRuntimeLevelStreamingCell::
 	// In Runtime, prepare LevelStreaming for activation
 	if (LevelStreaming)
 	{
-		// Setup pre-created LevelStreaming's outer to the WorldPartition owning world
 		const UWorldPartition* WorldPartition = GetOuterWorld()->GetWorldPartition();
+		
+		// Setup pre-created LevelStreaming's outer to the WorldPartition owning world.
+		// This is needed because ULevelStreaming is within=World, and ULevelStreaming::GetWorld() assumes that the outer world is the main world.		
 		UWorld* OwningWorld = GetOwningWorld();
 		if (LevelStreaming->GetWorld() != OwningWorld)
 		{
