Disabled for now assert for LightMapCoordinateIndex >= 0 test as it is triggered with PSO precaching atm

diff --git a/Engine/Source/Runtime/Engine/Private/Rendering/StaticMeshVertexBuffer.cpp b/Engine/Source/Runtime/Engine/Private/Rendering/StaticMeshVertexBuffer.cpp
--- a/Engine/Source/Runtime/Engine/Private/Rendering/StaticMeshVertexBuffer.cpp
+++ b/Engine/Source/Runtime/Engine/Private/Rendering/StaticMeshVertexBuffer.cpp
@@ -613,7 +613,8 @@ void FStaticMeshVertexBuffer::BindTexCoordVertexBuffer(const FVertexFactory* Ver
 void FStaticMeshVertexBuffer::BindLightMapVertexBuffer(const FVertexFactory* VertexFactory, FStaticMeshDataType& Data, int LightMapCoordinateIndex) const
 {
 	LightMapCoordinateIndex = LightMapCoordinateIndex < (int32)GetNumTexCoords() ? LightMapCoordinateIndex : (int32)GetNumTexCoords() - 1;
-	check(LightMapCoordinateIndex >= 0);
+	//FIXME: pso precache triggers this before mesh postload has completed. normally, EnforceLightmapRestrictions called from mesh postload prevents this
+	//check(LightMapCoordinateIndex >= 0);  
 
 	Data.LightMapCoordinateIndex = LightMapCoordinateIndex;
 	Data.NumTexCoords = GetNumTexCoords();
