[shaders] undo changes moving shader job submission into a task. race conditions mean this needs further thought.

diff --git a/Engine/Source/Runtime/Engine/Private/ShaderCompiler/ShaderCompiler.cpp b/Engine/Source/Runtime/Engine/Private/ShaderCompiler/ShaderCompiler.cpp
--- a/Engine/Source/Runtime/Engine/Private/ShaderCompiler/ShaderCompiler.cpp
+++ b/Engine/Source/Runtime/Engine/Private/ShaderCompiler/ShaderCompiler.cpp
@@ -592,17 +592,30 @@ void FShaderCompileJobCollection::SubmitJobs(const TArray<FShaderCommonCompileJo
 		// we may fulfill some of the jobs from the cache (and we will be subtracting them)
 		NumOutstandingJobs.Add(InJobs.Num());
 
-		for (FShaderCommonCompileJob* Job : InJobs)
+		int32 SubmittedJobsCount = 0;
+		int32 NumSubmittedJobs[NumShaderCompileJobPriorities] = { 0 };
 		{
-			UE::Tasks::Launch(UE_SOURCE_LOCATION, [Job, this]()
+			// Just precompute the InputHash for each job in multiple-thread.
+			if (ShaderCompiler::IsJobCacheEnabled())
 			{
-				TRACE_CPUPROFILER_EVENT_SCOPE(ShaderJobTask);
+				TRACE_CPUPROFILER_EVENT_SCOPE(ShaderCompiler.GetInputHash);
+				ParallelFor(TEXT("ShaderCompiler.GetInputHash.PF"), InJobs.Num(),1, [&InJobs](int32 Index) 
+				{
+					ConditionalPreprocessShader(InJobs[Index]);
+					InJobs[Index]->GetInputHash(); 
+				}, EParallelForFlags::Unbalanced);
+			}
+
+			FWriteScopeLock Locker(Lock);
 
+			// Optimization: Only search the linked list once to prevent O(n^2) behavior
+			FShaderCommonCompileJob* ExistingJobs[NumShaderCompileJobPriorities] = { nullptr };
+
+			for (FShaderCommonCompileJob* Job : InJobs)
+			{
 				check(Job->JobIndex != INDEX_NONE);
 				check(Job->Priority != EShaderCompileJobPriority::None);
 				check(Job->PendingPriority == EShaderCompileJobPriority::None);
-				
-				ConditionalPreprocessShader(Job);
 
 				const int32 PriorityIndex = (int32)Job->Priority;
 				bool bNewJob = true;
@@ -613,81 +626,68 @@ void FShaderCompileJobCollection::SubmitJobs(const TArray<FShaderCommonCompileJo
 
 					const bool bCheckDDC = !(Job->bIsDefaultMaterial || Job->bIsGlobalShader);
 
-					FSharedBuffer* ExistingOutput = nullptr;
-					{
-						FWriteScopeLock Locker(Lock);
-
-						// see if we can find the job in the cache first
-						ExistingOutput = CompletedJobsCache.Find(InputHash, bCheckDDC);
-					}
-				
-					if (ExistingOutput)
+					// see if we can find the job in the cache first
+					if (FSharedBuffer* ExistingOutput = CompletedJobsCache.Find(InputHash, bCheckDDC))
 					{
 						UE_LOG(LogShaderCompilers, UE_SHADERCACHE_LOG_LEVEL, TEXT("There is already a cached job with the ihash %s, processing the new one immediately."), *LexToString(InputHash));
 						FMemoryReaderView MemReader(*ExistingOutput);
 						Job->SerializeOutput(MemReader);
 
-						{
-							FWriteScopeLock WriteLocker(Lock);
-							ProcessFinishedJob(Job, true);
-						}
+						// finish the job instantly
+						ProcessFinishedJob(Job, true);
 
-						bNewJob = false;
+						continue;
 					}
-					else 
+					// see if another job with the same input hash is being worked on
+					else if (FShaderCommonCompileJob** DuplicateInFlight = JobsInFlight.Find(InputHash))
 					{
-						FWriteScopeLock Locker(Lock);
+						UE_LOG(LogShaderCompilers, UE_SHADERCACHE_LOG_LEVEL, TEXT("There is an outstanding job with the ihash %s, not submitting another one (adding to wait list)."), *LexToString(InputHash));
 
-						FShaderCommonCompileJob** DuplicateInFlight = nullptr;
-						// see if another job with the same input hash is being worked on
-						DuplicateInFlight = JobsInFlight.Find(InputHash);
-
-						if (DuplicateInFlight)
+						// because of the cloned jobs, we need to maintain a separate mapping
+						FShaderCommonCompileJob** WaitListHead = DuplicateJobsWaitList.Find(InputHash);
+						if (WaitListHead)
 						{
-							UE_LOG(LogShaderCompilers, UE_SHADERCACHE_LOG_LEVEL, TEXT("There is an outstanding job with the ihash %s, not submitting another one (adding to wait list)."), *LexToString(InputHash));
-
-							// because of the cloned jobs, we need to maintain a separate mapping
-							FShaderCommonCompileJob** WaitListHead = nullptr;
-							{
-								WaitListHead = DuplicateJobsWaitList.Find(InputHash);
-							}
-							
-							if (WaitListHead)
-							{
-								Job->LinkAfter(*WaitListHead);
-							}
-							else
-							{
-								DuplicateJobsWaitList.Add(InputHash, Job);
-							}
-							bNewJob = false;
+							Job->LinkAfter(*WaitListHead);
 						}
 						else
 						{
-							// track new jobs so we can de-dupe them
-							JobsInFlight.Add(InputHash, Job);
+							DuplicateJobsWaitList.Add(InputHash, Job);
 						}
+						bNewJob = false;
+					}
+					else
+					{
+						// track new jobs so we can dedupe them
+						JobsInFlight.Add(InputHash, Job);
 					}
 				}
 
 				// new job
 				if (bNewJob)
 				{
-					FWriteScopeLock Locker(Lock);
-
 					GShaderCompilerStats->RegisterNewPendingJob(*Job);
 					ensure(!ShaderCompiler::IsJobCacheEnabled() || Job->bInputHashSet);
 #if UE_SHADERCOMPILER_FIFO_JOB_EXECUTION
 					// link the job at the end of pending list, so we're executing them in a FIFO and not LIFO order
 					if (PendingJobs[PriorityIndex])
 					{
-						for (FShaderCommonCompileJob::TIterator It(PendingJobs[PriorityIndex]); It; ++It)
+						// Optimization: Only search the linked list if we have to
+						if (ExistingJobs[PriorityIndex])
 						{
-							FShaderCommonCompileJob& ExistingJob = *It;
-							if (ExistingJob.GetNextLink() == nullptr)
+							Job->LinkAfter(ExistingJobs[PriorityIndex]);
+							ExistingJobs[PriorityIndex] = Job;
+						}
+						else
+						{
+							for (FShaderCommonCompileJob::TIterator It(PendingJobs[PriorityIndex]); It; ++It)
 							{
-								Job->LinkAfter(&ExistingJob);
-								break;
+								FShaderCommonCompileJob& ExistingJob = *It;
+								if (ExistingJob.GetNextLink() == nullptr)
+								{
+									Job->LinkAfter(&ExistingJob);
+									ExistingJobs[PriorityIndex] = Job;
+									break;
+								}
 							}
 						}
 					}
@@ -696,11 +696,26 @@ void FShaderCompileJobCollection::SubmitJobs(const TArray<FShaderCommonCompileJo
 					{
 						Job->LinkHead(PendingJobs[PriorityIndex]);
 					}
+					
 
 					NumPendingJobs[PriorityIndex]++;
+					NumSubmittedJobs[PriorityIndex]++;
 					Job->PendingPriority = Job->Priority;
+					++SubmittedJobsCount;
 				}
-			});
+			}
+		}
+
+		UE_LOG(LogShaderCompilers, UE_SHADERCACHE_LOG_LEVEL, TEXT("Actual jobs submitted %d (of %d new), total outstanding jobs: %d."), SubmittedJobsCount, InJobs.Num(), NumOutstandingJobs.GetValue());
+
+		for (int32 PriorityIndex = 0; PriorityIndex < NumShaderCompileJobPriorities; ++PriorityIndex)
+		{
+			if (NumSubmittedJobs[PriorityIndex] > 0)
+			{
+				UE_LOG(LogShaderCompilers, Verbose, TEXT("Submitted %d shader compile jobs with '%s' priority"),
+					NumSubmittedJobs[PriorityIndex],
+					ShaderCompileJobPriorityToString((EShaderCompileJobPriority)PriorityIndex));
+			}
 		}
 	}
 }
