Fix for static meshes being marked dirty on load when built asynchronously : when static mesh async builds were initiated with GIsEditorLoadingPackage == true, they were ending up executing FinishBuildInternal with  GIsEditorLoadingPackage == false. We now store the value in the static mesh async build context so that we can temporarily set back GIsEditorLoadingPackage to its value at the start of the build when running FinishBuildInternal and other game thread-bound functions

diff --git a/Engine/Source/Runtime/Engine/Classes/Engine/StaticMesh.h b/Engine/Source/Runtime/Engine/Classes/Engine/StaticMesh.h
--- a/Engine/Source/Runtime/Engine/Classes/Engine/StaticMesh.h
+++ b/Engine/Source/Runtime/Engine/Classes/Engine/StaticMesh.h
@@ -2089,7 +2089,7 @@ private:
 class FStaticMeshCompilationContext
 {
 public:
-	FStaticMeshCompilationContext() = default;
+	FStaticMeshCompilationContext();
 	// Non-copyable
 	FStaticMeshCompilationContext(const FStaticMeshCompilationContext&) = delete;
 	FStaticMeshCompilationContext& operator=(const FStaticMeshCompilationContext&) = delete;
@@ -2098,6 +2098,7 @@ public:
 	FStaticMeshCompilationContext& operator=(FStaticMeshCompilationContext&&) = default;
 
 	bool bShouldComputeExtendedBounds = false;
+	bool bIsEditorLoadingPackage = false;
 };
 
 class FStaticMeshPostLoadContext : public FStaticMeshCompilationContext
diff --git a/Engine/Source/Runtime/Engine/Private/StaticMesh.cpp b/Engine/Source/Runtime/Engine/Private/StaticMesh.cpp
--- a/Engine/Source/Runtime/Engine/Private/StaticMesh.cpp
+++ b/Engine/Source/Runtime/Engine/Private/StaticMesh.cpp
@@ -8483,4 +8483,11 @@ void UStaticMeshSocket::Serialize(FArchive& Ar)
 	}
 }
 
-#undef LOCTEXT_NAMESPACE
+FStaticMeshCompilationContext::FStaticMeshCompilationContext() 
+{
+	// Remember if the editor was loading a package when initiating the build so that we can temporarily restore that state when 
+	//  executing FinishBuildInternal on the game thread at the end of the build :
+	bIsEditorLoadingPackage = GIsEditorLoadingPackage;
+}
+
+#undef LOCTEXT_NAMESPACE
\ No newline at end of file
diff --git a/Engine/Source/Runtime/Engine/Private/StaticMeshCompiler.cpp b/Engine/Source/Runtime/Engine/Private/StaticMeshCompiler.cpp
--- a/Engine/Source/Runtime/Engine/Private/StaticMeshCompiler.cpp
+++ b/Engine/Source/Runtime/Engine/Private/StaticMeshCompiler.cpp
@@ -248,6 +248,10 @@ void FStaticMeshCompilingManager::PostCompilation(UStaticMesh* StaticMesh)
 
 		FObjectCacheContextScope ObjectCacheScope;
 
+		// If async (post load or build), restore the state of GIsEditorLoadingPackage for the duration of this function (including outside of this scope) as it was when the build of the static mesh was initiated, 
+		//  so that async builds have the same result as synchronous ones (e.g. don't dirty packages when the components referencing this static mesh call Modify because GIsEditorLoadingPackage is true) :
+		TUniquePtr<TGuardValue<bool>> IsEditorLoadingPackageGuard;
+
 		// The scope is important here to destroy the FStaticMeshAsyncBuildScope before broadcasting events
 		{
 			// Acquire the async task locally to protect against re-entrance
@@ -267,6 +271,8 @@ void FStaticMeshCompilingManager::PostCompilation(UStaticMesh* StaticMesh)
 
 			if (LocalAsyncTask->GetTask().PostLoadContext.IsValid())
 			{
+				IsEditorLoadingPackageGuard.Reset(new TGuardValue<bool>(GIsEditorLoadingPackage, LocalAsyncTask->GetTask().PostLoadContext->bIsEditorLoadingPackage));
+
 				StaticMesh->FinishPostLoadInternal(*LocalAsyncTask->GetTask().PostLoadContext);
 
 				LocalAsyncTask->GetTask().PostLoadContext.Reset();
@@ -274,6 +280,8 @@ void FStaticMeshCompilingManager::PostCompilation(UStaticMesh* StaticMesh)
 
 			if (LocalAsyncTask->GetTask().BuildContext.IsValid())
 			{
+				IsEditorLoadingPackageGuard.Reset(new TGuardValue<bool>(GIsEditorLoadingPackage, LocalAsyncTask->GetTask().BuildContext->bIsEditorLoadingPackage));
+
 				TArray<UStaticMeshComponent*> ComponentsToUpdate;
 				for (UStaticMeshComponent* Component : ObjectCacheScope.GetContext().GetStaticMeshComponents(StaticMesh))
 				{
