Level Instance Partial Loading:
- Disable partial loading when changing the world asset.
- Properly handle actor filters.

diff --git a/Engine/Source/Runtime/Engine/Private/LevelInstance/LevelInstanceActor.cpp b/Engine/Source/Runtime/Engine/Private/LevelInstance/LevelInstanceActor.cpp
--- a/Engine/Source/Runtime/Engine/Private/LevelInstance/LevelInstanceActor.cpp
+++ b/Engine/Source/Runtime/Engine/Private/LevelInstance/LevelInstanceActor.cpp
@@ -139,6 +139,11 @@ TSubclassOf<AActor> ALevelInstance::GetEditorPivotClass() const
 	return ALevelInstancePivot::StaticClass();
 }
 
+bool ALevelInstance::SupportsPartialEditorLoading() const
+{
+	return ILevelInstanceInterface::SupportsPartialEditorLoading() && LevelInstanceActorImpl.SupportsPartialEditorLoading();
+}
+
 TUniquePtr<FWorldPartitionActorDesc> ALevelInstance::CreateClassActorDesc() const
 {
 	return TUniquePtr<FWorldPartitionActorDesc>(new FLevelInstanceActorDesc());
diff --git a/Engine/Source/Runtime/Engine/Private/LevelInstance/LevelInstanceActorImpl.cpp b/Engine/Source/Runtime/Engine/Private/LevelInstance/LevelInstanceActorImpl.cpp
--- a/Engine/Source/Runtime/Engine/Private/LevelInstance/LevelInstanceActorImpl.cpp
+++ b/Engine/Source/Runtime/Engine/Private/LevelInstance/LevelInstanceActorImpl.cpp
@@ -9,6 +9,7 @@
 #include "LevelInstance/LevelInstancePrivate.h"
 #include "LevelInstance/LevelInstanceComponent.h"
 #include "LevelInstance/LevelInstanceEditorInstanceActor.h"
+#include "WorldPartition/WorldPartitionActorLoaderInterface.h"
 #include "Engine/Level.h"
 #include "Engine/World.h"
 #include "LevelInstance/LevelInstanceSubsystem.h"
@@ -100,6 +101,10 @@ void FLevelInstanceActorImpl::OnLevelInstanceLoaded()
 }
 
 #if WITH_EDITOR
+bool FLevelInstanceActorImpl::SupportsPartialEditorLoading() const
+{
+	return bAllowPartialLoading;
+}
 
 void FLevelInstanceActorImpl::PreEditUndo(TFunctionRef<void()> SuperCall)
 {
@@ -210,7 +215,9 @@ void FLevelInstanceActorImpl::PreEditChange(FProperty* Property, bool bWorldAsse
 	if (bWorldAssetChange)
 	{
 		CachedWorldAsset = LevelInstance->GetWorldAsset();
-	}
+		bAllowPartialLoading = false;
+		IWorldPartitionActorLoaderInterface::RefreshLoadedState(true);
+	}	
 }
 
 void FLevelInstanceActorImpl::PostEditChangeProperty(FPropertyChangedEvent& PropertyChangedEvent, bool bWorldAssetChange, TFunction<void(FPropertyChangedEvent&)> SuperCall)
diff --git a/Engine/Source/Runtime/Engine/Private/LevelInstance/LevelInstanceInterface.cpp b/Engine/Source/Runtime/Engine/Private/LevelInstance/LevelInstanceInterface.cpp
--- a/Engine/Source/Runtime/Engine/Private/LevelInstance/LevelInstanceInterface.cpp
+++ b/Engine/Source/Runtime/Engine/Private/LevelInstance/LevelInstanceInterface.cpp
@@ -45,8 +45,8 @@ bool ILevelInstanceInterface::SupportsPartialEditorLoading() const
 					return false;
 				}
 
-				// If the level instance actor (or any parent actor) is dirty (has pending save changes), don't partially load.
-				if (LevelInstanceActor->GetPackage()->IsDirty())
+				// If the level instance actor (or any parent actor) is unsaved, don't partially load.
+				if (LevelInstanceActor->GetPackage()->HasAllPackagesFlags(PKG_NewlyCreated))
 				{
 					bResult = false;
 					return false;
diff --git a/Engine/Source/Runtime/Engine/Private/LevelInstance/LevelInstanceLevelStreaming.cpp b/Engine/Source/Runtime/Engine/Private/LevelInstance/LevelInstanceLevelStreaming.cpp
--- a/Engine/Source/Runtime/Engine/Private/LevelInstance/LevelInstanceLevelStreaming.cpp
+++ b/Engine/Source/Runtime/Engine/Private/LevelInstance/LevelInstanceLevelStreaming.cpp
@@ -27,7 +27,7 @@
 #include "Misc/ScopeExit.h"
 #include "UObject/LinkerLoad.h"
 
-static bool GDisableLevelInstanceEditorPartialLoading = true;
+static bool GDisableLevelInstanceEditorPartialLoading = false;
 FAutoConsoleVariableRef CVarDisableLevelInstanceEditorPartialLoading(
 	TEXT("wp.Editor.DisableLevelInstanceEditorPartialLoading"),
 	GDisableLevelInstanceEditorPartialLoading,
diff --git a/Engine/Source/Runtime/Engine/Private/WorldPartition/LoaderAdapter/LoaderAdapterSpatial.cpp b/Engine/Source/Runtime/Engine/Private/WorldPartition/LoaderAdapter/LoaderAdapterSpatial.cpp
--- a/Engine/Source/Runtime/Engine/Private/WorldPartition/LoaderAdapter/LoaderAdapterSpatial.cpp
+++ b/Engine/Source/Runtime/Engine/Private/WorldPartition/LoaderAdapter/LoaderAdapterSpatial.cpp
@@ -17,22 +17,20 @@ void ILoaderAdapterSpatial::ForEachActor(TFunctionRef<void(const FWorldPartition
 {
 	if (UWorldPartition* WorldPartition = GetWorld()->GetWorldPartition())
 	{
-		HandleIntersectingContainer(WorldPartition, InOperation);
+		HandleIntersectingContainer(WorldPartition, WorldPartition->GetInstanceTransform(), *GetBoundingBox(), InOperation);
 	}
 }
 
-void ILoaderAdapterSpatial::HandleIntersectingContainer(UWorldPartition* InWorldPartition, TFunctionRef<void(const FWorldPartitionHandle&)> InOperation) const
+void ILoaderAdapterSpatial::HandleIntersectingContainer(UWorldPartition* InWorldPartition, const FTransform InInstanceTransform, const FBox& InBoundingBox, TFunctionRef<void(const FWorldPartitionHandle&)> InOperation) const
 {
-	const FTransform InstanceTransform = InWorldPartition->GetInstanceTransform();
-	const FBox LocalBoundingBox = GetBoundingBox()->InverseTransformBy(InstanceTransform);
-
 	UWorldPartitionEditorHash::FForEachIntersectingActorParams ForEachIntersectingActorParams = UWorldPartitionEditorHash::FForEachIntersectingActorParams()
 		.SetIncludeSpatiallyLoadedActors(bIncludeSpatiallyLoadedActors)
 		.SetIncludeNonSpatiallyLoadedActors(bIncludeNonSpatiallyLoadedActors);
 
-	InWorldPartition->EditorHash->ForEachIntersectingActor(LocalBoundingBox, [this, InWorldPartition, &InstanceTransform, &InOperation](FWorldPartitionActorDesc* ActorDesc)
+	const FBox LocalBoundingBox = InBoundingBox.InverseTransformBy(InInstanceTransform);
+	InWorldPartition->EditorHash->ForEachIntersectingActor(LocalBoundingBox, [this, InWorldPartition, &InInstanceTransform, &InOperation](FWorldPartitionActorDesc* ActorDesc)
 	{
-		const FBox WorldActorEditorBox = ActorDesc->GetEditorBounds().TransformBy(InstanceTransform);
+		const FBox WorldActorEditorBox = ActorDesc->GetEditorBounds().TransformBy(InInstanceTransform);
 		if (Intersect(WorldActorEditorBox))
 		{
 			FWorldPartitionHandle ActorHandle(InWorldPartition, ActorDesc->GetGuid());
@@ -43,12 +41,10 @@ void ILoaderAdapterSpatial::HandleIntersectingContainer(UWorldPartition* InWorld
 				FWorldPartitionActorDesc::FContainerInstance ContainerInstance;
 				if (ActorHandle->GetContainerInstance(ContainerInstance))
 				{
-					if (ContainerInstance.bSupportsPartialEditorLoading)
+					if (UWorldPartition* ContainerWorldPartition = ContainerInstance.LoadedLevel ? ContainerInstance.LoadedLevel->GetWorldPartition() : nullptr)
 					{
-						if (UWorldPartition* ContainerWorldPartition = ContainerInstance.LoadedLevel ? ContainerInstance.LoadedLevel->GetWorldPartition() : nullptr)
-						{
-							HandleIntersectingContainer(ContainerWorldPartition, InOperation);
-						}
+						const FBox InnerBoundingBox = ContainerInstance.bSupportsPartialEditorLoading ? *GetBoundingBox() : FBox(FVector(-HALF_WORLD_MAX), FVector(HALF_WORLD_MAX));
+						HandleIntersectingContainer(ContainerWorldPartition, ContainerWorldPartition->GetInstanceTransform(), InnerBoundingBox, InOperation);
 					}
 				}
 			}
diff --git a/Engine/Source/Runtime/Engine/Public/LevelInstance/LevelInstanceActor.h b/Engine/Source/Runtime/Engine/Public/LevelInstance/LevelInstanceActor.h
--- a/Engine/Source/Runtime/Engine/Public/LevelInstance/LevelInstanceActor.h
+++ b/Engine/Source/Runtime/Engine/Public/LevelInstance/LevelInstanceActor.h
@@ -70,6 +70,7 @@ public:
 	virtual ELevelInstanceRuntimeBehavior GetDesiredRuntimeBehavior() const override { return DesiredRuntimeBehavior; }
 	virtual ELevelInstanceRuntimeBehavior GetDefaultRuntimeBehavior() const override { return ELevelInstanceRuntimeBehavior::Partitioned; }
 	virtual TSubclassOf<AActor> GetEditorPivotClass() const override;
+	virtual bool SupportsPartialEditorLoading() const override;
 	// End ILevelInstanceInterface
 			
 	// UObject overrides
diff --git a/Engine/Source/Runtime/Engine/Public/LevelInstance/LevelInstanceActorImpl.h b/Engine/Source/Runtime/Engine/Public/LevelInstance/LevelInstanceActorImpl.h
--- a/Engine/Source/Runtime/Engine/Public/LevelInstance/LevelInstanceActorImpl.h
+++ b/Engine/Source/Runtime/Engine/Public/LevelInstance/LevelInstanceActorImpl.h
@@ -26,6 +26,7 @@ private:
 	FLevelInstanceID CachedLevelInstanceID;
 	bool bCachedIsTemporarilyHiddenInEditor;
 	bool bGuardLoadUnload;
+	bool bAllowPartialLoading;
 public:
 	TSoftObjectPtr<UWorld> CachedWorldAsset;
 #endif
@@ -41,6 +42,7 @@ public:
 #if WITH_EDITOR
 		, bCachedIsTemporarilyHiddenInEditor(false)
 		, bGuardLoadUnload(false)
+		, bAllowPartialLoading(true)
 #endif
 	{
 	}
@@ -61,6 +63,10 @@ public:
 	virtual bool IsLoadingEnabled() const;
 	virtual void OnLevelInstanceLoaded();
 
+#if WITH_EDITOR
+	virtual bool SupportsPartialEditorLoading() const;
+#endif
+
 	/**
 	 * Begin AActor Implementation
 	 *
diff --git a/Engine/Source/Runtime/Engine/Public/WorldPartition/LoaderAdapter/LoaderAdapterSpatial.h b/Engine/Source/Runtime/Engine/Public/WorldPartition/LoaderAdapter/LoaderAdapterSpatial.h
--- a/Engine/Source/Runtime/Engine/Public/WorldPartition/LoaderAdapter/LoaderAdapterSpatial.h
+++ b/Engine/Source/Runtime/Engine/Public/WorldPartition/LoaderAdapter/LoaderAdapterSpatial.h
@@ -23,7 +23,7 @@ protected:
 	virtual bool Intersect(const FBox& Box) const =0;
 
 private:
-	void HandleIntersectingContainer(UWorldPartition* InWorldPartition, TFunctionRef<void(const FWorldPartitionHandle&)> InOperation) const;
+	void HandleIntersectingContainer(UWorldPartition* InWorldPartition, const FTransform InInstanceTransform, const FBox& InBoundingBox, TFunctionRef<void(const FWorldPartitionHandle&)> InOperation) const;
 
 protected:
 	bool bIncludeSpatiallyLoadedActors;
