Make sure PostLoad has been called on textures applied to dynamic material instances.

UMaterialInstance::PostLoad already does this, but it's possible for SetTextureParameterValue & co to be called from the PostLoad method of another object after the MID was created, but before the PostLoad of the referenced texture object has executed, which leads to uniform buffers containing invalid texture references and using the default white texture.

#rnx

diff --git a/Engine/Source/Runtime/Engine/Private/Materials/MaterialInstance.cpp b/Engine/Source/Runtime/Engine/Private/Materials/MaterialInstance.cpp
--- a/Engine/Source/Runtime/Engine/Private/Materials/MaterialInstance.cpp
+++ b/Engine/Source/Runtime/Engine/Private/Materials/MaterialInstance.cpp
@@ -3672,6 +3672,11 @@ void UMaterialInstance::SetTextureParameterValueInternal(const FMaterialParamete
 {
 	LLM_SCOPE(ELLMTag::MaterialInstance);
 
+	if (Value)
+	{
+		Value->ConditionalPostLoad();
+	}
+
 	FTextureParameterValue* ParameterValue = GameThread_FindParameterByName(TextureParameterValues, ParameterInfo);
 
 	bool bForceUpdate = false;
@@ -3706,6 +3711,11 @@ void UMaterialInstance::SetRuntimeVirtualTextureParameterValueInternal(const FMa
 {
 	LLM_SCOPE(ELLMTag::MaterialInstance);
 
+	if (Value)
+	{
+		Value->ConditionalPostLoad();
+	}
+
 	FRuntimeVirtualTextureParameterValue* ParameterValue = GameThread_FindParameterByName(RuntimeVirtualTextureParameterValues, ParameterInfo);
 
 	bool bForceUpdate = false;
@@ -3740,6 +3750,11 @@ void UMaterialInstance::SetSparseVolumeTextureParameterValueInternal(const FMate
 {
 	LLM_SCOPE(ELLMTag::MaterialInstance);
 
+	if (Value)
+	{
+		Value->ConditionalPostLoad();
+	}
+
 	FSparseVolumeTextureParameterValue* ParameterValue = GameThread_FindParameterByName(SparseVolumeTextureParameterValues, ParameterInfo);
 
 	bool bForceUpdate = false;
@@ -3773,6 +3788,11 @@ void UMaterialInstance::SetFontParameterValueInternal(const FMaterialParameterIn
 {
 	LLM_SCOPE(ELLMTag::MaterialInstance);
 
+	if (FontValue)
+	{
+		FontValue->ConditionalPostLoad();
+	}
+
 	FFontParameterValue* ParameterValue = GameThread_FindParameterByName(FontParameterValues, ParameterInfo);
 
 	bool bForceUpdate = false;
