Expose FDateTime's To/From Unix Timestamp functions to blueprints. Add the "IgnoreTypePromotion" metadata to DateTime structs to unblock users while we debug what the bigger type promotion problem is.

diff --git a/Engine/Source/Runtime/Engine/Classes/Kismet/KismetMathLibrary.h b/Engine/Source/Runtime/Engine/Classes/Kismet/KismetMathLibrary.h
--- a/Engine/Source/Runtime/Engine/Classes/Kismet/KismetMathLibrary.h
+++ b/Engine/Source/Runtime/Engine/Classes/Kismet/KismetMathLibrary.h
@@ -3266,51 +3266,51 @@ class ENGINE_API UKismetMathLibrary : public UBlueprintFunctionLibrary
 	//
 
 	/** Makes a DateTime struct */
-	UFUNCTION(BlueprintPure, Category="Math|DateTime", meta=(NativeMakeFunc, AdvancedDisplay = "3"))
+	UFUNCTION(BlueprintPure, Category="Math|DateTime", meta=(IgnoreTypePromotion, NativeMakeFunc, AdvancedDisplay = "3"))
 	static FDateTime MakeDateTime(int32 Year, int32 Month, int32 Day, int32 Hour = 0, int32 Minute = 0, int32 Second = 0, int32 Millisecond = 0);
 
 	/** Breaks a DateTime into its components */
-	UFUNCTION(BlueprintPure, Category="Math|DateTime", meta=(NativeBreakFunc, AdvancedDisplay = "4"))
+	UFUNCTION(BlueprintPure, Category="Math|DateTime", meta=(IgnoreTypePromotion, NativeBreakFunc, AdvancedDisplay = "4"))
 	static void BreakDateTime(FDateTime InDateTime, int32& Year, int32& Month, int32& Day, int32& Hour, int32& Minute, int32& Second, int32& Millisecond);
 
 	/** Addition (A + B) */
-	UFUNCTION(BlueprintPure, meta=(DisplayName="DateTime + Timespan", CompactNodeTitle="+", Keywords="+ add plus"), Category="Math|DateTime")
+	UFUNCTION(BlueprintPure, meta=(IgnoreTypePromotion, DisplayName="DateTime + Timespan", CompactNodeTitle="+", Keywords="+ add plus"), Category="Math|DateTime")
 	static FDateTime Add_DateTimeTimespan( FDateTime A, FTimespan B );
 
 	/** Subtraction (A - B) */
-	UFUNCTION(BlueprintPure, meta=(DisplayName="DateTime - Timespan", CompactNodeTitle="-", Keywords="- subtract minus"), Category="Math|DateTime")
+	UFUNCTION(BlueprintPure, meta=(IgnoreTypePromotion, DisplayName="DateTime - Timespan", CompactNodeTitle="-", Keywords="- subtract minus"), Category="Math|DateTime")
 	static FDateTime Subtract_DateTimeTimespan(FDateTime A, FTimespan B);
 
 	/** Addition (A + B) */
-	UFUNCTION(BlueprintPure, meta = (DisplayName = "DateTime + DateTime", CompactNodeTitle = "+", Keywords = "+ add plus"), Category = "Math|DateTime")
+	UFUNCTION(BlueprintPure, meta = (IgnoreTypePromotion, DisplayName = "DateTime + DateTime", CompactNodeTitle = "+", Keywords = "+ add plus"), Category = "Math|DateTime")
 	static FDateTime Add_DateTimeDateTime(FDateTime A, FDateTime B);
 
 	/** Subtraction (A - B) */
-	UFUNCTION(BlueprintPure, meta = (DisplayName = "DateTime - DateTime", CompactNodeTitle = "-", Keywords = "- subtract minus"), Category = "Math|DateTime")
+	UFUNCTION(BlueprintPure, meta = (IgnoreTypePromotion, DisplayName = "DateTime - DateTime", CompactNodeTitle = "-", Keywords = "- subtract minus"), Category = "Math|DateTime")
 	static FTimespan Subtract_DateTimeDateTime(FDateTime A, FDateTime B);
 
 	/** Returns true if the values are equal (A == B) */
-	UFUNCTION(BlueprintPure, meta=(DisplayName="Equal (DateTime)", CompactNodeTitle="==", Keywords="== equal"), Category="Math|DateTime")
+	UFUNCTION(BlueprintPure, meta=(IgnoreTypePromotion, DisplayName="Equal (DateTime)", CompactNodeTitle="==", Keywords="== equal"), Category="Math|DateTime")
 	static bool EqualEqual_DateTimeDateTime( FDateTime A, FDateTime B );
 
 	/** Returns true if the values are not equal (A != B) */
-	UFUNCTION(BlueprintPure, meta=(DisplayName="Not Equal (DateTime)", CompactNodeTitle="!=", Keywords="!= not equal"), Category="Math|DateTime")
+	UFUNCTION(BlueprintPure, meta=(IgnoreTypePromotion, DisplayName="Not Equal (DateTime)", CompactNodeTitle="!=", Keywords="!= not equal"), Category="Math|DateTime")
 	static bool NotEqual_DateTimeDateTime( FDateTime A, FDateTime B );
 
 	/** Returns true if A is greater than B (A > B) */
-	UFUNCTION(BlueprintPure, meta=(DisplayName="DateTime > DateTime", CompactNodeTitle=">", Keywords="> greater"), Category="Math|DateTime")
+	UFUNCTION(BlueprintPure, meta=(IgnoreTypePromotion, DisplayName="DateTime > DateTime", CompactNodeTitle=">", Keywords="> greater"), Category="Math|DateTime")
 	static bool Greater_DateTimeDateTime( FDateTime A, FDateTime B );
 
 	/** Returns true if A is greater than or equal to B (A >= B) */
-	UFUNCTION(BlueprintPure, meta=(DisplayName="DateTime >= DateTime", CompactNodeTitle=">=", Keywords=">= greater"), Category="Math|DateTime")
+	UFUNCTION(BlueprintPure, meta=(IgnoreTypePromotion, DisplayName="DateTime >= DateTime", CompactNodeTitle=">=", Keywords=">= greater"), Category="Math|DateTime")
 	static bool GreaterEqual_DateTimeDateTime( FDateTime A, FDateTime B );
 
 	/** Returns true if A is less than B (A < B) */
-	UFUNCTION(BlueprintPure, meta=(DisplayName="DateTime < DateTime", CompactNodeTitle="<", Keywords="< less"), Category="Math|DateTime")
+	UFUNCTION(BlueprintPure, meta=(IgnoreTypePromotion, DisplayName="DateTime < DateTime", CompactNodeTitle="<", Keywords="< less"), Category="Math|DateTime")
 	static bool Less_DateTimeDateTime( FDateTime A, FDateTime B );
 
 	/** Returns true if A is less than or equal to B (A <= B) */
-	UFUNCTION(BlueprintPure, meta=(DisplayName="DateTime <= DateTime", CompactNodeTitle="<=", Keywords="<= less"), Category="Math|DateTime")
+	UFUNCTION(BlueprintPure, meta=(IgnoreTypePromotion, DisplayName="DateTime <= DateTime", CompactNodeTitle="<=", Keywords="<= less"), Category="Math|DateTime")
 	static bool LessEqual_DateTimeDateTime( FDateTime A, FDateTime B );
 
 	/** Returns the date component of A */
@@ -3405,6 +3405,33 @@ class ENGINE_API UKismetMathLibrary : public UBlueprintFunctionLibrary
 	UFUNCTION(BlueprintPure, Category="Math|DateTime")
 	static bool DateTimeFromString(FString DateTimeString, FDateTime& Result);
 
+	/**
+	 * Returns this date as the number of seconds since the Unix Epoch (January 1st of 1970).
+	 *
+	 * @return Time of day.
+	 * @see FromUnixTimestamp
+	 */
+	UFUNCTION(BlueprintPure, Category="Math|DateTime")
+	static int64 ToUnixTimestamp(const FDateTime& Time);
+
+	/**
+	 * Returns this date as the number of seconds since the Unix Epoch (January 1st of 1970).
+	 *
+	 * @return Time of day.
+	 * @see FromUnixTimestamp
+	 */
+	UFUNCTION(BlueprintPure, Category="Math|DateTime")
+	static double ToUnixTimestampDouble(const FDateTime& Time);
+	
+	/**
+	 * Returns the date from Unix time (seconds from midnight 1970-01-01)
+	 *
+	 * @param UnixTime Unix time (seconds from midnight 1970-01-01)
+	 * @return Gregorian date and time.
+	 * @see ToUnixTimestamp
+	 */
+	UFUNCTION(BlueprintPure, Category="Math|DateTime")
+	static FDateTime FromUnixTimestamp(const int64 UnixTime);
 
 	//
 	// Timespan constants
diff --git a/Engine/Source/Runtime/Engine/Classes/Kismet/KismetMathLibrary.inl b/Engine/Source/Runtime/Engine/Classes/Kismet/KismetMathLibrary.inl
--- a/Engine/Source/Runtime/Engine/Classes/Kismet/KismetMathLibrary.inl
+++ b/Engine/Source/Runtime/Engine/Classes/Kismet/KismetMathLibrary.inl
@@ -2555,6 +2555,24 @@ bool UKismetMathLibrary::DateTimeFromString(FString DateTimeString, FDateTime& R
 	return FDateTime::Parse(DateTimeString, Result);
 }
 
+KISMET_MATH_FORCEINLINE
+int64 UKismetMathLibrary::ToUnixTimestamp(const FDateTime& Time)
+{
+	return Time.ToUnixTimestamp();
+}
+
+KISMET_MATH_FORCEINLINE
+double UKismetMathLibrary::ToUnixTimestampDouble(const FDateTime& Time)
+{
+	return Time.ToUnixTimestampDecimal();
+}
+
+KISMET_MATH_FORCEINLINE
+FDateTime UKismetMathLibrary::FromUnixTimestamp(const int64 UnixTime)
+{
+	return FDateTime::FromUnixTimestamp(UnixTime);
+}
+
 /* Timespan functions
 *****************************************************************************/
 
