- Allow creating a WorldPartition runtime cell and an external streaming object in engine
- Allow creating a WorldPartition runtime cell using an existing level package and its transform

diff --git a/Engine/Source/Runtime/Engine/Private/WorldPartition/WorldPartitionLevelStreamingDynamic.cpp b/Engine/Source/Runtime/Engine/Private/WorldPartition/WorldPartitionLevelStreamingDynamic.cpp
--- a/Engine/Source/Runtime/Engine/Private/WorldPartition/WorldPartitionLevelStreamingDynamic.cpp
+++ b/Engine/Source/Runtime/Engine/Private/WorldPartition/WorldPartitionLevelStreamingDynamic.cpp
@@ -30,6 +30,7 @@ UWorldPartitionLevelStreamingDynamic::UWorldPartitionLevelStreamingDynamic(const
 	, bLoadSucceeded(false)
 #endif
 	, bShouldBeAlwaysLoaded(false)
+	, bHasSetLevelTransform(false)
 {
 #if WITH_EDITOR
 	SetShouldBeVisibleInEditor(false);
@@ -68,6 +69,15 @@ void UWorldPartitionLevelStreamingDynamic::Initialize(const UWorldPartitionRunti
 	UpdateShouldSkipMakingVisibilityTransactionRequest();
 }
 
+void UWorldPartitionLevelStreamingDynamic::SetLevelTransform(const FTransform& InLevelTransform)
+{
+	if (!bHasSetLevelTransform)
+	{
+		LevelTransform = InLevelTransform;
+		bHasSetLevelTransform = true;
+	}
+}
+
 #if WITH_EDITOR
 
 UWorldPartitionLevelStreamingDynamic* UWorldPartitionLevelStreamingDynamic::LoadInEditor(UWorld* World, FName LevelStreamingName, const TArray<FWorldPartitionRuntimeCellObjectMapping>& InPackages)
@@ -178,6 +188,12 @@ void UWorldPartitionLevelStreamingDynamic::CreateRuntimeLevel()
  */
 bool UWorldPartitionLevelStreamingDynamic::RequestLevel(UWorld* InPersistentWorld, bool bInAllowLevelLoadRequests, EReqLevelBlock InBlockPolicy)
 {
+	// Use parent streaming implementation if world asset package is not a memory package
+	if (!FPackageName::IsMemoryPackage(GetWorldAssetPackageName()))
+	{
+		return Super::RequestLevel(InPersistentWorld, bInAllowLevelLoadRequests, InBlockPolicy);
+	}
+
 	// Quit early in case load request already issued
 	if (GetLevelStreamingState() == ELevelStreamingState::Loading)
 	{
diff --git a/Engine/Source/Runtime/Engine/Private/WorldPartition/WorldPartitionRuntimeHash.cpp b/Engine/Source/Runtime/Engine/Private/WorldPartition/WorldPartitionRuntimeHash.cpp
--- a/Engine/Source/Runtime/Engine/Private/WorldPartition/WorldPartitionRuntimeHash.cpp
+++ b/Engine/Source/Runtime/Engine/Private/WorldPartition/WorldPartitionRuntimeHash.cpp
@@ -9,8 +9,8 @@
 #include "WorldPartition/WorldPartitionStreamingSource.h"
 #include "WorldPartition/WorldPartitionRuntimeLevelStreamingCell.h"
 #include "WorldPartition/WorldPartitionLevelStreamingDynamic.h"
-#if WITH_EDITOR
 #include "Misc/ArchiveMD5.h"
+#if WITH_EDITOR
 #include "WorldPartition/Cook/WorldPartitionCookPackage.h"
 #endif
 
@@ -68,16 +68,40 @@ UWorldPartitionRuntimeHash::UWorldPartitionRuntimeHash(const FObjectInitializer&
 	: Super(ObjectInitializer)
 {}
 
-#if WITH_EDITOR
+UWorldPartitionRuntimeCell* UWorldPartitionRuntimeHash::CreateRuntimeCell(UClass* CellClass, UClass* CellDataClass, const FString& CellName, const FString& CellInstanceSuffix, UObject* InOuter)
+{
+	//@todo_ow: to reduce file paths on windows, we compute a MD5 hash from the unique cell name and use that hash as the cell filename. We
+	//			should create the runtime cell using the name but make sure the package that gets associated with it when cooking gets the
+	//			hash indtead.
+	auto GetCellObjectName = [](FString CellName) -> FString
+	{
+		FArchiveMD5 ArMD5;
+		ArMD5 << CellName;
+		FGuid CellNameGuid = ArMD5.GetGuidFromHash();
+		check(CellNameGuid.IsValid());
+
+		return CellNameGuid.ToString(EGuidFormats::Base36Encoded);
+	};
+
+	// Cooking should have an empty CellInstanceSuffix
+	check(!IsRunningCookCommandlet() || CellInstanceSuffix.IsEmpty());
+	const FString CellObjectName = GetCellObjectName(CellName) + CellInstanceSuffix;
+	// Use given outer if provided, else use hash as outer
+	UObject* Outer = InOuter ? InOuter : this;
+	UWorldPartitionRuntimeCell* RuntimeCell = NewObject<UWorldPartitionRuntimeCell>(Outer, CellClass, *CellObjectName);
+	RuntimeCell->RuntimeCellData = NewObject<UWorldPartitionRuntimeCellData>(RuntimeCell, CellDataClass);
+	return RuntimeCell;
+}
+
 URuntimeHashExternalStreamingObjectBase* UWorldPartitionRuntimeHash::CreateExternalStreamingObject(TSubclassOf<URuntimeHashExternalStreamingObjectBase> InClass, UObject* InOuter, FName InName, UWorld* InOwningWorld, UWorld* InOuterWorld)
 {
-	check(!InOwningWorld->IsGameWorld());
 	URuntimeHashExternalStreamingObjectBase* StreamingObject = NewObject<URuntimeHashExternalStreamingObjectBase>(InOuter, InClass, InName, RF_Public);
 	StreamingObject->OwningWorld = InOwningWorld;
 	StreamingObject->OuterWorld = InOuterWorld;
 	return StreamingObject;
 }
 
+#if WITH_EDITOR
 void UWorldPartitionRuntimeHash::OnBeginPlay()
 {
 	// Mark always loaded actors so that the Level will force reference to these actors for PIE.
@@ -96,33 +120,6 @@ void UWorldPartitionRuntimeHash::OnEndPlay()
 	ModifiedActorDescListForPIE.Empty();
 }
 
-UWorldPartitionRuntimeCell* UWorldPartitionRuntimeHash::CreateRuntimeCell(UClass* CellClass, UClass* CellDataClass, const FString& CellName, const FString& CellInstanceSuffix)
-{
-	//@todo_ow: to reduce file paths on windows, we compute a MD5 hash from the unique cell name and use that hash as the cell filename. We
-	//			should create the runtime cell using the name but make sure the package that gets associated with it when cooking gets the
-	//			hash indtead.
-	auto GetCellObjectName = [](FString CellName) -> FString
-	{
-		FArchiveMD5 ArMD5;
-		ArMD5 << CellName;
-
-		FMD5Hash MD5Hash;
-		ArMD5.GetHash(MD5Hash);
-
-		FGuid CellNameGuid = MD5HashToGuid(MD5Hash);
-		check(CellNameGuid.IsValid());
-
-		return CellNameGuid.ToString(EGuidFormats::Base36Encoded);
-	};
-
-	// Cooking should have an empty CellInstanceSuffix
-	check(!IsRunningCookCommandlet() || CellInstanceSuffix.IsEmpty());
-	const FString CellObjectName = GetCellObjectName(CellName) + CellInstanceSuffix;
-	UWorldPartitionRuntimeCell* RuntimeCell = NewObject<UWorldPartitionRuntimeCell>(this, CellClass, *CellObjectName);
-	RuntimeCell->RuntimeCellData = NewObject<UWorldPartitionRuntimeCellData>(RuntimeCell, CellDataClass);
-	return RuntimeCell;
-}
-
 bool UWorldPartitionRuntimeHash::GenerateStreaming(class UWorldPartitionStreamingPolicy* StreamingPolicy, const IStreamingGenerationContext* StreamingGenerationContext, TArray<FString>* OutPackagesToGenerate)
 {
 	return PackagesToGenerateForCook.IsEmpty();
diff --git a/Engine/Source/Runtime/Engine/Private/WorldPartition/WorldPartitionRuntimeLevelStreamingCell.cpp b/Engine/Source/Runtime/Engine/Private/WorldPartition/WorldPartitionRuntimeLevelStreamingCell.cpp
--- a/Engine/Source/Runtime/Engine/Private/WorldPartition/WorldPartitionRuntimeLevelStreamingCell.cpp
+++ b/Engine/Source/Runtime/Engine/Private/WorldPartition/WorldPartitionRuntimeLevelStreamingCell.cpp
@@ -103,6 +103,40 @@ void UWorldPartitionRuntimeLevelStreamingCell::CreateAndSetLevelStreaming(const
 	LevelStreaming = CreateLevelStreaming(InPackageName);
 }
 
+bool UWorldPartitionRuntimeLevelStreamingCell::CreateAndSetLevelStreaming(const TSoftObjectPtr<UWorld>& InWorldAsset, const FTransform& InInstanceTransform) const
+{
+	UWorld* OwningWorld = GetOwningWorld();
+	const FName LevelStreamingName = FName(*FString::Printf(TEXT("WorldPartitionLevelStreaming_%s"), *GetName()));
+	LevelStreaming = NewObject<UWorldPartitionLevelStreamingDynamic>(OwningWorld, UWorldPartitionLevelStreamingDynamic::StaticClass(), LevelStreamingName, RF_NoFlags, NULL);
+
+	// Generate unique Level Instance name assuming cell has a unique name
+	const FString LongPackageName = InWorldAsset.GetLongPackageName();
+	const FString PackagePath = FPackageName::GetLongPackagePath(LongPackageName);
+	const FString ShortPackageName = FPackageName::GetShortName(LongPackageName);
+	TStringBuilder<512> LevelPackageNameStrBuilder;
+	LevelPackageNameStrBuilder.Append(PackagePath);
+	LevelPackageNameStrBuilder.Append(TEXT("/"));
+	LevelPackageNameStrBuilder.Append(ShortPackageName);
+	LevelPackageNameStrBuilder.Append(TEXT("_LevelInstance_"));
+	LevelPackageNameStrBuilder.Append(GetName());
+	LevelPackageNameStrBuilder.Append(TEXT("."));
+	LevelPackageNameStrBuilder.Append(ShortPackageName);
+	LevelStreaming->SetWorldAsset(TSoftObjectPtr<UWorld>(FSoftObjectPath(LevelPackageNameStrBuilder.ToString())));
+
+	// Include WorldPartition's transform to Level
+	LevelStreaming->SetLevelTransform(GetOuterWorld()->GetWorldPartition()->GetInstanceTransform() * InInstanceTransform);
+	LevelStreaming->bClientOnlyVisible = GetClientOnlyVisible();
+	LevelStreaming->Initialize(*this);
+	LevelStreaming->PackageNameToLoad = FName(LongPackageName);
+
+	if (OwningWorld->IsPlayInEditor() && OwningWorld->GetPackage()->HasAnyPackageFlags(PKG_PlayInEditor) && OwningWorld->GetPackage()->GetPIEInstanceID() != INDEX_NONE)
+	{
+		// When renaming for PIE, make sure to keep World's name so that linker can properly remap with Package's instancing context
+		LevelStreaming->RenameForPIE(OwningWorld->GetPackage()->GetPIEInstanceID(), /*bKeepWorldAssetName*/true);
+	}
+	return true;
+}
+
 UWorldPartitionLevelStreamingDynamic* UWorldPartitionRuntimeLevelStreamingCell::CreateLevelStreaming(const FString& InPackageName) const
 {
 	if (HasActors())
@@ -353,16 +387,11 @@ void UWorldPartitionRuntimeLevelStreamingCell::DumpStateLog(FHierarchicalLogArch
 UWorldPartitionLevelStreamingDynamic* UWorldPartitionRuntimeLevelStreamingCell::GetOrCreateLevelStreaming() const
 {
 #if WITH_EDITOR
-	if (GetActorCount() == 0)
-	{
-		return nullptr;
-	}
-
-	if (!LevelStreaming)
+	if (!LevelStreaming && GetActorCount())
 	{
 		LevelStreaming = CreateLevelStreaming();
+		check(LevelStreaming);
 	}
-	check(LevelStreaming);
 #else
 	// In Runtime, always loaded cell level is handled by World directly
 	check(LevelStreaming || IsAlwaysLoaded());
@@ -383,7 +412,7 @@ UWorldPartitionLevelStreamingDynamic* UWorldPartitionRuntimeLevelStreamingCell::
 		}
 
 		// Transfer WorldPartition's transform to LevelStreaming
-		LevelStreaming->LevelTransform = WorldPartition->GetInstanceTransform();
+		LevelStreaming->SetLevelTransform(WorldPartition->GetInstanceTransform());
 
 		// LevelStreaming WorldAsset is a TSoftObjectPtr<UWorld>. If the WorldPartition's world is instanced then the TSoftObjectPtr<UWorld> will be remapped by FLinkerInstancingContext SoftObject remapping
 		// example if MainWorld is instanced it will have a package suffix and/or prefix like: /Temp+PATH+_LevelInstance1
@@ -454,11 +483,7 @@ bool UWorldPartitionRuntimeLevelStreamingCell::CanUnload() const
 void UWorldPartitionRuntimeLevelStreamingCell::Unload() const
 {
 #if WITH_EDITOR
-	if (GetActorCount() == 0)
-	{
-		return;
-	}
-	check(LevelStreaming);
+	check(LevelStreaming || GetActorCount());
 #else
 	// In Runtime, always loaded cell level is handled by World directly
 	check(LevelStreaming || IsAlwaysLoaded());
@@ -473,11 +498,7 @@ void UWorldPartitionRuntimeLevelStreamingCell::Unload() const
 void UWorldPartitionRuntimeLevelStreamingCell::Deactivate() const
 {
 #if WITH_EDITOR
-	if (GetActorCount() == 0)
-	{
-		return;
-	}
-	check(LevelStreaming);
+	check(LevelStreaming || GetActorCount());
 #else
 	// In Runtime, always loaded cell level is handled by World directly
 	check(LevelStreaming || IsAlwaysLoaded());
diff --git a/Engine/Source/Runtime/Engine/Public/WorldPartition/WorldPartitionLevelStreamingDynamic.h b/Engine/Source/Runtime/Engine/Public/WorldPartition/WorldPartitionLevelStreamingDynamic.h
--- a/Engine/Source/Runtime/Engine/Public/WorldPartition/WorldPartitionLevelStreamingDynamic.h
+++ b/Engine/Source/Runtime/Engine/Public/WorldPartition/WorldPartitionLevelStreamingDynamic.h
@@ -44,6 +44,7 @@ class ENGINE_API UWorldPartitionLevelStreamingDynamic : public ULevelStreamingDy
 #endif
 
 	void Initialize(const UWorldPartitionRuntimeLevelStreamingCell& InCell);
+	void SetLevelTransform(const FTransform& InLevelTransform);
 
 	virtual bool CanReplicateStreamingStatus() const override { return false; }
 
@@ -88,6 +89,8 @@ private:
 	UPROPERTY()
 	bool bShouldBeAlwaysLoaded;
 
+	bool bHasSetLevelTransform;
+
 #if WITH_EDITORONLY_DATA
 	UPROPERTY()
 	TObjectPtr<UActorContainer> UnsavedActorsContainer;
diff --git a/Engine/Source/Runtime/Engine/Public/WorldPartition/WorldPartitionRuntimeCell.h b/Engine/Source/Runtime/Engine/Public/WorldPartition/WorldPartitionRuntimeCell.h
--- a/Engine/Source/Runtime/Engine/Public/WorldPartition/WorldPartitionRuntimeCell.h
+++ b/Engine/Source/Runtime/Engine/Public/WorldPartition/WorldPartitionRuntimeCell.h
@@ -195,6 +195,7 @@ class ENGINE_API UWorldPartitionRuntimeCell : public UObject, public IWorldParti
 	virtual void SetStreamingPriority(int32 InStreamingPriority) const PURE_VIRTUAL(UWorldPartitionRuntimeCell::SetStreamingPriority,);
 	virtual EStreamingStatus GetStreamingStatus() const { return LEVEL_Unloaded; }
 	virtual bool IsLoading() const { return false; }
+	void SetClientOnlyVisible(bool bInClientOnlyVisible) { bClientOnlyVisible = bInClientOnlyVisible; }
 	bool GetClientOnlyVisible() const { return bClientOnlyVisible; }
 	virtual FGuid const& GetContentBundleID() const { return ContentBundleID; }
 	virtual TArray<FName> GetActors() const PURE_VIRTUAL(UWorldPartitionRuntimeCell::GetActors, return TArray<FName>(););
@@ -219,11 +220,10 @@ class ENGINE_API UWorldPartitionRuntimeCell : public UObject, public IWorldParti
 	UDataLayerManager* GetDataLayerManager() const;
 	bool HasAnyDataLayerInEffectiveRuntimeState(EDataLayerRuntimeState InState) const;
 
+	void SetBlockOnSlowLoading(bool bInBlockOnSlowLoading) { bBlockOnSlowLoading = bInBlockOnSlowLoading; }
 	bool GetBlockOnSlowLoading() const { return bBlockOnSlowLoading; }
 #if WITH_EDITOR
 	bool NeedsActorToCellRemapping() const;
-	void SetBlockOnSlowLoading(bool bInBlockOnSlowLoading) { bBlockOnSlowLoading = bInBlockOnSlowLoading; }
-	void SetClientOnlyVisible(bool bInClientOnlyVisible) { bClientOnlyVisible = bInClientOnlyVisible; }
 	void SetDataLayers(const TArray<const UDataLayerInstance*>& InDataLayerInstances);
 	void SetContentBundleUID(const FGuid& InContentBundleID) { ContentBundleID = InContentBundleID; }
 	void SetLevelPackageName(const FName& InLevelPackageName) { LevelPackageName = InLevelPackageName; }
diff --git a/Engine/Source/Runtime/Engine/Public/WorldPartition/WorldPartitionRuntimeHash.h b/Engine/Source/Runtime/Engine/Public/WorldPartition/WorldPartitionRuntimeHash.h
--- a/Engine/Source/Runtime/Engine/Public/WorldPartition/WorldPartitionRuntimeHash.h
+++ b/Engine/Source/Runtime/Engine/Public/WorldPartition/WorldPartitionRuntimeHash.h
@@ -103,8 +103,6 @@ class ENGINE_API UWorldPartitionRuntimeHash : public UObject
 	void OnBeginPlay();
 	void OnEndPlay();
 
-	UWorldPartitionRuntimeCell* CreateRuntimeCell(UClass* CellClass, UClass* CellDataClass, const FString& CellName, const FString& CellInstanceSuffix);
-
 protected:
 	bool ConditionalRegisterAlwaysLoadedActorsForPIE(const IStreamingGenerationContext::FActorSetInstance* ActorSetInstance, bool bIsMainWorldPartition, bool bIsMainContainer, bool bIsCellAlwaysLoaded);
 	bool PopulateCellActorInstances(const TArray<const IStreamingGenerationContext::FActorSetInstance*> ActorSetInstances, bool bIsMainWorldPartition, bool bIsCellAlwaysLoaded, TArray<IStreamingGenerationContext::FActorInstance>& OutCellActorInstances);
@@ -151,11 +149,11 @@ public:
 	virtual bool IsStreaming3D() const { return true; }
 
 protected:
+	static URuntimeHashExternalStreamingObjectBase* CreateExternalStreamingObject(TSubclassOf<URuntimeHashExternalStreamingObjectBase> InClass, UObject* InOuter, FName InName, UWorld* InOwningWorld, UWorld* InOuterWorld);
+	UWorldPartitionRuntimeCell* CreateRuntimeCell(UClass* CellClass, UClass* CellDataClass, const FString& CellName, const FString& CellInstanceSuffix, UObject* InOuter = nullptr);
 	virtual EWorldPartitionStreamingPerformance GetStreamingPerformanceForCell(const UWorldPartitionRuntimeCell* Cell) const { return EWorldPartitionStreamingPerformance::Good; }
 
 #if WITH_EDITOR
-	static URuntimeHashExternalStreamingObjectBase* CreateExternalStreamingObject(TSubclassOf<URuntimeHashExternalStreamingObjectBase> InClass, UObject* InOuter, FName InName, UWorld* InOwningWorld, UWorld* InOuterWorld);
-
 	template <class T>
 	T* CreateExternalStreamingObject(UObject* InOuter, FName InName)
 	{
diff --git a/Engine/Source/Runtime/Engine/Public/WorldPartition/WorldPartitionRuntimeLevelStreamingCell.h b/Engine/Source/Runtime/Engine/Public/WorldPartition/WorldPartitionRuntimeLevelStreamingCell.h
--- a/Engine/Source/Runtime/Engine/Public/WorldPartition/WorldPartitionRuntimeLevelStreamingCell.h
+++ b/Engine/Source/Runtime/Engine/Public/WorldPartition/WorldPartitionRuntimeLevelStreamingCell.h
@@ -43,6 +43,7 @@ class ENGINE_API UWorldPartitionRuntimeLevelStreamingCell : public UWorldPartiti
 	virtual TArray<FName> GetActors() const override;
 
 	void CreateAndSetLevelStreaming(const FString& InPackageName);
+	bool CreateAndSetLevelStreaming(const TSoftObjectPtr<UWorld>& InWorldAsset, const FTransform& InInstanceTransform) const;
 	class UWorldPartitionLevelStreamingDynamic* CreateLevelStreaming(const FString& InPackageName = FString()) const;
 	
 
