World Partition : Added helper function to insert a runtime cell in a FSpatialHashStreamingGrid

diff --git a/Engine/Source/Runtime/Engine/Private/WorldPartition/WorldPartitionRuntimeSpatialHash.cpp b/Engine/Source/Runtime/Engine/Private/WorldPartition/WorldPartitionRuntimeSpatialHash.cpp
--- a/Engine/Source/Runtime/Engine/Private/WorldPartition/WorldPartitionRuntimeSpatialHash.cpp
+++ b/Engine/Source/Runtime/Engine/Private/WorldPartition/WorldPartitionRuntimeSpatialHash.cpp
@@ -358,6 +358,40 @@ void FSpatialHashStreamingGrid::GetCells(const TArray<FWorldPartitionStreamingSo
 	}
 }
 
+bool FSpatialHashStreamingGrid::InsertGridCell(UWorldPartitionRuntimeCell* InCell, const FGridCellCoord& InGridCellCoords)
+{
+	check(InCell);
+
+	const int64 Level = InGridCellCoords.Z;
+	const int64 CellCoordX = InGridCellCoords.X;
+	const int64 CellCoordY = InGridCellCoords.Y;
+
+	// Add Cell to grid
+	if (GridLevels.IsValidIndex(Level))
+	{
+		uint64 CellIndex = 0;
+		if (GetGridHelper().Levels[Level].GetCellIndex(FGridCellCoord2(CellCoordX, CellCoordY), CellIndex))
+		{
+			FSpatialHashStreamingGridLevel& GridLevel = GridLevels[Level];
+			int32 LayerCellIndex;
+			int32* LayerCellIndexPtr = GridLevel.LayerCellsMapping.Find(CellIndex);
+			if (LayerCellIndexPtr)
+			{
+				LayerCellIndex = *LayerCellIndexPtr;
+			}
+			else
+			{
+				LayerCellIndex = GridLevel.LayerCells.AddDefaulted();
+				GridLevel.LayerCellsMapping.Add(CellIndex, LayerCellIndex);
+			}
+			GridLevel.LayerCells[LayerCellIndex].GridCells.Add(InCell);
+			return true;
+		}
+	}
+
+	return false;
+}
+
 void FSpatialHashStreamingGrid::InjectExternalStreamingObjectGrid(const FSpatialHashStreamingGrid& InExternalObjectStreamingGrid) const
 {
 	if (InjectedGridLevels.IsEmpty())
diff --git a/Engine/Source/Runtime/Engine/Public/WorldPartition/WorldPartitionRuntimeSpatialHash.h b/Engine/Source/Runtime/Engine/Public/WorldPartition/WorldPartitionRuntimeSpatialHash.h
--- a/Engine/Source/Runtime/Engine/Public/WorldPartition/WorldPartitionRuntimeSpatialHash.h
+++ b/Engine/Source/Runtime/Engine/Public/WorldPartition/WorldPartitionRuntimeSpatialHash.h
@@ -140,6 +140,8 @@ struct ENGINE_API FSpatialHashStreamingGrid
 	void InjectExternalStreamingObjectGrid(const FSpatialHashStreamingGrid& InExternalObjectStreamingGrid) const;
 	void RemoveExternalStreamingObjectGrid(const FSpatialHashStreamingGrid& InExternalObjectStreamingGrid) const;
 
+	bool InsertGridCell(UWorldPartitionRuntimeCell* InGridCell, const FGridCellCoord& InGridCellCoords);
+
 	// Used by PIE/Game
 	int64 GetCellSize(int32 Level) const;
 	void GetCells(const FWorldPartitionStreamingQuerySource& QuerySource, TSet<const UWorldPartitionRuntimeCell*>& OutCells, bool bEnableZCulling, FWorldPartitionQueryCache* QueryCache = nullptr) const;
