Exported URuntimeSpatialHashExternalStreamingObject

diff --git a/Engine/Source/Runtime/Engine/Public/WorldPartition/WorldPartitionRuntimeSpatialHash.h b/Engine/Source/Runtime/Engine/Public/WorldPartition/WorldPartitionRuntimeSpatialHash.h
--- a/Engine/Source/Runtime/Engine/Public/WorldPartition/WorldPartitionRuntimeSpatialHash.h
+++ b/Engine/Source/Runtime/Engine/Public/WorldPartition/WorldPartitionRuntimeSpatialHash.h
@@ -237,7 +237,7 @@ public:
 };
 
 UCLASS()
-class URuntimeSpatialHashExternalStreamingObject : public URuntimeHashExternalStreamingObjectBase
+class ENGINE_API URuntimeSpatialHashExternalStreamingObject : public URuntimeHashExternalStreamingObjectBase
 {
 	GENERATED_BODY()
 
