Skeletal Mesh: Remove unnecessary ensure.

#rnx

diff --git a/Engine/Source/Runtime/Engine/Private/Animation/SkeletalMeshVertexAttribute.cpp b/Engine/Source/Runtime/Engine/Private/Animation/SkeletalMeshVertexAttribute.cpp
--- a/Engine/Source/Runtime/Engine/Private/Animation/SkeletalMeshVertexAttribute.cpp
+++ b/Engine/Source/Runtime/Engine/Private/Animation/SkeletalMeshVertexAttribute.cpp
@@ -105,7 +105,6 @@ void FSkeletalMeshVertexAttributeRenderData::ReleaseResources()
 
 void FSkeletalMeshVertexAttributeRenderData::CleanUp()
 {
-	check(IsInGameThread());
 	for (const TTuple<FName, FSkeletalMeshAttributeVertexBuffer*>& BufferItem: Buffers)
 	{
 		BufferItem.Value->CleanUp();
