Fix scalar displacement pin going through the UMaterialExpressionBreakMaterialAttributes node

diff --git a/Engine/Source/Runtime/Engine/Private/Materials/MaterialExpressions.cpp b/Engine/Source/Runtime/Engine/Private/Materials/MaterialExpressions.cpp
--- a/Engine/Source/Runtime/Engine/Private/Materials/MaterialExpressions.cpp
+++ b/Engine/Source/Runtime/Engine/Private/Materials/MaterialExpressions.cpp
@@ -6850,6 +6850,7 @@ void UMaterialExpressionBreakMaterialAttributes::BuildPropertyToIOIndexMap()
 		PropertyToIOIndexMap.Add(MP_CustomizedUVs7,			23);
 		PropertyToIOIndexMap.Add(MP_PixelDepthOffset,		24);
 		PropertyToIOIndexMap.Add(MP_ShadingModel,			25);
+		PropertyToIOIndexMap.Add(MP_Displacement,			26);
 	}
 }
 
