World Partition : keep track of injected external streaming objects to prevent re-injecting multiple times the same content.

diff --git a/Engine/Source/Runtime/Engine/Private/WorldPartition/WorldPartition.cpp b/Engine/Source/Runtime/Engine/Private/WorldPartition/WorldPartition.cpp
--- a/Engine/Source/Runtime/Engine/Private/WorldPartition/WorldPartition.cpp
+++ b/Engine/Source/Runtime/Engine/Private/WorldPartition/WorldPartition.cpp
@@ -1469,30 +1469,33 @@ void UWorldPartition::UpdateStreamingState()
 	}
 }
 
-bool UWorldPartition::InjectExternalStreamingObject(URuntimeHashExternalStreamingObjectBase* ExternalStreamingObject)
+bool UWorldPartition::InjectExternalStreamingObject(URuntimeHashExternalStreamingObjectBase* InExternalStreamingObject)
 {
-	bool bInjected = RuntimeHash->InjectExternalStreamingObject(ExternalStreamingObject);
+	bool bInjected = RuntimeHash->InjectExternalStreamingObject(InExternalStreamingObject);
 	if (bInjected)
 	{
-		StreamingPolicy->InjectExternalStreamingObject(ExternalStreamingObject);
-		GetWorld()->GetSubsystem<UHLODSubsystem>()->OnExternalStreamingObjectInjected(ExternalStreamingObject);
+		if (StreamingPolicy)
+		{
+			StreamingPolicy->InjectExternalStreamingObject(InExternalStreamingObject);
+		}
+		GetWorld()->GetSubsystem<UHLODSubsystem>()->OnExternalStreamingObjectInjected(InExternalStreamingObject);
 		++StreamingStateEpoch;
 	}
 
 	return bInjected;
 }
 
-bool UWorldPartition::RemoveExternalStreamingObject(URuntimeHashExternalStreamingObjectBase* ExternalStreamingObject)
+bool UWorldPartition::RemoveExternalStreamingObject(URuntimeHashExternalStreamingObjectBase* InExternalStreamingObject)
 {
-	bool bRemoved = RuntimeHash->RemoveExternalStreamingObject(ExternalStreamingObject);
+	bool bRemoved = RuntimeHash->RemoveExternalStreamingObject(InExternalStreamingObject);
 	if (bRemoved)
 	{
 		if (StreamingPolicy)
 		{
-			StreamingPolicy->RemoveExternalStreamingObject(ExternalStreamingObject);
+			StreamingPolicy->RemoveExternalStreamingObject(InExternalStreamingObject);
 		}
 		
-		GetWorld()->GetSubsystem<UHLODSubsystem>()->OnExternalStreamingObjectRemoved(ExternalStreamingObject);
+		GetWorld()->GetSubsystem<UHLODSubsystem>()->OnExternalStreamingObjectRemoved(InExternalStreamingObject);
 		++StreamingStateEpoch;
 	}
 
diff --git a/Engine/Source/Runtime/Engine/Private/WorldPartition/WorldPartitionRuntimeHash.cpp b/Engine/Source/Runtime/Engine/Private/WorldPartition/WorldPartitionRuntimeHash.cpp
--- a/Engine/Source/Runtime/Engine/Private/WorldPartition/WorldPartitionRuntimeHash.cpp
+++ b/Engine/Source/Runtime/Engine/Private/WorldPartition/WorldPartitionRuntimeHash.cpp
@@ -463,6 +463,30 @@ EWorldPartitionStreamingPerformance UWorldPartitionRuntimeHash::GetStreamingPerf
 	return StreamingPerformance;
 }
 
+bool UWorldPartitionRuntimeHash::InjectExternalStreamingObject(URuntimeHashExternalStreamingObjectBase* InExternalStreamingObject)
+{
+	check(InExternalStreamingObject);
+	bool bAlreadyInSet = false;
+	InjectedExternalStreamingObjects.Add(InExternalStreamingObject, &bAlreadyInSet);
+	if (bAlreadyInSet)
+	{
+		UE_LOG(LogWorldPartition, Warning, TEXT("External streaming object %s already injected."), *InExternalStreamingObject->GetName());
+		return false;
+	}
+	return true;
+}
+
+bool UWorldPartitionRuntimeHash::RemoveExternalStreamingObject(URuntimeHashExternalStreamingObjectBase* InExternalStreamingObject)
+{
+	check(InExternalStreamingObject);
+	if (!InjectedExternalStreamingObjects.Remove(InExternalStreamingObject))
+	{
+		UE_LOG(LogWorldPartition, Warning, TEXT("External streaming object %s was not injected."), *InExternalStreamingObject->GetName());
+		return false;
+	}
+	return true;
+}
+
 void UWorldPartitionRuntimeHash::FStreamingSourceCells::AddCell(const UWorldPartitionRuntimeCell* Cell, const FWorldPartitionStreamingSource& Source, const FSphericalSector& SourceShape)
 {
 	if (Cell->ShouldResetStreamingSourceInfo())
diff --git a/Engine/Source/Runtime/Engine/Private/WorldPartition/WorldPartitionRuntimeSpatialHash.cpp b/Engine/Source/Runtime/Engine/Private/WorldPartition/WorldPartitionRuntimeSpatialHash.cpp
--- a/Engine/Source/Runtime/Engine/Private/WorldPartition/WorldPartitionRuntimeSpatialHash.cpp
+++ b/Engine/Source/Runtime/Engine/Private/WorldPartition/WorldPartitionRuntimeSpatialHash.cpp
@@ -1594,9 +1594,14 @@ void UWorldPartitionRuntimeSpatialHash::ForEachStreamingCellsSources(const TArra
 	}
 }
 
-bool UWorldPartitionRuntimeSpatialHash::InjectExternalStreamingObject(URuntimeHashExternalStreamingObjectBase* ExternalStreamingObject)
+bool UWorldPartitionRuntimeSpatialHash::InjectExternalStreamingObject(URuntimeHashExternalStreamingObjectBase* InExternalStreamingObject)
 {
-	URuntimeSpatialHashExternalStreamingObject* SpatialHashExternalStreamingObject = CastChecked<URuntimeSpatialHashExternalStreamingObject>(ExternalStreamingObject);
+	if (!Super::InjectExternalStreamingObject(InExternalStreamingObject))
+	{
+		return false;
+	}
+
+	URuntimeSpatialHashExternalStreamingObject* SpatialHashExternalStreamingObject = CastChecked<URuntimeSpatialHashExternalStreamingObject>(InExternalStreamingObject);
 
 	// Validate that there's a corresponding streaming grid for each streaming grid of this external streaming object
 	for (const FSpatialHashStreamingGrid& ExternalStreamingGrid : SpatialHashExternalStreamingObject->StreamingGrids)
@@ -1604,7 +1609,7 @@ bool UWorldPartitionRuntimeSpatialHash::InjectExternalStreamingObject(URuntimeHa
 		if (!GetStreamingGridByName(ExternalStreamingGrid.GridName))
 		{
 			UE_LOG(LogWorldPartition, Error, TEXT("Failed to inject external streaming object %s, can't find matching streaming grid %s."),
-				*ExternalStreamingObject->GetName(), *ExternalStreamingGrid.GridName.ToString());
+				*InExternalStreamingObject->GetName(), *ExternalStreamingGrid.GridName.ToString());
 			return false;
 		}
 	}
@@ -1618,9 +1623,14 @@ bool UWorldPartitionRuntimeSpatialHash::InjectExternalStreamingObject(URuntimeHa
 	return true;
 }
 
-bool UWorldPartitionRuntimeSpatialHash::RemoveExternalStreamingObject(URuntimeHashExternalStreamingObjectBase* ExternalStreamingObject)
+bool UWorldPartitionRuntimeSpatialHash::RemoveExternalStreamingObject(URuntimeHashExternalStreamingObjectBase* InExternalStreamingObject)
 {
-	URuntimeSpatialHashExternalStreamingObject* SpatialHashExternalStreamingObject = CastChecked<URuntimeSpatialHashExternalStreamingObject>(ExternalStreamingObject);
+	if (!Super::RemoveExternalStreamingObject(InExternalStreamingObject))
+	{
+		return false;
+	}
+
+	URuntimeSpatialHashExternalStreamingObject* SpatialHashExternalStreamingObject = CastChecked<URuntimeSpatialHashExternalStreamingObject>(InExternalStreamingObject);
 
 	for (const FSpatialHashStreamingGrid& ExternalStreamingGrid : SpatialHashExternalStreamingObject->StreamingGrids)
 	{
diff --git a/Engine/Source/Runtime/Engine/Public/WorldPartition/WorldPartitionRuntimeHash.h b/Engine/Source/Runtime/Engine/Public/WorldPartition/WorldPartitionRuntimeHash.h
--- a/Engine/Source/Runtime/Engine/Public/WorldPartition/WorldPartitionRuntimeHash.h
+++ b/Engine/Source/Runtime/Engine/Public/WorldPartition/WorldPartitionRuntimeHash.h
@@ -140,8 +140,8 @@ public:
 	bool IsCellRelevantFor(bool bClientOnlyVisible) const;
 	EWorldPartitionStreamingPerformance GetStreamingPerformance(const TSet<const UWorldPartitionRuntimeCell*>& CellToActivate) const;
 
-	virtual bool InjectExternalStreamingObject(URuntimeHashExternalStreamingObjectBase* ExternalStreamingObject) { return false; }
-	virtual bool RemoveExternalStreamingObject(URuntimeHashExternalStreamingObjectBase* ExternalStreamingObject) { return false; }
+	virtual bool InjectExternalStreamingObject(URuntimeHashExternalStreamingObjectBase* ExternalStreamingObject);
+	virtual bool RemoveExternalStreamingObject(URuntimeHashExternalStreamingObjectBase* ExternalStreamingObject);
 
 	virtual bool Draw2D(FWorldPartitionDraw2DContext& DrawContext) const { return false; }
 	virtual void Draw3D(const TArray<FWorldPartitionStreamingSource>& Sources) const {}
@@ -180,8 +180,10 @@ public:
 	mutable FActorDescList ModifiedActorDescListForPIE;
 #endif
 
-#if WITH_EDITOR
 private:
+#if WITH_EDITOR
 	void ForceExternalActorLevelReference(bool bForceExternalActorLevelReferenceForPIE);
 #endif
+
+	TSet<TWeakObjectPtr<URuntimeHashExternalStreamingObjectBase>> InjectedExternalStreamingObjects;
 };
\ No newline at end of file
