Attach VRAM info to AppendStaticClientPerformanceAttributes

diff --git a/Engine/Source/Runtime/Engine/Private/Analytics/EngineAnalyticsSessionSummary.cpp b/Engine/Source/Runtime/Engine/Private/Analytics/EngineAnalyticsSessionSummary.cpp
--- a/Engine/Source/Runtime/Engine/Private/Analytics/EngineAnalyticsSessionSummary.cpp
+++ b/Engine/Source/Runtime/Engine/Private/Analytics/EngineAnalyticsSessionSummary.cpp
@@ -143,6 +143,9 @@ FEngineAnalyticsSessionSummary::FEngineAnalyticsSessionSummary(TSharedPtr<IAnaly
 	FString OSMinor;
 	FPlatformMisc::GetOSVersions(/*out*/ OSMajor, /*out*/ OSMinor);
 
+	FTextureMemoryStats TextureMemStats;
+	RHIGetTextureMemoryStats(TextureMemStats);
+
 	// Get project settings.
 	const UGeneralProjectSettings& ProjectSettings = *GetDefault<UGeneralProjectSettings>();
 
@@ -174,6 +177,7 @@ FEngineAnalyticsSessionSummary::FEngineAnalyticsSessionSummary(TSharedPtr<IAnaly
 	Store->Set(TEXT("GRHIAdapterUserDriverVersion"), GRHIAdapterUserDriverVersion);
 	Store->Set(TEXT("GRHIName"), FString(GDynamicRHI ? GDynamicRHI->GetName() : TEXT("")));
 	Store->Set(TEXT("GRHIAdapterDriverOnDenyList"), GRHIAdapterDriverOnDenyList);
+	Store->Set(TEXT("GRHIAdapterMemory"), static_cast<uint64>(TextureMemStats.DedicatedVideoMemory));
 	Store->Set(TEXT("ProjectName"), EngineAnalyticsProperties::GetProjectName(ProjectSettings));
 	Store->Set(TEXT("ProjectID"), ProjectSettings.ProjectID.ToString(EGuidFormats::DigitsWithHyphens));
 	Store->Set(TEXT("ProjectDescription"), ProjectSettings.Description);
diff --git a/Engine/Source/Runtime/Engine/Private/EngineAnalytics.cpp b/Engine/Source/Runtime/Engine/Private/EngineAnalytics.cpp
--- a/Engine/Source/Runtime/Engine/Private/EngineAnalytics.cpp
+++ b/Engine/Source/Runtime/Engine/Private/EngineAnalytics.cpp
@@ -245,6 +245,9 @@ void FEngineAnalytics::AppendMachineStats(TArray<FAnalyticsEventAttribute>& Even
 	FString OSMinor;
 	FPlatformMisc::GetOSVersions(/*out*/ OSMajor, /*out*/ OSMinor);
 
+	FTextureMemoryStats TextureMemStats;
+	RHIGetTextureMemoryStats(TextureMemStats);
+
 	GEngine->CreateStartupAnalyticsAttributes( EventAttributes );
 
 	// Add project info whether we are in editor or game.
@@ -271,6 +274,7 @@ void FEngineAnalytics::AppendMachineStats(TArray<FAnalyticsEventAttribute>& Even
 	EventAttributes.Emplace(TEXT("Is64BitOS"), FPlatformMisc::Is64bitOperatingSystem());
 	EventAttributes.Emplace(TEXT("GPUVendorID"), GRHIVendorId);
 	EventAttributes.Emplace(TEXT("GPUDeviceID"), GRHIDeviceId);
+	EventAttributes.Emplace(TEXT("GPUMemory"), static_cast<uint64>(TextureMemStats.DedicatedVideoMemory));
 	EventAttributes.Emplace(TEXT("GRHIDeviceRevision"), GRHIDeviceRevision);
 	EventAttributes.Emplace(TEXT("GRHIAdapterInternalDriverVersion"), GRHIAdapterInternalDriverVersion);
 	EventAttributes.Emplace(TEXT("GRHIAdapterUserDriverVersion"), GRHIAdapterUserDriverVersion);
