Clarification for behavior of Attenuation Distance Scale on Sound Classes

diff --git a/Engine/Source/Runtime/Engine/Classes/Sound/SoundClass.h b/Engine/Source/Runtime/Engine/Classes/Sound/SoundClass.h
--- a/Engine/Source/Runtime/Engine/Classes/Sound/SoundClass.h
+++ b/Engine/Source/Runtime/Engine/Classes/Sound/SoundClass.h
@@ -66,8 +66,9 @@ struct FSoundClassProperties
 	UPROPERTY(EditAnywhere, BlueprintReadOnly, Category = General)
 	float LowPassFilterFrequency;
 
-	/** Distance scale to apply to sounds that play with this sound class.
-	  * Sounds will have their attenuation distance scaled by this amount.
+	/** Scales the distance measurement used by the audio engine when determining distance-based attenuation. 
+	  * E.g., a sound 1000 units away with an AttenuationDistanceScale of .5 will be attenuated
+	  * as if it is 500 units away from the listener.
 	  * Allows adjusting attenuation settings dynamically. */
 	UPROPERTY(EditAnywhere, BlueprintReadOnly, Category = General)
 	float AttenuationDistanceScale;
