Warn when our dest size isn't big enough now that we theoretically have good sizes coming back from RHI for the mips.

diff --git a/Engine/Source/Runtime/Engine/Private/Rendering/Texture2DResource.cpp b/Engine/Source/Runtime/Engine/Private/Rendering/Texture2DResource.cpp
--- a/Engine/Source/Runtime/Engine/Private/Rendering/Texture2DResource.cpp
+++ b/Engine/Source/Runtime/Engine/Private/Rendering/Texture2DResource.cpp
@@ -266,12 +266,18 @@ void FTexture2DResource::GetData( uint32 MipIndex, void* Dest, uint32 DestPitch,
 
 	int64 BulkDataSize = MipMap.BulkData.GetBulkDataSize();
 
-	UE_LOG(LogTextureUpload,Verbose,TEXT("Size: %dx%d , EffectiveSize=%d BulkDataSize=%d , SrcPitch=%d DestPitch=%d Format=%d, DestSize=%lld"),
+	UE_LOG(LogTextureUpload,Verbose,TEXT("Size: %dx%d , EffectiveSize=%d BulkDataSize=%lld , SrcPitch=%d DestPitch=%d Format=%d, DestSize=%lld"),
 		MipMap.SizeX,MipMap.SizeY,
-		EffectiveSize,(int)BulkDataSize,
+		EffectiveSize,BulkDataSize,
 		SrcPitch,DestPitch, PixelFormat, DestSize);
 
-	//check(DestSize >= (uint64)BulkDataSize);
+	if ((uint64)BulkDataSize > DestSize)
+	{
+		UE_LOG(LogTextureUpload, Warning, TEXT("DestSize is reported smaller than BulkDataSize in upload! Either RHI is reporting wrong sizes or we are stomping memory! Size: %dx%d, EffectiveSize=%d, BulkDataSize=%lld, SrcPitch=%d, DestPitch=%d, Format=%d, DestSize=%lld"),
+			MipMap.SizeX,MipMap.SizeY,
+			EffectiveSize,BulkDataSize,
+			SrcPitch,DestPitch, PixelFormat, DestSize);
+	}
 
 	// for platforms that returned 0 pitch from Lock, we need to just use the bulk data directly, never do 
 	// runtime block size checking, conversion, or the like
