Fixed non unity compile error

#Issue 300824

diff --git a/Engine/Source/Runtime/Engine/Private/Analytics/EngineAnalyticsSessionSummary.cpp b/Engine/Source/Runtime/Engine/Private/Analytics/EngineAnalyticsSessionSummary.cpp
--- a/Engine/Source/Runtime/Engine/Private/Analytics/EngineAnalyticsSessionSummary.cpp
+++ b/Engine/Source/Runtime/Engine/Private/Analytics/EngineAnalyticsSessionSummary.cpp
@@ -14,6 +14,7 @@
 #include "Misc/EngineVersion.h"
 #include "UserActivityTracking.h"
 #include "RHI.h"
+#include "RHIStats.h"
 
 #if PLATFORM_WINDOWS
 #include "HAL/FileManager.h"
diff --git a/Engine/Source/Runtime/Engine/Private/EngineAnalytics.cpp b/Engine/Source/Runtime/Engine/Private/EngineAnalytics.cpp
--- a/Engine/Source/Runtime/Engine/Private/EngineAnalytics.cpp
+++ b/Engine/Source/Runtime/Engine/Private/EngineAnalytics.cpp
@@ -12,6 +12,7 @@
 #include "Misc/EngineVersion.h"
 #include "BuildSettings.h"
 #include "RHI.h"
+#include "RHIStats.h"
 #include "GenericPlatform/GenericPlatformCrashContext.h"
 #include "StudioAnalytics.h"
 #include "UObject/Class.h"
