Fix unclamp radius/length for point/spot light causing visual artifact.

diff --git a/Engine/Source/Runtime/Engine/Classes/Components/PointLightComponent.h b/Engine/Source/Runtime/Engine/Classes/Components/PointLightComponent.h
--- a/Engine/Source/Runtime/Engine/Classes/Components/PointLightComponent.h
+++ b/Engine/Source/Runtime/Engine/Classes/Components/PointLightComponent.h
@@ -40,21 +40,21 @@ class ENGINE_API UPointLightComponent : public ULocalLightComponent
 	 * Radius of light source shape.
 	 * Note that light sources shapes which intersect shadow casting geometry can cause shadowing artifacts.
 	 */
-	UPROPERTY(EditAnywhere, BlueprintReadOnly, interp, Category=Light, meta=(ShouldShowInViewport = true))
+	UPROPERTY(EditAnywhere, BlueprintReadOnly, interp, Category=Light, meta=(UIMin = "0.0", UIMax = "1000.0", ClampMax = "100000", ShouldShowInViewport = true))
 	float SourceRadius;
 
 	/**
 	* Soft radius of light source shape.
 	* Note that light sources shapes which intersect shadow casting geometry can cause shadowing artifacts.
 	*/
-	UPROPERTY(EditAnywhere, BlueprintReadOnly, interp, Category = Light)
+	UPROPERTY(EditAnywhere, BlueprintReadOnly, interp, Category = Light, meta=(UIMin = "0.0", UIMax = "1000.0", ClampMax = "100000"))
 	float SoftSourceRadius;
 
 	/** 
 	 * Length of light source shape.
 	 * Note that light sources shapes which intersect shadow casting geometry can cause shadowing artifacts.
 	 */
-	UPROPERTY(EditAnywhere, BlueprintReadOnly, interp, Category=Light)
+	UPROPERTY(EditAnywhere, BlueprintReadOnly, interp, Category=Light, meta=(UIMin = "0.0", UIMax = "1000.0", ClampMax = "100000"))
 	float SourceLength;
 
 	UFUNCTION(BlueprintCallable, Category = "Rendering|Lighting")
