Fixed mismatched copy constructor/assignment operator in TQuadTree.

diff --git a/Engine/Source/Runtime/Engine/Public/GenericQuadTree.h b/Engine/Source/Runtime/Engine/Public/GenericQuadTree.h
--- a/Engine/Source/Runtime/Engine/Public/GenericQuadTree.h
+++ b/Engine/Source/Runtime/Engine/Public/GenericQuadTree.h
@@ -38,7 +38,8 @@ public:
 
 	void Serialize(FArchive& Ar);
 
-	TreeType& operator=(const TreeType& Other);
+	TQuadTree(const TQuadTree&);
+	TQuadTree& operator=(const TQuadTree& Other);
 
 	~TQuadTree();
 
@@ -112,7 +113,13 @@ private:
 };
 
 template <typename ElementType, int32 NodeCapacity /*= 4*/>
-typename TQuadTree<ElementType, NodeCapacity>::TreeType& TQuadTree<ElementType, NodeCapacity>::operator=(const TreeType& Other)
+TQuadTree<ElementType, NodeCapacity>::TQuadTree(const TQuadTree& Other)
+{
+	Other.Duplicate(*this);
+}
+
+template <typename ElementType, int32 NodeCapacity /*= 4*/>
+TQuadTree<ElementType, NodeCapacity>& TQuadTree<ElementType, NodeCapacity>::operator=(const TQuadTree& Other)
 {
 	Other.Duplicate(*this);
 	return *this;
