Iris
- Fixed bug in RepMovementNetSerialzier which would not properly set bRepPhysics flag

diff --git a/Engine/Source/Runtime/Engine/Private/RepMovementNetSerializer.cpp b/Engine/Source/Runtime/Engine/Private/RepMovementNetSerializer.cpp
--- a/Engine/Source/Runtime/Engine/Private/RepMovementNetSerializer.cpp
+++ b/Engine/Source/Runtime/Engine/Private/RepMovementNetSerializer.cpp
@@ -559,8 +559,8 @@ void FRepMovementNetSerializer::Dequantize(FNetSerializationContext& Context, co
 	const QuantizedType& Source = *reinterpret_cast<const QuantizedType*>(Args.Source);
 	SourceType& Target = *reinterpret_cast<SourceType*>(Args.Target);
 
-	Target.bSimulatedPhysicSleep = (Source.Flags & Flag_SimulatedPhysicSleep);
-	Target.bRepPhysics = (Source.Flags & Flag_RepPhysics);
+	Target.bSimulatedPhysicSleep = (Source.Flags & Flag_SimulatedPhysicSleep) ? 1 : 0;
+	Target.bRepPhysics = (Source.Flags & Flag_RepPhysics) ? 1 : 0;
 	Target.ServerFrame = (Source.Flags & Flag_ServerFrameIsPresent ? Source.ServerFrame : 0);
 	Target.ServerPhysicsHandle = (Source.Flags & Flag_ServerPhysicsHandleIsPresent ? Source.ServerPhysicsHandle : INDEX_NONE);
 
