Fixed mismatched copy constructor/assignment operator in FMeshBoneInfo.

diff --git a/Engine/Source/Runtime/Engine/Public/ReferenceSkeleton.h b/Engine/Source/Runtime/Engine/Public/ReferenceSkeleton.h
--- a/Engine/Source/Runtime/Engine/Public/ReferenceSkeleton.h
+++ b/Engine/Source/Runtime/Engine/Public/ReferenceSkeleton.h
@@ -33,14 +33,6 @@ struct FMeshBoneInfo
 #endif
 	{}
 
-	FMeshBoneInfo(const FMeshBoneInfo& Other)
-		: Name(Other.Name)
-		, ParentIndex(Other.ParentIndex)
-#if WITH_EDITORONLY_DATA
-		, ExportName(Other.ExportName)
-#endif
-	{}
-
 	bool operator==(const FMeshBoneInfo& B) const
 	{
 		return(Name == B.Name);
