Avoid using IterateDirectoryRecursively to gather external actors in ULevel::GetOnDiskExternalActorPackages, rely on the asset registry instead.

diff --git a/Engine/Source/Runtime/Engine/Private/Level.cpp b/Engine/Source/Runtime/Engine/Private/Level.cpp
--- a/Engine/Source/Runtime/Engine/Private/Level.cpp
+++ b/Engine/Source/Runtime/Engine/Private/Level.cpp
@@ -134,7 +134,7 @@ void FLevelActorFoldersHelper::RenameFolder(ULevel* InLevel, const FFolder& InOl
 		ActorFolder->SetLabel(FolderLabel);
 		check(ActorFolder->GetPath().IsEqual(InNewFolder.GetPath(), ENameCase::CaseSensitive));
 	}
-};
+}
 
 static bool GAllowCleanupActorFolders = true;
 
@@ -155,7 +155,7 @@ void FLevelActorFoldersHelper::DeleteFolder(ULevel* InLevel, const FFolder& InFo
 			ActorFolder->MarkAsDeleted();
 		}
 	}
-};
+}
 
 #endif
 
@@ -3965,19 +3965,20 @@ TArray<FString> ULevel::GetOnDiskExternalActorPackages(const FString& ExternalAc
 	TArray<FString> ActorPackageNames;
 	if (!ExternalActorsPath.IsEmpty())
 	{
-		IFileManager::Get().IterateDirectoryRecursively(*FPackageName::LongPackageNameToFilename(ExternalActorsPath), [&ActorPackageNames](const TCHAR* FilenameOrDirectory, bool bIsDirectory)
-		{
-			if (!bIsDirectory)
-			{
-				FString Filename(FilenameOrDirectory);
-				if (Filename.EndsWith(FPackageName::GetAssetPackageExtension()))
-				{
-					ActorPackageNames.Add(FPackageName::FilenameToLongPackageName(Filename));
-				}
-			}
-			return true;
-		});
+		FARFilter Filter;
+		Filter.bIncludeOnlyOnDiskAssets = true;
+		Filter.PackagePaths.Add(*ExternalActorsPath);
+		Filter.bRecursivePaths = true;
+	
+		TArray<FAssetData> ActorAssets;
+		IAssetRegistry& AssetRegistry = FModuleManager::LoadModuleChecked<FAssetRegistryModule>(TEXT("AssetRegistry")).Get();
+		AssetRegistry.ScanSynchronous({ ExternalActorsPath }, TArray<FString>());
+		AssetRegistry.GetAssets(Filter, ActorAssets);
+
+		ActorPackageNames.Reserve(ActorAssets.Num());
+		Algo::Transform(ActorAssets, ActorPackageNames, [](const FAssetData& ActorAssetData) { return ActorAssetData.PackageName.ToString(); });
 	}
+
 	return ActorPackageNames;
 }
 
@@ -4350,7 +4351,7 @@ void ULevel::ApplyWorldOffset(const FVector& InWorldOffset, bool bWorldShift)
 void ULevel::RegisterActorForAutoReceiveInput(AActor* Actor, const int32 PlayerIndex)
 {
 	PendingAutoReceiveInputActors.Add(FPendingAutoReceiveInputActor(Actor, PlayerIndex));
-};
+}
 
 void ULevel::PushPendingAutoReceiveInput(APlayerController* InPlayerController)
 {
