Remove workaround in FExternalPackageHelper::LoadObjectsFromExternalPackages now that we can scan asset registry paths without actually invoking BP functions.

diff --git a/Engine/Source/Runtime/Engine/Public/ExternalPackageHelper.h b/Engine/Source/Runtime/Engine/Public/ExternalPackageHelper.h
--- a/Engine/Source/Runtime/Engine/Public/ExternalPackageHelper.h
+++ b/Engine/Source/Runtime/Engine/Public/ExternalPackageHelper.h
@@ -12,7 +12,6 @@
 #include "Misc/ArchiveMD5.h"
 #include "Misc/PackageName.h"
 #include "Misc/PackagePath.h"
-#include "HAL/FileManager.h"
 #include "Delegates/DelegateCombinations.h"
 
 class ENGINE_API FExternalPackageHelper
@@ -84,63 +83,30 @@ private:
 	static FString GetExternalObjectPackageInstanceName(const FString& OuterPackageName, const FString& ObjectPackageName);
 };
 
-static TArray<FString> GetOnDiskExternalObjectPackages(const FString& ExternalObjectsPath)
-{
-	TArray<FString> ObjectPackageNames;
-	if (!ExternalObjectsPath.IsEmpty())
-	{
-		IFileManager::Get().IterateDirectoryRecursively(*FPackageName::LongPackageNameToFilename(ExternalObjectsPath), [&ObjectPackageNames](const TCHAR* FilenameOrDirectory, bool bIsDirectory)
-		{
-			if (!bIsDirectory)
-			{
-				FString Filename(FilenameOrDirectory);
-				if (Filename.EndsWith(FPackageName::GetAssetPackageExtension()))
-				{
-					ObjectPackageNames.Add(FPackageName::FilenameToLongPackageName(Filename));
-				}
-			}
-			return true;
-		});
-	}
-	return ObjectPackageNames;
-}
-
 template<typename T>
 void FExternalPackageHelper::LoadObjectsFromExternalPackages(UObject* InOuter, TFunctionRef<void(T*)> Operation)
 {
 	const FString ExternalObjectsPath = FExternalPackageHelper::GetExternalObjectsPath(InOuter->GetPackage(), FString(), /*bTryUsingPackageLoadedPath*/ true);
 	TArray<FString> ObjectPackageNames;
 
-	// Test if function is called while inside postload
-	if (!FUObjectThreadContext::Get().IsRoutingPostLoad)
-	{
-		// Do a synchronous scan of the world external objects path.			
-		IAssetRegistry& AssetRegistry = FModuleManager::LoadModuleChecked<FAssetRegistryModule>(TEXT("AssetRegistry")).Get();
-		AssetRegistry.ScanPathsSynchronous({ ExternalObjectsPath }, /*bForceRescan*/false, /*bIgnoreDenyListScanFilters*/false);
-
-		FARFilter Filter;
-		Filter.bRecursivePaths = true;
-		Filter.bIncludeOnlyOnDiskAssets = true;
-		Filter.ClassPaths.Add(T::StaticClass()->GetClassPathName());
-		Filter.bRecursiveClasses = true;
-		Filter.PackagePaths.Add(*ExternalObjectsPath);
-
-		TArray<FAssetData> Assets;
-		AssetRegistry.GetAssets(Filter, Assets);
-
-		ObjectPackageNames.Reserve(Assets.Num());
-		for (const FAssetData& Asset : Assets)
-		{
-			ObjectPackageNames.Add(Asset.PackageName.ToString());
-		}
-	}
-	// If so, we can't use AssetRegistry's ScanPathsSynchronous as it might update AssetRegistry object tags which can call BP events 
-	// (which is not allowed while inside postload). Instead, get all external object packages on disk.
-	// @todo_ow : This is far from ideal because once we move External Actors to share the same folder as external objects (or any new object we will externalize),
-	//            we will discover and iterate on all external packages regarless of the class type we actually want to load.
-	else
+	// Do a synchronous scan of the world external objects path.			
+	IAssetRegistry& AssetRegistry = FModuleManager::LoadModuleChecked<FAssetRegistryModule>(TEXT("AssetRegistry")).Get();
+	AssetRegistry.ScanSynchronous({ ExternalObjectsPath }, TArray<FString>());
+
+	FARFilter Filter;
+	Filter.bRecursivePaths = true;
+	Filter.bIncludeOnlyOnDiskAssets = true;
+	Filter.ClassPaths.Add(T::StaticClass()->GetClassPathName());
+	Filter.bRecursiveClasses = true;
+	Filter.PackagePaths.Add(*ExternalObjectsPath);
+
+	TArray<FAssetData> Assets;
+	AssetRegistry.GetAssets(Filter, Assets);
+
+	ObjectPackageNames.Reserve(Assets.Num());
+	for (const FAssetData& Asset : Assets)
 	{
-		 ObjectPackageNames = GetOnDiskExternalObjectPackages(ExternalObjectsPath);
+		ObjectPackageNames.Add(Asset.PackageName.ToString());
 	}
 
 	FLinkerInstancingContext InstancingContext;
