Animation editor will now correctly apply preview mesh from a compatible skeleton #jira UE-185278, #rb Thomas.Sarkanen

diff --git a/Engine/Source/Runtime/Engine/Private/AnimationAsset.cpp b/Engine/Source/Runtime/Engine/Private/AnimationAsset.cpp
--- a/Engine/Source/Runtime/Engine/Private/AnimationAsset.cpp
+++ b/Engine/Source/Runtime/Engine/Private/AnimationAsset.cpp
@@ -387,7 +387,7 @@ USkeletalMesh* UAnimationAsset::GetPreviewMesh(bool bFindIfNotSet)
 #if WITH_EDITORONLY_DATA
 	USkeletalMesh* PreviewMesh = PreviewSkeletalMesh.LoadSynchronous();
 	// if somehow skeleton changes, just nullify it. 
-	if (PreviewMesh && PreviewMesh->GetSkeleton() != Skeleton)
+	if (PreviewMesh && !PreviewMesh->GetSkeleton()->IsCompatibleForEditor(Skeleton))
 	{
 		PreviewMesh = nullptr;
 		SetPreviewMesh(nullptr);
