World Partition: support patching class descriptors during project upload staging.

diff --git a/Engine/Source/Runtime/Engine/Private/WorldPartition/WorldPartitionActorDesc.cpp b/Engine/Source/Runtime/Engine/Private/WorldPartition/WorldPartitionActorDesc.cpp
--- a/Engine/Source/Runtime/Engine/Private/WorldPartition/WorldPartitionActorDesc.cpp
+++ b/Engine/Source/Runtime/Engine/Private/WorldPartition/WorldPartitionActorDesc.cpp
@@ -42,6 +42,13 @@ FWorldPartitionActorDesc::FWorldPartitionActorDesc()
 	, UnloadedReason(nullptr)
 {}
 
+FGuid FWorldPartitionActorDesc::GetDefaultActorDescGuid(FString ClassPath)
+{
+	FArchiveMD5 MD5Ar;
+	MD5Ar << ClassPath;
+	return ArMD5.GetGuidFromHash();
+}
+
 void FWorldPartitionActorDesc::Init(const AActor* InActor)
 {	
 	check(IsValid(InActor));
@@ -62,13 +69,8 @@ void FWorldPartitionActorDesc::Init(const AActor* InActor)
 	{
 		check(!InActor->IsPackageExternal());
 		check(!InActor->GetActorGuid().IsValid());
-
-		FArchiveMD5 ArMD5;
-		FString ClassPath = BaseClass.IsValid() ? BaseClass.ToString() : NativeClass.ToString();
-		ArMD5 << ClassPath;
-		
-		Guid = ArMD5.GetGuidFromHash();
-
+		const FString ClassPath = BaseClass.IsValid() ? BaseClass.ToString() : NativeClass.ToString();
+		Guid = GetDefaultActorDescGuid(ClassPath);
 		bIsDefaultActorDesc = true;
 	}
 	else
diff --git a/Engine/Source/Runtime/Engine/Public/WorldPartition/WorldPartitionActorDesc.h b/Engine/Source/Runtime/Engine/Public/WorldPartition/WorldPartitionActorDesc.h
--- a/Engine/Source/Runtime/Engine/Public/WorldPartition/WorldPartitionActorDesc.h
+++ b/Engine/Source/Runtime/Engine/Public/WorldPartition/WorldPartitionActorDesc.h
@@ -78,6 +78,8 @@ class ENGINE_API FWorldPartitionActorDesc
 	friend class FAssetRootPackagePatcher;
 	friend class FActorDescArchive;
 
+	static FGuid GetDefaultActorDescGuid(FString ClassPath);
+
 public:
 	struct FContainerInstance
 	{
