Expose SoundNode properties to BP for editor scripting usage

diff --git a/Engine/Source/Runtime/Engine/Classes/Sound/SoundCue.h b/Engine/Source/Runtime/Engine/Classes/Sound/SoundCue.h
--- a/Engine/Source/Runtime/Engine/Classes/Sound/SoundCue.h
+++ b/Engine/Source/Runtime/Engine/Classes/Sound/SoundCue.h
@@ -91,7 +91,7 @@ class ENGINE_API USoundCue : public USoundBase
 {
 	GENERATED_UCLASS_BODY()
 
-	UPROPERTY()
+	UPROPERTY(BlueprintReadOnly, Category = Sound)
 	TObjectPtr<USoundNode> FirstNode;
 
 	/* Base volume multiplier */
diff --git a/Engine/Source/Runtime/Engine/Classes/Sound/SoundNode.h b/Engine/Source/Runtime/Engine/Classes/Sound/SoundNode.h
--- a/Engine/Source/Runtime/Engine/Classes/Sound/SoundNode.h
+++ b/Engine/Source/Runtime/Engine/Classes/Sound/SoundNode.h
@@ -60,7 +60,7 @@ class ENGINE_API USoundNode : public UObject
 
 	static const int32 MAX_ALLOWED_CHILD_NODES = 32;
 
-	UPROPERTY()
+	UPROPERTY(BlueprintReadOnly, Category = "SoundNode")
 	TArray<TObjectPtr<class USoundNode>> ChildNodes;
 
 #if WITH_EDITORONLY_DATA
diff --git a/Engine/Source/Runtime/Engine/Classes/Sound/SoundNodeModulatorContinuous.h b/Engine/Source/Runtime/Engine/Classes/Sound/SoundNodeModulatorContinuous.h
--- a/Engine/Source/Runtime/Engine/Classes/Sound/SoundNodeModulatorContinuous.h
+++ b/Engine/Source/Runtime/Engine/Classes/Sound/SoundNodeModulatorContinuous.h
@@ -22,7 +22,7 @@ enum ModulationParamMode : int
 	MPM_MAX,
 };
 
-USTRUCT()
+USTRUCT(BlueprintType)
 struct FModulatorContinuousParams
 {
 	GENERATED_USTRUCT_BODY()
@@ -38,7 +38,7 @@ struct FModulatorContinuousParams
 	}
 
 	/** The name of the sound instance parameter that specifies the current value. */
-	UPROPERTY(EditAnywhere, Category=ModulatorContinousParameters)
+	UPROPERTY(EditAnywhere, BlueprintReadOnly, Category=ModulatorContinousParameters)
 	FName ParameterName;
 
 	/** The default value to be used if the parameter is not found. */
