fix bad merge

diff --git a/Engine/Source/Runtime/Engine/Private/WorldPartition/WorldPartitionActorDesc.cpp b/Engine/Source/Runtime/Engine/Private/WorldPartition/WorldPartitionActorDesc.cpp
--- a/Engine/Source/Runtime/Engine/Private/WorldPartition/WorldPartitionActorDesc.cpp
+++ b/Engine/Source/Runtime/Engine/Private/WorldPartition/WorldPartitionActorDesc.cpp
@@ -44,8 +44,8 @@ FWorldPartitionActorDesc::FWorldPartitionActorDesc()
 
 FGuid FWorldPartitionActorDesc::GetDefaultActorDescGuid(FString ClassPath)
 {
-	FArchiveMD5 MD5Ar;
-	MD5Ar << ClassPath;
+	FArchiveMD5 ArMD5;
+	ArMD5 << ClassPath;
 	return ArMD5.GetGuidFromHash();
 }
 
