Fix bad Unsaved Asset count when Editing level instances:
- UEditorLevelUtils::PrivateRemoveLevelFromWorld: Avoid Modify being called on actor components and dirtying packages on remove world

#rnx

diff --git a/Engine/Source/Editor/UnrealEd/Private/EditorLevelUtils.cpp b/Engine/Source/Editor/UnrealEd/Private/EditorLevelUtils.cpp
--- a/Engine/Source/Editor/UnrealEd/Private/EditorLevelUtils.cpp
+++ b/Engine/Source/Editor/UnrealEd/Private/EditorLevelUtils.cpp
@@ -1000,7 +1000,8 @@ void UEditorLevelUtils::PrivateRemoveLevelFromWorld(ULevel* InLevel)
 	{
 		if (Actor != nullptr)
 		{
-			Actor->MarkComponentsAsPendingKill();
+			const bool bModify = false;
+			Actor->MarkComponentsAsGarbage(bModify);
 			Actor->MarkAsGarbage();
 		}
 	}
diff --git a/Engine/Source/Editor/UnrealEd/Private/EditorServer.cpp b/Engine/Source/Editor/UnrealEd/Private/EditorServer.cpp
--- a/Engine/Source/Editor/UnrealEd/Private/EditorServer.cpp
+++ b/Engine/Source/Editor/UnrealEd/Private/EditorServer.cpp
@@ -2824,7 +2824,7 @@ bool UEditorEngine::Map_Load(const TCHAR* Str, FOutputDevice& Ar)
 						{
 							UE_LOG(LogEditorServer, Log,  TEXT("Destroying orphan Actor: %s"), *Actor->GetName() );					
 							Actor->MarkAsGarbage();
-							Actor->MarkComponentsAsPendingKill();
+							Actor->MarkComponentsAsGarbage();
 						}
 					}
 					UE_LOG(LogEditorServer, Log,  TEXT("Finished looking for orphan Actors (%3.3lf secs)"), FPlatformTime::Seconds() - StartTime );
diff --git a/Engine/Source/Runtime/Engine/Classes/GameFramework/Actor.h b/Engine/Source/Runtime/Engine/Classes/GameFramework/Actor.h
--- a/Engine/Source/Runtime/Engine/Classes/GameFramework/Actor.h
+++ b/Engine/Source/Runtime/Engine/Classes/GameFramework/Actor.h
@@ -2973,9 +2973,16 @@ public:
 	/** Debug rendering to visualize the component tree for this actor. */
 	void DrawDebugComponents(FColor const& BaseColor=FColor::White) const;
 
-	/** Called to mark all components as pending kill when the actor is being destroyed */
-	virtual void MarkComponentsAsPendingKill();
+	
+	UE_DEPRECATED(5.3, "Use MarkComponentsAsGarbage instead.")
+	virtual void MarkComponentsAsPendingKill() { MarkComponentsAsGarbage(); }
 
+	/** Called to mark all components as garbage when the actor is being destroyed
+	 *
+	 *  @param bModify if True, Modify will be called on actor before marking components
+	 */
+	virtual void MarkComponentsAsGarbage(bool bModify = true);
+	
 	/**
 	 * Returns true if this actor has begun the destruction process.
 	 * This is set to true in UWorld::DestroyActor, after the network connection has been closed but before any other shutdown has been performed.
diff --git a/Engine/Source/Runtime/Engine/Private/Actor.cpp b/Engine/Source/Runtime/Engine/Private/Actor.cpp
--- a/Engine/Source/Runtime/Engine/Private/Actor.cpp
+++ b/Engine/Source/Runtime/Engine/Private/Actor.cpp
@@ -5452,15 +5452,15 @@ bool AActor::HasValidRootComponent()
 	return (RootComponent != nullptr && RootComponent->IsRegistered()); 
 }
 
-void AActor::MarkComponentsAsPendingKill()
+void AActor::MarkComponentsAsGarbage(bool bModify)
 {
-	// Iterate components and mark them all as pending kill.
+	// Iterate components and mark them all as garbage.
 	TInlineComponentArray<UActorComponent*> Components(this);
 
 	for (UActorComponent* Component : Components)
 	{
 		// Modify component so undo/ redo works in the editor.
-		if (GIsEditor)
+		if (bModify && GIsEditor)
 		{
 			Component->Modify();
 		}
diff --git a/Engine/Source/Runtime/Engine/Private/LevelActor.cpp b/Engine/Source/Runtime/Engine/Private/LevelActor.cpp
--- a/Engine/Source/Runtime/Engine/Private/LevelActor.cpp
+++ b/Engine/Source/Runtime/Engine/Private/LevelActor.cpp
@@ -1002,7 +1002,7 @@ bool UWorld::DestroyActor( AActor* ThisActor, bool bNetForce, bool bShouldModify
 	// Mark the actor and its direct components as pending kill.
 	ThisActor->MarkAsGarbage();
 	ThisActor->MarkPackageDirty();
-	ThisActor->MarkComponentsAsPendingKill();
+	ThisActor->MarkComponentsAsGarbage();
 
 	// Unregister the actor's tick function
 	const bool bRegisterTickFunctions = false;
diff --git a/Engine/Source/Runtime/NavigationSystem/Private/NavigationData.cpp b/Engine/Source/Runtime/NavigationSystem/Private/NavigationData.cpp
--- a/Engine/Source/Runtime/NavigationSystem/Private/NavigationData.cpp
+++ b/Engine/Source/Runtime/NavigationSystem/Private/NavigationData.cpp
@@ -518,7 +518,7 @@ void ANavigationData::CleanUpAndMarkPendingKill()
 			World->RemoveNetworkActor(this);
 		}
 		MarkAsGarbage();
-		MarkComponentsAsPendingKill();
+		MarkComponentsAsGarbage();
 	}
 }
 
