World Partition Class Descriptors: fixed delta serialization bug against default actor class descriptor when serializing with a missing class.

diff --git a/Engine/Plugins/Experimental/Water/Source/Runtime/Private/WaterZoneActorDesc.h b/Engine/Plugins/Experimental/Water/Source/Runtime/Private/WaterZoneActorDesc.h
--- a/Engine/Plugins/Experimental/Water/Source/Runtime/Private/WaterZoneActorDesc.h
+++ b/Engine/Plugins/Experimental/Water/Source/Runtime/Private/WaterZoneActorDesc.h
@@ -16,6 +16,7 @@ public:
 
 	int32 GetOverlapPriority() const { return OverlapPriority; }
 protected:
+	virtual uint32 GetSizeOf() const override { return sizeof(FWaterZoneActorDesc); }
 	virtual void Serialize(FArchive& Ar) override;
 
 	int32 OverlapPriority;
diff --git a/Engine/Source/Runtime/Engine/Private/WorldPartition/WorldPartitionActorDescArchive.cpp b/Engine/Source/Runtime/Engine/Private/WorldPartition/WorldPartitionActorDescArchive.cpp
--- a/Engine/Source/Runtime/Engine/Private/WorldPartition/WorldPartitionActorDescArchive.cpp
+++ b/Engine/Source/Runtime/Engine/Private/WorldPartition/WorldPartitionActorDescArchive.cpp
@@ -13,6 +13,7 @@
 FActorDescArchive::FActorDescArchive(FArchive& InArchive, FWorldPartitionActorDesc* InActorDesc)
 	: FArchiveProxy(InArchive)
 	, ActorDesc(InActorDesc)
+	, bIsMissingClassDesc(false)
 {
 	check(InArchive.IsPersistent());
 
@@ -65,8 +66,10 @@ PRAGMA_ENABLE_DEPRECATION_WARNINGS
 	{
 		if (IsLoading())
 		{
+			bIsMissingClassDesc = true;
+
 			ClassDesc = ClassDescRegistry.GetClassDescDefault(FTopLevelAssetPath(TEXT("/Script/Engine.Actor")));
-			check(ClassDesc);
+			check(ClassDesc);			
 
 			UE_LOG(LogWorldPartition, Log, TEXT("Can't find class descriptor '%s' for loading '%s', using '%s'"), *ClassPath.ToString(), *InActorDesc->GetActorSoftPath().ToString(), *ClassDesc->GetActorSoftPath().ToString());
 		}
@@ -75,5 +78,7 @@ PRAGMA_ENABLE_DEPRECATION_WARNINGS
 			UE_LOG(LogWorldPartition, Log, TEXT("Can't find class descriptor '%s' for saving '%s'"), *ClassPath.ToString(), *InActorDesc->GetActorSoftPath().ToString());
 		}
 	}
+
+	ClassDescSizeof = ClassDesc ? ClassDesc->GetSizeOf() : 0;
 }
 #endif
\ No newline at end of file
diff --git a/Engine/Source/Runtime/Engine/Public/WorldPartition/ActorPartition/PartitionActorDesc.h b/Engine/Source/Runtime/Engine/Public/WorldPartition/ActorPartition/PartitionActorDesc.h
--- a/Engine/Source/Runtime/Engine/Public/WorldPartition/ActorPartition/PartitionActorDesc.h
+++ b/Engine/Source/Runtime/Engine/Public/WorldPartition/ActorPartition/PartitionActorDesc.h
@@ -22,6 +22,7 @@ public:
 protected:
 	virtual void Init(const AActor* InActor) override;
 	virtual bool Equals(const FWorldPartitionActorDesc* Other) const override;
+	virtual uint32 GetSizeOf() const override { return sizeof(FPartitionActorDesc); }
 	virtual void Serialize(FArchive& Ar) override;
 	virtual FBox GetEditorBounds() const override;
 	virtual void TransferWorldData(const FWorldPartitionActorDesc* From) override;
diff --git a/Engine/Source/Runtime/Engine/Public/WorldPartition/DataLayer/WorldDataLayersActorDesc.h b/Engine/Source/Runtime/Engine/Public/WorldPartition/DataLayer/WorldDataLayersActorDesc.h
--- a/Engine/Source/Runtime/Engine/Public/WorldPartition/DataLayer/WorldDataLayersActorDesc.h
+++ b/Engine/Source/Runtime/Engine/Public/WorldPartition/DataLayer/WorldDataLayersActorDesc.h
@@ -71,6 +71,7 @@ protected:
 	//~ Begin FWorldPartitionActorDesc Interface.
 	virtual void Init(const AActor* InActor) override;
 	virtual bool Equals(const FWorldPartitionActorDesc* Other) const override;
+	virtual uint32 GetSizeOf() const override { return sizeof(FWorldDataLayersActorDesc); }
 	virtual void Serialize(FArchive& Ar) override;
 	virtual bool IsResaveNeeded() const override { return !IsValid(); }
 	//~ End FWorldPartitionActorDesc Interface.
diff --git a/Engine/Source/Runtime/Engine/Public/WorldPartition/HLOD/HLODActorDesc.h b/Engine/Source/Runtime/Engine/Public/WorldPartition/HLOD/HLODActorDesc.h
--- a/Engine/Source/Runtime/Engine/Public/WorldPartition/HLOD/HLODActorDesc.h
+++ b/Engine/Source/Runtime/Engine/Public/WorldPartition/HLOD/HLODActorDesc.h
@@ -37,6 +37,7 @@ protected:
 	//~ Begin FWorldPartitionActorDesc Interface.
 	virtual void Init(const AActor* InActor) override;
 	virtual bool Equals(const FWorldPartitionActorDesc* Other) const override;
+	virtual uint32 GetSizeOf() const override { return sizeof(FHLODActorDesc); }
 	virtual void Serialize(FArchive& Ar) override;
 	virtual bool IsRuntimeRelevant(const FActorContainerID& InContainerID) const override { return !bIsForcedNonSpatiallyLoaded; }
 	virtual bool ShouldValidateRuntimeGrid() const override { return false; }
diff --git a/Engine/Source/Runtime/Engine/Public/WorldPartition/Landscape/LandscapeActorDesc.h b/Engine/Source/Runtime/Engine/Public/WorldPartition/Landscape/LandscapeActorDesc.h
--- a/Engine/Source/Runtime/Engine/Public/WorldPartition/Landscape/LandscapeActorDesc.h
+++ b/Engine/Source/Runtime/Engine/Public/WorldPartition/Landscape/LandscapeActorDesc.h
@@ -19,6 +19,7 @@ public:
 	virtual const FGuid& GetSceneOutlinerParent() const override;
 
 protected:
+	virtual uint32 GetSizeOf() const override { return sizeof(FLandscapeActorDesc); }
 	virtual void Serialize(FArchive& Ar) override;
 	virtual FBox GetEditorBounds() const override;
 
diff --git a/Engine/Source/Runtime/Engine/Public/WorldPartition/LevelInstance/LevelInstanceActorDesc.h b/Engine/Source/Runtime/Engine/Public/WorldPartition/LevelInstance/LevelInstanceActorDesc.h
--- a/Engine/Source/Runtime/Engine/Public/WorldPartition/LevelInstance/LevelInstanceActorDesc.h
+++ b/Engine/Source/Runtime/Engine/Public/WorldPartition/LevelInstance/LevelInstanceActorDesc.h
@@ -37,6 +37,7 @@ protected:
 	virtual void Init(const FWorldPartitionActorDescInitData& DescData) override;
 	virtual bool Equals(const FWorldPartitionActorDesc* Other) const override;
 	virtual void TransferFrom(const FWorldPartitionActorDesc* From) override;
+	virtual uint32 GetSizeOf() const override { return sizeof(FLevelInstanceActorDesc); }
 	virtual void Serialize(FArchive& Ar) override;
 	virtual void SetContainer(UActorDescContainer* InContainer, UWorld* InWorldContext) override;
 
diff --git a/Engine/Source/Runtime/Engine/Public/WorldPartition/WorldPartitionActorDesc.h b/Engine/Source/Runtime/Engine/Public/WorldPartition/WorldPartitionActorDesc.h
--- a/Engine/Source/Runtime/Engine/Public/WorldPartition/WorldPartitionActorDesc.h
+++ b/Engine/Source/Runtime/Engine/Public/WorldPartition/WorldPartitionActorDesc.h
@@ -281,6 +281,8 @@ protected:
 		BoundsExtent = From->BoundsExtent;
 	}
 
+	virtual uint32 GetSizeOf() const { return sizeof(FWorldPartitionActorDesc); }
+
 	virtual void Serialize(FArchive& Ar);
 
 	// Persistent
diff --git a/Engine/Source/Runtime/Engine/Public/WorldPartition/WorldPartitionActorDescArchive.h b/Engine/Source/Runtime/Engine/Public/WorldPartition/WorldPartitionActorDescArchive.h
--- a/Engine/Source/Runtime/Engine/Public/WorldPartition/WorldPartitionActorDescArchive.h
+++ b/Engine/Source/Runtime/Engine/Public/WorldPartition/WorldPartitionActorDescArchive.h
@@ -39,11 +39,19 @@ public:
 
 			Ar.UsingCustomVersion(FFortniteMainBranchObjectVersion::GUID);
 
-			auto GetClassDefaultValue = [&V, &ActorDescAr]() -> const T&
+			auto GetClassDefaultValue = [&V, &ActorDescAr]() -> const T*
 			{
 				const UPTRINT PropertyOffset = (UPTRINT)&V.Value - *(UPTRINT*)&ActorDescAr.ActorDesc;
-				const T& RefValue = *(const T*)(*(UPTRINT*)&ActorDescAr.ClassDesc + PropertyOffset);
-				return RefValue;
+
+				if (PropertyOffset < ActorDescAr.ClassDescSizeof)
+				{
+					check((PropertyOffset + sizeof(V)) <= ActorDescAr.ClassDescSizeof);
+					const T* RefValue = (const T*)(*(UPTRINT*)&ActorDescAr.ClassDesc + PropertyOffset);
+					return RefValue;
+				}
+
+				check(ActorDescAr.bIsMissingClassDesc);
+				return nullptr;
 			};
 
 			uint8 bSerialize = 1;
@@ -52,7 +60,11 @@ public:
 			{
 				if (Ar.IsSaving() && ActorDescAr.ClassDesc)
 				{
-					bSerialize = (V.Value != GetClassDefaultValue()) ? 1 : 0;
+					// When saving, we expect the class descriptor to be the exact type as what we are serializing.
+					const T* ClassDefaultValue = GetClassDefaultValue();
+					check(ClassDefaultValue);
+
+					bSerialize = (V.Value != *ClassDefaultValue) ? 1 : 0;
 				}
 
 				Ar << bSerialize;
@@ -64,7 +76,11 @@ public:
 			}
 			else if (Ar.IsLoading())
 			{
-				V.Value = GetClassDefaultValue();
+				// When loading, we need to handle a different class descriptor in case of missing classes, etc.
+				if (const T* ClassDefaultValue = GetClassDefaultValue())
+				{
+					V.Value = *ClassDefaultValue;
+				}
 			}
 
 			return Ar;
@@ -75,6 +91,8 @@ public:
 
 	FWorldPartitionActorDesc* ActorDesc;
 	const FWorldPartitionActorDesc* ClassDesc;
+	uint32 ClassDescSizeof;
+	bool bIsMissingClassDesc;
 };
 
 template <typename T>
