Use GPU distance cull settings for ISM.

diff --git a/Engine/Source/Runtime/Engine/Classes/Engine/InstancedStaticMesh.h b/Engine/Source/Runtime/Engine/Classes/Engine/InstancedStaticMesh.h
--- a/Engine/Source/Runtime/Engine/Classes/Engine/InstancedStaticMesh.h
+++ b/Engine/Source/Runtime/Engine/Classes/Engine/InstancedStaticMesh.h
@@ -591,6 +591,8 @@ public:
 	 */
 	virtual HHitProxy* CreateHitProxies(UPrimitiveComponent* Component,TArray<TRefCountPtr<HHitProxy> >& OutHitProxies) override;
 
+	virtual bool GetInstanceDrawDistanceMinMax(FVector2f& OutDistanceMinMax) const override;
+
 	virtual bool IsDetailMesh() const override
 	{
 		return true;
diff --git a/Engine/Source/Runtime/Engine/Private/InstancedStaticMesh.cpp b/Engine/Source/Runtime/Engine/Private/InstancedStaticMesh.cpp
--- a/Engine/Source/Runtime/Engine/Private/InstancedStaticMesh.cpp
+++ b/Engine/Source/Runtime/Engine/Private/InstancedStaticMesh.cpp
@@ -2119,6 +2119,20 @@ HHitProxy* FInstancedStaticMeshSceneProxy::CreateHitProxies(UPrimitiveComponent*
 	return FStaticMeshSceneProxy::CreateHitProxies(Component, OutHitProxies);
 }
 
+bool FInstancedStaticMeshSceneProxy::GetInstanceDrawDistanceMinMax(FVector2f& OutDistanceMinMax) const
+{
+	if (UserData_AllInstances.EndCullDistance > 0)
+	{
+		OutDistanceMinMax = FVector2f(0.0f, float(UserData_AllInstances.EndCullDistance));
+		return true;
+	}
+	else
+	{
+		OutDistanceMinMax = FVector2f(0.0f);
+		return false;
+	}
+}
+
 #if RHI_RAYTRACING
 bool FInstancedStaticMeshSceneProxy::HasRayTracingRepresentation() const
 {
