Iris - Client will disconnect on bunch error when Iris is enabled.

diff --git a/Engine/Source/Runtime/Engine/Private/NetConnection.cpp b/Engine/Source/Runtime/Engine/Private/NetConnection.cpp
--- a/Engine/Source/Runtime/Engine/Private/NetConnection.cpp
+++ b/Engine/Source/Runtime/Engine/Private/NetConnection.cpp
@@ -3559,6 +3559,15 @@ void UNetConnection::ReceivedPacket( FBitReader& Reader, bool bIsReinjectedPacke
 					else
 					{
 						UE_LOG(LogNetTraffic, Error, TEXT("Received corrupted packet data with SequenceId: %d from server %s"), InPacketId, *LowLevelGetRemoteAddress());
+#if UE_WITH_IRIS
+						// If Iris reports errors they are unrecoverable.
+						if (Driver->GetReplicationSystem() != nullptr)
+						{
+							ensureAlwaysMsgf(false, TEXT("Received corrupted packet data with SequenceId: %d from server."), InPacketId);
+							Close(AddToAndConsumeChainResultPtr(Bunch.ExtendedError, ENetCloseResult::CorruptData));
+							return;
+						}
+#endif
 					}
 				}
 				// In replay, if the bunch generated an error but the channel isn't actually open, clean it up so the channel index remains free
