Fix Substrate missing tangent frame when using parameter blending.

diff --git a/Engine/Shaders/Private/Strata/Strata.ush b/Engine/Shaders/Private/Strata/Strata.ush
--- a/Engine/Shaders/Private/Strata/Strata.ush
+++ b/Engine/Shaders/Private/Strata/Strata.ush
@@ -2219,7 +2219,7 @@ FStrataData FStrataTree::StrataVerticalLayering(FStrataData Top, FStrataData Bas
 	Operators[OperatorIndex].LeftIndex				= Top.OperatorIndex;
 	Operators[OperatorIndex].RightIndex				= Base.OperatorIndex;
 
-	Operators[Top.OperatorIndex].ParentIndex			= OperatorIndex;
+	Operators[Top.OperatorIndex].ParentIndex		= OperatorIndex;
 	Operators[Base.OperatorIndex].ParentIndex		= OperatorIndex;
 
 	return StrataData;
@@ -2276,7 +2276,7 @@ uint StrataMergeSSSType(uint SSSType0, uint SSSType1)
 //  - Slab node with SSSprofile cannot be blendend. Those will have to be specifically handled by the user and not mixed with any other slab if it can be demoted to parameter blending
 
 // ==> NOTE: Always pair with the compiler behavior in StrataCompilationInfoHorizontalMixingParamBlend
-FStrataData StrataHorizontalMixingParameterBlending(FStrataData Background, FStrataData Foreground, float HorizontalMixValue, float NormalMixCodeChunk, uint NewNormalIndex)
+FStrataData StrataHorizontalMixingParameterBlending(FStrataData Background, FStrataData Foreground, float HorizontalMixValue, float NormalMixCodeChunk, uint NewNormalIndex, inout uint OutSharedLocalBasisTypes)
 {
 	const float ForegroundOtherMixFactor = saturate(NormalMixCodeChunk);	// Horizontal mixing value combined with coverage
 	const float ForegroundLightingMixFactor = saturate(HorizontalMixValue);
@@ -2345,6 +2345,11 @@ FStrataData StrataHorizontalMixingParameterBlending(FStrataData Background, FStr
 		}
 	}
 
+	// Request a tangent basis if the blended BSDF has anisotropy
+	if (!STRATA_FASTPATH && BSDF_GETHASANISOTROPY(ResultBSDF))
+	{
+		StrataRequestSharedLocalBasisTangent(OutSharedLocalBasisTypes, NewNormalIndex);
+	}
 
 #if STRATA_INLINE_SHADING
 	// At this stage we do not know if this strata BSDF is going to be using simple volume lighting so always store the MFP on the side to not override SSSPROFILE data.
diff --git a/Engine/Shaders/Private/Strata/StrataVisualize.usf b/Engine/Shaders/Private/Strata/StrataVisualize.usf
--- a/Engine/Shaders/Private/Strata/StrataVisualize.usf
+++ b/Engine/Shaders/Private/Strata/StrataVisualize.usf
@@ -154,7 +154,7 @@ void PrintBSDF(inout FShaderPrintContext Context, inout FStrataAddressing Strata
 
 			Print(Context, TEXT("Slab "), FontBSDFType);
 			Newline(Context, RectMax);
-			Print(Context, TEXT(" NormalID Aniso TopLayer Scattering SimpleVolume IsThin WeightV "), FontBSDFStateName);
+			Print(Context, TEXT(" NormalID BasisType Aniso TopLayer Scattering SimpleVolume IsThin WeightV "), FontBSDFStateName);
 			if (!bGreyWeightV)
 			{
 				Print(Context, TEXT("            "), FontBSDFStateName);
@@ -170,8 +170,18 @@ void PrintBSDF(inout FShaderPrintContext Context, inout FStrataAddressing Strata
 			}
 			Newline(Context, RectMax);
 
+			const uint BasisType  = StrataGetSharedLocalBasisType(Header.SharedLocalBasesTypes_PackedHeader, BSDF_GETSHAREDLOCALBASISID(BSDF));
+
 			Print(Context, TEXT(" "), FontSilver);
 			Print(Context, BSDF_GETSHAREDLOCALBASISID(BSDF), FontBSDF, 2, 1);	Print(Context, TEXT("       "), FontSilver);
+			if (BasisType == 0)
+			{
+				Print(Context, TEXT("Normal"), FontBSDF);	Print(Context, TEXT("    "), FontSilver);
+			}
+			else
+			{
+				Print(Context, TEXT("Tangent"), FontBSDF);	Print(Context, TEXT("   "), FontSilver);
+			}
 			PrintBool(Context, BSDF_GETHASANISOTROPY(BSDF));					Print(Context, TEXT("  "));
 			PrintBool(Context, BSDF_GETISTOPLAYER(BSDF));						Print(Context, TEXT("     "));
 			PrintBool(Context, BSDF_GETSSSTYPE(BSDF) != SSS_TYPE_INVALID);		Print(Context, TEXT("       "));
diff --git a/Engine/Source/Runtime/Engine/Private/Materials/HLSLMaterialTranslator.cpp b/Engine/Source/Runtime/Engine/Private/Materials/HLSLMaterialTranslator.cpp
--- a/Engine/Source/Runtime/Engine/Private/Materials/HLSLMaterialTranslator.cpp
+++ b/Engine/Source/Runtime/Engine/Private/Materials/HLSLMaterialTranslator.cpp
@@ -11095,7 +11095,7 @@ void FHLSLMaterialTranslator::FStrataCompilationContext::StrataEvaluateSharedLoc
 	* FStrataData BSDF1 = GetStrataSlabBSDF(... SHAREDLOCALBASIS_1, NormalCode1 ...)
 	*
 	* float3 NormalCode2 = lerp(NormalCode0, NormalCode1, mix)
-	* FStrataData BSDF2 = StrataHorizontalMixingParameterBlending(BSDF0, BSDF1, mix, NormalCode2, SHAREDLOCALBASIS_INDEX_2) // will internally create NormalCode2
+	* FStrataData BSDF2 = StrataHorizontalMixingParameterBlending(BSDF0, BSDF1, mix, NormalCode2, SHAREDLOCALBASIS_INDEX_2, SharedLocalBases.Types) // will internally create NormalCode2
 	*
 	* tParameters.SharedLocalBases.Normals[SHAREDLOCALBASIS_INDEX_2] = NormalCode2;
 	* #if MATERIAL_TANGENTSPACENORMAL
@@ -12258,13 +12258,14 @@ int32 FHLSLMaterialTranslator::StrataHorizontalMixingParameterBlending(
 		check(PromoteToOperator->Index != INDEX_NONE);
 		check(PromoteToOperator->BSDFIndex != INDEX_NONE);
 		return AddCodeChunk(
-			MCT_Strata, TEXT("Parameters.%s.PromoteParameterBlendedBSDFToOperator(StrataHorizontalMixingParameterBlending(%s, %s, %s, %s, %s), %u, %u, %u, %u)"),
+			MCT_Strata, TEXT("Parameters.%s.PromoteParameterBlendedBSDFToOperator(StrataHorizontalMixingParameterBlending(%s, %s, %s, %s, %s, Parameters.%s.Types), %u, %u, %u, %u)"),
 			*GetParametersStrataTreeName(CurrentStrataCompilationContext),
 			*GetParameterCode(Background),
 			*GetParameterCode(Foreground),
 			*GetParameterCode(HorizontalMixCodeChunk),
 			*GetParameterCode(NormalMixCodeChunk),
 			*SharedLocalBasisIndexMacro,
+			*GetParametersSharedLocalBasesName(CurrentStrataCompilationContext),
 			PromoteToOperator->Index,
 			PromoteToOperator->BSDFIndex,
 			PromoteToOperator->LayerDepth,
@@ -12273,12 +12274,13 @@ int32 FHLSLMaterialTranslator::StrataHorizontalMixingParameterBlending(
 	}
 
 	return AddCodeChunk(
-		MCT_Strata, TEXT("StrataHorizontalMixingParameterBlending(%s, %s, %s, %s, %s)"),
+		MCT_Strata, TEXT("StrataHorizontalMixingParameterBlending(%s, %s, %s, %s, %s, Parameters.%s.Types)"),
 		*GetParameterCode(Background),
 		*GetParameterCode(Foreground),
 		*GetParameterCode(HorizontalMixCodeChunk),
 		*GetParameterCode(NormalMixCodeChunk),
-		*SharedLocalBasisIndexMacro
+		*SharedLocalBasisIndexMacro,
+		*GetParametersSharedLocalBasesName(CurrentStrataCompilationContext)
 	);
 }
 
