Fix incorrect warning when cooking Substrate material in non-Substrate mode.

diff --git a/Engine/Source/Runtime/Engine/Private/Materials/HLSLMaterialTranslator.cpp b/Engine/Source/Runtime/Engine/Private/Materials/HLSLMaterialTranslator.cpp
--- a/Engine/Source/Runtime/Engine/Private/Materials/HLSLMaterialTranslator.cpp
+++ b/Engine/Source/Runtime/Engine/Private/Materials/HLSLMaterialTranslator.cpp
@@ -1234,7 +1234,9 @@ bool FHLSLMaterialTranslator::Translate()
 						((MaterialShadingModels.HasShadingModel(MSM_SubsurfaceProfile) && IsMaterialPropertyUsed(MP_CustomData0, Chunk[MP_CustomData0], FLinearColor(1, 0, 0, 0), 1))
 						|| (MaterialShadingModels.HasShadingModel(MSM_Eye) && bOpacityPropertyIsUsed));
 
-		if (!bStrataEnabled && IsModulateBlendMode(BlendMode) && MaterialShadingModels.IsLit() && !Material->IsDeferredDecal())
+		// If Substrate is enabled or if this is a Substrate material cooked/used in non-Substrate mode, 
+		// we disable this warning as Substrate supports 'colored transmittance only' mode (i.e, modulate).
+		if (!bStrataEnabled && !FrontMaterialExpr && IsModulateBlendMode(BlendMode) && MaterialShadingModels.IsLit() && !Material->IsDeferredDecal())
 		{
 			Errorf(TEXT("Dynamically lit translucency is not supported for BLEND_Modulate materials."));
 		}
