World Partition Actor Descriptors:
- Added support for FSoftObjectPath serialization through FActorDescArchive, that will automatically remap redirectors when loading.
- Converted HLOD layers and actor filters manual FSoftObjectPath serialization and redirector remapping.
- Todo: convert data layers.

diff --git a/Engine/Source/Runtime/Core/Public/UObject/FortniteMainBranchObjectVersion.h b/Engine/Source/Runtime/Core/Public/UObject/FortniteMainBranchObjectVersion.h
--- a/Engine/Source/Runtime/Core/Public/UObject/FortniteMainBranchObjectVersion.h
+++ b/Engine/Source/Runtime/Core/Public/UObject/FortniteMainBranchObjectVersion.h
@@ -310,6 +310,9 @@ struct CORE_API FFortniteMainBranchObjectVersion
 		// Serialize editor only references in world partition actor descriptors
 		WorldPartitionActorDescSerializeEditorOnlyReferences,
 
+		// Add support for soft object paths in actor descriptors
+		WorldPartitionActorDescSerializeSoftObjectPathSupport,
+
 		// -----<new versions can be added above this line>-------------------------------------------------
 		VersionPlusOne,
 		LatestVersion = VersionPlusOne - 1
diff --git a/Engine/Source/Runtime/Engine/Private/WorldPartition/Filter/WorldPartitionActorFilter.cpp b/Engine/Source/Runtime/Engine/Private/WorldPartition/Filter/WorldPartitionActorFilter.cpp
--- a/Engine/Source/Runtime/Engine/Private/WorldPartition/Filter/WorldPartitionActorFilter.cpp
+++ b/Engine/Source/Runtime/Engine/Private/WorldPartition/Filter/WorldPartitionActorFilter.cpp
@@ -3,6 +3,7 @@
 #include "WorldPartition/Filter/WorldPartitionActorFilter.h"
 #include "WorldPartition/WorldPartitionActorLoaderInterface.h"
 #include "WorldPartition/WorldPartitionHelpers.h"
+#include "UObject/FortniteMainBranchObjectVersion.h"
 
 #if WITH_EDITORONLY_DATA
 FWorldPartitionActorFilter::FOnWorldPartitionActorFilterChanged FWorldPartitionActorFilter::OnWorldPartitionActorFilterChanged;
@@ -146,6 +147,8 @@ bool FWorldPartitionActorFilter::operator==(const FWorldPartitionActorFilter& Ot
 
 bool FWorldPartitionActorFilter::Serialize(FArchive& Ar)
 {
+	Ar.UsingCustomVersion(FFortniteMainBranchObjectVersion::GUID);
+
 	uint32 DataLayerFilterCount = DataLayerFilters.Num();
 	Ar << DataLayerFilterCount;
 
@@ -155,23 +158,31 @@ bool FWorldPartitionActorFilter::Serialize(FArchive& Ar)
 
 		for (uint32 i = 0; i < DataLayerFilterCount; ++i)
 		{
-			FString AssetPath;
-			Ar << AssetPath;
+			FSoftObjectPath AssetPath;
+
+			if (Ar.CustomVer(FFortniteMainBranchObjectVersion::GUID) < FFortniteMainBranchObjectVersion::WorldPartitionActorDescSerializeSoftObjectPathSupport)
+			{
+				FString AssetPathStr;
+				Ar << AssetPathStr;
+				AssetPath = FSoftObjectPath(AssetPathStr);
+				FWorldPartitionHelpers::FixupRedirectedAssetPath(AssetPath);
+			}
+			else
+			{
+				Ar << AssetPath;
+			}
+
 			bool bIncluded;
 			Ar << bIncluded;
-
-			FSoftObjectPath SoftObjectPath(AssetPath);
-			FWorldPartitionHelpers::FixupRedirectedAssetPath(SoftObjectPath);
-			DataLayerFilters.Add(SoftObjectPath).bIncluded = bIncluded;
+			
+			DataLayerFilters.Add(AssetPath).bIncluded = bIncluded;
 		}
 	}
 	else
 	{
 		for (auto& [AssetPath, DataLayerFilter] : DataLayerFilters)
 		{
-			FString AssetPathStr = AssetPath.ToString();
-			Ar << AssetPathStr;
-			Ar << DataLayerFilter.bIncluded;
+			Ar << AssetPath << DataLayerFilter.bIncluded;
 		}
 	}
 
diff --git a/Engine/Source/Runtime/Engine/Private/WorldPartition/HLOD/HLODLayer.cpp b/Engine/Source/Runtime/Engine/Private/WorldPartition/HLOD/HLODLayer.cpp
--- a/Engine/Source/Runtime/Engine/Private/WorldPartition/HLOD/HLODLayer.cpp
+++ b/Engine/Source/Runtime/Engine/Private/WorldPartition/HLOD/HLODLayer.cpp
@@ -61,8 +61,8 @@ UHLODLayer* UHLODLayer::GetHLODLayer(const FWorldPartitionActorDescView& InActor
 {
 	check(InWorldPartition);
 
-	const FName HLODLayerName = InActorDesc.GetHLODLayer();
-	if (UHLODLayer* HLODLayer = HLODLayerName.IsNone() ? nullptr : Cast<UHLODLayer>(FSoftObjectPath(HLODLayerName.ToString()).TryLoad()))
+	const FSoftObjectPath HLODLayerPath = InActorDesc.GetHLODLayer();
+	if (UHLODLayer* HLODLayer = Cast<UHLODLayer>(HLODLayerPath.TryLoad()))
 	{
 		return HLODLayer;
 	}
diff --git a/Engine/Source/Runtime/Engine/Private/WorldPartition/LevelInstance/LevelInstanceActorDesc.cpp b/Engine/Source/Runtime/Engine/Private/WorldPartition/LevelInstance/LevelInstanceActorDesc.cpp
--- a/Engine/Source/Runtime/Engine/Private/WorldPartition/LevelInstance/LevelInstanceActorDesc.cpp
+++ b/Engine/Source/Runtime/Engine/Private/WorldPartition/LevelInstance/LevelInstanceActorDesc.cpp
@@ -83,9 +83,9 @@ void FLevelInstanceActorDesc::UpdateBounds()
 	UWorldPartitionSubsystem* WorldPartitionSubsystem = UWorld::GetSubsystem<UWorldPartitionSubsystem>(LevelInstanceContainerWorldContext.Get());
 	check(WorldPartitionSubsystem);
 
-	FTransform LevelInstancePivotOffsetTransform = FTransform(ULevel::GetLevelInstancePivotOffsetFromPackage(LevelPackage));
+	FTransform LevelInstancePivotOffsetTransform = FTransform(ULevel::GetLevelInstancePivotOffsetFromPackage(GetContainerPackage()));
 	FTransform FinalLevelTransform = LevelInstancePivotOffsetTransform * LevelInstanceTransform;
-	FBox ContainerBounds = WorldPartitionSubsystem->GetContainerBounds(LevelPackage).TransformBy(FinalLevelTransform);
+	FBox ContainerBounds = WorldPartitionSubsystem->GetContainerBounds(GetContainerPackage()).TransformBy(FinalLevelTransform);
 
 	ContainerBounds.GetCenterAndExtents(BoundsLocation, BoundsExtent);
 }
@@ -105,7 +105,7 @@ void FLevelInstanceActorDesc::RegisterContainerInstance(UWorld* InWorldContext)
 			UWorldPartitionSubsystem* WorldPartitionSubsystem = UWorld::GetSubsystem<UWorldPartitionSubsystem>(InWorldContext);
 			check(WorldPartitionSubsystem);
 
-			LevelInstanceContainer = WorldPartitionSubsystem->RegisterContainer(LevelPackage);
+			LevelInstanceContainer = WorldPartitionSubsystem->RegisterContainer(GetContainerPackage());
 			check(LevelInstanceContainer.IsValid());
 
 			// Should only be called on RegisterContainerInstance before ActorDesc is hashed
@@ -161,17 +161,17 @@ bool FLevelInstanceActorDesc::IsContainerInstanceInternal() const
 		return false;
 	}
 
-	if (LevelPackage.IsNone())
+	if (LevelPackage.IsNull())
 	{
 		return false;
 	}
 	
-	if (!ULevel::GetIsLevelUsingExternalActorsFromPackage(LevelPackage))
+	if (!ULevel::GetIsLevelUsingExternalActorsFromPackage(GetContainerPackage()))
 	{
 		return false;
 	}
 
-	return ULevelInstanceSubsystem::CanUsePackage(LevelPackage);
+	return ULevelInstanceSubsystem::CanUsePackage(GetContainerPackage());
 }
 
 bool FLevelInstanceActorDesc::GetContainerInstance(FContainerInstance& OutContainerInstance) const
@@ -198,22 +198,22 @@ void FLevelInstanceActorDesc::CheckForErrors(IStreamingGenerationErrorHandler* E
 	FWorldPartitionActorDesc::CheckForErrors(ErrorHandler);
 
 	FPackagePath WorldAssetPath;
-	if (!FPackagePath::TryFromPackageName(LevelPackage, WorldAssetPath) || !FPackageName::DoesPackageExist(WorldAssetPath))
+	if (!FPackagePath::TryFromPackageName(GetContainerPackage(), WorldAssetPath) || !FPackageName::DoesPackageExist(WorldAssetPath))
 	{
-		ErrorHandler->OnLevelInstanceInvalidWorldAsset(this, LevelPackage, IStreamingGenerationErrorHandler::ELevelInstanceInvalidReason::WorldAssetNotFound);
+		ErrorHandler->OnLevelInstanceInvalidWorldAsset(this, GetContainerPackage(), IStreamingGenerationErrorHandler::ELevelInstanceInvalidReason::WorldAssetNotFound);
 	}
-	else if (!ULevel::GetIsLevelUsingExternalActorsFromPackage(LevelPackage))
+	else if (!ULevel::GetIsLevelUsingExternalActorsFromPackage(GetContainerPackage()))
 	{
 		if (DesiredRuntimeBehavior != ELevelInstanceRuntimeBehavior::LevelStreaming)
 		{
-			ErrorHandler->OnLevelInstanceInvalidWorldAsset(this, LevelPackage, IStreamingGenerationErrorHandler::ELevelInstanceInvalidReason::WorldAssetNotUsingExternalActors);
+			ErrorHandler->OnLevelInstanceInvalidWorldAsset(this, GetContainerPackage(), IStreamingGenerationErrorHandler::ELevelInstanceInvalidReason::WorldAssetNotUsingExternalActors);
 		}
 	}
-	else if (ULevel::GetIsLevelPartitionedFromPackage(LevelPackage))
+	else if (ULevel::GetIsLevelPartitionedFromPackage(GetContainerPackage()))
 	{
-		if ((DesiredRuntimeBehavior != ELevelInstanceRuntimeBehavior::Partitioned) || !ULevel::GetPartitionedLevelCanBeUsedByLevelInstanceFromPackage(LevelPackage))
+		if ((DesiredRuntimeBehavior != ELevelInstanceRuntimeBehavior::Partitioned) || !ULevel::GetPartitionedLevelCanBeUsedByLevelInstanceFromPackage(GetContainerPackage()))
 		{
-			ErrorHandler->OnLevelInstanceInvalidWorldAsset(this, LevelPackage, IStreamingGenerationErrorHandler::ELevelInstanceInvalidReason::WorldAssetImcompatiblePartitioned);
+			ErrorHandler->OnLevelInstanceInvalidWorldAsset(this, GetContainerPackage(), IStreamingGenerationErrorHandler::ELevelInstanceInvalidReason::WorldAssetImcompatiblePartitioned);
 		}
 	}
 }
@@ -241,7 +241,17 @@ void FLevelInstanceActorDesc::Serialize(FArchive& Ar)
 	Ar.UsingCustomVersion(FUE5MainStreamObjectVersion::GUID);
 	Ar.UsingCustomVersion(FFortniteMainBranchObjectVersion::GUID);
 
-	Ar << TDeltaSerialize<FName>(LevelPackage);
+	if (Ar.CustomVer(FFortniteMainBranchObjectVersion::GUID) < FFortniteMainBranchObjectVersion::WorldPartitionActorDescSerializeSoftObjectPathSupport)
+	{
+		Ar << TDeltaSerialize<FSoftObjectPath, FName>(LevelPackage, [](FSoftObjectPath& Value, const FName& DeprecatedValue)
+		{
+			Value = FSoftObjectPath(*DeprecatedValue.ToString());
+		});
+	}
+	else
+	{
+		Ar << TDeltaSerialize<FSoftObjectPath>(LevelPackage);
+	}
 
 	if (!bIsDefaultActorDesc)
 	{
@@ -276,12 +286,12 @@ void FLevelInstanceActorDesc::Serialize(FArchive& Ar)
 
 			const AActor* CDO = GetActorNativeClass()->GetDefaultObject<AActor>();
 			const ILevelInstanceInterface* LevelInstanceCDO = CastChecked<ILevelInstanceInterface>(CDO);
-			if (!LevelPackage.IsNone() && (LevelInstanceCDO->IsLoadingEnabled() || bFixupOldVersion))
+			if (LevelPackage.IsValid() && (LevelInstanceCDO->IsLoadingEnabled() || bFixupOldVersion))
 			{
 				if (!IsContainerInstance())
 				{
 					FBox OutBounds;
-					if (ULevelInstanceSubsystem::GetLevelInstanceBoundsFromPackage(LevelInstanceTransform, LevelPackage, OutBounds))
+					if (ULevelInstanceSubsystem::GetLevelInstanceBoundsFromPackage(LevelInstanceTransform, GetContainerPackage(), OutBounds))
 					{
 						OutBounds.GetCenterAndExtents(BoundsLocation, BoundsExtent);
 					}
@@ -302,5 +312,4 @@ void FLevelInstanceActorDesc::Serialize(FArchive& Ar)
 		}
 	}
 }
-
 #endif
diff --git a/Engine/Source/Runtime/Engine/Private/WorldPartition/WorldPartitionActorDesc.cpp b/Engine/Source/Runtime/Engine/Private/WorldPartition/WorldPartitionActorDesc.cpp
--- a/Engine/Source/Runtime/Engine/Private/WorldPartition/WorldPartitionActorDesc.cpp
+++ b/Engine/Source/Runtime/Engine/Private/WorldPartition/WorldPartitionActorDesc.cpp
@@ -91,7 +91,7 @@ void FWorldPartitionActorDesc::Init(const AActor* InActor)
 	bActorIsEditorOnly = InActor->IsEditorOnly();
 	bActorIsRuntimeOnly = InActor->IsRuntimeOnly();
 	bActorIsHLODRelevant = InActor->IsHLODRelevant();
-	HLODLayer = InActor->GetHLODLayer() ? FName(InActor->GetHLODLayer()->GetPathName()) : FName();
+	HLODLayer = InActor->GetHLODLayer() ? FSoftObjectPath(InActor->GetHLODLayer()->GetPathName()) : FSoftObjectPath();
 	
 	// DataLayers
 	{
@@ -413,7 +413,7 @@ FString FWorldPartitionActorDesc::ToString(EToStringMode Mode) const
 			Result.Appendf(TEXT(" Parent:%s"), *ParentActor.ToString());
 		}
 
-		if (!HLODLayer.IsNone())
+		if (HLODLayer.IsValid())
 		{
 			Result.Appendf(TEXT(" HLODLayer:%s"), *HLODLayer.ToString());
 		}
@@ -581,12 +581,23 @@ PRAGMA_ENABLE_DEPRECATION_WARNINGS
 		(Ar.CustomVer(FUE5ReleaseStreamObjectVersion::GUID) >= FUE5ReleaseStreamObjectVersion::WorldPartitionActorDescSerializeHLODInfo))
 	{
 		Ar << TDeltaSerialize<bool>(bActorIsHLODRelevant);
-		Ar << TDeltaSerialize<FName>(HLODLayer);
+
+		if (Ar.CustomVer(FFortniteMainBranchObjectVersion::GUID) < FFortniteMainBranchObjectVersion::WorldPartitionActorDescSerializeSoftObjectPathSupport)
+		{
+			Ar << TDeltaSerialize<FSoftObjectPath, FName>(HLODLayer, [](FSoftObjectPath& Value, const FName& DeprecatedValue)
+			{
+				Value = FSoftObjectPath(*DeprecatedValue.ToString());
+			});
+		}
+		else
+		{
+			Ar << TDeltaSerialize<FSoftObjectPath>(HLODLayer);
+		}
 	}
 	else
 	{
 		bActorIsHLODRelevant = true;
-		HLODLayer = FName();
+		HLODLayer = FSoftObjectPath();
 	}
 
 	if (!bIsDefaultActorDesc)
diff --git a/Engine/Source/Runtime/Engine/Private/WorldPartition/WorldPartitionActorDescArchive.cpp b/Engine/Source/Runtime/Engine/Private/WorldPartition/WorldPartitionActorDescArchive.cpp
--- a/Engine/Source/Runtime/Engine/Private/WorldPartition/WorldPartitionActorDescArchive.cpp
+++ b/Engine/Source/Runtime/Engine/Private/WorldPartition/WorldPartitionActorDescArchive.cpp
@@ -4,6 +4,7 @@
 #include "WorldPartition/WorldPartitionActorDescArchive.h"
 #include "WorldPartition/WorldPartitionClassDescRegistry.h"
 #include "WorldPartition/WorldPartitionActorDesc.h"
+#include "WorldPartition/WorldPartitionHelpers.h"
 #include "WorldPartition/WorldPartitionLog.h"
 #include "UObject/UE5MainStreamObjectVersion.h"
 #include "UObject/FortniteNCBranchObjectVersion.h"
@@ -81,4 +82,16 @@ PRAGMA_ENABLE_DEPRECATION_WARNINGS
 
 	ClassDescSizeof = ClassDesc ? ClassDesc->GetSizeOf() : 0;
 }
+
+FArchive& FActorDescArchive::operator<<(FSoftObjectPath& Value)
+{
+	Value.SerializePathWithoutFixup(*this);
+
+	if (IsLoading())
+	{
+		FWorldPartitionHelpers::FixupRedirectedAssetPath(Value);
+	}
+
+	return *this;
+}
 #endif
\ No newline at end of file
diff --git a/Engine/Source/Runtime/Engine/Private/WorldPartition/WorldPartitionActorDescView.cpp b/Engine/Source/Runtime/Engine/Private/WorldPartition/WorldPartitionActorDescView.cpp
--- a/Engine/Source/Runtime/Engine/Private/WorldPartition/WorldPartitionActorDescView.cpp
+++ b/Engine/Source/Runtime/Engine/Private/WorldPartition/WorldPartitionActorDescView.cpp
@@ -69,7 +69,7 @@ bool FWorldPartitionActorDescView::GetActorIsHLODRelevant() const
 	return ActorDesc->GetActorIsHLODRelevant();
 }
 
-FName FWorldPartitionActorDescView::GetHLODLayer() const
+FSoftObjectPath FWorldPartitionActorDescView::GetHLODLayer() const
 {
 	return ActorDesc->GetHLODLayer();
 }
diff --git a/Engine/Source/Runtime/Engine/Public/WorldPartition/LevelInstance/LevelInstanceActorDesc.h b/Engine/Source/Runtime/Engine/Public/WorldPartition/LevelInstance/LevelInstanceActorDesc.h
--- a/Engine/Source/Runtime/Engine/Public/WorldPartition/LevelInstance/LevelInstanceActorDesc.h
+++ b/Engine/Source/Runtime/Engine/Public/WorldPartition/LevelInstance/LevelInstanceActorDesc.h
@@ -27,7 +27,7 @@ public:
 
 	virtual bool IsContainerInstance() const override;
 	virtual bool IsContainerFilter() const override { return IsContainerInstance(); }
-	virtual FName GetContainerPackage() const override { return LevelPackage; }
+	virtual FName GetContainerPackage() const override { return *LevelPackage.ToString(); }
 	virtual bool GetContainerInstance(FContainerInstance& OutContainerInstance) const override;
 	virtual const FWorldPartitionActorFilter* GetContainerFilter() const override { return &Filter; }
 	virtual void CheckForErrors(IStreamingGenerationErrorHandler* ErrorHandler) const override;
@@ -41,7 +41,7 @@ protected:
 	virtual void Serialize(FArchive& Ar) override;
 	virtual void SetContainer(UActorDescContainer* InContainer, UWorld* InWorldContext) override;
 
-	FName LevelPackage;
+	FSoftObjectPath LevelPackage;
 	FTransform LevelInstanceTransform;
 	ELevelInstanceRuntimeBehavior DesiredRuntimeBehavior;
 
diff --git a/Engine/Source/Runtime/Engine/Public/WorldPartition/WorldPartitionActorDesc.h b/Engine/Source/Runtime/Engine/Public/WorldPartition/WorldPartitionActorDesc.h
--- a/Engine/Source/Runtime/Engine/Public/WorldPartition/WorldPartitionActorDesc.h
+++ b/Engine/Source/Runtime/Engine/Public/WorldPartition/WorldPartitionActorDesc.h
@@ -114,7 +114,7 @@ PRAGMA_ENABLE_DEPRECATION_WARNINGS
 	inline void SetIsSpatiallyLoadedRaw(bool bNewIsSpatiallyLoaded) { bIsSpatiallyLoaded = bNewIsSpatiallyLoaded; }
 
 	inline bool GetActorIsHLODRelevant() const { return bActorIsHLODRelevant; }
-	inline FName GetHLODLayer() const { return HLODLayer; }
+	inline FSoftObjectPath GetHLODLayer() const { return HLODLayer; }
 	inline const TArray<FName>& GetDataLayers() const { return DataLayers; }
 	inline bool HasResolvedDataLayerInstanceNames() const { return ResolvedDataLayerInstanceNames.IsSet(); }
 	const TArray<FName>& GetDataLayerInstanceNames() const;
@@ -301,7 +301,7 @@ protected:
 	bool							bActorIsHLODRelevant;
 	bool							bIsUsingDataLayerAsset; // Used to know if DataLayers array represents DataLayers Asset paths or the FNames of the deprecated version of Data Layers
 	bool							bIsBoundsValid;
-	FName							HLODLayer;
+	FSoftObjectPath					HLODLayer;
 	TArray<FName>					DataLayers;
 	TArray<FGuid>					References;
 	TArray<FGuid>					EditorOnlyReferences; // References that aren't necessarily editor only but referenced through an editor only property.
diff --git a/Engine/Source/Runtime/Engine/Public/WorldPartition/WorldPartitionActorDescArchive.h b/Engine/Source/Runtime/Engine/Public/WorldPartition/WorldPartitionActorDescArchive.h
--- a/Engine/Source/Runtime/Engine/Public/WorldPartition/WorldPartitionActorDescArchive.h
+++ b/Engine/Source/Runtime/Engine/Public/WorldPartition/WorldPartitionActorDescArchive.h
@@ -19,18 +19,25 @@ public:
 	virtual FArchive& operator<<(FLazyObjectPtr& Value) override { unimplemented(); return *this; }
 	virtual FArchive& operator<<(FObjectPtr& Value) override { unimplemented(); return *this; }
 	virtual FArchive& operator<<(FSoftObjectPtr& Value) override { unimplemented(); return *this; }
-	virtual FArchive& operator<<(FSoftObjectPath& Value) override { unimplemented(); return *this; }
+	virtual FArchive& operator<<(FSoftObjectPath& Value) override;
 	virtual FArchive& operator<<(FWeakObjectPtr& Value) override { unimplemented(); return *this; }
 	//~ End FArchive Interface
 
-	template <typename T>
+	template <typename DestPropertyType, typename SourcePropertyType>
 	struct TDeltaSerializer
 	{
-		explicit TDeltaSerializer(T& InValue)
+		using FDeprecateFunction = TFunction<void(DestPropertyType&, const SourcePropertyType&)>;
+
+		explicit TDeltaSerializer(DestPropertyType& InValue)
 			: Value(InValue)
 		{}
 
-		friend FArchive& operator<<(FArchive& Ar, const TDeltaSerializer<T>& V)
+		explicit TDeltaSerializer(DestPropertyType& InValue, FDeprecateFunction InFunc)
+			: Value(InValue)
+			, Func(InFunc)
+		{}
+
+		friend FArchive& operator<<(FArchive& Ar, const TDeltaSerializer<DestPropertyType, SourcePropertyType>& V)
 		{
 			FActorDescArchive& ActorDescAr = (FActorDescArchive&)Ar;
 			
@@ -39,14 +46,14 @@ public:
 
 			Ar.UsingCustomVersion(FFortniteMainBranchObjectVersion::GUID);
 
-			auto GetClassDefaultValue = [&V, &ActorDescAr]() -> const T*
+			auto GetClassDefaultValue = [&V, &ActorDescAr]() -> const DestPropertyType*
 			{
 				const UPTRINT PropertyOffset = (UPTRINT)&V.Value - *(UPTRINT*)&ActorDescAr.ActorDesc;
 
 				if (PropertyOffset < ActorDescAr.ClassDescSizeof)
 				{
 					check((PropertyOffset + sizeof(V)) <= ActorDescAr.ClassDescSizeof);
-					const T* RefValue = (const T*)(*(UPTRINT*)&ActorDescAr.ClassDesc + PropertyOffset);
+					const DestPropertyType* RefValue = (const DestPropertyType*)(*(UPTRINT*)&ActorDescAr.ClassDesc + PropertyOffset);
 					return RefValue;
 				}
 
@@ -60,11 +67,14 @@ public:
 			{
 				if (Ar.IsSaving() && ActorDescAr.ClassDesc)
 				{
-					// When saving, we expect the class descriptor to be the exact type as what we are serializing.
-					const T* ClassDefaultValue = GetClassDefaultValue();
-					check(ClassDefaultValue);
-
-					bSerialize = (V.Value != *ClassDefaultValue) ? 1 : 0;
+					if constexpr (std::is_same_v<DestPropertyType, SourcePropertyType>)
+					{
+						// When saving, we expect the class descriptor to be the exact type as what we are serializing.
+						const DestPropertyType* ClassDefaultValue = GetClassDefaultValue();
+						check(ClassDefaultValue);
+
+						bSerialize = (V.Value != *ClassDefaultValue) ? 1 : 0;
+					}
 				}
 
 				Ar << bSerialize;
@@ -72,12 +82,23 @@ public:
 
 			if (bSerialize)
 			{
-				Ar << V.Value;
+				if constexpr (std::is_same_v<DestPropertyType, SourcePropertyType>)
+				{
+					Ar << V.Value;
+				}
+				else
+				{
+					check(Ar.IsLoading());
+
+					SourcePropertyType SourceValue;
+					Ar << SourceValue;
+					V.Func(V.Value, SourceValue);
+				}
 			}
 			else if (Ar.IsLoading())
 			{
 				// When loading, we need to handle a different class descriptor in case of missing classes, etc.
-				if (const T* ClassDefaultValue = GetClassDefaultValue())
+				if (const DestPropertyType* ClassDefaultValue = GetClassDefaultValue())
 				{
 					V.Value = *ClassDefaultValue;
 				}
@@ -86,7 +107,8 @@ public:
 			return Ar;
 		}
 
-		T& Value;
+		DestPropertyType& Value;
+		FDeprecateFunction Func;
 	};
 
 	FWorldPartitionActorDesc* ActorDesc;
@@ -95,6 +117,6 @@ public:
 	bool bIsMissingClassDesc;
 };
 
-template <typename T>
-using TDeltaSerialize = FActorDescArchive::TDeltaSerializer<T>;
+template <typename DestPropertyType, typename SourcePropertyType = DestPropertyType>
+using TDeltaSerialize = FActorDescArchive::TDeltaSerializer<DestPropertyType, SourcePropertyType>;
 #endif
\ No newline at end of file
diff --git a/Engine/Source/Runtime/Engine/Public/WorldPartition/WorldPartitionActorDescView.h b/Engine/Source/Runtime/Engine/Public/WorldPartition/WorldPartitionActorDescView.h
--- a/Engine/Source/Runtime/Engine/Public/WorldPartition/WorldPartitionActorDescView.h
+++ b/Engine/Source/Runtime/Engine/Public/WorldPartition/WorldPartitionActorDescView.h
@@ -37,7 +37,7 @@ public:
 	bool GetActorIsEditorOnly() const;
 	bool GetActorIsRuntimeOnly() const;
 	bool GetActorIsHLODRelevant() const;
-	FName GetHLODLayer() const;
+	FSoftObjectPath GetHLODLayer() const;
 	const TArray<FName>& GetDataLayerInstanceNames() const;
 	const TArray<FName>& GetRuntimeDataLayerInstanceNames() const;
 	const TArray<FName>& GetTags() const;
