World Partition Class Descriptors Registry:
- Optimize object redirectors resolving through the global Redirect Collector.
- Switched Asset Registry synchronous scan to new API.

diff --git a/Engine/Source/Runtime/Engine/Private/WorldPartition/WorldPartitionClassDescRegistry.cpp b/Engine/Source/Runtime/Engine/Private/WorldPartition/WorldPartitionClassDescRegistry.cpp
--- a/Engine/Source/Runtime/Engine/Private/WorldPartition/WorldPartitionClassDescRegistry.cpp
+++ b/Engine/Source/Runtime/Engine/Private/WorldPartition/WorldPartitionClassDescRegistry.cpp
@@ -9,6 +9,7 @@
 #include "Algo/Transform.h"
 #include "HAL/FileManager.h"
 #include "Misc/PackageName.h"
+#include "Misc/RedirectCollector.h"
 #include "UObject/CoreRedirects.h"
 #include "UObject/ObjectSaveContext.h"
 #include "UObject/UObjectIterator.h"
@@ -192,10 +193,22 @@ void FWorldPartitionClassDescRegistry::PrefetchClassDescs(const TArray<FTopLevel
 		return;
 	}
 
-	// Make sure the asset registry has all paths scanned
+	// Resolve object redirectors and prepare asset registry file paths to be scanned.
 	TArray<FString> FilePaths;
 	FilePaths.Reserve(ClassPaths.Num());
-	Algo::Transform(ClassPaths, FilePaths, [this](const FTopLevelAssetPath& ClassPath) { return ClassPath.GetPackageName().ToString(); });
+	Algo::Transform(ClassPaths, FilePaths, [this](const FTopLevelAssetPath& ClassPath)
+	{
+		const FSoftObjectPath FoundRedirection = GRedirectCollector.GetAssetPathRedirection(FSoftObjectPath(ClassPath.ToString()));
+		
+		FString NewResult;
+		if (FoundRedirection.IsValid())
+		{
+			RedirectClassMap.Add(ClassPath, FoundRedirection.GetAssetPath());
+			return FoundRedirection.GetAssetPath().GetPackageName().ToString();
+		}
+
+		return ClassPath.GetPackageName().ToString();
+	});
 
 	// In the editor, the asset registry will not return blueprint generated classes or their default object but will during cooks (see AssetRegistry::FFiltering::ShouldSkipAsset),
 	// so we must filter our results by UBlueprint and UObjectRedirector. For redirectors, they will show up for the blueprint, generated class and default objects in both situations,
@@ -214,7 +227,7 @@ void FWorldPartitionClassDescRegistry::PrefetchClassDescs(const TArray<FTopLevel
 		Filter.PackageNames.Reserve(FilePaths.Num());
 		Algo::Transform(FilePaths, Filter.PackageNames, [](const FString& ClassPath) { return *ClassPath; });
 	
-		AssetRegistry.ScanFilesSynchronous(FilePaths, /*bForceRescan*/false);
+		AssetRegistry.ScanSynchronous(FilePaths, TArray<FString>());
 		AssetRegistry.GetAssets(Filter, Assets);
 	};
 
@@ -650,4 +663,4 @@ FTopLevelAssetPath FWorldPartitionClassDescRegistry::RedirectClassPath(const FTo
 
 	return ResolvedClassPath;
 }
-#endif
\ No newline at end of file
+#endif
