Initialize bInitialEvaluateWorldPositionOffset earlier. BeginPlay seems too late. Things like network replication can happen before BeginPlay and bEvaluateWorldPositionOffset can be changed before the original value is copied.

diff --git a/Engine/Source/Runtime/Engine/Private/Components/StaticMeshComponent.cpp b/Engine/Source/Runtime/Engine/Private/Components/StaticMeshComponent.cpp
--- a/Engine/Source/Runtime/Engine/Private/Components/StaticMeshComponent.cpp
+++ b/Engine/Source/Runtime/Engine/Private/Components/StaticMeshComponent.cpp
@@ -403,6 +403,11 @@ void UStaticMeshComponent::Serialize(FArchive& Ar)
 		UpdateCollisionFromStaticMesh();
 		BodyInstance.FixupData(this);
 	}
+
+	if (Ar.IsLoading())
+	{
+		bInitialEvaluateWorldPositionOffset = bEvaluateWorldPositionOffset;
+	}
 }
 
 void UStaticMeshComponent::PostApplyToComponent()
@@ -417,6 +422,8 @@ void UStaticMeshComponent::PostReinitProperties()
 	NotifyIfStaticMeshChanged();
 
 	Super::PostReinitProperties();
+
+	bInitialEvaluateWorldPositionOffset = bEvaluateWorldPositionOffset;
 }
 
 void UStaticMeshComponent::PostInitProperties()
@@ -429,6 +436,8 @@ void UStaticMeshComponent::PostInitProperties()
 	{
 		LODData[LODIndex].OwningComponent = this;
 	}
+
+	bInitialEvaluateWorldPositionOffset = bEvaluateWorldPositionOffset;
 }
 
 bool UStaticMeshComponent::AreNativePropertiesIdenticalTo( UObject* Other ) const
@@ -771,8 +780,6 @@ void UStaticMeshComponent::OnUnregister()
 
 void UStaticMeshComponent::BeginPlay()
 {
-	bInitialEvaluateWorldPositionOffset = bEvaluateWorldPositionOffset;
-
 	Super::BeginPlay();
 }
 
