Additional safe-guarding ensuring that fetched and compressed animation data is _valid_ (data is as expected according to raw data)

diff --git a/Engine/Source/Runtime/Engine/Private/Animation/AnimCompressionTypes.cpp b/Engine/Source/Runtime/Engine/Private/Animation/AnimCompressionTypes.cpp
--- a/Engine/Source/Runtime/Engine/Private/Animation/AnimCompressionTypes.cpp
+++ b/Engine/Source/Runtime/Engine/Private/Animation/AnimCompressionTypes.cpp
@@ -1487,7 +1487,7 @@ bool FCompressedAnimSequence::IsValid(const UAnimSequence* AnimSequence, bool bL
 	const bool bHasValidCompressedCurveData = CompressedCurveByteStream.Num() != 0;
 	const bool bValidCompressedData = (!bHasBoneTracks || bHasValidCompressedBoneData) && (!bHasCurveData || bHasValidCompressedCurveData || AnimSequence->IsValidAdditive());
 
-	if(bLogInformation)
+	if(bLogInformation && !bValidCompressedData)
 	{
 		UE_LOG(LogAnimation, Warning, TEXT("Num Bone Tracks: %i\nNum Curves: %i\nCompressed bone data: %i (%i)\nCompressed curve data %i\nValid additive: %i"), AnimSequence->GetDataModelInterface()->GetNumBoneTracks(), AnimSequence->GetDataModelInterface()->GetNumberOfFloatCurves(), CompressedDataStructure != nullptr ? 1 : 0, CompressedTrackToSkeletonMapTable.Num(), CompressedCurveByteStream.Num() != 0 ? 1 : 0, AnimSequence->IsValidAdditive() ? 1 : 0);
 	}
diff --git a/Engine/Source/Runtime/Engine/Private/Animation/AnimSequence.cpp b/Engine/Source/Runtime/Engine/Private/Animation/AnimSequence.cpp
--- a/Engine/Source/Runtime/Engine/Private/Animation/AnimSequence.cpp
+++ b/Engine/Source/Runtime/Engine/Private/Animation/AnimSequence.cpp
@@ -1169,7 +1169,7 @@ bool UAnimSequence::IsCachedCookedPlatformDataLoaded(const ITargetPlatform* Targ
 
 	if (KeyHash == DataKeyHash)
 	{
-		if (CompressedData.IsValid(this) && !CacheTasksByKeyHash.Contains(KeyHash))
+		if (!CacheTasksByKeyHash.Contains(KeyHash) && CompressedData.IsValid(this, true))
 		{
 			return true;
 		}
@@ -1182,7 +1182,7 @@ bool UAnimSequence::IsCachedCookedPlatformDataLoaded(const ITargetPlatform* Targ
 	{
 		if(const TUniquePtr<FCompressedAnimSequence>* CompressedDataPtr = DataByPlatformKeyHash.Find(KeyHash))
 		{
-			if ((*CompressedDataPtr)->IsValid(this) && !CacheTasksByKeyHash.Contains(KeyHash))
+			if (!CacheTasksByKeyHash.Contains(KeyHash) && (*CompressedDataPtr)->IsValid(this, true))
 			{
 				return true;
 			}
diff --git a/Engine/Source/Runtime/Engine/Private/Animation/AnimationCompressionDerivedData.cpp b/Engine/Source/Runtime/Engine/Private/Animation/AnimationCompressionDerivedData.cpp
--- a/Engine/Source/Runtime/Engine/Private/Animation/AnimationCompressionDerivedData.cpp
+++ b/Engine/Source/Runtime/Engine/Private/Animation/AnimationCompressionDerivedData.cpp
@@ -138,7 +138,16 @@ void FAnimationSequenceAsyncCacheTask::EndCache(UE::DerivedData::FCacheGetValueR
 				FMemoryReaderView Ar(RecordData, /*bIsPersistent*/ true);
 				CompressedData->SerializeCompressedData(Ar, true, AnimSequence, AnimSequence->GetSkeleton(), CompressibleAnimPtr->BoneCompressionSettings, CompressibleAnimPtr->CurveCompressionSettings);
 
-				UE_LOG(LogAnimationCompression, Verbose, TEXT("Fetched compressed animation data for %s"), *CompressibleAnimPtr->FullName);
+				if (!CompressedData->IsValid(AnimSequence, true))
+				{
+					UE_LOG(LogAnimationCompression, Warning, TEXT("Fetched invalid compressed animation data for %s"), *CompressibleAnimPtr->FullName);
+					CompressedData->Reset();
+				}
+				else
+				{
+					UE_LOG(LogAnimationCompression, Verbose, TEXT("Fetched compressed animation data for %s"), *CompressibleAnimPtr->FullName);
+				}
+
 				if (Compression::FAnimationCompressionMemorySummaryScope::ShouldStoreCompressionResults())
 				{
 					const double CompressionEndTime = FPlatformTime::Seconds();
@@ -169,12 +178,20 @@ void FAnimationSequenceAsyncCacheTask::EndCache(UE::DerivedData::FCacheGetValueR
 			
 			if (const UAnimSequence* AnimSequence = WeakAnimSequence.Get())
 			{
-				TArray64<uint8> RecordData;
-				FMemoryWriter64 Ar(RecordData, /*bIsPersistent*/ true);
-				CompressedData->SerializeCompressedData(Ar, true, nullptr, nullptr, CompressibleAnimPtr->BoneCompressionSettings, CompressibleAnimPtr->CurveCompressionSettings);
+				if (!CompressedData->IsValid(AnimSequence, true))
+				{
+					UE_LOG(LogAnimationCompression, Warning, TEXT("Generated invalid compressed animation data for %s"), *CompressibleAnimPtr->FullName);
+				}
+				else
+				{
+					TArray64<uint8> RecordData;
+					FMemoryWriter64 Ar(RecordData, /*bIsPersistent*/ true);
+					CompressedData->SerializeCompressedData(Ar, true, nullptr, nullptr, CompressibleAnimPtr->BoneCompressionSettings, CompressibleAnimPtr->CurveCompressionSettings);
+					UE_LOG(LogAnimationCompression, Display, TEXT("Storing compressed animation data for %s, at %s/%s"), *Name, *FString(Key.Bucket.ToString()), *LexToString(Key.Hash));
+					GetCache().PutValue({ {Name, Key, FValue::Compress(MakeSharedBufferFromArray(MoveTemp(RecordData)))} }, Owner);
 
-				UE_LOG(LogAnimationCompression, Display, TEXT("Storing compressed animation data for %s, at %s/%s"), *Name, *FString(Key.Bucket.ToString()), *LexToString(Key.Hash));
-				GetCache().PutValue({ {Name, Key, FValue::Compress(MakeSharedBufferFromArray(MoveTemp(RecordData)))} }, Owner);
+					COOK_STAT(Timer.AddMiss(int64(Ar.Tell())));
+				}
 				
 				if (Compression::FAnimationCompressionMemorySummaryScope::ShouldStoreCompressionResults())
 				{
@@ -183,7 +200,6 @@ void FAnimationSequenceAsyncCacheTask::EndCache(UE::DerivedData::FCacheGetValueR
 					Compression::FAnimationCompressionMemorySummaryScope::CompressionResultSummary().GatherPostCompressionStats(*CompressedData, CompressibleAnimPtr->BoneData, AnimSequence->GetFName(), CompressionTime, true);
 				}
 				
-				COOK_STAT(Timer.AddMiss(int64(Ar.Tell())));
 			}
 		});
 	}
