World Partition: class descriptors GUIDs are now transient and always recomputed from their class path.

diff --git a/Engine/Source/Runtime/Core/Public/UObject/FortniteMainBranchObjectVersion.h b/Engine/Source/Runtime/Core/Public/UObject/FortniteMainBranchObjectVersion.h
--- a/Engine/Source/Runtime/Core/Public/UObject/FortniteMainBranchObjectVersion.h
+++ b/Engine/Source/Runtime/Core/Public/UObject/FortniteMainBranchObjectVersion.h
@@ -313,6 +313,9 @@ struct CORE_API FFortniteMainBranchObjectVersion
 		// Add support for soft object paths in actor descriptors
 		WorldPartitionActorDescSerializeSoftObjectPathSupport,
 
+		// Don't serialize class descriptor GUIDs
+		WorldPartitionClasDescGuidTransient,
+
 		// -----<new versions can be added above this line>-------------------------------------------------
 		VersionPlusOne,
 		LatestVersion = VersionPlusOne - 1
diff --git a/Engine/Source/Runtime/Engine/Private/WorldPartition/WorldPartitionActorDesc.cpp b/Engine/Source/Runtime/Engine/Private/WorldPartition/WorldPartitionActorDesc.cpp
--- a/Engine/Source/Runtime/Engine/Private/WorldPartition/WorldPartitionActorDesc.cpp
+++ b/Engine/Source/Runtime/Engine/Private/WorldPartition/WorldPartitionActorDesc.cpp
@@ -31,6 +31,14 @@
 
 TMap<TSubclassOf<AActor>, FWorldPartitionActorDesc::FActorDescDeprecator> FWorldPartitionActorDesc::Deprecators;
 
+static FGuid GetDefaultActorDescGuid(const FWorldPartitionActorDesc* ActorDesc)
+{
+	FArchiveMD5 ArMD5;
+	FString ClassPath = ActorDesc->GetBaseClass().IsValid() ? ActorDesc->GetBaseClass().ToString() : ActorDesc->GetNativeClass().ToString();	
+	ArMD5 << ClassPath;
+	return ArMD5.GetGuidFromHash();
+}
+
 FWorldPartitionActorDesc::FWorldPartitionActorDesc()
 	: bIsUsingDataLayerAsset(false)
 	, bIsBoundsValid(false)
@@ -42,13 +50,6 @@ FWorldPartitionActorDesc::FWorldPartitionActorDesc()
 	, UnloadedReason(nullptr)
 {}
 
-FGuid FWorldPartitionActorDesc::GetDefaultActorDescGuid(FString ClassPath)
-{
-	FArchiveMD5 ArMD5;
-	ArMD5 << ClassPath;
-	return ArMD5.GetGuidFromHash();
-}
-
 void FWorldPartitionActorDesc::Init(const AActor* InActor)
 {	
 	check(IsValid(InActor));
@@ -69,8 +70,7 @@ void FWorldPartitionActorDesc::Init(const AActor* InActor)
 	{
 		check(!InActor->IsPackageExternal());
 		check(!InActor->GetActorGuid().IsValid());
-		const FString ClassPath = BaseClass.IsValid() ? BaseClass.ToString() : NativeClass.ToString();
-		Guid = GetDefaultActorDescGuid(ClassPath);
+		Guid = GetDefaultActorDescGuid(this);
 		bIsDefaultActorDesc = true;
 	}
 	else
@@ -469,7 +469,24 @@ void FWorldPartitionActorDesc::Serialize(FArchive& Ar)
 	Ar.UsingCustomVersion(FFortniteNCBranchObjectVersion::GUID);
 	Ar.UsingCustomVersion(FFortniteMainBranchObjectVersion::GUID);
 
-	Ar << Guid;
+	if (bIsDefaultActorDesc)
+	{
+		if (Ar.IsLoading())
+		{
+			Guid = GetDefaultActorDescGuid(this);
+
+			if (Ar.CustomVer(FFortniteMainBranchObjectVersion::GUID) < FFortniteMainBranchObjectVersion::WorldPartitionClasDescGuidTransient)
+			{
+				FGuid ClassDescGuid;
+				Ar << ClassDescGuid;
+				check(Guid == ClassDescGuid);
+			}
+		}
+	}
+	else
+	{
+		Ar << Guid;
+	}
 
 	if(Ar.CustomVer(FUE5ReleaseStreamObjectVersion::GUID) < FUE5ReleaseStreamObjectVersion::LargeWorldCoordinates)
 	{
diff --git a/Engine/Source/Runtime/Engine/Public/WorldPartition/WorldPartitionActorDesc.h b/Engine/Source/Runtime/Engine/Public/WorldPartition/WorldPartitionActorDesc.h
--- a/Engine/Source/Runtime/Engine/Public/WorldPartition/WorldPartitionActorDesc.h
+++ b/Engine/Source/Runtime/Engine/Public/WorldPartition/WorldPartitionActorDesc.h
@@ -78,8 +78,6 @@ class ENGINE_API FWorldPartitionActorDesc
 	friend class FAssetRootPackagePatcher;
 	friend class FActorDescArchive;
 
-	static FGuid GetDefaultActorDescGuid(FString ClassPath);
-
 public:
 	struct FContainerInstance
 	{
