#ue5 - Prevent AActor::RerunConstructionScripts() from calling a redundant SetWorldTransform() inside ExecuteConstruction() when the transform hasn't changed. This can cause error accumulation in the relative and world transforms of a component that has an attach parent. This also avoids pointless matrix and transform math during construction in these cases, even if the values wouldn't be affected.

diff --git a/Engine/Source/Runtime/Engine/Private/ActorConstruction.cpp b/Engine/Source/Runtime/Engine/Private/ActorConstruction.cpp
--- a/Engine/Source/Runtime/Engine/Private/ActorConstruction.cpp
+++ b/Engine/Source/Runtime/Engine/Private/ActorConstruction.cpp
@@ -557,8 +557,20 @@ void AActor::RerunConstructionScripts()
 			ExchangeNetRoles(true);
 		}
 
+		// Determine if we already have the correct world transform on the RootComponent. If so, we don't want to try to set it again in ExecuteConstruction()
+		// or else a re-computation of the relative transform can cause error accumulation on the RelativeLocation/etc which is supposed to derive the ComponentToWorld.
+		bool bIsDefaultTransform = false;
+		if (RootComponent != nullptr && bUseRootComponentProperties)
+		{
+			const double TransformTolerance = 0.0;
+			if (OldTransform.Equals(RootComponent->GetComponentTransform(), TransformTolerance))
+			{
+				bIsDefaultTransform = true;
+			}
+		}
+
 		// Run the construction scripts
-		const bool bErrorFree = ExecuteConstruction(OldTransform, &OldTransformRotationCache, InstanceDataCache);
+		const bool bErrorFree = ExecuteConstruction(OldTransform, &OldTransformRotationCache, InstanceDataCache, bIsDefaultTransform);
 
 		if(Parent)
 		{
