Changed default value of wp.Runtime.UpdateStreaming.EnableOptimization to true.

diff --git a/Engine/Source/Runtime/Engine/Private/WorldPartition/WorldPartitionStreamingPolicy.cpp b/Engine/Source/Runtime/Engine/Private/WorldPartition/WorldPartitionStreamingPolicy.cpp
--- a/Engine/Source/Runtime/Engine/Private/WorldPartition/WorldPartitionStreamingPolicy.cpp
+++ b/Engine/Source/Runtime/Engine/Private/WorldPartition/WorldPartitionStreamingPolicy.cpp
@@ -44,7 +44,7 @@ static FAutoConsoleVariableRef CVarServerDisallowStreamingOutDataLayers(
 	TEXT("Comma separated list of data layer names that aren't allowed to be unloaded or deactivated on the server"),
 	ECVF_ReadOnly);
 
-bool UWorldPartitionStreamingPolicy::IsUpdateOptimEnabled = false;
+bool UWorldPartitionStreamingPolicy::IsUpdateOptimEnabled = true;
 FAutoConsoleVariableRef UWorldPartitionStreamingPolicy::CVarUpdateOptimEnabled(
 	TEXT("wp.Runtime.UpdateStreaming.EnableOptimization"),
 	UWorldPartitionStreamingPolicy::IsUpdateOptimEnabled,
