- Don't allow us to spawn more particle in a trail emitter than we can pack
- If we exceed this limit it can result in infinite loops / memory overruns when traversing the particle trail

diff --git a/Engine/Source/Runtime/Engine/Private/Particles/ParticleTrail2EmitterInstance.cpp b/Engine/Source/Runtime/Engine/Private/Particles/ParticleTrail2EmitterInstance.cpp
--- a/Engine/Source/Runtime/Engine/Private/Particles/ParticleTrail2EmitterInstance.cpp
+++ b/Engine/Source/Runtime/Engine/Private/Particles/ParticleTrail2EmitterInstance.cpp
@@ -3408,7 +3408,14 @@ float FParticleAnimTrailEmitterInstance::Spawn(float DeltaTime)
 	// Handle growing arrays.
 	bool bProcessSpawn = true;
 	int32 TotalCount = ActiveParticles + NewCount;
-	if (TotalCount >= MaxActiveParticles)
+
+	constexpr int32 MaxTrailParticles = FMath::Min(TRAIL_EMITTER_NEXT_MASK >> TRAIL_EMITTER_NEXT_SHIFT, TRAIL_EMITTER_PREV_MASK >> TRAIL_EMITTER_PREV_SHIFT);
+	if (TotalCount >= MaxTrailParticles)
+	{
+		// Note the above test is >= as null is considered == max value
+		bProcessSpawn = false;
+	}
+	else if (TotalCount >= MaxActiveParticles)
 	{
 		if (DeltaTime < 0.25f)
 		{
