Fix nested FGlobalComponentReregisterContext behavior

Previously, every global context would reregister, even if there were already active global contexts. It contained functionality that was supposed to address this, but it did not.

diff --git a/Engine/Source/Runtime/Engine/Private/Components/ActorComponent.cpp b/Engine/Source/Runtime/Engine/Private/Components/ActorComponent.cpp
--- a/Engine/Source/Runtime/Engine/Private/Components/ActorComponent.cpp
+++ b/Engine/Source/Runtime/Engine/Private/Components/ActorComponent.cpp
@@ -196,41 +196,46 @@ FGlobalComponentReregisterContext::FGlobalComponentReregisterContext()
 
 FGlobalComponentReregisterContext::FGlobalComponentReregisterContext(const TArray<UClass*>& ExcludeComponents)
 {
-	ActiveGlobalReregisterContextCount++;
-
-	// wait until resources are released
-	FlushRenderingCommands();
-
-	// Detach only actor components that are not in the excluded list
-	for (UActorComponent* Component : TObjectRange<UActorComponent>())
+	// Check if this is the first active context
+	if (++ActiveGlobalReregisterContextCount == 1)
 	{
-		bool bShouldReregister=true;
-		for (UClass* ExcludeClass : ExcludeComponents)
+		// wait until resources are released
+		FlushRenderingCommands();
+		
+		// Detach only actor components that are not in the excluded list
+		for (UActorComponent* Component : TObjectRange<UActorComponent>())
 		{
-			if( ExcludeClass &&
-				Component->IsA(ExcludeClass) )
+			bool bShouldReregister=true;
+			for (UClass* ExcludeClass : ExcludeComponents)
 			{
-				bShouldReregister = false;
-				break;
+				if( ExcludeClass &&
+					Component->IsA(ExcludeClass) )
+				{
+					bShouldReregister = false;
+					break;
+				}
+			}
+			if( bShouldReregister )
+			{
+				ComponentContexts.Add(new FComponentReregisterContext(Component, &ScenesToUpdateAllPrimitiveSceneInfos));
 			}
 		}
-		if( bShouldReregister )
-		{
-			ComponentContexts.Add(new FComponentReregisterContext(Component, &ScenesToUpdateAllPrimitiveSceneInfos));
-		}
+		
+		UpdateAllPrimitiveSceneInfos();
 	}
-
-	UpdateAllPrimitiveSceneInfos();
 }
 
 FGlobalComponentReregisterContext::~FGlobalComponentReregisterContext()
 {
 	check(ActiveGlobalReregisterContextCount > 0);
-	// We empty the array now, to ensure that the FComponentReregisterContext destructors are called while ActiveGlobalReregisterContextCount still indicates activity
-	ComponentContexts.Empty();
-	ActiveGlobalReregisterContextCount--;
 
-	UpdateAllPrimitiveSceneInfos();
+	// Check if this is the last active context
+	if (--ActiveGlobalReregisterContextCount == 0)
+	{
+		ComponentContexts.Empty();
+		
+		UpdateAllPrimitiveSceneInfos();
+	}
 }
 
 void FGlobalComponentReregisterContext::UpdateAllPrimitiveSceneInfos()
