Markup up native serializers for object ref archives

diff --git a/Engine/Source/Runtime/ClothingSystemRuntimeCommon/Public/ClothTetherData.h b/Engine/Source/Runtime/ClothingSystemRuntimeCommon/Public/ClothTetherData.h
--- a/Engine/Source/Runtime/ClothingSystemRuntimeCommon/Public/ClothTetherData.h
+++ b/Engine/Source/Runtime/ClothingSystemRuntimeCommon/Public/ClothTetherData.h
@@ -46,4 +46,5 @@ struct TStructOpsTypeTraits<FClothTetherData> : public TStructOpsTypeTraitsBase2
 	{
 		WithSerializer = true,
 	};
+	static constexpr EPropertyObjectReferenceType WithSerializerObjectReferences = EPropertyObjectReferenceType::None;
 };
diff --git a/Engine/Source/Runtime/Engine/Classes/AI/Navigation/NavAgentSelector.h b/Engine/Source/Runtime/Engine/Classes/AI/Navigation/NavAgentSelector.h
--- a/Engine/Source/Runtime/Engine/Classes/AI/Navigation/NavAgentSelector.h
+++ b/Engine/Source/Runtime/Engine/Classes/AI/Navigation/NavAgentSelector.h
@@ -112,5 +112,6 @@ struct TStructOpsTypeTraits< FNavAgentSelector > : public TStructOpsTypeTraitsBa
 	{
 		WithSerializer = true,
 	};
+	static constexpr EPropertyObjectReferenceType WithSerializerObjectReferences = EPropertyObjectReferenceType::None;
 };
 
diff --git a/Engine/Source/Runtime/Engine/Classes/Animation/SmartName.h b/Engine/Source/Runtime/Engine/Classes/Animation/SmartName.h
--- a/Engine/Source/Runtime/Engine/Classes/Animation/SmartName.h
+++ b/Engine/Source/Runtime/Engine/Classes/Animation/SmartName.h
@@ -262,4 +262,5 @@ struct TStructOpsTypeTraits<FSmartName> : public TStructOpsTypeTraitsBase2<FSmar
 		WithSerializer = true,
 		WithIdenticalViaEquality = true
 	};
+	static constexpr EPropertyObjectReferenceType WithSerializerObjectReferences = EPropertyObjectReferenceType::None;
 };
diff --git a/Engine/Source/Runtime/Engine/Classes/Curves/KeyHandle.h b/Engine/Source/Runtime/Engine/Classes/Curves/KeyHandle.h
--- a/Engine/Source/Runtime/Engine/Classes/Curves/KeyHandle.h
+++ b/Engine/Source/Runtime/Engine/Classes/Curves/KeyHandle.h
@@ -121,6 +121,7 @@ struct TStructOpsTypeTraits<FKeyHandleMap>
 		WithCopy = false,
 		WithIdenticalViaEquality = true,
 	};
+	static constexpr EPropertyObjectReferenceType WithSerializerObjectReferences = EPropertyObjectReferenceType::None;
 };
 
 /**
@@ -207,4 +208,5 @@ struct TStructOpsTypeTraits<FKeyHandleLookupTable>
 	{
 		WithSerializer = true,
 	};
+	static constexpr EPropertyObjectReferenceType WithSerializerObjectReferences = EPropertyObjectReferenceType::None;
 };
\ No newline at end of file
diff --git a/Engine/Source/Runtime/Engine/Classes/Curves/NameCurve.h b/Engine/Source/Runtime/Engine/Classes/Curves/NameCurve.h
--- a/Engine/Source/Runtime/Engine/Classes/Curves/NameCurve.h
+++ b/Engine/Source/Runtime/Engine/Classes/Curves/NameCurve.h
@@ -71,6 +71,7 @@ struct TStructOpsTypeTraits<FNameCurveKey>
 		WithCopy = false,
 		WithIdenticalViaEquality = true,
 	};
+	static constexpr EPropertyObjectReferenceType WithSerializerObjectReferences = EPropertyObjectReferenceType::None;
 };
 
 
diff --git a/Engine/Source/Runtime/Engine/Classes/Curves/RichCurve.h b/Engine/Source/Runtime/Engine/Classes/Curves/RichCurve.h
--- a/Engine/Source/Runtime/Engine/Classes/Curves/RichCurve.h
+++ b/Engine/Source/Runtime/Engine/Classes/Curves/RichCurve.h
@@ -188,6 +188,7 @@ struct TStructOpsTypeTraits<FRichCurveKey>
 		WithCopy = false,
 		WithIdenticalViaEquality = true,
 	};
+	static constexpr EPropertyObjectReferenceType WithSerializerObjectReferences = EPropertyObjectReferenceType::None;
 };
 
 
@@ -436,6 +437,7 @@ struct TStructOpsTypeTraits<FCompressedRichCurve>
 		WithCopy = false,
 		WithIdenticalViaEquality = true,
 	};
+	static constexpr EPropertyObjectReferenceType WithSerializerObjectReferences = EPropertyObjectReferenceType::None;
 };
 
 /**
diff --git a/Engine/Source/Runtime/Engine/Classes/Curves/SimpleCurve.h b/Engine/Source/Runtime/Engine/Classes/Curves/SimpleCurve.h
--- a/Engine/Source/Runtime/Engine/Classes/Curves/SimpleCurve.h
+++ b/Engine/Source/Runtime/Engine/Classes/Curves/SimpleCurve.h
@@ -60,6 +60,7 @@ struct TStructOpsTypeTraits<FSimpleCurveKey>
 		WithCopy = false,
 		WithIdenticalViaEquality = true,
 	};
+	static constexpr EPropertyObjectReferenceType WithSerializerObjectReferences = EPropertyObjectReferenceType::None;
 };
 
 
diff --git a/Engine/Source/Runtime/Engine/Classes/Curves/StringCurve.h b/Engine/Source/Runtime/Engine/Classes/Curves/StringCurve.h
--- a/Engine/Source/Runtime/Engine/Classes/Curves/StringCurve.h
+++ b/Engine/Source/Runtime/Engine/Classes/Curves/StringCurve.h
@@ -63,6 +63,7 @@ struct TStructOpsTypeTraits<FStringCurveKey>
 		WithCopy = false,
 		WithIdenticalViaEquality = true,
 	};
+	static constexpr EPropertyObjectReferenceType WithSerializerObjectReferences = EPropertyObjectReferenceType::None;
 };
 
 
diff --git a/Engine/Source/Runtime/Engine/Classes/EdGraph/EdGraphPin.h b/Engine/Source/Runtime/Engine/Classes/EdGraph/EdGraphPin.h
--- a/Engine/Source/Runtime/Engine/Classes/EdGraph/EdGraphPin.h
+++ b/Engine/Source/Runtime/Engine/Classes/EdGraph/EdGraphPin.h
@@ -232,6 +232,7 @@ struct TStructOpsTypeTraits< FEdGraphPinType > : public TStructOpsTypeTraitsBase
 		WithSerializer = true,
 		WithPostSerialize = true
 	};
+	static constexpr EPropertyObjectReferenceType WithSerializerObjectReferences = EPropertyObjectReferenceType::Strong;
 };
 
 UENUM()
diff --git a/Engine/Source/Runtime/Engine/Classes/Engine/Font.h b/Engine/Source/Runtime/Engine/Classes/Engine/Font.h
--- a/Engine/Source/Runtime/Engine/Classes/Engine/Font.h
+++ b/Engine/Source/Runtime/Engine/Classes/Engine/Font.h
@@ -86,6 +86,7 @@ struct TStructOpsTypeTraits<FFontCharacter> : public TStructOpsTypeTraitsBase2<F
 	{
 		WithSerializer = true,
 	};
+	static constexpr EPropertyObjectReferenceType WithSerializerObjectReferences = EPropertyObjectReferenceType::None;
 };
 
 
diff --git a/Engine/Source/Runtime/Engine/Classes/Engine/SkeletalMeshSampling.h b/Engine/Source/Runtime/Engine/Classes/Engine/SkeletalMeshSampling.h
--- a/Engine/Source/Runtime/Engine/Classes/Engine/SkeletalMeshSampling.h
+++ b/Engine/Source/Runtime/Engine/Classes/Engine/SkeletalMeshSampling.h
@@ -52,6 +52,7 @@ struct ENGINE_API FSkeletalMeshSamplingRegionBuiltData
 template<> struct TStructOpsTypeTraits<FSkeletalMeshSamplingRegionBuiltData> : public TStructOpsTypeTraitsBase2<FSkeletalMeshSamplingRegionBuiltData>
 {
 	enum{ WithSerializer = true };
+	static constexpr EPropertyObjectReferenceType WithSerializerObjectReferences = EPropertyObjectReferenceType::None;
 };
 
 /** Built data for sampling a the whole mesh at a particular LOD. */
@@ -75,6 +76,7 @@ struct ENGINE_API FSkeletalMeshSamplingLODBuiltData
 template<> struct TStructOpsTypeTraits<FSkeletalMeshSamplingLODBuiltData> : public TStructOpsTypeTraitsBase2<FSkeletalMeshSamplingLODBuiltData>
 {
 	enum{ WithSerializer = true };
+	static constexpr EPropertyObjectReferenceType WithSerializerObjectReferences = EPropertyObjectReferenceType::None;
 };
 
 USTRUCT()
diff --git a/Engine/Source/Runtime/Engine/Classes/GameFramework/OnlineReplStructs.h b/Engine/Source/Runtime/Engine/Classes/GameFramework/OnlineReplStructs.h
--- a/Engine/Source/Runtime/Engine/Classes/GameFramework/OnlineReplStructs.h
+++ b/Engine/Source/Runtime/Engine/Classes/GameFramework/OnlineReplStructs.h
@@ -149,6 +149,7 @@ struct TStructOpsTypeTraits<FUniqueNetIdRepl> : public TStructOpsTypeTraitsBase2
 		// Import string contents as a unique id
 		WithImportTextItem = true
 	};
+	static constexpr EPropertyObjectReferenceType WithSerializerObjectReferences = EPropertyObjectReferenceType::None;
 };
 
 /** Test harness for Unique Id replication */
diff --git a/Engine/Source/Runtime/Engine/Classes/PhysicsEngine/LevelSetElem.h b/Engine/Source/Runtime/Engine/Classes/PhysicsEngine/LevelSetElem.h
--- a/Engine/Source/Runtime/Engine/Classes/PhysicsEngine/LevelSetElem.h
+++ b/Engine/Source/Runtime/Engine/Classes/PhysicsEngine/LevelSetElem.h
@@ -93,4 +93,5 @@ struct TStructOpsTypeTraits<FKLevelSetElem> : public TStructOpsTypeTraitsBase2<F
 	{
 		WithSerializer = true
 	};
+	static constexpr EPropertyObjectReferenceType WithSerializerObjectReferences = EPropertyObjectReferenceType::None;
 };
diff --git a/Engine/Source/Runtime/Engine/Public/Animation/AnimTypes.h b/Engine/Source/Runtime/Engine/Public/Animation/AnimTypes.h
--- a/Engine/Source/Runtime/Engine/Public/Animation/AnimTypes.h
+++ b/Engine/Source/Runtime/Engine/Public/Animation/AnimTypes.h
@@ -539,6 +539,7 @@ template<> struct TStructOpsTypeTraits<FAnimSyncMarker> : public TStructOpsTypeT
 	{ 
 		WithSerializer = true
 	};
+	static constexpr EPropertyObjectReferenceType WithSerializerObjectReferences = EPropertyObjectReferenceType::None;
 };
 #endif
 
@@ -933,6 +934,7 @@ struct ENGINE_API FRawAnimSequenceTrack
 template<> struct TStructOpsTypeTraits<FRawAnimSequenceTrack> : public TStructOpsTypeTraitsBase2<FRawAnimSequenceTrack>
 {
 	enum { WithSerializer = true };
+	static constexpr EPropertyObjectReferenceType WithSerializerObjectReferences = EPropertyObjectReferenceType::None;
 };
 
 UCLASS()
diff --git a/Engine/Source/Runtime/Engine/Public/MaterialExpressionIO.h b/Engine/Source/Runtime/Engine/Public/MaterialExpressionIO.h
--- a/Engine/Source/Runtime/Engine/Public/MaterialExpressionIO.h
+++ b/Engine/Source/Runtime/Engine/Public/MaterialExpressionIO.h
@@ -119,6 +119,7 @@ struct TStructOpsTypeTraits<FExpressionInput>
 	{
 		WithSerializer = true,
 	};
+	static constexpr EPropertyObjectReferenceType WithSerializerObjectReferences = EPropertyObjectReferenceType::None;
 };
 
 //
@@ -199,6 +200,8 @@ struct TStructOpsTypeTraits<FColorMaterialInput>
 	{
 		WithSerializer = true,
 	};
+
+	static constexpr EPropertyObjectReferenceType WithSerializerObjectReferences = EPropertyObjectReferenceType::None;
 };
 
 struct FScalarMaterialInput : FMaterialInput<float>
@@ -218,6 +221,7 @@ struct TStructOpsTypeTraits<FScalarMaterialInput>
 	{
 		WithSerializer = true,
 	};
+	static constexpr EPropertyObjectReferenceType WithSerializerObjectReferences = EPropertyObjectReferenceType::None;
 };
 
 struct FShadingModelMaterialInput : FMaterialInput<uint32>
@@ -275,6 +279,7 @@ struct TStructOpsTypeTraits<FVectorMaterialInput>
 	{
 		WithSerializer = true,
 	};
+	static constexpr EPropertyObjectReferenceType WithSerializerObjectReferences = EPropertyObjectReferenceType::None;
 };
 
 struct FVector2MaterialInput : FMaterialInput<FVector2f>
@@ -294,6 +299,7 @@ struct TStructOpsTypeTraits<FVector2MaterialInput>
 	{
 		WithSerializer = true,
 	};
+	static constexpr EPropertyObjectReferenceType WithSerializerObjectReferences = EPropertyObjectReferenceType::None;
 };
 
 struct FMaterialAttributesInput : FExpressionInput
@@ -330,4 +336,5 @@ struct TStructOpsTypeTraits<FMaterialAttributesInput>
 	{
 		WithSerializer = true,
 	};
+	static constexpr EPropertyObjectReferenceType WithSerializerObjectReferences = EPropertyObjectReferenceType::None;
 };
diff --git a/Engine/Source/Runtime/Engine/Public/PerPlatformProperties.h b/Engine/Source/Runtime/Engine/Public/PerPlatformProperties.h
--- a/Engine/Source/Runtime/Engine/Public/PerPlatformProperties.h
+++ b/Engine/Source/Runtime/Engine/Public/PerPlatformProperties.h
@@ -295,6 +295,7 @@ struct TStructOpsTypeTraits<FPerPlatformInt>
 		WithSerializeFromMismatchedTag = true,
 		WithSerializer = true
 	};
+	static constexpr EPropertyObjectReferenceType WithSerializerObjectReferences = EPropertyObjectReferenceType::None;
 };
 
 struct FFreezablePerPlatformFloat;
@@ -376,6 +377,7 @@ struct TStructOpsTypeTraits<FPerPlatformFloat>
 		WithSerializeFromMismatchedTag = true,
 		WithSerializer = true
 	};
+	static constexpr EPropertyObjectReferenceType WithSerializerObjectReferences = EPropertyObjectReferenceType::None;
 };
 
 /** FPerPlatformBool - bool property with per-platform overrides */
@@ -416,6 +418,7 @@ struct TStructOpsTypeTraits<FPerPlatformBool>
 		WithSerializeFromMismatchedTag = true,
 		WithSerializer = true
 	};
+	static constexpr EPropertyObjectReferenceType WithSerializerObjectReferences = EPropertyObjectReferenceType::None;
 };
 
 /** FPerPlatformFrameRate - FFrameRate property with per-platform overrides */
@@ -472,4 +475,5 @@ struct TStructOpsTypeTraits<FPerPlatformFrameRate>
 		WithSerializeFromMismatchedTag = false,
 		WithSerializer = true
 	};
+	static constexpr EPropertyObjectReferenceType WithSerializerObjectReferences = EPropertyObjectReferenceType::None;
 };
\ No newline at end of file
diff --git a/Engine/Source/Runtime/Engine/Public/PerQualityLevelProperties.h b/Engine/Source/Runtime/Engine/Public/PerQualityLevelProperties.h
--- a/Engine/Source/Runtime/Engine/Public/PerQualityLevelProperties.h
+++ b/Engine/Source/Runtime/Engine/Public/PerQualityLevelProperties.h
@@ -213,6 +213,7 @@ struct TStructOpsTypeTraits<FPerQualityLevelInt>
 		WithSerializeFromMismatchedTag = true,
 		WithSerializer = true
 	};
+	static constexpr EPropertyObjectReferenceType WithSerializerObjectReferences = EPropertyObjectReferenceType::None;
 };
 
 USTRUCT()
@@ -255,5 +256,6 @@ struct TStructOpsTypeTraits<FPerQualityLevelFloat>
 		WithSerializeFromMismatchedTag = true,
 		WithSerializer = true
 	};
+	static constexpr EPropertyObjectReferenceType WithSerializerObjectReferences = EPropertyObjectReferenceType::None;
 };
 
diff --git a/Engine/Source/Runtime/Experimental/GeometryCollectionEngine/Public/GeometryCollection/GeometryCollectionObject.h b/Engine/Source/Runtime/Experimental/GeometryCollectionEngine/Public/GeometryCollection/GeometryCollectionObject.h
--- a/Engine/Source/Runtime/Experimental/GeometryCollectionEngine/Public/GeometryCollection/GeometryCollectionObject.h
+++ b/Engine/Source/Runtime/Experimental/GeometryCollectionEngine/Public/GeometryCollection/GeometryCollectionObject.h
@@ -322,6 +322,7 @@ struct TStructOpsTypeTraits<FGeometryCollectionSizeSpecificData> : public TStruc
 		WithPostSerialize = true
 #endif
 	};
+	static constexpr EPropertyObjectReferenceType WithSerializerObjectReferences = EPropertyObjectReferenceType::None;
 };
 
 
diff --git a/Engine/Source/Runtime/GameplayTags/Classes/GameplayTagContainer.h b/Engine/Source/Runtime/GameplayTags/Classes/GameplayTagContainer.h
--- a/Engine/Source/Runtime/GameplayTags/Classes/GameplayTagContainer.h
+++ b/Engine/Source/Runtime/GameplayTags/Classes/GameplayTagContainer.h
@@ -713,6 +713,7 @@ struct TStructOpsTypeTraits<FGameplayTagContainer> : public TStructOpsTypeTraits
 		WithCopy = true,
 		WithPostScriptConstruct = true,
 	};
+	static constexpr EPropertyObjectReferenceType WithSerializerObjectReferences = EPropertyObjectReferenceType::None;
 };
 
 /** Class that can be subclassed by a game/plugin to allow easily adding native gameplay tags at startup */
diff --git a/Engine/Source/Runtime/MovieScene/Public/Channels/MovieSceneChannelData.h b/Engine/Source/Runtime/MovieScene/Public/Channels/MovieSceneChannelData.h
--- a/Engine/Source/Runtime/MovieScene/Public/Channels/MovieSceneChannelData.h
+++ b/Engine/Source/Runtime/MovieScene/Public/Channels/MovieSceneChannelData.h
@@ -46,6 +46,7 @@ struct TStructOpsTypeTraits<FMovieSceneKeyHandleMap>
 	{
 		WithSerializer = true,
 	};
+	static constexpr EPropertyObjectReferenceType WithSerializerObjectReferences = EPropertyObjectReferenceType::None;
 };
 
 namespace UE
diff --git a/Engine/Source/Runtime/MovieScene/Public/Channels/MovieSceneCurveChannelCommon.h b/Engine/Source/Runtime/MovieScene/Public/Channels/MovieSceneCurveChannelCommon.h
--- a/Engine/Source/Runtime/MovieScene/Public/Channels/MovieSceneCurveChannelCommon.h
+++ b/Engine/Source/Runtime/MovieScene/Public/Channels/MovieSceneCurveChannelCommon.h
@@ -92,5 +92,6 @@ struct TStructOpsTypeTraits<FMovieSceneTangentData>
 		WithCopy = false,
 		WithIdenticalViaEquality = true,
 	};
+	static constexpr EPropertyObjectReferenceType WithSerializerObjectReferences = EPropertyObjectReferenceType::None;
 };
 
diff --git a/Engine/Source/Runtime/MovieScene/Public/Channels/MovieSceneDoubleChannel.h b/Engine/Source/Runtime/MovieScene/Public/Channels/MovieSceneDoubleChannel.h
--- a/Engine/Source/Runtime/MovieScene/Public/Channels/MovieSceneDoubleChannel.h
+++ b/Engine/Source/Runtime/MovieScene/Public/Channels/MovieSceneDoubleChannel.h
@@ -108,6 +108,7 @@ struct TStructOpsTypeTraits<FMovieSceneDoubleValue>
 		WithCopy = false,
 		WithIdenticalViaEquality = true,
 	};
+	static constexpr EPropertyObjectReferenceType WithSerializerObjectReferences = EPropertyObjectReferenceType::None;
 };
 
 
@@ -353,6 +354,7 @@ struct TStructOpsTypeTraits<FMovieSceneDoubleChannel> : public TStructOpsTypeTra
 		WithStructuredSerializeFromMismatchedTag = true, 
 	    WithSerializer = true
 	};
+	static constexpr EPropertyObjectReferenceType WithSerializerObjectReferences = EPropertyObjectReferenceType::None;
 };
 
 template<>
diff --git a/Engine/Source/Runtime/MovieScene/Public/Channels/MovieSceneFloatChannel.h b/Engine/Source/Runtime/MovieScene/Public/Channels/MovieSceneFloatChannel.h
--- a/Engine/Source/Runtime/MovieScene/Public/Channels/MovieSceneFloatChannel.h
+++ b/Engine/Source/Runtime/MovieScene/Public/Channels/MovieSceneFloatChannel.h
@@ -107,6 +107,7 @@ struct TStructOpsTypeTraits<FMovieSceneFloatValue>
 		WithCopy = false,
 		WithIdenticalViaEquality = true,
 	};
+	static constexpr EPropertyObjectReferenceType WithSerializerObjectReferences = EPropertyObjectReferenceType::None;
 };
 
 
@@ -353,6 +354,7 @@ struct TStructOpsTypeTraits<FMovieSceneFloatChannel> : public TStructOpsTypeTrai
 		WithPostSerialize = true,
 #endif
 	};
+	static constexpr EPropertyObjectReferenceType WithSerializerObjectReferences = EPropertyObjectReferenceType::None;
 };
 
 template<>
diff --git a/Engine/Source/Runtime/MovieScene/Public/Evaluation/MovieSceneEvaluationKey.h b/Engine/Source/Runtime/MovieScene/Public/Evaluation/MovieSceneEvaluationKey.h
--- a/Engine/Source/Runtime/MovieScene/Public/Evaluation/MovieSceneEvaluationKey.h
+++ b/Engine/Source/Runtime/MovieScene/Public/Evaluation/MovieSceneEvaluationKey.h
@@ -124,4 +124,5 @@ struct TStructOpsTypeTraits<FMovieSceneEvaluationKey> : public TStructOpsTypeTra
 	{
 		WithSerializer = true, WithIdenticalViaEquality = true
 	};
+	static constexpr EPropertyObjectReferenceType WithSerializerObjectReferences = EPropertyObjectReferenceType::None;
 };
diff --git a/Engine/Source/Runtime/MovieScene/Public/Evaluation/MovieSceneSegment.h b/Engine/Source/Runtime/MovieScene/Public/Evaluation/MovieSceneSegment.h
--- a/Engine/Source/Runtime/MovieScene/Public/Evaluation/MovieSceneSegment.h
+++ b/Engine/Source/Runtime/MovieScene/Public/Evaluation/MovieSceneSegment.h
@@ -72,6 +72,7 @@ struct FMovieSceneSegmentIdentifier
 template<> struct TStructOpsTypeTraits<FMovieSceneSegmentIdentifier> : public TStructOpsTypeTraitsBase2<FMovieSceneSegmentIdentifier>
 {
 	enum { WithSerializer = true, WithIdenticalViaEquality = true };
+	static constexpr EPropertyObjectReferenceType WithSerializerObjectReferences = EPropertyObjectReferenceType::None;
 };
 
 /**
@@ -246,4 +247,5 @@ struct FMovieSceneSegment
 template<> struct TStructOpsTypeTraits<FMovieSceneSegment> : public TStructOpsTypeTraitsBase2<FMovieSceneSegment>
 {
 	enum { WithSerializer = true, WithCopy = true, WithIdenticalViaEquality = true };
+	static constexpr EPropertyObjectReferenceType WithSerializerObjectReferences = EPropertyObjectReferenceType::None;
 };
diff --git a/Engine/Source/Runtime/MovieScene/Public/Evaluation/MovieSceneTrackIdentifier.h b/Engine/Source/Runtime/MovieScene/Public/Evaluation/MovieSceneTrackIdentifier.h
--- a/Engine/Source/Runtime/MovieScene/Public/Evaluation/MovieSceneTrackIdentifier.h
+++ b/Engine/Source/Runtime/MovieScene/Public/Evaluation/MovieSceneTrackIdentifier.h
@@ -88,4 +88,5 @@ struct TStructOpsTypeTraits<FMovieSceneTrackIdentifier> : public TStructOpsTypeT
 	{
 		WithSerializer = true, WithIdenticalViaEquality = true
 	};
+	static constexpr EPropertyObjectReferenceType WithSerializerObjectReferences = EPropertyObjectReferenceType::None;
 };
diff --git a/Engine/Source/Runtime/MovieScene/Public/MovieSceneFrameMigration.h b/Engine/Source/Runtime/MovieScene/Public/MovieSceneFrameMigration.h
--- a/Engine/Source/Runtime/MovieScene/Public/MovieSceneFrameMigration.h
+++ b/Engine/Source/Runtime/MovieScene/Public/MovieSceneFrameMigration.h
@@ -109,6 +109,7 @@ struct TStructOpsTypeTraits<FMovieSceneFrameRange> : public TStructOpsTypeTraits
 {
 	enum { WithStructuredSerializeFromMismatchedTag = true, WithSerializer = true, WithIdenticalViaEquality = true,
 		   WithExportTextItem = true, WithImportTextItem = true};
+	static constexpr EPropertyObjectReferenceType WithSerializerObjectReferences = EPropertyObjectReferenceType::None;
 };
 
 /**
diff --git a/Engine/Source/Runtime/MovieScene/Public/MovieSceneSequenceID.h b/Engine/Source/Runtime/MovieScene/Public/MovieSceneSequenceID.h
--- a/Engine/Source/Runtime/MovieScene/Public/MovieSceneSequenceID.h
+++ b/Engine/Source/Runtime/MovieScene/Public/MovieSceneSequenceID.h
@@ -87,6 +87,7 @@ struct TStructOpsTypeTraits<FMovieSceneSequenceID> : public TStructOpsTypeTraits
 		WithSerializer = true,
 		WithCopy = true
 	};
+	static constexpr EPropertyObjectReferenceType WithSerializerObjectReferences = EPropertyObjectReferenceType::None;
 };
 
 #if UE_MOVIESCENE_ENTITY_DEBUG
