World Partition : When sorting world partition cells, don't rely on cell grid level, but instead use grid cell extent as we can sort cells between multiple grids.

diff --git a/Engine/Source/Runtime/Engine/Private/WorldPartition/WorldPartitionRuntimeCellDataSpatialHash.cpp b/Engine/Source/Runtime/Engine/Private/WorldPartition/WorldPartitionRuntimeCellDataSpatialHash.cpp
--- a/Engine/Source/Runtime/Engine/Private/WorldPartition/WorldPartitionRuntimeCellDataSpatialHash.cpp
+++ b/Engine/Source/Runtime/Engine/Private/WorldPartition/WorldPartitionRuntimeCellDataSpatialHash.cpp
@@ -12,6 +12,12 @@ static FAutoConsoleVariableRef CVarRuntimeSpatialHashCellToSourceAngleContributi
 	GRuntimeSpatialHashCellToSourceAngleContributionToCellImportance,
 	TEXT("Value between 0 and 1 that modulates the contribution of the angle between streaming source-to-cell vector and source-forward vector to the cell importance. The closest to 0, the less the angle will contribute to the cell importance."));
 
+static bool GRuntimeSpatialHashSortUsingCellExtent = true;
+static FAutoConsoleVariableRef CVarRuntimeSpatialHashSortUsingCellExtent(
+	TEXT("wp.Runtime.RuntimeSpatialHashSortUsingCellExtent"),
+	GRuntimeSpatialHashSortUsingCellExtent,
+	TEXT("Set to 1 to use cell extent instead of cell grid level when sorting cells by importance."));
+
 UWorldPartitionRuntimeCellDataSpatialHash::UWorldPartitionRuntimeCellDataSpatialHash(const FObjectInitializer& ObjectInitializer)
 	: Super(ObjectInitializer)
 	, Extent(0)
@@ -142,8 +148,8 @@ int32 UWorldPartitionRuntimeCellDataSpatialHash::SortCompare(const UWorldPartiti
 	{
 		const UWorldPartitionRuntimeCellDataSpatialHash* Other = (UWorldPartitionRuntimeCellDataSpatialHash*)InOther;
 		
-		// Level (higher value is higher prio)
-		Result = Other->Level - Level;
+		// By default, now compare cell's extent instead of its grid level since we compare cells across multiple WPs/grids (higher value is higher prio)
+		Result = GRuntimeSpatialHashSortUsingCellExtent ? int32(Other->Extent - Extent) : (Other->Level - Level);
 		if (bCanUseSortingCache && (Result == 0))
 		{
 			// Closest distance (lower value is higher prio)
