- Fixed WorldPartitionLevelStreamingDynamic to properly maintain World's NumStreamingLevelsBeingLoaded when running PIE.
- Added getter in World for NumStreamingLevelsBeingLoaded (for future use by world partition)

diff --git a/Engine/Source/Runtime/Engine/Classes/Engine/LevelStreaming.h b/Engine/Source/Runtime/Engine/Classes/Engine/LevelStreaming.h
--- a/Engine/Source/Runtime/Engine/Classes/Engine/LevelStreaming.h
+++ b/Engine/Source/Runtime/Engine/Classes/Engine/LevelStreaming.h
@@ -681,6 +681,9 @@ protected:
 	/** Called by SetLoadedLevel */
 	virtual void OnLevelLoadedChanged(ULevel* Level) {}
 
+	void OnLoadingStarted();
+	void OnLoadingFinished();
+
 	/** @return Name of the level package that is currently loaded.																	*/
 	FName GetLoadedLevelPackageName() const;
 
diff --git a/Engine/Source/Runtime/Engine/Classes/Engine/World.h b/Engine/Source/Runtime/Engine/Classes/Engine/World.h
--- a/Engine/Source/Runtime/Engine/Classes/Engine/World.h
+++ b/Engine/Source/Runtime/Engine/Classes/Engine/World.h
@@ -1041,6 +1041,8 @@ public:
 	/** Return a const version of the streaming levels array */
 	const TArray<ULevelStreaming*>& GetStreamingLevels() const { return StreamingLevels; }
 
+	uint16 GetNumStreamingLevelsBeingLoaded() const { return NumStreamingLevelsBeingLoaded; }
+
 	/** Returns true if StreamingLevel is part of the levels being considered for update */
 	bool IsStreamingLevelBeingConsidered(ULevelStreaming* StreamingLevel) const { return StreamingLevelsToConsider.Contains(StreamingLevel); }
 	
diff --git a/Engine/Source/Runtime/Engine/Private/LevelStreaming.cpp b/Engine/Source/Runtime/Engine/Private/LevelStreaming.cpp
--- a/Engine/Source/Runtime/Engine/Private/LevelStreaming.cpp
+++ b/Engine/Source/Runtime/Engine/Private/LevelStreaming.cpp
@@ -1605,7 +1605,7 @@ bool ULevelStreaming::RequestLevel(UWorld* PersistentWorld, bool bAllowLevelLoad
 		if (FPackageName::DoesPackageExist(PackagePath, &PackagePath))
 		{
 			SetCurrentState(ELevelStreamingState::Loading);
-			FWorldNotifyStreamingLevelLoading::Started(PersistentWorld);
+			OnLoadingStarted();
 			
 			ULevel::StreamedLevelsOwningWorld.Add(DesiredPackageName, PersistentWorld);
 			UWorld::WorldTypePreLoadMap.FindOrAdd(DesiredPackageName) = PersistentWorld->WorldType;
@@ -1672,13 +1672,7 @@ void ULevelStreaming::AsyncLevelLoadComplete(const FName& InPackageName, UPackag
 {
 	// TODO: Should not set state here so that observers will have access to LoadedLevel
 	SetCurrentState(ELevelStreamingState::LoadedNotVisible);
-	if (UWorld* World = GetWorld())
-	{
-		if (World->GetStreamingLevels().Contains(this))
-		{
-			FWorldNotifyStreamingLevelLoading::Finished(World);
-		}
-	}
+	OnLoadingFinished();
 
 	if (InLoadedPackage)
 	{
@@ -1982,6 +1976,27 @@ FName ULevelStreaming::GetLoadedLevelPackageName() const
 	return CachedLoadedLevelPackageName;
 }
 
+void ULevelStreaming::OnLoadingStarted()
+{
+	UWorld* World = GetWorld();
+	if (World && World->IsGameWorld())
+	{
+		FWorldNotifyStreamingLevelLoading::Started(World);
+	}
+}
+
+void ULevelStreaming::OnLoadingFinished()
+{
+	UWorld* World = GetWorld();
+	if (World && World->IsGameWorld())
+	{
+		if (World->GetStreamingLevels().Contains(this))
+		{
+			FWorldNotifyStreamingLevelLoading::Finished(World);
+		}
+	}
+}
+
 void ULevelStreaming::SetWorldAssetByPackageName(FName InPackageName)
 {
 	// Need to strip PIE prefix from object name, only the package has it
diff --git a/Engine/Source/Runtime/Engine/Private/WorldPartition/WorldPartitionLevelStreamingDynamic.cpp b/Engine/Source/Runtime/Engine/Private/WorldPartition/WorldPartitionLevelStreamingDynamic.cpp
--- a/Engine/Source/Runtime/Engine/Private/WorldPartition/WorldPartitionLevelStreamingDynamic.cpp
+++ b/Engine/Source/Runtime/Engine/Private/WorldPartition/WorldPartitionLevelStreamingDynamic.cpp
@@ -289,8 +289,14 @@ bool UWorldPartitionLevelStreamingDynamic::RequestLevel(UWorld* InPersistentWorl
 
 			if (IssueLoadRequests())
 			{
+				bool bBlock = (InBlockPolicy == AlwaysBlock || (ShouldBeAlwaysLoaded() && InBlockPolicy != NeverBlock));
+				if (const UWorldPartitionRuntimeLevelStreamingCell* RuntimeLevelStreamingCell = StreamingCell.Get())
+				{
+					bBlock = RuntimeLevelStreamingCell->GetOwningWorld()->GetWorldPartition()->IsMainWorldPartition();
+				}
+
 				// Editor immediately blocks on load and we also block if background level streaming is disabled.
-				if (InBlockPolicy == AlwaysBlock || (ShouldBeAlwaysLoaded() && InBlockPolicy != NeverBlock))
+				if (bBlock)
 				{
 					if (IsAsyncLoading())
 					{
@@ -414,8 +420,10 @@ bool UWorldPartitionLevelStreamingDynamic::IssueLoadRequests()
 		ActorContainerDup->MarkAsGarbage();
 	}
 
+	OnLoadingStarted();
 	auto FinalizeLoading = [this](bool bSucceeded)
 	{
+		OnLoadingFinished();
 		check(bLoadRequestInProgress);
 		bLoadRequestInProgress = false;
 		bLoadSucceeded = bSucceeded;
