Scoped FSoundWaveData::LoadZerothChunk code and its mutex to be editor only as it is unused by runtime code and initialize a critical section object for nothing

#rnx

diff --git a/Engine/Source/Runtime/Engine/Classes/Sound/SoundWave.h b/Engine/Source/Runtime/Engine/Classes/Sound/SoundWave.h
--- a/Engine/Source/Runtime/Engine/Classes/Sound/SoundWave.h
+++ b/Engine/Source/Runtime/Engine/Classes/Sound/SoundWave.h
@@ -1355,7 +1355,9 @@ private:
 
 	/** Zeroth Chunk of audio for sources that use Load On Demand. */
 	FBulkDataBuffer<uint8> ZerothChunkData;
+#if WITH_EDITOR
 	mutable FCriticalSection LoadZerothChunkDataCriticalSection;
+#endif
 
 	/* Accessor to get the zeroth chunk which might perform additional work in editor to handle async tasks. */
 	FBulkDataBuffer<uint8>& GetZerothChunkData() const;
diff --git a/Engine/Source/Runtime/Engine/Private/SoundWave.cpp b/Engine/Source/Runtime/Engine/Private/SoundWave.cpp
--- a/Engine/Source/Runtime/Engine/Private/SoundWave.cpp
+++ b/Engine/Source/Runtime/Engine/Private/SoundWave.cpp
@@ -434,6 +434,7 @@ const TArrayView<uint8> FSoundWaveData::GetZerothChunkDataView() const
 
 bool FSoundWaveData::LoadZerothChunk()
 {
+#if WITH_EDITOR
 	TRACE_CPUPROFILER_EVENT_SCOPE(FSoundWaveData::LoadZerothChunk);
 
 	LLM_SCOPE(ELLMTag::AudioSoundWaves);
@@ -448,7 +449,6 @@ bool FSoundWaveData::LoadZerothChunk()
 
 		if (ZerothChunkData.GetView().Num() == 0)
 		{
-#if WITH_EDITOR
 			if (!bLoadedFromCookedData)
 			{
 				// In editor, we wait until the zeroth chunk is required to force finish compilation
@@ -479,7 +479,6 @@ bool FSoundWaveData::LoadZerothChunk()
 				}
 			}
 			else
-#endif // WITH_EDITOR
 			{
 				// The zeroth chunk is cooked out to RunningPlatformData, so retrieve it.
 				if (ensureMsgf(GetNumChunks() > 0, TEXT("SoundWave '%s' is cooked but contains no chunks."), *GetFName().ToString()))
@@ -503,6 +502,7 @@ bool FSoundWaveData::LoadZerothChunk()
 			}
 		}
 	}
+#endif
 
 	return true;
 }
