Reset must be mapped guid lists instead of emptying them when we're not doing cleanup

diff --git a/Engine/Source/Runtime/Engine/Private/DataChannel.cpp b/Engine/Source/Runtime/Engine/Private/DataChannel.cpp
--- a/Engine/Source/Runtime/Engine/Private/DataChannel.cpp
+++ b/Engine/Source/Runtime/Engine/Private/DataChannel.cpp
@@ -1068,7 +1068,7 @@ void UChannel::AppendMustBeMappedGuids( FOutBunch* Bunch )
 
 		Bunch->bHasMustBeMappedGUIDs = 1;
 
-		MustBeMappedGuidsInLastBunch.Empty();
+		MustBeMappedGuidsInLastBunch.Reset();
 	}
 }
 PRAGMA_ENABLE_DEPRECATION_WARNINGS
@@ -1140,7 +1140,7 @@ void UActorChannel::AppendMustBeMappedGuids( FOutBunch* Bunch )
 
 			PackageMapClient->GetMustBeMappedGuidsInLastBunch().Append( QueuedMustBeMappedGuidsInLastBunch );
 
-			QueuedMustBeMappedGuidsInLastBunch.Empty();
+			QueuedMustBeMappedGuidsInLastBunch.Reset();
 		}
 	}
 
diff --git a/Engine/Source/Runtime/Engine/Private/NetConnection.cpp b/Engine/Source/Runtime/Engine/Private/NetConnection.cpp
--- a/Engine/Source/Runtime/Engine/Private/NetConnection.cpp
+++ b/Engine/Source/Runtime/Engine/Private/NetConnection.cpp
@@ -4989,7 +4989,7 @@ void UNetConnection::FlushDormancyForObject(AActor* DormantActor, UObject* Repli
 		if (UPackageMapClient* PackageMapClient = CastChecked<UPackageMapClient>(PackageMap))
 		{
 			TArray< FNetworkGUID >& MustBeMappedGuidsInLastBunch = PackageMapClient->GetMustBeMappedGuidsInLastBunch();
-			MustBeMappedGuidsInLastBunch.Empty();
+			MustBeMappedGuidsInLastBunch.Reset();
 		}
 	}
 }
