Partial backout of CL 25412692 (some debug code that should have never been submitted)

diff --git a/Engine/Source/Runtime/Engine/Private/WorldPartition/WorldPartitionLevelStreamingDynamic.cpp b/Engine/Source/Runtime/Engine/Private/WorldPartition/WorldPartitionLevelStreamingDynamic.cpp
--- a/Engine/Source/Runtime/Engine/Private/WorldPartition/WorldPartitionLevelStreamingDynamic.cpp
+++ b/Engine/Source/Runtime/Engine/Private/WorldPartition/WorldPartitionLevelStreamingDynamic.cpp
@@ -289,14 +289,8 @@ bool UWorldPartitionLevelStreamingDynamic::RequestLevel(UWorld* InPersistentWorl
 
 			if (IssueLoadRequests())
 			{
-				bool bBlock = (InBlockPolicy == AlwaysBlock || (ShouldBeAlwaysLoaded() && InBlockPolicy != NeverBlock));
-				if (const UWorldPartitionRuntimeLevelStreamingCell* RuntimeLevelStreamingCell = StreamingCell.Get())
-				{
-					bBlock = RuntimeLevelStreamingCell->GetOwningWorld()->GetWorldPartition()->IsMainWorldPartition();
-				}
-
 				// Editor immediately blocks on load and we also block if background level streaming is disabled.
-				if (bBlock)
+				if (InBlockPolicy == AlwaysBlock || (ShouldBeAlwaysLoaded() && InBlockPolicy != NeverBlock))
 				{
 					if (IsAsyncLoading())
 					{
