Chunked Seektables for BinkAudio

diff --git a/Engine/Source/Developer/AudioFormatBink/Private/AudioFormatBink.cpp b/Engine/Source/Developer/AudioFormatBink/Private/AudioFormatBink.cpp
--- a/Engine/Source/Developer/AudioFormatBink/Private/AudioFormatBink.cpp
+++ b/Engine/Source/Developer/AudioFormatBink/Private/AudioFormatBink.cpp
@@ -40,7 +40,7 @@ class FAudioFormatBink : public IAudioFormat
 	enum
 	{
 		/** Version for Bink Audio format, this becomes part of the DDC key. */
-		UE_AUDIO_BINK_VER = 5,
+		UE_AUDIO_BINK_VER = 7,
 	};
 
 public:
@@ -217,6 +217,86 @@ public:
 
 		return true;
 	}
+
+	static void StripSeek(TArray<uint8>& InOutBuffer)
+	{
+		void* CompressedData = InOutBuffer.GetData();
+		uint32_t CompressedDataLen = InOutBuffer.Num();
+
+		check(CompressedDataLen > sizeof(BinkAudioFileHeader));
+
+		BinkAudioFileHeader* Header = (BinkAudioFileHeader*)CompressedData;
+		uint32 SeekTableBytes = Header->seek_table_entry_count * sizeof(uint16);
+		check(CompressedDataLen > sizeof(BinkAudioFileHeader) + SeekTableBytes);
+
+		// Mark in the header we don't have any seek table 
+		Header->seek_table_entry_count = 0;
+
+		// Copy the rest of the encoded data over the seek table.
+		uint8* SeekTableStart = (uint8*)(Header + 1);
+		FMemory::Memmove(SeekTableStart, SeekTableStart + SeekTableBytes, CompressedDataLen - SeekTableBytes - sizeof(BinkAudioFileHeader));
+
+		CompressedDataLen -= SeekTableBytes;
+
+		InOutBuffer.SetNum(CompressedDataLen);
+	}
+	
+	virtual bool RequiresStreamingSeekTable() const
+	{
+		return true;
+	}
+	
+	// Taken from BinkAudioInfo.cpp 
+	static uint32 GetMaxFrameSizeSamples(uint32 SampleRate) 
+	{
+		if (SampleRate >= 44100)
+		{
+			return 1920;
+		}
+		else if (SampleRate >= 22050)
+		{
+			return 960;
+		}
+		else
+		{
+			return 480;
+		}
+	}
+
+	virtual bool ExtractSeekTableForStreaming(TArray<uint8>& InOutBuffer, IAudioFormat::FSeekTable& OutSeektable) const
+	{
+		BinkAudioFileHeader const* Header = (BinkAudioFileHeader const*)InOutBuffer.GetData();
+		if (InOutBuffer.Num() < sizeof(BinkAudioFileHeader) || Header->tag != 'UEBA' || Header->seek_table_entry_count==0)
+		{
+			return false;
+		}
+		
+		// The outer logic that manages the seek-table is unaware of how big the header can be, so for the sake of simplicity,
+		// offset for the size of the header in the seek-table entries, so we don't need to worry about adjusting for it later.
+		static const uint32 ActualAudioOffset = sizeof(BinkAudioFileHeader) + 0; // No entries as we're stripping it.
+
+		// Decode and copy out the seek-table. (it's stored as deltas).
+		uint16* EncodedSeekTable = (uint16*)(InOutBuffer.GetData() + sizeof(BinkAudioFileHeader));
+		uint32 CurrentSeekOffset = ActualAudioOffset;
+		uint32 CurrentTimeOffset = 0;
+		
+		int32 SamplesPerEntry = Header->blocks_per_seek_table_entry * GetMaxFrameSizeSamples(Header->rate);
+		OutSeektable.Offsets.SetNum(Header->seek_table_entry_count);
+		OutSeektable.Times.SetNum(Header->seek_table_entry_count);
+
+		for (int32 i = 0; i < Header->seek_table_entry_count; ++i)
+		{
+			OutSeektable.Times[i] = CurrentTimeOffset;
+			OutSeektable.Offsets[i] = CurrentSeekOffset;
+			CurrentSeekOffset += EncodedSeekTable[i];
+			CurrentTimeOffset += SamplesPerEntry;
+		}
+
+		// Strip the seek-table from the buffer now we've copied it.
+		StripSeek(InOutBuffer);
+
+		return true;
+	}
 };
 
 class FAudioPlatformBinkModule : public IAudioFormatModule
diff --git a/Engine/Source/Developer/TargetPlatform/Public/Interfaces/IAudioFormat.h b/Engine/Source/Developer/TargetPlatform/Public/Interfaces/IAudioFormat.h
--- a/Engine/Source/Developer/TargetPlatform/Public/Interfaces/IAudioFormat.h
+++ b/Engine/Source/Developer/TargetPlatform/Public/Interfaces/IAudioFormat.h
@@ -123,6 +123,23 @@ public:
 	 */
 	virtual bool SplitDataForStreaming(const TArray<uint8>& SrcBuffer, TArray<TArray<uint8>>& OutBuffers, const int32 FirstChunkMaxSize, const int32 MaxChunkSize) const {return false;}
 
+	virtual bool RequiresStreamingSeekTable() const { return false; }
+
+	struct FSeekTable
+	{
+		TArray<uint32> Times;			// Times in AudioFrames.
+		TArray<uint32> Offsets;			// Offset in the compressed data.
+	};	
+
+	/**
+	* Extracts the embedded seek-table, removing it, and outputting it separately.
+	* NOTE: TArray is modified in place. The seek-table is parsed and removed from it.	
+	* @param InOutBuffer Pre-compressed data containing seek-table and compressed data as array of bytes.
+	* @param OutSeekTable Seektable in its generic form.
+	* @return Success or failure
+	*/
+	virtual bool ExtractSeekTableForStreaming(TArray<uint8>& InOutBuffer, FSeekTable& OutSeektable) const { return false; }
+
 public:
 
 	/** Virtual destructor. */
diff --git a/Engine/Source/Runtime/BinkAudioDecoder/Module/Private/BinkAudioInfo.cpp b/Engine/Source/Runtime/BinkAudioDecoder/Module/Private/BinkAudioInfo.cpp
--- a/Engine/Source/Runtime/BinkAudioDecoder/Module/Private/BinkAudioInfo.cpp
+++ b/Engine/Source/Runtime/BinkAudioDecoder/Module/Private/BinkAudioInfo.cpp
@@ -114,17 +114,24 @@ void FBinkAudioInfo::PrepareToLoop()
 //-----------------------------------------------------------------------------
 void FBinkAudioInfo::SeekToTime(const float SeekTimeSeconds)
 {
-	if (Decoder->SeekTableCount == 0)
-	{
-		// Can't seek without a seek table.
-		return;
-	}
-
 	uint32 SeekTimeSamples = 0;
 	if (SeekTimeSeconds > 0)
 	{
 		SeekTimeSamples = (uint32)(SeekTimeSeconds * SampleRate);
 	}
+	
+	// If there's no seek table on the header, see if we there's an embedded one on the chunk.
+	if (Decoder->SeekTableCount == 0)
+	{
+		int32 ChunkIndexToSeekTo = StreamingSoundWave->GetSoundWaveData()->FindChunkIndexForSeeking(SeekTimeSamples);
+		if (ChunkIndexToSeekTo >= 0)
+		{
+			this->StreamSeekBlockIndex = ChunkIndexToSeekTo;
+			this->StreamSeekBlockOffset = 0;
+			this->StreamSeekToSamples = SeekTimeSamples;
+		}
+		return;
+	}
 
 	uint32 SamplesInFrame = GetMaxFrameSizeSamples();
 	if (SeekTimeSamples > this->TrueSampleCount)
diff --git a/Engine/Source/Runtime/Engine/Classes/Sound/SoundWave.h b/Engine/Source/Runtime/Engine/Classes/Sound/SoundWave.h
--- a/Engine/Source/Runtime/Engine/Classes/Sound/SoundWave.h
+++ b/Engine/Source/Runtime/Engine/Classes/Sound/SoundWave.h
@@ -76,9 +76,12 @@ struct FStreamedAudioChunk
 	/** Size of the chunk of data in bytes including zero padding */
 	int32 DataSize = 0;
 
-	/** Size of the audio data. */
+	/** Size of the audio data. (NOTE: This includes a seek-table if its present.) */
 	int32 AudioDataSize = 0;
 
+	/** Chunk position in samples frames in the stream. (NOTE: != INDEX_NONE will assume the presence of a seektable). */
+	uint32 SeekOffsetInAudioFrames = INDEX_NONE;
+
 	/** Bulk data if stored in the package. */
 	FByteBulkData BulkData;
 
@@ -1327,6 +1330,9 @@ public:
  	bool IsZerothChunkDataLoaded() const;
  	const TArrayView<uint8> GetZerothChunkDataView() const;
 
+	bool HasChunkSeekTable(int32 InChunkIndex) const;
+	int32 FindChunkIndexForSeeking(uint32 InTimeInAudioFrames) const;
+
 	// Returns true if the zeroth chunk is loaded, or attempts to load it if not already loaded,
 	// returning true if the load was successful. Can return false if either an error was encountered
 	// in attempting to load the chunk or if stream caching is not enabled for the given sound.
diff --git a/Engine/Source/Runtime/Engine/Classes/Sound/StreamedAudioChunkSeekTable.h b/Engine/Source/Runtime/Engine/Classes/Sound/StreamedAudioChunkSeekTable.h
new file mode 100644
--- /dev/null
+++ b/Engine/Source/Runtime/Engine/Classes/Sound/StreamedAudioChunkSeekTable.h
@@ -0,0 +1,122 @@
+// Copyright Epic Games, Inc. All Rights Reserved.
+
+#pragma once
+
+#include "HAL/Platform.h"
+#include "Containers/Array.h"
+#include "Templates/PimplPtr.h"
+
+#include "StreamedAudioChunkSeekTable.generated.h"
+
+// Forward declarations.
+class FArchive;
+
+UENUM()
+enum class EChunkSeekTableMode : uint8
+{
+	ConstantSamplesPerEntry = 0,
+	VariableSamplesPerEntry,
+};
+
+/**
+ * Class representing a audio seek-table. Typically embedded in a bit-stream.
+ */
+class FStreamedAudioChunkSeekTable
+{
+public:
+	
+	/**
+	 * Parse a seek-table from memory.
+	 *
+	 * @param InMemory Pointer to memory where to load from.
+	 * @param InSize Size of memory
+	 * @param InOutOffset Output the offset of the end of the table
+	 * @param OutTable Table that's loaded
+	 * @return true on Success, false on Failure
+	 */
+	static bool Parse(const uint8* InMemory, uint32 InSize, uint32& InOutOffset, FStreamedAudioChunkSeekTable& OutTable);
+
+	struct ISeekTableImpl
+	{
+		virtual ~ISeekTableImpl() = default;
+		virtual int32 Num() const = 0;
+		virtual uint32 FindOffset(uint32 InTimeInAudioFrames) const = 0;
+		virtual uint32 FindTime(uint32 InOffset) const = 0;
+		virtual void Add(uint32 InTimeInAudioFrames, uint32 InOffset) = 0;
+		virtual bool Serialize(FArchive& Ar) = 0;
+	};
+
+	/**
+	 * Construct a seek-table. The mode determines the internal representation.
+	 *
+	 * @param InMode Mode the table will operate in.
+	 */
+	FStreamedAudioChunkSeekTable(EChunkSeekTableMode InMode = EChunkSeekTableMode::ConstantSamplesPerEntry);	
+
+	/**
+	 * Add an item to the seek table.
+	 *
+	 * @param InTimeInAudioFrames Time of this entry in audio frames
+	 * @param InOffset Offset in the file of where this time happens
+	 */
+	void Add(uint32 InTimeInAudioFrames, uint32 InOffset)
+	{
+		Impl->Add(InTimeInAudioFrames,InOffset);
+	}
+
+	/**
+	 * Finds an offset for seeking given a time
+	 * 
+	 * @param InTimeInAudioFrames Time to search for in audio frames
+	 * @return Offset indicated by the table. INDEX_NONE (~0) on error.
+	 */
+	uint32 FindOffset(uint32 InTimeInAudioFrames) const
+	{
+		return Impl->FindOffset(InTimeInAudioFrames);
+	}
+
+	/**
+	 * Finds a time given an offset (reverse look up)
+	 *
+	 * @param InOffset Offset in the file
+	 * @return Time indicated by the table. INDEX_NONE (~0) on error. 
+	 */
+	uint32 FindTime(uint32 InOffset) const
+	{
+		return Impl->FindTime(InOffset);
+	}
+	
+	/**
+	 * Get the number of Entries in the seek table.
+	 * @return Num of Items.
+	 */
+	int32 Num() const
+	{
+		return Impl->Num();
+	}
+	
+	/**
+	 * Serialize the table to/from the archive
+	 *
+	 * @param Ar Archive to load/save
+	 * @return True on success, false on failure.
+	 */
+	bool Serialize(FArchive& Ar);
+
+	/**
+	 * Empties the table.
+	 *
+	 */
+	void Reset();
+
+private:
+
+	static uint32 GetMagic();
+	static TUniquePtr<ISeekTableImpl> CreateImpl(EChunkSeekTableMode InMode);
+
+	void SetMode(EChunkSeekTableMode InMode);
+	
+	TUniquePtr<ISeekTableImpl> Impl;
+	EChunkSeekTableMode Mode = EChunkSeekTableMode::ConstantSamplesPerEntry;
+};
+
diff --git a/Engine/Source/Runtime/Engine/Private/AudioDecompress.cpp b/Engine/Source/Runtime/Engine/Private/AudioDecompress.cpp
--- a/Engine/Source/Runtime/Engine/Private/AudioDecompress.cpp
+++ b/Engine/Source/Runtime/Engine/Private/AudioDecompress.cpp
@@ -9,6 +9,7 @@
 #include "Misc/CoreStats.h"
 #include "Misc/ScopeRWLock.h"
 #include "Stats/StatsTrace.h"
+#include "Sound/StreamedAudioChunkSeekTable.h"
 
 IStreamedCompressedInfo::IStreamedCompressedInfo()
 	: bIsStreaming(false)
@@ -157,6 +158,18 @@ bool IStreamedCompressedInfo::StreamCompressedInfoInternal(const FSoundWaveProxy
 	uint32 ChunkSize = 0;
 	const uint8* FirstChunk = GetLoadedChunk(StreamingSoundWave, CurrentChunkIndex, ChunkSize);
 
+	// If there is a seek-table present in the chunk, we'll need to parse that first.
+	if (InWaveProxy->GetSoundWaveData()->HasChunkSeekTable(CurrentChunkIndex))
+	{
+		uint32 SeekTableEnd = 0;
+		if (ensureMsgf(FStreamedAudioChunkSeekTable::Parse(FirstChunk, ChunkSize, SeekTableEnd, GetCurrentSeekTable()),
+			TEXT("Failed to parse seektable in '%s' chunk=%d"), *StreamingSoundWave->GetFName().ToString(), CurrentChunkIndex))
+		{
+			ChunkSize -= SeekTableEnd;
+			FirstChunk += SeekTableEnd;
+		}
+	}
+	
 	bIsStreaming = true;
 	if (FirstChunk)
 	{
@@ -218,6 +231,27 @@ bool IStreamedCompressedInfo::StreamCompressedData(uint8* Destination, bool bLoo
 		SrcBufferDataSize = ChunkSize;
 		SrcBufferOffset = StreamSeekBlockOffset;
 		StreamSeekBlockIndex = INDEX_NONE;
+
+		// If we are using a chunked seek-table, we need to parse the table first.
+		if (StreamingSoundWave->GetSoundWaveData()->HasChunkSeekTable(CurrentChunkIndex))
+		{
+			uint32 TableOffset = 0;
+			if (ensureMsgf(FStreamedAudioChunkSeekTable::Parse(SrcBufferData, ChunkSize, TableOffset, GetCurrentSeekTable()), 
+				TEXT("Failed to parse seektable in '%s', chunk=%d"), *StreamingSoundWave->GetFName().ToString(), CurrentChunkIndex))
+			{
+				if (StreamSeekToSamples != INDEX_NONE)
+				{
+					const FStreamedAudioChunk& Chunk = StreamingSoundWave->GetSoundWaveData()->GetChunk(CurrentChunkIndex);
+					int32 SeekOffsetInChunkSpace = StreamSeekToSamples - Chunk.SeekOffsetInAudioFrames;
+					uint32 Offset = GetCurrentSeekTable().FindOffset(StreamSeekToSamples);
+					if (Offset != INDEX_NONE)
+					{
+						SrcBufferOffset = Offset + TableOffset;
+					}
+					StreamSeekToSamples = INDEX_NONE;
+				}
+			}
+		}
 	}
 
 	// If next chunk wasn't loaded when last one finished reading, try to get it again now
@@ -230,6 +264,13 @@ bool IStreamedCompressedInfo::StreamCompressedData(uint8* Destination, bool bLoo
 			bPrintChunkFailMessage = true;
 			SrcBufferDataSize = ChunkSize;
 			SrcBufferOffset = CurrentChunkIndex == 0 ? AudioDataOffset : 0;
+
+			// If we are using a chunked seek-table, we need to parse the table first.
+			if (StreamingSoundWave->GetSoundWaveData()->HasChunkSeekTable(CurrentChunkIndex))
+			{
+				ensureMsgf(FStreamedAudioChunkSeekTable::Parse(SrcBufferData, SrcBufferDataSize, SrcBufferOffset, GetCurrentSeekTable()), 
+					TEXT("Failed to parse seektable in '%s' chunk=%d"), *StreamingSoundWave->GetFName().ToString(), CurrentChunkIndex );
+			}
 		}
 		else
 		{
@@ -304,7 +345,6 @@ bool IStreamedCompressedInfo::StreamCompressedData(uint8* Destination, bool bLoo
 					{
 						RawPCMOffset += ZeroBuffer(Destination + RawPCMOffset, BufferSize - RawPCMOffset);
 						OutNumBytesStreamed = BufferSize - RawPCMOffset;
-						break;
 					}
 				}
 				else
@@ -316,6 +356,12 @@ bool IStreamedCompressedInfo::StreamCompressedData(uint8* Destination, bool bLoo
 				SrcBufferData = GetLoadedChunk(StreamingSoundWave, CurrentChunkIndex, SrcBufferDataSize);
 				if (SrcBufferData)
 				{
+					// If we are using a chunked seek-table, we need to parse the table first.
+					if (StreamingSoundWave->GetSoundWaveData()->HasChunkSeekTable(CurrentChunkIndex))
+					{
+						ensureMsgf(FStreamedAudioChunkSeekTable::Parse(SrcBufferData, SrcBufferDataSize, SrcBufferOffset, GetCurrentSeekTable()), 
+							TEXT("Failed to parse seektable in '%s' chunk=%d"), *StreamingSoundWave->GetFName().ToString(), CurrentChunkIndex);
+					}
 					UE_CLOG(PreviousChunkIndex != CurrentChunkIndex, LogAudio, Log, TEXT("Changed current chunk '%s' from %d to %d, Offset %d"), *StreamingSoundWave->GetFName().ToString(), PreviousChunkIndex, CurrentChunkIndex, SrcBufferOffset);
 				}
 				else
@@ -442,6 +488,20 @@ const uint8* IStreamedCompressedInfo::GetLoadedChunk(const FSoundWaveProxyPtr& I
 	}
 }
 
+FStreamedAudioChunkSeekTable& IStreamedCompressedInfo::GetCurrentSeekTable()
+{
+	if (!CurrentChunkSeekTable.IsValid())
+	{
+		CurrentChunkSeekTable = MakePimpl<FStreamedAudioChunkSeekTable>();
+	}
+	return *CurrentChunkSeekTable;
+}
+
+const FStreamedAudioChunkSeekTable& IStreamedCompressedInfo::GetCurrentSeekTable() const
+{
+	return const_cast<IStreamedCompressedInfo*>(this)->GetCurrentSeekTable();
+}
+
 //////////////////////////////////////////////////////////////////////////
 // Copied from IOS - probably want to split and share
 //////////////////////////////////////////////////////////////////////////
diff --git a/Engine/Source/Runtime/Engine/Private/AudioDerivedData.cpp b/Engine/Source/Runtime/Engine/Private/AudioDerivedData.cpp
--- a/Engine/Source/Runtime/Engine/Private/AudioDerivedData.cpp
+++ b/Engine/Source/Runtime/Engine/Private/AudioDerivedData.cpp
@@ -20,6 +20,7 @@
 #include "DSP/FloatArrayMath.h"
 #include "DSP/MultichannelBuffer.h"
 #include "Containers/UnrealString.h"
+#include "Sound/StreamedAudioChunkSeekTable.h"
 
 DEFINE_LOG_CATEGORY_STATIC(LogAudioDerivedData, Log, All);
 
@@ -165,7 +166,7 @@ Derived data key generation.
 
 // If you want to bump this version, generate a new guid using
 // VS->Tools->Create GUID and paste it here. https://www.guidgen.com works too.
-#define STREAMEDAUDIO_DERIVEDDATA_VER		TEXT("acc1e60e03be438ab23c-a3d318425d34")
+#define STREAMEDAUDIO_DERIVEDDATA_VER		TEXT("05921D80AD3149058866915162EE9903")
 
 #ifndef CASE_ENUM_TO_TEXT
 #define CASE_ENUM_TO_TEXT(X) case X: return TEXT(#X);
@@ -445,6 +446,93 @@ class FStreamedAudioCacheDerivedDataWorker : public FNonAbandonableTask
 		return false;
 	}
 
+	static void MakeChunkSeekTable(const IAudioFormat::FSeekTable& InTable, uint32 InChunkStart, uint32 InChunkEnd, FStreamedAudioChunkSeekTable& Out)
+	{
+		TArray<uint32> ChunkTable;
+		for (int32 i=0;i<InTable.Offsets.Num();++i)
+		{
+			uint32 Offset = InTable.Offsets[i];
+			uint32 TimeInAudioFrames = InTable.Times[i];
+
+			if (Offset > InChunkEnd) break;
+			if (Offset >= InChunkStart && Offset < InChunkEnd)
+			{
+				Out.Add(TimeInAudioFrames,Offset - InChunkStart);
+			}
+		}
+	}
+
+	static int32 PrefixChunkWithSeekTable(FStreamedAudioChunkSeekTable& InSeekTable, TArray<uint8>& InOutChunkBytes)
+	{	
+		TArray<uint8> Bytes;
+		FMemoryWriter Writer(Bytes);
+		InSeekTable.Serialize(Writer);
+		InOutChunkBytes.Insert(Bytes,0);
+		return Bytes.Num();
+	}
+
+	static int32 InjectSeekTable(
+		const IAudioFormat::FSeekTable& InTable,
+		const int32 InChunkOffset,
+		TArray<uint8>& InChunkBytes,
+		FStreamedAudioChunk* InNewChunk,
+		EChunkSeekTableMode InMode)
+	{	
+		// Build per chunk seek table and insert in chunk byte array.
+		FStreamedAudioChunkSeekTable ChunkTable(InMode);
+		int32 Start = InChunkOffset;
+		int32 End = Start + InChunkBytes.Num();
+		MakeChunkSeekTable(InTable, Start, End, ChunkTable);
+		int32 SeekTableSize = PrefixChunkWithSeekTable(ChunkTable, InChunkBytes);
+
+		// Mark the chunks offset in sample frames if it contains any entries.
+		if (ChunkTable.Num() > 0)
+		{
+			// Ask for item at beginning of this new chunk table for its time.
+			InNewChunk->SeekOffsetInAudioFrames = ChunkTable.FindTime(0);
+		}
+		else
+		{
+			InNewChunk->SeekOffsetInAudioFrames = INDEX_NONE;
+		}
+
+		return SeekTableSize;
+	}
+
+	EChunkSeekTableMode DetermineSeekTableMode(const IAudioFormat::FSeekTable& InTable)
+	{
+		const TArray<uint32>& Times = InTable.Times;
+		if (InTable.Times.Num()>2)
+		{
+			int32 Delta = Delta = Times[1] - Times[0];
+			for (int32 i = 2; i < Times.Num(); ++i)
+			{
+				if (Times[i] - Times[i - 1] != Delta)
+				{
+					return EChunkSeekTableMode::VariableSamplesPerEntry;
+				}
+			}
+		}
+		return EChunkSeekTableMode::ConstantSamplesPerEntry;
+	}
+	
+	static int32 CalculateWorseCaseSeektableSize(const IAudioFormat::FSeekTable& InTable, const int32 InOffset, int32 InMaxChunkSize, int32 InDataSize, EChunkSeekTableMode InMode)
+	{
+		// Do dry run of creating the tables simply using MaxChunkSize, to gage the rough size to reserve.
+		FStreamedAudioChunkSeekTable ChunkTable(InMode);
+		
+		int32 WorseCase = 0;
+		for(int32 Offset = InOffset; Offset < InDataSize; Offset += InMaxChunkSize)
+		{
+			MakeChunkSeekTable(InTable, Offset, InMaxChunkSize, ChunkTable);
+			TArray<uint8> Bytes;
+			FMemoryWriter Counter(Bytes);
+			ChunkTable.Serialize(Counter);
+			WorseCase = FMath::Max(Bytes.Num(), WorseCase);
+		}
+		return WorseCase;
+	}
+
 	/** Build the streamed audio. This function is safe to call from any thread. */
 	void BuildStreamedAudio()
 	{
@@ -479,6 +567,14 @@ class FStreamedAudioCacheDerivedDataWorker : public FNonAbandonableTask
 				// By default, the first chunk's max size is the same as the other chunks.
 				int32 FirstChunkSize = MaxChunkSizeForCurrentWave;
 
+				// Extract and remove any seek table to split up and embed separately.
+				IAudioFormat::FSeekTable Table;
+				bool bRequiresSeekTable = AudioFormat->RequiresStreamingSeekTable();
+				if (bRequiresSeekTable && !AudioFormat->ExtractSeekTableForStreaming(CompressedBuffer, Table))
+				{
+					UE_LOG(LogAudioDerivedData, Warning, TEXT("SoundWave: %s requires a Seektable, but doesn't contain one. "), *SoundWaveFullName);
+				}
+
 				const int32 MinimumChunkSize = AudioFormat->GetMinimumSizeForInitialChunk(AudioFormatName, CompressedBuffer);
 				const bool bForceLegacyStreamChunking = bIsStreaming && CompressionOverrides && CompressionOverrides->StreamCachingSettings.bForceLegacyStreamChunking;
 
@@ -522,26 +618,48 @@ class FStreamedAudioCacheDerivedDataWorker : public FNonAbandonableTask
 						MaxChunkSizeForCurrentWave = FMath::Min(MaxChunkSizeOverrideBytes, MaxChunkSizeForCurrentWave);
 					}
 
-					UE_LOG(LogAudio, Display, TEXT("Chunk size for %s: %d"), *SoundWaveFullName, MaxChunkSizeForCurrentWave);
+					UE_LOG(LogAudio, Display, TEXT("Chunk sizes for %s: Chunk0=%d, Chunk1-N=%d"), *SoundWaveFullName, FirstChunkSize, MaxChunkSizeForCurrentWave);
 				}
 				
 				check(FirstChunkSize != 0 && MaxChunkSizeForCurrentWave != 0);
+				
+				// Before splitting, reserve some space in each chunk for the seek table.
+				int32 ChunkSizeAdjusted = MaxChunkSizeForCurrentWave;
 
-				if (AudioFormat->SplitDataForStreaming(CompressedBuffer, ChunkBuffers, FirstChunkSize, MaxChunkSizeForCurrentWave))
+				EChunkSeekTableMode SeekMode = EChunkSeekTableMode::ConstantSamplesPerEntry;
+				if (bRequiresSeekTable)
+				{
+					// Guesstimate ahead of splitting the worse case size for a seek-table and adjust the budget of each chunk.
+					SeekMode = DetermineSeekTableMode(Table);
+					int32 ReserveForSeekTable = CalculateWorseCaseSeektableSize(Table, MinimumChunkSize, MaxChunkSizeForCurrentWave, CompressedBuffer.Num(), SeekMode);
+					ChunkSizeAdjusted -= ReserveForSeekTable;
+				}
+				
+				if (AudioFormat->SplitDataForStreaming(CompressedBuffer, ChunkBuffers, FirstChunkSize, ChunkSizeAdjusted))
 				{
 					if (ChunkBuffers.Num() > 32)
 					{
 						UE_LOG(LogAudio, Display, TEXT("Sound Wave %s is very large, requiring %d chunks."), *SoundWaveFullName, ChunkBuffers.Num());
 					}
 
+					int32 TotalAudioBytesExcludingSeektable = 0;
 					if (ChunkBuffers.Num() > 0)
 					{
 						// The zeroth chunk should not be zero-padded.
-						const int32 AudioDataSize = ChunkBuffers[0].Num();
+						int32 AudioDataSize = ChunkBuffers[0].Num();
 
 						//FStreamedAudioChunk* NewChunk = new(DerivedData->Chunks) FStreamedAudioChunk();
 						int32 ChunkIndex = DerivedData->Chunks.Add(new FStreamedAudioChunk());
 						FStreamedAudioChunk* NewChunk = &(DerivedData->Chunks[ChunkIndex]);
+
+						// Prefix the chunk a seek-table if required. 
+						// On chunk 0, this should be largely empty unless the format has chosen to include audio here.
+						if (bRequiresSeekTable)
+						{
+							TotalAudioBytesExcludingSeektable += AudioDataSize;
+							AudioDataSize += InjectSeekTable(Table, 0, ChunkBuffers[0], NewChunk, SeekMode);
+						}
+
 						// Store both the audio data size and the data size so decoders will know what portion of the bulk data is real audio
 						NewChunk->AudioDataSize = AudioDataSize;
 						NewChunk->DataSize = AudioDataSize;
@@ -561,7 +679,7 @@ class FStreamedAudioCacheDerivedDataWorker : public FNonAbandonableTask
 					for (int32 ChunkIndex = 1; ChunkIndex < ChunkBuffers.Num(); ++ChunkIndex)
 					{
 						// Zero pad the reallocation if the chunk isn't precisely the max chunk size to keep the reads aligned to MaxChunkSize
-						const int32 AudioDataSize = ChunkBuffers[ChunkIndex].Num();
+						int32 AudioDataSize = ChunkBuffers[ChunkIndex].Num();
 						check(AudioDataSize != 0 && AudioDataSize <= MaxChunkSizeForCurrentWave);
 
 						int32 ZeroPadBytes = 0;
@@ -575,10 +693,19 @@ class FStreamedAudioCacheDerivedDataWorker : public FNonAbandonableTask
 						FStreamedAudioChunk* NewChunk = new FStreamedAudioChunk();
 						DerivedData->Chunks.Add(NewChunk);
 
+						// Prefix the chunk with a seek-table if required.
+						// AudioDataSize on the chunk will include the Seek-table, which is what is loaded into the ChunkCache.
+						if (bRequiresSeekTable)
+						{
+							int32 SeekTableSize = InjectSeekTable(Table, TotalAudioBytesExcludingSeektable, ChunkBuffers[ChunkIndex], NewChunk, SeekMode);
+							TotalAudioBytesExcludingSeektable += AudioDataSize;
+							AudioDataSize += SeekTableSize;
+						}
+
 						// Store both the audio data size and the data size so decoders will know what portion of the bulk data is real audio
 						NewChunk->AudioDataSize = AudioDataSize;
 						NewChunk->DataSize = AudioDataSize + ZeroPadBytes;
-
+						
 						if (NewChunk->BulkData.IsLocked())
 						{
 							UE_LOG(LogAudioDerivedData, Warning, TEXT("While building split chunk for streaming: Raw PCM data already being written to. Chunk Index: %d SoundWave: %s "), ChunkIndex, *SoundWaveFullName);
@@ -604,8 +731,16 @@ class FStreamedAudioCacheDerivedDataWorker : public FNonAbandonableTask
 					// Could not split so copy compressed data into a single chunk
 					FStreamedAudioChunk* NewChunk = new FStreamedAudioChunk();
 					DerivedData->Chunks.Add(NewChunk);
+
+					// Prefix the chunk with a seek-table if required.
+					// AudioDataSize on the chunk will include the Seek-table, which is what is loaded into the ChunkCache.
+					if (bRequiresSeekTable)
+					{
+						InjectSeekTable(Table, 0, CompressedBuffer, NewChunk, SeekMode);
+					}
+					
 					NewChunk->DataSize = CompressedBuffer.Num();
-					NewChunk->AudioDataSize = NewChunk->DataSize;
+					NewChunk->AudioDataSize = CompressedBuffer.Num();
 
 					if (NewChunk->BulkData.IsLocked())
 					{
diff --git a/Engine/Source/Runtime/Engine/Private/SoundWave.cpp b/Engine/Source/Runtime/Engine/Private/SoundWave.cpp
--- a/Engine/Source/Runtime/Engine/Private/SoundWave.cpp
+++ b/Engine/Source/Runtime/Engine/Private/SoundWave.cpp
@@ -28,6 +28,8 @@
 #include "Misc/OutputDeviceArchiveWrapper.h"
 #include "Async/Async.h"
 #include "UObject/UE5MainStreamObjectVersion.h"
+#include "Algo/BinarySearch.h"
+#include "Templates/UnrealTemplate.h"
 
 #if WITH_EDITOR
 #endif // WITH_EDITOR
@@ -432,6 +434,43 @@ const TArrayView<uint8> FSoundWaveData::GetZerothChunkDataView() const
 	return TArrayView<uint8>(View.GetData(), static_cast<int32>(View.Num()));
 }
 
+bool FSoundWaveData::HasChunkSeekTable(int32 InChunkIndex) const
+{
+	const TIndirectArray<FStreamedAudioChunk>& Chunks = RunningPlatformData.Chunks;
+	if (Chunks.IsValidIndex(InChunkIndex))
+	{
+		return Chunks[InChunkIndex].SeekOffsetInAudioFrames != INDEX_NONE;
+	}
+	return false;
+}
+
+int32 FSoundWaveData::FindChunkIndexForSeeking(uint32 InTimeInAudioFrames) const
+{	
+	const TIndirectArray<FStreamedAudioChunk>& Chunks = RunningPlatformData.Chunks;
+	int32 NumChunks = Chunks.Num();
+
+	// Find most in range chunk.
+	int32 InRangeChunk = INDEX_NONE;
+	for (int32 i=0; i < NumChunks; ++i)
+	{
+		// Ignore chunks without tables. (should only be chunk0).
+		if (Chunks[i].SeekOffsetInAudioFrames == INDEX_NONE)
+		{
+			continue;
+		}
+
+		// If its out of bound, use the last valid one.
+		if (Chunks[i].SeekOffsetInAudioFrames > InTimeInAudioFrames)
+		{
+			return InRangeChunk;
+		}
+
+		// In bounds and valid... 
+		InRangeChunk = i;
+	}
+	return INDEX_NONE;
+}
+
 bool FSoundWaveData::LoadZerothChunk()
 {
 #if WITH_EDITOR
@@ -653,7 +692,8 @@ void FStreamedAudioChunk::Serialize(FArchive& Ar, UObject* Owner, int32 ChunkInd
 	BulkData.Serialize(Ar, Owner, ChunkIndex, false);
 	Ar << DataSize;
 	Ar << AudioDataSize;
-
+	Ar << SeekOffsetInAudioFrames;
+	
 #if WITH_EDITORONLY_DATA
 	if (!bCooked)
 	{
@@ -4382,4 +4422,4 @@ bool USoundWave::HasError() const
 {
 	check(SoundWaveDataPtr);
 	return SoundWaveDataPtr->HasError();
-}
\ No newline at end of file
+}
diff --git a/Engine/Source/Runtime/Engine/Private/StreamedAudioChunkSeekTable.cpp b/Engine/Source/Runtime/Engine/Private/StreamedAudioChunkSeekTable.cpp
new file mode 100644
--- /dev/null
+++ b/Engine/Source/Runtime/Engine/Private/StreamedAudioChunkSeekTable.cpp
@@ -0,0 +1,284 @@
+// Copyright Epic Games, Inc. All Rights Reserved.
+#include "Sound/StreamedAudioChunkSeekTable.h"
+
+#include "Serialization/MemoryReader.h"
+#include "Misc/AssertionMacros.h"
+#include "Audio.h"			// LogAudio
+
+namespace StreamedAudioChunkSeekTable_Private
+{
+	static void DeltaEncode(const TArray<uint32>& InValues, TArray<uint16>& OutDeltas)
+	{
+		if (InValues.Num() > 0)
+		{
+			OutDeltas.SetNum(InValues.Num());
+			for (int32 i = 1; i < InValues.Num(); ++i)
+			{
+				int32 Delta = InValues[i] - InValues[i - 1];
+				check(Delta >= 0);
+				check(Delta < MAX_uint16);
+				OutDeltas[i] = static_cast<uint16>(Delta);
+			}
+		}
+	}
+	static void DeltaDecode(const TArray<uint16>& InDeltas, TArray<uint32>& OutValues)
+	{
+		OutValues.SetNum(InDeltas.Num());
+		uint32 Sum = 0;
+		for (int32 i = 0; i < InDeltas.Num(); ++i)
+		{
+			Sum += InDeltas[i];
+			OutValues[i] = Sum;
+		}
+	}
+}
+
+void FStreamedAudioChunkSeekTable::SetMode(EChunkSeekTableMode InMode)
+{
+	Mode = InMode;
+	Impl = CreateImpl(InMode);
+}
+
+void FStreamedAudioChunkSeekTable::Reset()
+{
+	// Will delete recreate the instance.
+	SetMode(Mode);
+}
+
+FStreamedAudioChunkSeekTable::FStreamedAudioChunkSeekTable(EChunkSeekTableMode InMode)
+{
+	SetMode(InMode);
+}
+
+bool FStreamedAudioChunkSeekTable::Parse(const uint8* InMemory, uint32 InSize, uint32& InOutOffset, FStreamedAudioChunkSeekTable& OutTable)
+{
+	FMemoryReaderView Reader(MakeArrayView(InMemory, InSize));
+	OutTable.Reset();
+	if (OutTable.Serialize(Reader))
+	{
+		int32 Size = Reader.Tell();
+		InOutOffset += Size;
+		UE_LOG(LogAudio, Verbose, TEXT("Successfully parsed seektable: Entries=%d, Size=%d"), OutTable.Num(), Size);
+		return true;
+	}
+	return false;
+}
+
+uint32 FStreamedAudioChunkSeekTable::GetMagic()
+{
+	// Magic number is 'SEEK', so it's obvious in memory.
+	return (uint32)'KEES';
+}
+
+bool FStreamedAudioChunkSeekTable::Serialize(FArchive& Ar)
+{
+	// For minimizing memory. 
+	// Don't save anything if the table is empty, including the magic.
+	check(Impl.IsValid());
+	if (Ar.IsSaving() && Impl->Num() == 0)
+	{
+		return true;
+	}
+	
+	// Magic number. (in the case of loading an empty table, this will fail and return false)
+	// The outer logic will need to handle that case. 
+	uint32 Magic = GetMagic();
+	Ar << Magic;
+	if (Magic != GetMagic())
+	{
+		return false;
+	}
+	
+	// Mode.
+	uint8 ModeAsByte = static_cast<uint8>(Mode);
+	Ar << ModeAsByte;
+
+	if (Ar.IsLoading())
+	{
+		EChunkSeekTableMode NewMode = static_cast<EChunkSeekTableMode>(ModeAsByte);
+		if (NewMode!=Mode)
+		{
+			SetMode(NewMode);
+		}
+	}
+
+	check(Impl.IsValid());
+	return Impl->Serialize(Ar);
+}
+
+struct FConstantRateSeekTable : public FStreamedAudioChunkSeekTable::ISeekTableImpl
+{
+	TArray<uint32> Table;
+	uint32 StartTimeOffset = INDEX_NONE;
+	uint16 AudioFramesPerEntry = MAX_uint16;
+
+	int32 Num() const override
+	{
+		return Table.Num();
+	}
+	void Add(uint32 InTimeInAudioFrames, uint32 InOffset) override
+	{
+		if (Table.IsEmpty())
+		{
+			StartTimeOffset = InTimeInAudioFrames;
+		}
+		else if(Table.Num() == 1)
+		{
+			// Imply Samples Per Entry from step 1.
+			check(InTimeInAudioFrames > StartTimeOffset);
+			int32 DeltaFrames = InTimeInAudioFrames - StartTimeOffset;
+			check(DeltaFrames >= 0);
+			check(DeltaFrames < MAX_uint16);
+			AudioFramesPerEntry = static_cast<uint16>(DeltaFrames);
+		}		
+		Table.Add(InOffset);
+	}
+
+	uint32 FindOffset(uint32 InTimeInAudioFrames) const override
+	{
+		if (!Table.IsEmpty() && AudioFramesPerEntry > 0)
+		{
+			int32 Index = (InTimeInAudioFrames - StartTimeOffset) / AudioFramesPerEntry;
+			if (Table.IsValidIndex(Index))
+			{
+				return Table[Index];
+			}
+		}
+		return INDEX_NONE;
+	}
+	uint32 FindTime(uint32 InOffset) const override
+	{
+		if (!Table.IsEmpty() && AudioFramesPerEntry > 0 && StartTimeOffset >= 0)
+		{
+			uint32 TimeInSamples = StartTimeOffset + (InOffset * AudioFramesPerEntry);
+			return TimeInSamples;
+		}
+		return INDEX_NONE;
+	}
+
+	bool Serialize(FArchive& Ar) override
+	{
+		Ar << AudioFramesPerEntry;
+		Ar << StartTimeOffset;
+
+		if (Ar.IsSaving())
+		{
+			Encode(Ar);
+		}
+		else if (Ar.IsLoading())
+		{
+			if (!Decode(Ar))
+			{
+				return false;
+			}
+		}
+		return true;
+	}
+	void Encode(FArchive& Ar) const
+	{
+		TArray<uint16> Deltas;
+		StreamedAudioChunkSeekTable_Private::DeltaEncode(Table,Deltas);
+
+		// Save in stream.
+		Ar << Deltas;
+	}
+	bool Decode(FArchive& Ar) 
+	{
+		TArray<uint16> Deltas;
+		Ar << Deltas;
+
+		StreamedAudioChunkSeekTable_Private::DeltaDecode(Deltas,Table);
+		TArray<uint8> Bytes;
+		return true;
+	}
+};
+
+struct FVariableRateSeekTable : public FStreamedAudioChunkSeekTable::ISeekTableImpl
+{
+	TArray<uint32> Offsets;
+	TArray<uint32> Times;
+
+	void Add(uint32 InTimeInAudioFrames, uint32 InOffset) override
+	{
+		Offsets.Add(InOffset);
+		Times.Add(InTimeInAudioFrames);
+	}
+	int32 Num() const override
+	{
+		check(Offsets.Num() == Times.Num());
+		return Times.Num();
+	}
+	uint32 FindOffset(uint32 InTimeInAudioFrames) const override
+	{
+		check(Offsets.Num() == Times.Num());
+		int32 Index = Algo::LowerBound(Times, InTimeInAudioFrames);
+		if (Index != INDEX_NONE)
+		{
+			return Offsets[Index];
+		}
+		return INDEX_NONE;
+	}
+	uint32 FindTime(uint32 InOffset) const override
+	{
+		int32 Index = Algo::LowerBound(Offsets,InOffset);
+		if (Index != INDEX_NONE)
+		{
+			return Times[Index];
+		}
+		return INDEX_NONE;
+	}
+
+	bool Serialize(FArchive& Ar) override
+	{
+		if (Ar.IsSaving())
+		{
+			Encode(Ar);
+		}
+		else if (Ar.IsLoading())
+		{
+			if (!Decode(Ar))
+			{
+				return false;
+			}
+		}
+		return true;
+	}
+
+	void Encode(FArchive& Ar) const
+	{
+		TArray<uint16> DeltaOffsets;
+		StreamedAudioChunkSeekTable_Private::DeltaEncode(Offsets,DeltaOffsets);
+		Ar << DeltaOffsets;
+		
+		TArray<uint16> DeltaTimes;
+		StreamedAudioChunkSeekTable_Private::DeltaEncode(Times, DeltaTimes);
+		Ar << DeltaTimes;
+	}
+
+	bool Decode(FArchive& Ar)
+	{
+		TArray<uint16> DeltaOffsets;
+		Ar << DeltaOffsets;
+		StreamedAudioChunkSeekTable_Private::DeltaDecode(DeltaOffsets, Offsets);
+		
+		TArray<uint16> DeltaSamples;
+		Ar << DeltaSamples;
+		StreamedAudioChunkSeekTable_Private::DeltaDecode(DeltaSamples, Times);
+		return true;
+	}
+};
+
+TUniquePtr<FStreamedAudioChunkSeekTable::ISeekTableImpl> FStreamedAudioChunkSeekTable::CreateImpl(EChunkSeekTableMode InMode)
+{
+	switch(InMode)
+	{
+		case EChunkSeekTableMode::ConstantSamplesPerEntry:
+			return MakeUnique<FConstantRateSeekTable>();
+		case EChunkSeekTableMode::VariableSamplesPerEntry:
+			return MakeUnique<FVariableRateSeekTable>();
+		default:
+			break;
+	}
+	checkNoEntry();
+	return {};
+}
diff --git a/Engine/Source/Runtime/Engine/Public/AudioDecompress.h b/Engine/Source/Runtime/Engine/Public/AudioDecompress.h
--- a/Engine/Source/Runtime/Engine/Public/AudioDecompress.h
+++ b/Engine/Source/Runtime/Engine/Public/AudioDecompress.h
@@ -20,6 +20,7 @@
 #define MONO_PCM_BUFFER_SIZE		( MONO_PCM_BUFFER_SAMPLES * sizeof( int16 ) )
 
 struct FSoundQualityInfo;
+class FStreamedAudioChunkSeekTable;
 
 /**
  * Interface class to decompress various types of audio data
@@ -278,6 +279,13 @@ protected:
 	 */
 	ENGINE_API const uint8* GetLoadedChunk(const FSoundWaveProxyPtr& InSoundWave, uint32 ChunkIndex, uint32& OutChunkSize);
 
+	/**
+	 * Gets the current chunks seektable instance (or creates one)
+	 * @return the currently loaded chunk seektable
+	 */
+	const FStreamedAudioChunkSeekTable& GetCurrentSeekTable() const;
+	FStreamedAudioChunkSeekTable& GetCurrentSeekTable();
+
 	/** bool set before ParseHeader. Whether we are streaming a file or not. */
 	bool bIsStreaming;
 	/** Ptr to the current streamed chunk. */
@@ -315,8 +323,9 @@ protected:
 	/** Number of bytes of padding used, overridden in some implementations. Defaults to 0. */
 	uint32 SrcBufferPadding;
 	/** Chunk Handle to ensure that this chunk of streamed audio is not deleted while we are using it. */
-	FAudioChunkHandle CurCompressedChunkHandle;
-
+	FAudioChunkHandle CurCompressedChunkHandle;		
+	/** If there's a chunked seek-table present, this contains the current chunks portion. */
+	TPimplPtr<FStreamedAudioChunkSeekTable> CurrentChunkSeekTable;
 	/** 
 		When a streaming seek request comes down, this is the block we are going to. INDEX_NONE means no seek pending.
 		When using the legacy streaming system this is read on a thread other than the decompression thread
@@ -328,6 +337,9 @@ protected:
 	std::atomic<int32> StreamSeekBlockIndex;
 	/** When a streaming seek request comes down, this is the offset in to the block we want to start decoding from. */
 	int32 StreamSeekBlockOffset;
+
+	/** If using the Chunked seek-tables, we request the seek in samples so we can resolve after the chunk table loads. */
+	uint32 StreamSeekToSamples = INDEX_NONE;	
 };
 
 struct IAudioInfoFactory 
