Revert CL 25386902 because it is breaking onefileperactor maps that have actors that have been renamed.
"Avoid using IterateDirectoryRecursively to gather external actors in ULevel::GetOnDiskExternalActorPackages, rely on the asset registry instead"
When an external actor package is renamed, the redirector and the actor have the same object path but exist in two different pacakges. Their entry collides in the AssetRegistry, which stores objects by objectpath, and whichever one was discovered by the AR later wins, intermittently removing the real actor from the AssetRegistry.
We will need to fix the AssetRegistry to do ... something ... with the redirectors and actors having the same objectpath before we can restore 25386902.
#rnx

diff --git a/Engine/Source/Runtime/Engine/Private/Level.cpp b/Engine/Source/Runtime/Engine/Private/Level.cpp
--- a/Engine/Source/Runtime/Engine/Private/Level.cpp
+++ b/Engine/Source/Runtime/Engine/Private/Level.cpp
@@ -58,6 +58,7 @@ Level.cpp: Level-related functions
 #include "EditorActorFolders.h"
 #include "UObject/MetaData.h"
 #include "WorldPartition/WorldPartitionHelpers.h"
+#include "HAL/FileManager.h"
 #include "HAL/PlatformFileManager.h"
 #include "Misc/PathViews.h"
 #include "Selection.h"
@@ -3965,20 +3966,19 @@ TArray<FString> ULevel::GetOnDiskExternalActorPackages(const FString& ExternalAc
 	TArray<FString> ActorPackageNames;
 	if (!ExternalActorsPath.IsEmpty())
 	{
-		FARFilter Filter;
-		Filter.bIncludeOnlyOnDiskAssets = true;
-		Filter.PackagePaths.Add(*ExternalActorsPath);
-		Filter.bRecursivePaths = true;
-	
-		TArray<FAssetData> ActorAssets;
-		IAssetRegistry& AssetRegistry = FModuleManager::LoadModuleChecked<FAssetRegistryModule>(TEXT("AssetRegistry")).Get();
-		AssetRegistry.ScanSynchronous({ ExternalActorsPath }, TArray<FString>());
-		AssetRegistry.GetAssets(Filter, ActorAssets);
-
-		ActorPackageNames.Reserve(ActorAssets.Num());
-		Algo::Transform(ActorAssets, ActorPackageNames, [](const FAssetData& ActorAssetData) { return ActorAssetData.PackageName.ToString(); });
+		IFileManager::Get().IterateDirectoryRecursively(*FPackageName::LongPackageNameToFilename(ExternalActorsPath), [&ActorPackageNames](const TCHAR* FilenameOrDirectory, bool bIsDirectory)
+		{
+			if (!bIsDirectory)
+			{
+				FString Filename(FilenameOrDirectory);
+				if (Filename.EndsWith(FPackageName::GetAssetPackageExtension()))
+				{
+					ActorPackageNames.Add(FPackageName::FilenameToLongPackageName(Filename));
+				}
+			}
+			return true;
+		});
 	}
-
 	return ActorPackageNames;
 }
 
