#UE Do not attempt to redirect empty paths from WP sources. Do not attempt to ask the asset registry for assets in the path if it is only a package name.

diff --git a/Engine/Source/Runtime/Engine/Private/WorldPartition/WorldPartitionHelpers.cpp b/Engine/Source/Runtime/Engine/Private/WorldPartition/WorldPartitionHelpers.cpp
--- a/Engine/Source/Runtime/Engine/Private/WorldPartition/WorldPartitionHelpers.cpp
+++ b/Engine/Source/Runtime/Engine/Private/WorldPartition/WorldPartitionHelpers.cpp
@@ -306,6 +306,12 @@ bool FWorldPartitionHelpers::ConvertRuntimePathToEditorPath(const FSoftObjectPat
 
 bool FWorldPartitionHelpers::FixupRedirectedAssetPath(FSoftObjectPath& InOutSoftObjectPath)
 {
+	if (InOutSoftObjectPath.IsNull())
+	{
+		// Empty path, no redirect
+		return true;
+	}
+
 	// Check GRedirectCollector first for faster fixup
 	FSoftObjectPath FoundRedirection = GRedirectCollector.GetAssetPathRedirection(InOutSoftObjectPath.GetWithoutSubPath());
 	if (!FoundRedirection.IsNull())
@@ -314,6 +320,12 @@ bool FWorldPartitionHelpers::FixupRedirectedAssetPath(FSoftObjectPath& InOutSoft
 		return true;
 	}
 
+	if (InOutSoftObjectPath.GetAssetName().IsEmpty())
+	{
+		// A package name. No need to ask the asset registry for assets, it wont find any
+		return true;
+	}
+
 	const FAssetData* AssetData;
 	IAssetRegistry& AssetRegistry = FModuleManager::LoadModuleChecked<FAssetRegistryModule>(TEXT("AssetRegistry")).Get();
 
