LWI - Change check to ensure to prevent server crashes in test builds

diff --git a/Engine/Source/Runtime/Engine/Private/LightWeightInstanceManager.cpp b/Engine/Source/Runtime/Engine/Private/LightWeightInstanceManager.cpp
--- a/Engine/Source/Runtime/Engine/Private/LightWeightInstanceManager.cpp
+++ b/Engine/Source/Runtime/Engine/Private/LightWeightInstanceManager.cpp
@@ -145,7 +145,15 @@ void ALightWeightInstanceManager::OnSpawnedActorDestroyed(AActor* DestroyedActor
 {
 	check(DestroyedActor);
 	const int32 DestroyedActorInstanceIndex = FindIndexForActor(DestroyedActor);
-	check(DestroyedActorInstanceIndex != INDEX_NONE);
+
+	if (!ensure(DestroyedActorInstanceIndex != INDEX_NONE))
+	{
+		UE_LOG(LogLightWeightInstance, Error,
+			TEXT("OnSpawnedActorDestroyed - actor [ %s ] is not being tracked by manager [ %s ]"),
+			*DestroyedActor->GetName(),
+			*BaseInstanceName);
+	}
+
 	Actors.Remove(DestroyedActorInstanceIndex);
 	DestroyedActor->OnDestroyed.RemoveAll(this);
 }
