Chaos Cloth - Replace an ensure by a verbose log when not finding a supported collision type during cloth collision with the environment.

diff --git a/Engine/Source/Runtime/Engine/Private/SkeletalMeshComponentPhysics.cpp b/Engine/Source/Runtime/Engine/Private/SkeletalMeshComponentPhysics.cpp
--- a/Engine/Source/Runtime/Engine/Private/SkeletalMeshComponentPhysics.cpp
+++ b/Engine/Source/Runtime/Engine/Private/SkeletalMeshComponentPhysics.cpp
@@ -3355,7 +3355,7 @@ void USkeletalMeshComponent::ProcessClothCollisionWithEnvironment()
 							break;
 
 						default: 
-							ensure(false);  // Is there a missing collision type?
+							UE_LOG(LogSkeletalMesh, Verbose, TEXT("Found unsupported collision type during environmental collision with the cloth in [%s]"), !Component->GetOwner() ? TEXT("Unknown") : *Component->GetOwner()->GetFName().ToString());
 							break;
 						}
 					}
