LevelInstance: New "IsMainWorldOnly" property allows filtering actors when they are being used in a Level Instance

Good example use case is having a Lighting setup in the Level Instance level being checked as "IsMainWorldOnly" so that those lights do not get loaded when this level gets used as a Level Instance in another World Partition world.

diff --git a/Engine/Source/Runtime/Core/Public/UObject/FortniteMainBranchObjectVersion.h b/Engine/Source/Runtime/Core/Public/UObject/FortniteMainBranchObjectVersion.h
--- a/Engine/Source/Runtime/Core/Public/UObject/FortniteMainBranchObjectVersion.h
+++ b/Engine/Source/Runtime/Core/Public/UObject/FortniteMainBranchObjectVersion.h
@@ -316,6 +316,9 @@ struct CORE_API FFortniteMainBranchObjectVersion
 		// Don't serialize class descriptor GUIDs
 		WorldPartitionClasDescGuidTransient,
 
+		// Serialize ActorDesc bIsMainWorldOnly
+		WorldPartitionActorDescIsMainWorldOnly,
+
 		// -----<new versions can be added above this line>-------------------------------------------------
 		VersionPlusOne,
 		LatestVersion = VersionPlusOne - 1
diff --git a/Engine/Source/Runtime/Engine/Classes/Engine/LevelBounds.h b/Engine/Source/Runtime/Engine/Classes/Engine/LevelBounds.h
--- a/Engine/Source/Runtime/Engine/Classes/Engine/LevelBounds.h
+++ b/Engine/Source/Runtime/Engine/Classes/Engine/LevelBounds.h
@@ -108,6 +108,6 @@ protected:
 	FDelegateHandle OnLevelActorAddedDelegateHandle;
 
 private:
-	virtual bool ActorTypeShouldSkipFromLevelInstance() const override { return true; }
+	virtual bool ActorTypeIsMainWorldOnly() const override { return true; }
 #endif
 };
\ No newline at end of file
diff --git a/Engine/Source/Runtime/Engine/Classes/GameFramework/Actor.h b/Engine/Source/Runtime/Engine/Classes/GameFramework/Actor.h
--- a/Engine/Source/Runtime/Engine/Classes/GameFramework/Actor.h
+++ b/Engine/Source/Runtime/Engine/Classes/GameFramework/Actor.h
@@ -294,6 +294,11 @@ private:
 	 */
 	UPROPERTY(Transient)
 	uint8 bIsInEditingLevelInstance:1;
+
+public:
+	UPROPERTY(EditAnywhere, AdvancedDisplay, Category = LevelInstance, meta = (Tooltip = "If checked, this Actor will only get loaded in a main world (persistent level), it will not be loaded through Level Instances."))
+	uint8 bIsMainWorldOnly : 1;
+private:
 #endif
 
 	/** If true, PreReplication will be called on this actor before each potential replication. */
@@ -961,9 +966,10 @@ public:
 	virtual bool ShouldLevelKeepRefIfExternal() const { return false; }
 
 	/**
-	 * Whether this actor should be ignored when it is loaded from a level instance
+	 * Whether this actor should be ignored when it is not loaded as part of the Main World (Persistent Level).
+	 * eg. Will not be loaded through Level Instances.
 	 */
-	bool ShouldSkipFromLevelInstance() const;
+	bool IsMainWorldOnly() const;
 
 	FActorOnPackagingModeChanged OnPackagingModeChanged;
 
@@ -1265,7 +1271,7 @@ public:
 #if WITH_EDITOR
 private:
 	virtual bool ActorTypeSupportsDataLayer() const { return true; }
-	virtual bool ActorTypeShouldSkipFromLevelInstance() const { return false; }
+	virtual bool ActorTypeIsMainWorldOnly() const { return false; }
 	TArray<const UDataLayerAsset*> ResolveDataLayerAssets(const TArray<TSoftObjectPtr<UDataLayerAsset>>& InDataLayerAssets) const;
 public:
 	bool AddDataLayer(const UDataLayerInstance* DataLayerInstance);
diff --git a/Engine/Source/Runtime/Engine/Private/ActorEditor.cpp b/Engine/Source/Runtime/Engine/Private/ActorEditor.cpp
--- a/Engine/Source/Runtime/Engine/Private/ActorEditor.cpp
+++ b/Engine/Source/Runtime/Engine/Private/ActorEditor.cpp
@@ -83,8 +83,9 @@ bool AActor::CanEditChange(const FProperty* PropertyThatWillChange) const
 	const bool bIsRuntimeGridProperty = PropertyThatWillChange->GetFName() == GET_MEMBER_NAME_CHECKED(AActor, RuntimeGrid);
 	const bool bIsDataLayersProperty = PropertyThatWillChange->GetFName() == GET_MEMBER_NAME_CHECKED(AActor, DataLayerAssets);
 	const bool bIsHLODLayerProperty = PropertyThatWillChange->GetFName() == GET_MEMBER_NAME_CHECKED(AActor, HLODLayer);
+	const bool bIsMainWorldOnlyProperty = PropertyThatWillChange->GetFName() == GET_MEMBER_NAME_CHECKED(AActor, bIsMainWorldOnly);
 
-	if (bIsSpatiallyLoadedProperty || bIsRuntimeGridProperty || bIsDataLayersProperty || bIsHLODLayerProperty)
+	if (bIsSpatiallyLoadedProperty || bIsRuntimeGridProperty || bIsDataLayersProperty || bIsHLODLayerProperty || bIsMainWorldOnlyProperty)
 	{
 		if (!IsTemplate())
 		{
@@ -1040,15 +1041,20 @@ void AActor::SetHLODLayer(class UHLODLayer* InHLODLayer)
 	HLODLayer = InHLODLayer;
 }
 
-bool AActor::ShouldSkipFromLevelInstance() const
+bool AActor::IsMainWorldOnly() const
 {
+	if (bIsMainWorldOnly)
+	{
+		return true;
+	}
+	
 	if (HasAnyFlags(RF_ClassDefaultObject))
 	{
-		return ActorTypeShouldSkipFromLevelInstance();
+		return ActorTypeIsMainWorldOnly();
 	}
 	else
 	{
-		return CastChecked<AActor>(GetClass()->GetDefaultObject())->ShouldSkipFromLevelInstance();
+		return CastChecked<AActor>(GetClass()->GetDefaultObject())->IsMainWorldOnly();
 	}
 }
 
diff --git a/Engine/Source/Runtime/Engine/Private/LevelInstance/LevelInstanceSubsystem.cpp b/Engine/Source/Runtime/Engine/Private/LevelInstance/LevelInstanceSubsystem.cpp
--- a/Engine/Source/Runtime/Engine/Private/LevelInstance/LevelInstanceSubsystem.cpp
+++ b/Engine/Source/Runtime/Engine/Private/LevelInstance/LevelInstanceSubsystem.cpp
@@ -1351,7 +1351,7 @@ void ULevelInstanceSubsystem::BreakLevelInstance_Impl(ILevelInstanceInterface* L
 			// Skip some actor types
 			if ((Actor != Actor->GetLevel()->GetDefaultBrush()) &&
 				!Actor->IsA<AWorldSettings>() &&
-				!Actor->ShouldSkipFromLevelInstance())
+				!Actor->IsMainWorldOnly())
 			{
 				if (CanMoveActorToLevel(Actor))
 				{
diff --git a/Engine/Source/Runtime/Engine/Private/WorldPartition/WorldPartitionActorDesc.cpp b/Engine/Source/Runtime/Engine/Private/WorldPartition/WorldPartitionActorDesc.cpp
--- a/Engine/Source/Runtime/Engine/Private/WorldPartition/WorldPartitionActorDesc.cpp
+++ b/Engine/Source/Runtime/Engine/Private/WorldPartition/WorldPartitionActorDesc.cpp
@@ -91,6 +91,7 @@ void FWorldPartitionActorDesc::Init(const AActor* InActor)
 	bActorIsEditorOnly = InActor->IsEditorOnly();
 	bActorIsRuntimeOnly = InActor->IsRuntimeOnly();
 	bActorIsHLODRelevant = InActor->IsHLODRelevant();
+	bActorIsMainWorldOnly = InActor->bIsMainWorldOnly;
 	HLODLayer = InActor->GetHLODLayer() ? FSoftObjectPath(InActor->GetHLODLayer()->GetPathName()) : FSoftObjectPath();
 	
 	// DataLayers
@@ -257,6 +258,7 @@ bool FWorldPartitionActorDesc::Equals(const FWorldPartitionActorDesc* Other) con
 		bIsSpatiallyLoaded == Other->bIsSpatiallyLoaded &&
 		bActorIsEditorOnly == Other->bActorIsEditorOnly &&
 		bActorIsRuntimeOnly == Other->bActorIsRuntimeOnly &&
+		bActorIsMainWorldOnly == Other->bActorIsMainWorldOnly &&
 		bActorIsHLODRelevant == Other->bActorIsHLODRelevant &&
 		bIsUsingDataLayerAsset == Other->bIsUsingDataLayerAsset &&
 		HLODLayer == Other->HLODLayer &&
@@ -281,6 +283,7 @@ bool FWorldPartitionActorDesc::ShouldResave(const FWorldPartitionActorDesc* Othe
 		bIsSpatiallyLoaded != Other->bIsSpatiallyLoaded ||
 		bActorIsEditorOnly != Other->bActorIsEditorOnly ||
 		bActorIsRuntimeOnly != Other->bActorIsRuntimeOnly ||
+		bActorIsMainWorldOnly != Other->bActorIsMainWorldOnly ||
 		bIsBoundsValid != Other->bIsBoundsValid ||
 		HLODLayer != Other->HLODLayer ||
 		ParentActor != Other->ParentActor ||
@@ -395,7 +398,7 @@ FString FWorldPartitionActorDesc::ToString(EToStringMode Mode) const
 		}
 
 		Result.Appendf(
-			TEXT(" BaseClass:%s NativeClass:%s Name:%s Label:%s SpatiallyLoaded:%s Bounds:%s RuntimeGrid:%s EditorOnly:%s RuntimeOnly:%s HLODRelevant:%s"),
+			TEXT(" BaseClass:%s NativeClass:%s Name:%s Label:%s SpatiallyLoaded:%s Bounds:%s RuntimeGrid:%s EditorOnly:%s RuntimeOnly:%s HLODRelevant:%s IsMainWorldOnly:%s"),
 			*BaseClass.ToString(), 
 			*NativeClass.ToString(), 
 			*GetActorName().ToString(),
@@ -405,7 +408,8 @@ FString FWorldPartitionActorDesc::ToString(EToStringMode Mode) const
 			*RuntimeGrid.ToString(),
 			GetBoolStr(bActorIsEditorOnly),
 			GetBoolStr(bActorIsRuntimeOnly),
-			GetBoolStr(bActorIsHLODRelevant)
+			GetBoolStr(bActorIsHLODRelevant),
+			GetBoolStr(IsMainWorldOnly())
 		);
 
 		if (ParentActor.IsValid())
@@ -660,6 +664,11 @@ PRAGMA_ENABLE_DEPRECATION_WARNINGS
 		}
 	}
 
+	if (Ar.CustomVer(FFortniteMainBranchObjectVersion::GUID) >= FFortniteMainBranchObjectVersion::WorldPartitionActorDescIsMainWorldOnly)
+	{
+		Ar << TDeltaSerialize<bool>(bActorIsMainWorldOnly);
+	}
+
 	// Fixup redirected data layer asset paths
 	if (Ar.IsLoading() && bIsUsingDataLayerAsset)
 	{
@@ -722,6 +731,11 @@ void FWorldPartitionActorDesc::CheckForErrors(IStreamingGenerationErrorHandler*
 	}
 }
 
+bool FWorldPartitionActorDesc::IsMainWorldOnly() const
+{
+	return bActorIsMainWorldOnly || CastChecked<AActor>(ActorNativeClass->GetDefaultObject())->IsMainWorldOnly();
+}
+
 bool FWorldPartitionActorDesc::IsEditorRelevant() const
 {
 	if (GetActorIsRuntimeOnly())
@@ -729,7 +743,7 @@ bool FWorldPartitionActorDesc::IsEditorRelevant() const
 		return false;
 	}
 
-	if (CastChecked<AActor>(ActorNativeClass->GetDefaultObject())->ShouldSkipFromLevelInstance())
+	if (IsMainWorldOnly())
 	{
 		return GetContainer() && !GetContainer()->IsTemplateContainer() && GetContainer()->GetWorldPartition()->IsMainWorldPartition();
 	}
@@ -739,7 +753,7 @@ bool FWorldPartitionActorDesc::IsEditorRelevant() const
 
 bool FWorldPartitionActorDesc::IsRuntimeRelevant(const FActorContainerID& InContainerID) const
 {
-	return InContainerID.IsMainContainer() || !CastChecked<AActor>(ActorNativeClass->GetDefaultObject())->ShouldSkipFromLevelInstance();
+	return InContainerID.IsMainContainer() || !IsMainWorldOnly();
 }
 
 bool FWorldPartitionActorDesc::IsLoaded(bool bEvenIfPendingKill) const
diff --git a/Engine/Source/Runtime/Engine/Public/LevelInstance/LevelInstanceEditorInstanceActor.h b/Engine/Source/Runtime/Engine/Public/LevelInstance/LevelInstanceEditorInstanceActor.h
--- a/Engine/Source/Runtime/Engine/Public/LevelInstance/LevelInstanceEditorInstanceActor.h
+++ b/Engine/Source/Runtime/Engine/Public/LevelInstance/LevelInstanceEditorInstanceActor.h
@@ -32,7 +32,7 @@ public:
 	virtual AActor* GetSelectionParent() const override;
 
 private:
-	virtual bool ActorTypeShouldSkipFromLevelInstance() const override { return true; }
+	virtual bool ActorTypeIsMainWorldOnly() const override { return true; }
 	friend class ULevelInstanceComponent;
 	void UpdateWorldTransform(const FTransform& WorldTransform);
 
diff --git a/Engine/Source/Runtime/Engine/Public/WorldPartition/DataLayer/WorldDataLayers.h b/Engine/Source/Runtime/Engine/Public/WorldPartition/DataLayer/WorldDataLayers.h
--- a/Engine/Source/Runtime/Engine/Public/WorldPartition/DataLayer/WorldDataLayers.h
+++ b/Engine/Source/Runtime/Engine/Public/WorldPartition/DataLayer/WorldDataLayers.h
@@ -179,7 +179,7 @@ private:
 	void DumpDataLayerRecursively(const UDataLayerInstance* DataLayer, FString Prefix, FOutputDevice& OutputDevice) const;
 
 #if WITH_EDITOR
-	virtual bool ActorTypeShouldSkipFromLevelInstance() const override { return true; }
+	virtual bool ActorTypeIsMainWorldOnly() const override { return true; }
 
 	void ConvertDataLayerToInstancces();
 	void UpdateContainsDeprecatedDataLayers();
diff --git a/Engine/Source/Runtime/Engine/Public/WorldPartition/WorldPartitionActorDesc.h b/Engine/Source/Runtime/Engine/Public/WorldPartition/WorldPartitionActorDesc.h
--- a/Engine/Source/Runtime/Engine/Public/WorldPartition/WorldPartitionActorDesc.h
+++ b/Engine/Source/Runtime/Engine/Public/WorldPartition/WorldPartitionActorDesc.h
@@ -268,6 +268,7 @@ public:
 	using FActorDescDeprecator = TFunction<void(FArchive&, FWorldPartitionActorDesc*)>;
 	static void RegisterActorDescDeprecator(TSubclassOf<AActor> ActorClass, const FActorDescDeprecator& Deprecator);
 
+	bool IsMainWorldOnly() const;
 protected:
 	FWorldPartitionActorDesc();
 
@@ -296,6 +297,7 @@ protected:
 	bool							bIsSpatiallyLoaded;
 	bool							bActorIsEditorOnly;
 	bool							bActorIsRuntimeOnly;
+	bool							bActorIsMainWorldOnly;
 	bool							bActorIsHLODRelevant;
 	bool							bIsUsingDataLayerAsset; // Used to know if DataLayers array represents DataLayers Asset paths or the FNames of the deprecated version of Data Layers
 	bool							bIsBoundsValid;
diff --git a/Engine/Source/Runtime/Engine/Public/WorldPartition/WorldPartitionMiniMap.h b/Engine/Source/Runtime/Engine/Public/WorldPartition/WorldPartitionMiniMap.h
--- a/Engine/Source/Runtime/Engine/Public/WorldPartition/WorldPartitionMiniMap.h
+++ b/Engine/Source/Runtime/Engine/Public/WorldPartition/WorldPartitionMiniMap.h
@@ -19,7 +19,7 @@ class ENGINE_API AWorldPartitionMiniMap : public AInfo
 private:
 #if WITH_EDITOR
 	virtual bool ActorTypeSupportsDataLayer() const final { return false; }
-	virtual bool ActorTypeShouldSkipFromLevelInstance() const override  { return true; }
+	virtual bool ActorTypeIsMainWorldOnly() const override  { return true; }
 #endif
 
 public:
