WorldPartition: Test for null components in AActor GetHLODRelevantComponents() / HasHLODRelevantComponents()

diff --git a/Engine/Source/Runtime/Engine/Private/Actor.cpp b/Engine/Source/Runtime/Engine/Private/Actor.cpp
--- a/Engine/Source/Runtime/Engine/Private/Actor.cpp
+++ b/Engine/Source/Runtime/Engine/Private/Actor.cpp
@@ -5974,13 +5974,13 @@ bool AActor::IsHLODRelevant() const
 
 bool AActor::HasHLODRelevantComponents() const
 {
-	return Algo::AnyOf(GetComponents(), [](const UActorComponent* Component) { return Component->IsHLODRelevant(); });
+	return Algo::AnyOf(GetComponents(), [](const UActorComponent* Component) { return Component && Component->IsHLODRelevant(); });
 }
 
 TArray<UActorComponent*> AActor::GetHLODRelevantComponents() const
 {
 	TArray<UActorComponent*> HLODRelevantComponents;
-	auto IsHLODRelevant = [](UActorComponent* Component) { return Component->IsHLODRelevant(); };
+	auto IsHLODRelevant = [](UActorComponent* Component) { return Component && Component->IsHLODRelevant(); };
 	auto GetComponent = [](UActorComponent* Component) { return Component; };
 	Algo::TransformIf(GetComponents(), HLODRelevantComponents, IsHLODRelevant, GetComponent);
 	return HLODRelevantComponents;
