Cooker and AssetManager: Change the AssetManager to directly report a PrimaryAsset's bundled assets to the Cooker as soft dependencies, rather than storing them in GObjectRedirector. This is required for EditorOnlyEditor to identify them as SoftObjectPaths that should be cooked; OnlyEditorOnly does not follow the legacy behavior of cooking every loaded SoftObjectPath; it only cooks SoftObjectPaths that are saved into cooked pacakges or that are identified as necessary by system specific code such as the AssetManager.
#rnx

diff --git a/Engine/Source/Editor/UnrealEd/Private/CookOnTheFlyServer.cpp b/Engine/Source/Editor/UnrealEd/Private/CookOnTheFlyServer.cpp
--- a/Engine/Source/Editor/UnrealEd/Private/CookOnTheFlyServer.cpp
+++ b/Engine/Source/Editor/UnrealEd/Private/CookOnTheFlyServer.cpp
@@ -6176,6 +6176,17 @@ void FSaveCookedPackageContext::FinishPackage()
 					FDiscoveredPlatformSet(ReachablePlatforms));
 			}
 		}
+		// Also add any references from the package that are required by the AssetManager
+		for (FName AMPackageName : FRequestCluster::GetAssetManagerReferences(PackageFName))
+		{
+			UE::Cook::FPackageData* AMPackageData = COTFS.PackageDatas->TryAddPackageDataByPackageName(AMPackageName);
+			if (AMPackageData)
+			{
+				COTFS.QueueDiscoveredPackage(*AMPackageData,
+					UE::Cook::FInstigator(UE::Cook::EInstigator::SoftDependency, PackageFName),
+					FDiscoveredPlatformSet(ReachablePlatforms));
+			}
+		}
 
 		// When using legacy WhatGetsCookedRules, add all the SoftObjectPaths discovered during the package's load, plus any added on
 		// during save, to the cook for all platforms
diff --git a/Engine/Source/Editor/UnrealEd/Private/Cooker/CookRequestCluster.cpp b/Engine/Source/Editor/UnrealEd/Private/Cooker/CookRequestCluster.cpp
--- a/Engine/Source/Editor/UnrealEd/Private/Cooker/CookRequestCluster.cpp
+++ b/Engine/Source/Editor/UnrealEd/Private/Cooker/CookRequestCluster.cpp
@@ -18,6 +18,7 @@
 #include "DerivedDataCache.h"
 #include "DerivedDataRequestOwner.h"
 #include "EditorDomain/EditorDomainUtils.h"
+#include "Engine/AssetManager.h"
 #include "Engine/Level.h"
 #include "HAL/Event.h"
 #include "HAL/PlatformProcess.h"
@@ -1005,6 +1006,9 @@ void FRequestCluster::FGraphSearch::ExploreVertexEdges(FVertexData& Vertex)
 
 			// LocalizationReferences are a source of SoftGameDependencies that are not present in the AssetRegistry
 			SoftGameDependencies.Append(GetLocalizationReferences(PackageName, Cluster.COTFS));
+
+			// The AssetManager can provide additional SoftGameDependencies
+			SoftGameDependencies.Append(GetAssetManagerReferences(PackageName));
 		}
 	}
 
@@ -1570,6 +1574,13 @@ TConstArrayView<FName> FRequestCluster::GetLocalizationReferences(FName PackageN
 	return TConstArrayView<FName>();
 }
 
+TArray<FName> FRequestCluster::GetAssetManagerReferences(FName PackageName)
+{
+	TArray<FName> Results;
+	UAssetManager::Get().ModifyCookReferences(PackageName, Results);
+	return Results;
+}
+
 template <typename T>
 static void ArrayShuffle(TArray<T>& Array)
 {
diff --git a/Engine/Source/Editor/UnrealEd/Private/Cooker/CookRequestCluster.h b/Engine/Source/Editor/UnrealEd/Private/Cooker/CookRequestCluster.h
--- a/Engine/Source/Editor/UnrealEd/Private/Cooker/CookRequestCluster.h
+++ b/Engine/Source/Editor/UnrealEd/Private/Cooker/CookRequestCluster.h
@@ -80,6 +80,7 @@ public:
 		TMap<FPackageData*, TArray<FPackageData*>>& OutRequestGraph);
 
 	static TConstArrayView<FName> GetLocalizationReferences(FName PackageName, UCookOnTheFlyServer& InCOTFS);
+	static TArray<FName> GetAssetManagerReferences(FName PackageName);
 	static void IsRequestCookable(const ITargetPlatform* TargetPlatform, FPackageData& PackageData,
 		UCookOnTheFlyServer& COTFS, ESuppressCookReason& OutReason, bool& bOutCookable, bool& bOutExplorable);
 
diff --git a/Engine/Source/Runtime/Engine/Classes/Engine/AssetManager.h b/Engine/Source/Runtime/Engine/Classes/Engine/AssetManager.h
--- a/Engine/Source/Runtime/Engine/Classes/Engine/AssetManager.h
+++ b/Engine/Source/Runtime/Engine/Classes/Engine/AssetManager.h
@@ -527,6 +527,13 @@ public:
 	virtual void ModifyDLCCook(const FString& DLCName, TConstArrayView<const ITargetPlatform*> TargetPlatforms,
 		TArray<FName>& PackagesToCook, TArray<FName>& PackagesToNeverCook);
 
+	/**
+	 * If the given package contains a primary asset, get the packages referenced by its AssetBundleEntries.
+	 * Used to inform the cook of should-be-cooked dependencies of PrimaryAssets for PrimaryAssets that
+	 * are recorded in the AssetManager but have cooktype Unknown and so are not returned from ModifyCook.
+	 */
+	virtual void ModifyCookReferences(FName PackageName, TArray<FName>& PackagesToCook);
+
 	/** Returns whether or not a specific UPackage should be cooked for the provied TargetPlatform */
 	virtual bool ShouldCookForPlatform(const UPackage* Package, const ITargetPlatform* TargetPlatform);
 
@@ -870,6 +877,14 @@ private:
 	/** Per-type asset information, cannot be accessed by children as it is defined in CPP file */
 	TMap<FName, TSharedRef<FPrimaryAssetTypeData>> AssetTypeMap;
 
+#if WITH_EDITOR
+	/**
+	 * Map from the PackageName of a PrimaryAsset to the packages that should be added to the cook
+	 * based on the PrimaryAsset's AssetBundleData, if the PrimaryAsset is ever referenced.
+	 */
+	TMap<FName, TArray<FTopLevelAssetPath>> AssetBundlePathsForPackage;
+#endif
+
 	mutable class IAssetRegistry* CachedAssetRegistry;
 	mutable const class UAssetManagerSettings* CachedSettings;
 
diff --git a/Engine/Source/Runtime/Engine/Private/AssetManager.cpp b/Engine/Source/Runtime/Engine/Private/AssetManager.cpp
--- a/Engine/Source/Runtime/Engine/Private/AssetManager.cpp
+++ b/Engine/Source/Runtime/Engine/Private/AssetManager.cpp
@@ -1215,32 +1215,23 @@ bool UAssetManager::TryUpdateCachedAssetData(const FPrimaryAssetId& PrimaryAsset
 		{
 			CachedAssetBundles.Add(PrimaryAssetId, NewAssetData.TaggedAssetBundles);
 
-			// FSoftObjectPathSerializationScope is costly and FSoftObjectPath::PostLoadPath()
-			// is only useful to make sure soft bundle references are cooked.
-			//
-			// Ideally we should not need this code, we should be able to pick up these
-			// references when we load the package header and then queue them up for cooking.
-			//
-			// For some reason, this doesn't work. StartupSoftObjectPackages in CookOnTheFlyServer.cpp
-			// become the same without this code, but GRedirectCollector picks up other soft references 
-			// thanks to this code that are not detected otherwise.
 #if WITH_EDITOR
 			if (GIsEditor)
 			{
-				static FName AssetBundleDataName("AssetBundleData");
-
-				// Mark these as editor only if our type is editor only
-				FSoftObjectPathSerializationScope SerializationScope(NewAssetData.PackageName, AssetBundleDataName,
-					TypeData.Info.bIsEditorOnly ? ESoftObjectPathCollectType::EditorOnlyCollect : ESoftObjectPathCollectType::AlwaysCollect,
-					ESoftObjectPathSerializeType::AlwaysSerialize);
-			
-				for (const FAssetBundleEntry& Entry : NewAssetData.TaggedAssetBundles->Bundles)
+				// Add the bundles to our data that notifies the cooker of extra soft references from
+				// primary assets, if the cooker encounters those primary assets referenced from elsewhere
+				// rather than us reporting them as AlwaysCook in ModifyCook
+				// Ignore the links from editor-only TypeDatas; the cooker only cares about the UsedInGame links.
+				if (!TypeData.Info.bIsEditorOnly)
 				{
-					for (const FTopLevelAssetPath& Path : Entry.AssetPaths)
+					TArray<FTopLevelAssetPath>& Paths = AssetBundlePathsForPackage.FindOrAdd(NewAssetData.PackageName);
+					Paths.Empty();
+					for (const FAssetBundleEntry& Entry : NewAssetData.TaggedAssetBundles->Bundles)
 					{
-						FSoftObjectPath FullPath(Path);
-						FullPath.PostLoadPath(nullptr);
+						Paths.Append(Entry.AssetPaths);
 					}
+					Paths.Sort([](const FTopLevelAssetPath& A, const FTopLevelAssetPath& B) { return A.CompareFast(B) < 0; });
+					Paths.SetNum(Algo::Unique(Paths));
 				}
 			}
 #endif
@@ -4229,6 +4220,22 @@ void UAssetManager::ModifyDLCCook(const FString& DLCName, TConstArrayView<const
 	}
 }
 
+void UAssetManager::ModifyCookReferences(FName PackageName, TArray<FName>& PackagesToCook)
+{
+	TArray<FTopLevelAssetPath>* Paths = AssetBundlePathsForPackage.Find(PackageName);
+	if (!Paths)
+	{
+		return;
+	}
+	PackagesToCook.Reserve(Paths->Num());
+	for (const FTopLevelAssetPath& Path : *Paths)
+	{
+		PackagesToCook.Add(Path.GetPackageName());
+	}
+	Algo::Sort(PackagesToCook, FNameFastLess());
+	PackagesToCook.SetNum(Algo::Unique(PackagesToCook));
+}
+
 void UAssetManager::GatherPublicAssetsForPackage(FName PackagePath, TArray<FName>& PackagesToCook) const
 {
 	FARFilter Filter;
