OnlyEditorOnly: Mark up some more hidden dependencies as UsedInGame or EditorOnly so that OnlyEditorOnly knows what to do with them.
#rnx

diff --git a/Engine/Source/Runtime/Engine/Private/Components/SceneComponent.cpp b/Engine/Source/Runtime/Engine/Private/Components/SceneComponent.cpp
--- a/Engine/Source/Runtime/Engine/Private/Components/SceneComponent.cpp
+++ b/Engine/Source/Runtime/Engine/Private/Components/SceneComponent.cpp
@@ -30,6 +30,7 @@
 #include "Interfaces/ITargetPlatform.h"
 #include "DeviceProfiles/DeviceProfile.h"
 #include "Net/Core/PushModel/PushModel.h"
+#include  "UObject/ICookInfo.h"
 
 #define LOCTEXT_NAMESPACE "SceneComponent"
 
@@ -713,7 +714,10 @@ void USceneComponent::CreateSpriteComponent(class UTexture2D* SpriteTexture, boo
 	if (bVisualizeComponent && SpriteComponent == nullptr && GetOwner() && !GetWorld()->IsGameWorld())
 	{
 		// Create a new billboard component to serve as a visualization of the actor until there is another primitive component
-		SpriteComponent = NewObject<UBillboardComponent>(GetOwner(), NAME_None, RF_Transactional | RF_Transient | RF_TextExportTransient);
+		{
+			FCookLoadScope EditorOnlyLoadScope(ECookLoadType::EditorOnly);
+			SpriteComponent = NewObject<UBillboardComponent>(GetOwner(), NAME_None, RF_Transactional | RF_Transient | RF_TextExportTransient);
+		}
 
 		SpriteComponent->Sprite = SpriteTexture? SpriteTexture : LoadObject<UTexture2D>(nullptr, TEXT("/Engine/EditorResources/EmptyActor.EmptyActor"));
 		SpriteComponent->SetRelativeScale3D_Direct(FVector(0.5f, 0.5f, 0.5f));
