Pinned Actors: Test SupportsPinning + Add new condition for !IsEditorRelevant() && !IsRuntimeOnly() actors

Fixes WorldMiniMap getting pinned when creating a new Level Instance and will support ShouldSkipFromLevelInstance actors properly

#rnx

diff --git a/Engine/Source/Editor/WorldPartitionEditor/Private/WorldPartition/SWorldPartitionEditorGrid2D.cpp b/Engine/Source/Editor/WorldPartitionEditor/Private/WorldPartition/SWorldPartitionEditorGrid2D.cpp
--- a/Engine/Source/Editor/WorldPartitionEditor/Private/WorldPartition/SWorldPartitionEditorGrid2D.cpp
+++ b/Engine/Source/Editor/WorldPartitionEditor/Private/WorldPartition/SWorldPartitionEditorGrid2D.cpp
@@ -1065,7 +1065,7 @@ void SWorldPartitionEditorGrid2D::Tick(const FGeometry& AllottedGeometry, const
 
 	// Also include transient actor loader adapters that might have been spawned by blutilities, etc. Since these actors can't be saved because they are transient,
 	// they will never get an actor descriptor so they will never appear in the world partition editor. Also include unsaved, newly created actors for convenience.
-	for (auto& [Reference, Actor] : GetWorldPartition()->GetDirtyActors())
+	for (auto& [Reference, Actor] : GetWorldPartition()->GetForceLoadedActors())
 	{
 		if (!GetWorldPartition()->GetActorDesc(Actor->GetActorGuid()) && Actor->Implements<UWorldPartitionActorLoaderInterface>())
 		{
diff --git a/Engine/Source/Runtime/Engine/Private/WorldPartition/LoaderAdapter/LoaderAdapterPinnedActors.cpp b/Engine/Source/Runtime/Engine/Private/WorldPartition/LoaderAdapter/LoaderAdapterPinnedActors.cpp
--- a/Engine/Source/Runtime/Engine/Private/WorldPartition/LoaderAdapter/LoaderAdapterPinnedActors.cpp
+++ b/Engine/Source/Runtime/Engine/Private/WorldPartition/LoaderAdapter/LoaderAdapterPinnedActors.cpp
@@ -13,7 +13,7 @@
 bool FLoaderAdapterPinnedActors::PassActorDescFilter(const FWorldPartitionHandle& ActorHandle) const
 {
 	// We want to be able to pin any type of actors (HLODs, etc).
-	return ActorHandle.IsValid() && !ActorsToRemove.Contains(ActorHandle);
+	return ActorHandle.IsValid() && !ActorsToRemove.Contains(ActorHandle) && SupportsPinning(ActorHandle.Get());
 }
 
 bool FLoaderAdapterPinnedActors::SupportsPinning(FWorldPartitionActorDesc* InActorDesc)
@@ -29,6 +29,12 @@ bool FLoaderAdapterPinnedActors::SupportsPinning(FWorldPartitionActorDesc* InAct
 		return false;
 	}
 
+	// This allows skipping actors that can never be loaded in current context: ex: bActorShouldSkipLevelInstance
+	if (!InActorDesc->IsEditorRelevant() && !InActorDesc->GetActorIsRuntimeOnly())
+	{
+		return false;
+	}
+
 	if (UActorDescContainer* Container = InActorDesc->GetContainer())
 	{
 		if (Container->IsMainPartitionContainer())
diff --git a/Engine/Source/Runtime/Engine/Private/WorldPartition/WorldPartition.cpp b/Engine/Source/Runtime/Engine/Private/WorldPartition/WorldPartition.cpp
--- a/Engine/Source/Runtime/Engine/Private/WorldPartition/WorldPartition.cpp
+++ b/Engine/Source/Runtime/Engine/Private/WorldPartition/WorldPartition.cpp
@@ -341,7 +341,7 @@ void UWorldPartition::OnPackageDirtyStateChanged(UPackage* Package)
 		{
 			if (Package->IsDirty())
 			{
-				DirtyActors.Add(ActorHandle.ToReference(), Actor);
+				ForceLoadedActors.Add(ActorHandle.ToReference(), Actor);
 			}
 		}
 	}
@@ -709,7 +709,7 @@ void UWorldPartition::Uninitialize()
 			RegisteredEditorLoaderAdapters.Empty();
 		}
 
-		DirtyActors.Empty();
+		ForceLoadedActors.Empty();
 
 		UninitializeActorDescContainers();
 		ActorDescContainer = nullptr;
@@ -833,11 +833,11 @@ void UWorldPartition::RegisterDelegates()
 			GEditor->OnPostBugItGoCalled().AddUObject(this, &UWorldPartition::OnPostBugItGoCalled);
 			GEditor->OnEditorClose().AddUObject(this, &UWorldPartition::SavePerUserSettings);
 			FWorldDelegates::OnPostWorldRename.AddUObject(this, &UWorldPartition::OnWorldRenamed);
+		}
 
-			if (!IsRunningCommandlet())
-			{
-				UPackage::PackageDirtyStateChangedEvent.AddUObject(this, &UWorldPartition::OnPackageDirtyStateChanged);
-			}
+		if (!IsRunningCommandlet())
+		{
+			UPackage::PackageDirtyStateChangedEvent.AddUObject(this, &UWorldPartition::OnPackageDirtyStateChanged);
 		}
 	}
 #endif
@@ -881,11 +881,11 @@ void UWorldPartition::UnregisterDelegates()
 
 			GEditor->OnPostBugItGoCalled().RemoveAll(this);
 			GEditor->OnEditorClose().RemoveAll(this);
+		}
 
-			if (!IsRunningCommandlet())
-			{
-				UPackage::PackageDirtyStateChangedEvent.RemoveAll(this);
-			}
+		if (!IsRunningCommandlet())
+		{
+			UPackage::PackageDirtyStateChangedEvent.RemoveAll(this);
 		}
 	}
 #endif
@@ -1140,7 +1140,7 @@ void UWorldPartition::OnActorDescAdded(FWorldPartitionActorDesc* NewActorDesc)
 
 	if (AActor* NewActor = NewActorDesc->GetActor())
 	{
-		DirtyActors.Add(FWorldPartitionReference(NewActorDesc->GetContainer(), NewActorDesc->GetGuid()), NewActor);
+		ForceLoadedActors.Add(FWorldPartitionReference(NewActorDesc->GetContainer(), NewActorDesc->GetGuid()), NewActor);
 	}
 
 	if (WorldPartitionEditor)
@@ -1373,7 +1373,7 @@ void UWorldPartition::AddReferencedObjects(UObject* InThis, FReferenceCollector&
 	if (IsGarbageCollecting())
 	{
 		Collector.AllowEliminatingReferences(false);
-		for (auto& [ActorReference, Actor] : This->DirtyActors)
+		for (auto& [ActorReference, Actor] : This->ForceLoadedActors)
 		{
 			Collector.AddReferencedObject(Actor);
 		}
@@ -1397,19 +1397,23 @@ void UWorldPartition::Tick(float DeltaSeconds)
 		EditorHash->Tick(DeltaSeconds);
 	}
 
-	if (PinnedActors)
+	for (TMap<FWorldPartitionReference, AActor*>::TIterator ForceLoadedActorIt(ForceLoadedActors); ForceLoadedActorIt; ++ForceLoadedActorIt)
 	{
-		for (TMap<FWorldPartitionReference, AActor*>::TIterator DirtyActorIt(DirtyActors); DirtyActorIt; ++DirtyActorIt)
+		if (!ForceLoadedActorIt.Key().IsValid())
+		{
+			ForceLoadedActorIt.RemoveCurrent();
+		}
+		else if (!ForceLoadedActorIt.Value()->GetPackage()->IsDirty())
 		{
-			if (!DirtyActorIt.Key().IsValid() || !DirtyActorIt.Value()->GetPackage()->IsDirty())
+			// Transfer ownership of the last ref if actor can be pinned
+			if (ForceLoadedActorIt.Key()->GetHardRefCount() <= 1 && PinnedActors && FLoaderAdapterPinnedActors::SupportsPinning(ForceLoadedActorIt.Key().Get()))
 			{
-				// If we hold the last reference to that actor (or no reference are held at all), pin it to avoid unloading
-				if (DirtyActorIt.Key().IsValid() && DirtyActorIt.Key()->GetHardRefCount() <= 1)
-				{
-					PinnedActors->AddActors({ DirtyActorIt.Key().ToHandle() });
-				}
-
-				DirtyActorIt.RemoveCurrent();
+				PinnedActors->AddActors({ ForceLoadedActorIt.Key().ToHandle() });
+				ForceLoadedActorIt.RemoveCurrent();
+			} // Clean up if a loader took a reference on the actor (ILoaderAdapter::RefreshLoadedState was called since)
+			else if (ForceLoadedActorIt.Key()->GetHardRefCount() > 1)
+			{
+				ForceLoadedActorIt.RemoveCurrent();
 			}
 		}
 	}
@@ -1853,11 +1857,11 @@ bool UWorldPartition::UnregisterActorDescContainer(UActorDescContainer* InActorD
 			{
 				ActorGuids.Add(It->GetGuid());
 
-				for (TMap<FWorldPartitionReference, AActor*>::TIterator DirtyActorIt(DirtyActors); DirtyActorIt; ++DirtyActorIt)
+				for (TMap<FWorldPartitionReference, AActor*>::TIterator ForceLoadedActorIt(ForceLoadedActors); ForceLoadedActorIt; ++ForceLoadedActorIt)
 				{
-					if (DirtyActorIt.Key() == ActorHandle)
+					if (ForceLoadedActorIt.Key() == ActorHandle)
 					{
-						DirtyActorIt.RemoveCurrent();
+						ForceLoadedActorIt.RemoveCurrent();
 					}
 				}
 			}
diff --git a/Engine/Source/Runtime/Engine/Public/WorldPartition/WorldPartition.h b/Engine/Source/Runtime/Engine/Public/WorldPartition/WorldPartition.h
--- a/Engine/Source/Runtime/Engine/Public/WorldPartition/WorldPartition.h
+++ b/Engine/Source/Runtime/Engine/Public/WorldPartition/WorldPartition.h
@@ -305,7 +305,10 @@ public:
 
 	bool IsEnablingStreamingJustified() const { return bEnablingStreamingJustified; }
 
-	const TMap<FWorldPartitionReference, AActor*>& GetDirtyActors() const { return DirtyActors; }
+	UE_DEPRECATED(5.3, "Use GetForceLoadedActors instead.")
+	const TMap<FWorldPartitionReference, AActor*>& GetDirtyActors() const { return ForceLoadedActors; }
+	const TMap<FWorldPartitionReference, AActor*>& GetForceLoadedActors() const { return ForceLoadedActors; }
+
 #endif
 
 public:
@@ -419,7 +422,7 @@ private:
 
 	TArray<FWorldPartitionReference> LoadedSubobjects;
 
-	TMap<FWorldPartitionReference, AActor*> DirtyActors;
+	TMap<FWorldPartitionReference, AActor*> ForceLoadedActors;
 
 	TSet<FString> GeneratedStreamingPackageNames;
 
