Fixed StreamingSource.bRemote not properly set by APlayerController::GetStreamingSourcesInternal.

diff --git a/Engine/Source/Runtime/Engine/Private/PlayerController.cpp b/Engine/Source/Runtime/Engine/Private/PlayerController.cpp
--- a/Engine/Source/Runtime/Engine/Private/PlayerController.cpp
+++ b/Engine/Source/Runtime/Engine/Private/PlayerController.cpp
@@ -3664,6 +3664,7 @@ bool APlayerController::GetStreamingSourcesInternal(TArray<FWorldPartitionStream
 	StreamingSource.bBlockOnSlowLoading = StreamingSourceShouldBlockOnSlowStreaming();
 	StreamingSource.DebugColor = StreamingSourceDebugColor;
 	StreamingSource.Priority = GetStreamingSourcePriority();
+	StreamingSource.bRemote = !IsLocalController();
 	GetStreamingSourceShapes(StreamingSource.Shapes);
 	return true;
 }
