Added OnCellShown/OnCellHidden in UWorldPartitionRuntimeLevelStreamingCell for derived classes to notify world partition when the activated cell is actually shown/hidden.

diff --git a/Engine/Source/Runtime/Engine/Private/WorldPartition/WorldPartitionRuntimeLevelStreamingCell.cpp b/Engine/Source/Runtime/Engine/Private/WorldPartition/WorldPartitionRuntimeLevelStreamingCell.cpp
--- a/Engine/Source/Runtime/Engine/Private/WorldPartition/WorldPartitionRuntimeLevelStreamingCell.cpp
+++ b/Engine/Source/Runtime/Engine/Private/WorldPartition/WorldPartitionRuntimeLevelStreamingCell.cpp
@@ -511,6 +511,11 @@ void UWorldPartitionRuntimeLevelStreamingCell::Deactivate() const
 }
 
 void UWorldPartitionRuntimeLevelStreamingCell::OnLevelShown()
+{
+	OnCellShown();
+}
+
+void UWorldPartitionRuntimeLevelStreamingCell::OnCellShown() const
 {
 	UWorldPartition* OuterWorldPartition = GetOuterWorld()->GetWorldPartition();
 	if (OuterWorldPartition && OuterWorldPartition->IsInitialized())
@@ -520,6 +525,11 @@ void UWorldPartitionRuntimeLevelStreamingCell::OnLevelShown()
 }
 
 void UWorldPartitionRuntimeLevelStreamingCell::OnLevelHidden()
+{
+	OnCellHidden();
+}
+
+void UWorldPartitionRuntimeLevelStreamingCell::OnCellHidden() const
 {
 	UWorldPartition* OuterWorldPartition = GetOuterWorld()->GetWorldPartition();
 	if (OuterWorldPartition && OuterWorldPartition->IsInitialized())
diff --git a/Engine/Source/Runtime/Engine/Public/WorldPartition/WorldPartitionRuntimeLevelStreamingCell.h b/Engine/Source/Runtime/Engine/Public/WorldPartition/WorldPartitionRuntimeLevelStreamingCell.h
--- a/Engine/Source/Runtime/Engine/Public/WorldPartition/WorldPartitionRuntimeLevelStreamingCell.h
+++ b/Engine/Source/Runtime/Engine/Public/WorldPartition/WorldPartitionRuntimeLevelStreamingCell.h
@@ -66,6 +66,12 @@ class ENGINE_API UWorldPartitionRuntimeLevelStreamingCell : public UWorldPartiti
 	const TArray<FWorldPartitionRuntimeCellObjectMapping>& GetPackages() const { return Packages; }
 #endif
 
+protected:
+	// Called when cell is shown
+	virtual void OnCellShown() const;
+	// Called when cell is hidden
+	virtual void OnCellHidden() const;
+
 private:
 	UFUNCTION()
 	void OnLevelShown();
