Fix uninitialized variable: bIsMainWorldOnly

#rnx

diff --git a/Engine/Source/Runtime/Engine/Private/WorldPartition/WorldPartitionActorDesc.cpp b/Engine/Source/Runtime/Engine/Private/WorldPartition/WorldPartitionActorDesc.cpp
--- a/Engine/Source/Runtime/Engine/Private/WorldPartition/WorldPartitionActorDesc.cpp
+++ b/Engine/Source/Runtime/Engine/Private/WorldPartition/WorldPartitionActorDesc.cpp
@@ -48,6 +48,7 @@ FWorldPartitionActorDesc::FWorldPartitionActorDesc()
 	, bIsForcedNonSpatiallyLoaded(false)
 	, bIsDefaultActorDesc(false)
 	, UnloadedReason(nullptr)
+	, bActorIsMainWorldOnly(false)
 {}
 
 void FWorldPartitionActorDesc::Init(const AActor* InActor)
