Fix init order.

#rnx

diff --git a/Engine/Source/Runtime/Engine/Private/WorldPartition/WorldPartitionActorDesc.cpp b/Engine/Source/Runtime/Engine/Private/WorldPartition/WorldPartitionActorDesc.cpp
--- a/Engine/Source/Runtime/Engine/Private/WorldPartition/WorldPartitionActorDesc.cpp
+++ b/Engine/Source/Runtime/Engine/Private/WorldPartition/WorldPartitionActorDesc.cpp
@@ -40,7 +40,8 @@ static FGuid GetDefaultActorDescGuid(const FWorldPartitionActorDesc* ActorDesc)
 }
 
 FWorldPartitionActorDesc::FWorldPartitionActorDesc()
-	: bIsUsingDataLayerAsset(false)
+	: bActorIsMainWorldOnly(false)
+	, bIsUsingDataLayerAsset(false)
 	, bIsBoundsValid(false)
 	, SoftRefCount(0)
 	, HardRefCount(0)
@@ -48,7 +49,6 @@ FWorldPartitionActorDesc::FWorldPartitionActorDesc()
 	, bIsForcedNonSpatiallyLoaded(false)
 	, bIsDefaultActorDesc(false)
 	, UnloadedReason(nullptr)
-	, bActorIsMainWorldOnly(false)
 {}
 
 void FWorldPartitionActorDesc::Init(const AActor* InActor)
