LevelInstanceActorDesc: Save a full asset path instead of a package name for the level asset

#rnx

diff --git a/Engine/Source/Runtime/Engine/Private/WorldPartition/LevelInstance/LevelInstanceActorDesc.cpp b/Engine/Source/Runtime/Engine/Private/WorldPartition/LevelInstance/LevelInstanceActorDesc.cpp
--- a/Engine/Source/Runtime/Engine/Private/WorldPartition/LevelInstance/LevelInstanceActorDesc.cpp
+++ b/Engine/Source/Runtime/Engine/Private/WorldPartition/LevelInstance/LevelInstanceActorDesc.cpp
@@ -42,7 +42,7 @@ void FLevelInstanceActorDesc::Init(const AActor* InActor)
 	FWorldPartitionActorDesc::Init(InActor);
 
 	const ILevelInstanceInterface* LevelInstance = CastChecked<ILevelInstanceInterface>(InActor);
-	LevelPackage = *LevelInstance->GetWorldAssetPackage();
+	WorldAsset = LevelInstance->GetWorldAsset().ToSoftObjectPath();
 	LevelInstanceTransform = InActor->GetActorTransform();
 	DesiredRuntimeBehavior = LevelInstance->GetDesiredRuntimeBehavior();
 	Filter = LevelInstance->GetFilter();
@@ -68,7 +68,7 @@ bool FLevelInstanceActorDesc::Equals(const FWorldPartitionActorDesc* Other) cons
 		const FLevelInstanceActorDesc* LevelInstanceActorDesc = (FLevelInstanceActorDesc*)Other;
 
 		return
-			LevelPackage == LevelInstanceActorDesc->LevelPackage &&
+			WorldAsset == LevelInstanceActorDesc->WorldAsset &&
 			LevelInstanceTransform.Equals(LevelInstanceActorDesc->LevelInstanceTransform, 0.1f) &&
 			DesiredRuntimeBehavior == LevelInstanceActorDesc->DesiredRuntimeBehavior;
 	}
@@ -161,7 +161,7 @@ bool FLevelInstanceActorDesc::IsContainerInstanceInternal() const
 		return false;
 	}
 
-	if (LevelPackage.IsNull())
+	if (WorldAsset.IsNull())
 	{
 		return false;
 	}
@@ -243,14 +243,14 @@ void FLevelInstanceActorDesc::Serialize(FArchive& Ar)
 
 	if (Ar.CustomVer(FFortniteMainBranchObjectVersion::GUID) < FFortniteMainBranchObjectVersion::WorldPartitionActorDescSerializeSoftObjectPathSupport)
 	{
-		Ar << TDeltaSerialize<FSoftObjectPath, FName>(LevelPackage, [](FSoftObjectPath& Value, const FName& DeprecatedValue)
+		Ar << TDeltaSerialize<FSoftObjectPath, FName>(WorldAsset, [](FSoftObjectPath& Value, const FName& DeprecatedValue)
 		{
 			Value = FSoftObjectPath(*DeprecatedValue.ToString());
 		});
 	}
 	else
 	{
-		Ar << TDeltaSerialize<FSoftObjectPath>(LevelPackage);
+		Ar << TDeltaSerialize<FSoftObjectPath>(WorldAsset);
 	}
 
 	if (!bIsDefaultActorDesc)
@@ -286,7 +286,7 @@ void FLevelInstanceActorDesc::Serialize(FArchive& Ar)
 
 			const AActor* CDO = GetActorNativeClass()->GetDefaultObject<AActor>();
 			const ILevelInstanceInterface* LevelInstanceCDO = CastChecked<ILevelInstanceInterface>(CDO);
-			if (LevelPackage.IsValid() && (LevelInstanceCDO->IsLoadingEnabled() || bFixupOldVersion))
+			if (WorldAsset.IsValid() && (LevelInstanceCDO->IsLoadingEnabled() || bFixupOldVersion))
 			{
 				if (!IsContainerInstance())
 				{
diff --git a/Engine/Source/Runtime/Engine/Public/WorldPartition/LevelInstance/LevelInstanceActorDesc.h b/Engine/Source/Runtime/Engine/Public/WorldPartition/LevelInstance/LevelInstanceActorDesc.h
--- a/Engine/Source/Runtime/Engine/Public/WorldPartition/LevelInstance/LevelInstanceActorDesc.h
+++ b/Engine/Source/Runtime/Engine/Public/WorldPartition/LevelInstance/LevelInstanceActorDesc.h
@@ -27,7 +27,7 @@ public:
 
 	virtual bool IsContainerInstance() const override;
 	virtual bool IsContainerFilter() const override { return IsContainerInstance(); }
-	virtual FName GetContainerPackage() const override { return *LevelPackage.ToString(); }
+	virtual FName GetContainerPackage() const override { return WorldAsset.GetLongPackageFName(); }
 	virtual bool GetContainerInstance(FContainerInstance& OutContainerInstance) const override;
 	virtual const FWorldPartitionActorFilter* GetContainerFilter() const override { return &Filter; }
 	virtual void CheckForErrors(IStreamingGenerationErrorHandler* ErrorHandler) const override;
@@ -41,7 +41,7 @@ protected:
 	virtual void Serialize(FArchive& Ar) override;
 	virtual void SetContainer(UActorDescContainer* InContainer, UWorld* InWorldContext) override;
 
-	FSoftObjectPath LevelPackage;
+	FSoftObjectPath WorldAsset;
 	FTransform LevelInstanceTransform;
 	ELevelInstanceRuntimeBehavior DesiredRuntimeBehavior;
 
