some additional profile scopes for isolating rebuild delays

diff --git a/Engine/Source/Runtime/Engine/Private/Texture.cpp b/Engine/Source/Runtime/Engine/Private/Texture.cpp
--- a/Engine/Source/Runtime/Engine/Private/Texture.cpp
+++ b/Engine/Source/Runtime/Engine/Private/Texture.cpp
@@ -2489,6 +2489,8 @@ FTextureSource::FMipData FTextureSource::GetMipData(IImageWrapperModule* )
 	check(NumLockedMips == 0);
 
 #if WITH_EDITOR
+	// We can end up waiting here a lot as the bulk data gets serialized for entry in to the transaction buffer.
+	TRACE_CPUPROFILER_EVENT_SCOPE(FTextureSource::GetMipData (StartLock) );
 	FReadScopeLock _(BulkDataLock.Get());
 #endif //WITH_EDITOR
 
diff --git a/Engine/Source/Runtime/Engine/Private/TextureDerivedDataTask.cpp b/Engine/Source/Runtime/Engine/Private/TextureDerivedDataTask.cpp
--- a/Engine/Source/Runtime/Engine/Private/TextureDerivedDataTask.cpp
+++ b/Engine/Source/Runtime/Engine/Private/TextureDerivedDataTask.cpp
@@ -273,41 +273,44 @@ void FTextureSourceData::GetSourceMips(FTextureSource& Source, IImageWrapperModu
 
 		const FTextureSource::FMipData ScopedMipData = Source.GetMipData(InImageWrapper);
 
-		for (int32 BlockIndex = 0; BlockIndex < Blocks.Num(); ++BlockIndex)
 		{
-			FTextureSourceBlock SourceBlock;
-			Source.GetBlock(BlockIndex, SourceBlock);
-
-			FTextureSourceBlockData& BlockData = Blocks[BlockIndex];
-			for (int32 LayerIndex = 0; LayerIndex < Layers.Num(); ++LayerIndex)
+			TRACE_CPUPROFILER_EVENT_SCOPE(FTextureSourceData::GetSourceMips_CopyMips);
+			for (int32 BlockIndex = 0; BlockIndex < Blocks.Num(); ++BlockIndex)
 			{
-				const FTextureSourceLayerData& LayerData = Layers[LayerIndex];
-				if (!BlockData.MipsPerLayer[LayerIndex].Num()) // If we already got valid data, nothing to do.
+				FTextureSourceBlock SourceBlock;
+				Source.GetBlock(BlockIndex, SourceBlock);
+
+				FTextureSourceBlockData& BlockData = Blocks[BlockIndex];
+				for (int32 LayerIndex = 0; LayerIndex < Layers.Num(); ++LayerIndex)
 				{
-					int32 MipSizeX = SourceBlock.SizeX;
-					int32 MipSizeY = SourceBlock.SizeY;
-					int32 MipSizeZ = SourceBlock.NumSlices;
-					for (int32 MipIndex = 0; MipIndex < BlockData.NumMips; ++MipIndex)
+					const FTextureSourceLayerData& LayerData = Layers[LayerIndex];
+					if (!BlockData.MipsPerLayer[LayerIndex].Num()) // If we already got valid data, nothing to do.
 					{
-						FImage* SourceMip = new(BlockData.MipsPerLayer[LayerIndex]) FImage(
-							MipSizeX, MipSizeY, MipSizeZ,
-							LayerData.ImageFormat,
-							LayerData.SourceGammaSpace
-						);
-
-						if (!ScopedMipData.GetMipData(SourceMip->RawData, BlockIndex, LayerIndex, MipIndex))
+						int32 MipSizeX = SourceBlock.SizeX;
+						int32 MipSizeY = SourceBlock.SizeY;
+						int32 MipSizeZ = SourceBlock.NumSlices;
+						for (int32 MipIndex = 0; MipIndex < BlockData.NumMips; ++MipIndex)
 						{
-							UE_LOG(LogTexture, Warning, TEXT("Cannot retrieve source data for mip %d of %s"), MipIndex, *TextureFullName);
-							ReleaseMemory();
-							bValid = false;
-							break;
-						}
-
-						MipSizeX = FMath::Max(MipSizeX / 2, 1);
-						MipSizeY = FMath::Max(MipSizeY / 2, 1);
-						if ( Source.IsVolume() )
-						{
-							MipSizeZ = FMath::Max(MipSizeZ / 2, 1);
+							FImage* SourceMip = new(BlockData.MipsPerLayer[LayerIndex]) FImage(
+								MipSizeX, MipSizeY, MipSizeZ,
+								LayerData.ImageFormat,
+								LayerData.SourceGammaSpace
+							);
+
+							if (!ScopedMipData.GetMipData(SourceMip->RawData, BlockIndex, LayerIndex, MipIndex))
+							{
+								UE_LOG(LogTexture, Warning, TEXT("Cannot retrieve source data for mip %d of %s"), MipIndex, *TextureFullName);
+								ReleaseMemory();
+								bValid = false;
+								break;
+							}
+
+							MipSizeX = FMath::Max(MipSizeX / 2, 1);
+							MipSizeY = FMath::Max(MipSizeY / 2, 1);
+							if ( Source.IsVolume() )
+							{
+								MipSizeZ = FMath::Max(MipSizeZ / 2, 1);
+							}
 						}
 					}
 				}
@@ -534,6 +537,8 @@ static void DDC1_StoreClassicTextureInDerivedData(
 	uint32 ExtData, bool bReplaceExistingDDC, const FString& TexturePathName, const FString& KeySuffix, int64& BytesCached
 	)
 {
+	TRACE_CPUPROFILER_EVENT_SCOPE(Texture.DDC1_StoreClassicTextureInDerivedData);
+
 	const int32 MipCount = CompressedMips.Num();
 
 	// VT can be bigger than (1<<(MAX_TEXTURE_MIP_COUNT-1)) , but doesn't actually make all those mips
@@ -1369,6 +1374,8 @@ static void DDC1_FetchAndFillDerivedData(
 	using namespace UE;
 	using namespace UE::DerivedData;
 
+	TRACE_CPUPROFILER_EVENT_SCOPE(Texture.DDC1_FetchAndFillDerivedData);
+
 	bool bForceRebuild = EnumHasAnyFlags(CacheFlags, ETextureCacheFlags::ForceRebuild);
 	bool bForVirtualTextureStreamingBuild = EnumHasAnyFlags(CacheFlags, ETextureCacheFlags::ForVirtualTextureStreamingBuild);
 
@@ -1638,9 +1645,12 @@ static bool DDC1_LoadAndValidateTextureData(
 	bool bAllowAsyncLoading
 )
 {
+	TRACE_CPUPROFILER_EVENT_SCOPE(Texture.DDC1_LoadAndValidateTextureData);
+
 	bool bHasTextureSourceMips = false;
 	if (TextureData.IsValid() && Texture.Source.IsBulkDataLoaded())
 	{
+		TRACE_CPUPROFILER_EVENT_SCOPE(GetSourceMips);
 		TextureData.GetSourceMips(Texture.Source, ImageWrapper);
 		bHasTextureSourceMips = true;
 	}
@@ -1648,6 +1658,7 @@ static bool DDC1_LoadAndValidateTextureData(
 	bool bHasCompositeTextureSourceMips = false;
 	if (CompositeTextureData.IsValid() && Texture.CompositeTexture && Texture.CompositeTexture->Source.IsBulkDataLoaded())
 	{
+		TRACE_CPUPROFILER_EVENT_SCOPE(GetCompositeSourceMips);
 		check( Texture.CompositeTexture->Source.IsValid() );
 		CompositeTextureData.GetSourceMips(Texture.CompositeTexture->Source, ImageWrapper);
 		bHasCompositeTextureSourceMips = true;
@@ -1657,12 +1668,14 @@ static bool DDC1_LoadAndValidateTextureData(
 	{
 		if (!bHasTextureSourceMips)
 		{
+			TRACE_CPUPROFILER_EVENT_SCOPE(GetAsyncSourceMips);
 			TextureData.GetAsyncSourceMips(ImageWrapper);
 			TextureData.AsyncSource.RemoveBulkData();
 		}
 
 		if (!bHasCompositeTextureSourceMips)
 		{
+			TRACE_CPUPROFILER_EVENT_SCOPE(GetAsyncCompositeSourceMips);
 			CompositeTextureData.GetAsyncSourceMips(ImageWrapper);
 			CompositeTextureData.AsyncSource.RemoveBulkData();
 		}
