Building HLODs fails with Return Code: 3

diff --git a/Engine/Source/Runtime/Engine/Private/WorldPartition/WorldPartitionLevelStreamingDynamic.cpp b/Engine/Source/Runtime/Engine/Private/WorldPartition/WorldPartitionLevelStreamingDynamic.cpp
--- a/Engine/Source/Runtime/Engine/Private/WorldPartition/WorldPartitionLevelStreamingDynamic.cpp
+++ b/Engine/Source/Runtime/Engine/Private/WorldPartition/WorldPartitionLevelStreamingDynamic.cpp
@@ -85,11 +85,8 @@ UWorldPartitionLevelStreamingDynamic* UWorldPartitionLevelStreamingDynamic::Load
 	check(World->WorldType == EWorldType::Editor);
 	UWorldPartitionLevelStreamingDynamic* LevelStreaming = NewObject<UWorldPartitionLevelStreamingDynamic>(World, LevelStreamingName, RF_Transient);
 	
-	const FString WorldPackageName = World->GetPackage()->GetPathName();
-	const FStringView WorldMountPointName = FPathViews::GetMountPointNameFromPath(WorldPackageName);
-
-	FString PackageName = FString::Printf(TEXT("Temp/%s"), *LevelStreamingName.ToString());
-	TSoftObjectPtr<UWorld> WorldAsset(FSoftObjectPath(FString::Printf(TEXT("/%.*s/%s.%s"), WorldMountPointName.Len(), WorldMountPointName.GetData(), *PackageName, *World->GetName())));
+	FString PackageName = FString::Printf(TEXT("/Temp/%s"), *LevelStreamingName.ToString());
+	TSoftObjectPtr<UWorld> WorldAsset(FSoftObjectPath(FString::Printf(TEXT("%s.%s"), *PackageName, *World->GetName())));
 	LevelStreaming->SetWorldAsset(WorldAsset);
 	
 	LevelStreaming->LevelTransform = FTransform::Identity;
@@ -191,8 +188,11 @@ void UWorldPartitionLevelStreamingDynamic::CreateRuntimeLevel()
  */
 bool UWorldPartitionLevelStreamingDynamic::RequestLevel(UWorld* InPersistentWorld, bool bInAllowLevelLoadRequests, EReqLevelBlock InBlockPolicy)
 {
+	const FName WorldAssetPackageFName = GetWorldAssetPackageFName();
+	const FString WorldAssetPackageName = GetWorldAssetPackageName();
+
 	// Use parent streaming implementation if world asset package is not a memory package
-	if (!FPackageName::IsMemoryPackage(GetWorldAssetPackageName()))
+	if (!FPackageName::IsMemoryPackage(WorldAssetPackageName) && !FPackageName::IsTempPackage(WorldAssetPackageName))
 	{
 		return Super::RequestLevel(InPersistentWorld, bInAllowLevelLoadRequests, InBlockPolicy);
 	}
@@ -226,7 +226,7 @@ bool UWorldPartitionLevelStreamingDynamic::RequestLevel(UWorld* InPersistentWorl
 	ULevel* PendingLevelVisOrInvis = (InPersistentWorld->GetCurrentLevelPendingVisibility() ? InPersistentWorld->GetCurrentLevelPendingVisibility() : InPersistentWorld->GetCurrentLevelPendingInvisibility());
 	if (PendingLevelVisOrInvis && PendingLevelVisOrInvis == LoadedLevel)
 	{
-		UE_LOG(LogLevelStreaming, Verbose, TEXT("Delaying load of new level %s, because still processing visibility request."), *GetWorldAssetPackageName());
+		UE_LOG(LogLevelStreaming, Verbose, TEXT("Delaying load of new level %s, because still processing visibility request."), *WorldAssetPackageName);
 		return false;
 	}
 
@@ -234,8 +234,7 @@ bool UWorldPartitionLevelStreamingDynamic::RequestLevel(UWorld* InPersistentWorl
 	FScopeCycleCounterUObject Context(InPersistentWorld);
 
 	// Try to find the package to load
-	const FName DesiredPackageName = GetWorldAssetPackageFName();
-	UPackage* LevelPackage = (UPackage*)StaticFindObjectFast(UPackage::StaticClass(), nullptr, DesiredPackageName, /*bExactClass =*/ false, RF_NoFlags, EInternalObjectFlags::Garbage);
+	UPackage* LevelPackage = (UPackage*)StaticFindObjectFast(UPackage::StaticClass(), nullptr, WorldAssetPackageFName, /*bExactClass =*/ false, RF_NoFlags, EInternalObjectFlags::Garbage);
 	UWorld* FoundWorld = LevelPackage ? UWorld::FindWorldInPackage(LevelPackage) : nullptr;
 	check(!FoundWorld || IsValidChecked(FoundWorld));
 	check(!FoundWorld || FoundWorld->PersistentLevel);
@@ -297,7 +296,7 @@ bool UWorldPartitionLevelStreamingDynamic::RequestLevel(UWorld* InPersistentWorl
 				{
 					if (IsAsyncLoading())
 					{
-						UE_LOG(LogStreaming, Display, TEXT("UWorldPartitionLevelStreamingDynamic::RequestLevel(%s) is flushing async loading"), *GetWorldAssetPackageName());
+						UE_LOG(LogStreaming, Display, TEXT("UWorldPartitionLevelStreamingDynamic::RequestLevel(%s) is flushing async loading"), *WorldAssetPackageName);
 					}
 
 					// Finish all async loading.
