ActorDesc: Add missing constructor defaults

#rnx

diff --git a/Engine/Plugins/Experimental/Water/Source/Runtime/Private/WaterZoneActorDesc.cpp b/Engine/Plugins/Experimental/Water/Source/Runtime/Private/WaterZoneActorDesc.cpp
--- a/Engine/Plugins/Experimental/Water/Source/Runtime/Private/WaterZoneActorDesc.cpp
+++ b/Engine/Plugins/Experimental/Water/Source/Runtime/Private/WaterZoneActorDesc.cpp
@@ -5,6 +5,11 @@
 #if WITH_EDITOR
 #include "WaterZoneActor.h"
 
+FWaterZoneActorDesc::FWaterZoneActorDesc()
+	: OverlapPriority(0)
+{
+}
+
 void FWaterZoneActorDesc::Init(const AActor* InActor)
 {
 	FWorldPartitionActorDesc::Init(InActor);
diff --git a/Engine/Plugins/Experimental/Water/Source/Runtime/Private/WaterZoneActorDesc.h b/Engine/Plugins/Experimental/Water/Source/Runtime/Private/WaterZoneActorDesc.h
--- a/Engine/Plugins/Experimental/Water/Source/Runtime/Private/WaterZoneActorDesc.h
+++ b/Engine/Plugins/Experimental/Water/Source/Runtime/Private/WaterZoneActorDesc.h
@@ -11,6 +11,7 @@
 class FWaterZoneActorDesc : public FWorldPartitionActorDesc
 {
 public:
+	FWaterZoneActorDesc();
 	virtual void Init(const AActor* InActor) override;
 	virtual bool Equals(const FWorldPartitionActorDesc* Other) const override;
 
diff --git a/Engine/Source/Runtime/Engine/Private/WorldPartition/ActorPartition/PartitionActorDesc.cpp b/Engine/Source/Runtime/Engine/Private/WorldPartition/ActorPartition/PartitionActorDesc.cpp
--- a/Engine/Source/Runtime/Engine/Private/WorldPartition/ActorPartition/PartitionActorDesc.cpp
+++ b/Engine/Source/Runtime/Engine/Private/WorldPartition/ActorPartition/PartitionActorDesc.cpp
@@ -6,6 +6,14 @@
 #if WITH_EDITOR
 #include "UObject/UE5MainStreamObjectVersion.h"
 
+FPartitionActorDesc::FPartitionActorDesc()
+	: GridSize(0)
+	, GridIndexX(0)
+	, GridIndexY(0)
+	, GridIndexZ(0)
+{
+}
+
 void FPartitionActorDesc::Init(const AActor* InActor)
 {
 	FWorldPartitionActorDesc::Init(InActor);
diff --git a/Engine/Source/Runtime/Engine/Private/WorldPartition/WorldPartitionActorDesc.cpp b/Engine/Source/Runtime/Engine/Private/WorldPartition/WorldPartitionActorDesc.cpp
--- a/Engine/Source/Runtime/Engine/Private/WorldPartition/WorldPartitionActorDesc.cpp
+++ b/Engine/Source/Runtime/Engine/Private/WorldPartition/WorldPartitionActorDesc.cpp
@@ -40,11 +40,16 @@ static FGuid GetDefaultActorDescGuid(const FWorldPartitionActorDesc* ActorDesc)
 }
 
 FWorldPartitionActorDesc::FWorldPartitionActorDesc()
-	: bActorIsMainWorldOnly(false)
+	: bIsSpatiallyLoaded(false)
+	, bActorIsEditorOnly(false)
+	, bActorIsRuntimeOnly(false)
+	, bActorIsMainWorldOnly(false)
+	, bActorIsHLODRelevant(false)
 	, bIsUsingDataLayerAsset(false)
 	, bIsBoundsValid(false)
 	, SoftRefCount(0)
 	, HardRefCount(0)
+	, ActorNativeClass(nullptr)
 	, Container(nullptr)
 	, bIsForcedNonSpatiallyLoaded(false)
 	, bIsDefaultActorDesc(false)
diff --git a/Engine/Source/Runtime/Engine/Public/WorldPartition/ActorPartition/PartitionActorDesc.h b/Engine/Source/Runtime/Engine/Public/WorldPartition/ActorPartition/PartitionActorDesc.h
--- a/Engine/Source/Runtime/Engine/Public/WorldPartition/ActorPartition/PartitionActorDesc.h
+++ b/Engine/Source/Runtime/Engine/Public/WorldPartition/ActorPartition/PartitionActorDesc.h
@@ -19,6 +19,7 @@ public:
 	int64 GridIndexZ;
 	FGuid GridGuid;
 
+	FPartitionActorDesc();
 protected:
 	virtual void Init(const AActor* InActor) override;
 	virtual bool Equals(const FWorldPartitionActorDesc* Other) const override;
