Fix RBAN Bone Space Simulation collision.

Body bounds were not being updated for kinematic bodies that never move, resulting in missed collisions.

diff --git a/Engine/Source/Runtime/Engine/Private/PhysicsEngine/ImmediatePhysics/ImmediatePhysicsChaos/ImmediatePhysicsActorHandle_Chaos.cpp b/Engine/Source/Runtime/Engine/Private/PhysicsEngine/ImmediatePhysics/ImmediatePhysicsChaos/ImmediatePhysicsActorHandle_Chaos.cpp
--- a/Engine/Source/Runtime/Engine/Private/PhysicsEngine/ImmediatePhysics/ImmediatePhysicsChaos/ImmediatePhysicsActorHandle_Chaos.cpp
+++ b/Engine/Source/Runtime/Engine/Private/PhysicsEngine/ImmediatePhysics/ImmediatePhysicsChaos/ImmediatePhysicsActorHandle_Chaos.cpp
@@ -351,6 +351,10 @@ namespace ImmediatePhysics_Chaos
 			Kinematic->W() = FVec3(0);
 			Kinematic->KinematicTarget().Clear();
 		}
+
+		// Initialize the bounds. Important because if the particle never moves its 
+		// bounds will never get updated (see FPBDMinEvolution::ApplyKinematicTargets) 
+		ParticleHandle->UpdateWorldSpaceState(FRigidTransform3(ParticleHandle->X(), ParticleHandle->R()), FVec3(0));
 	}
 
 	void FActorHandle::SetWorldTransform(const FTransform& WorldTM)
