Don't mark external level/OFPA assets as dirty when reinstancing and/or reconstructing an Actor-based Blueprint class instance that is attached to one or more children contained within those assets.

diff --git a/Engine/Source/Editor/UnrealEd/Private/Kismet2/KismetReinstanceUtilities.cpp b/Engine/Source/Editor/UnrealEd/Private/Kismet2/KismetReinstanceUtilities.cpp
--- a/Engine/Source/Editor/UnrealEd/Private/Kismet2/KismetReinstanceUtilities.cpp
+++ b/Engine/Source/Editor/UnrealEd/Private/Kismet2/KismetReinstanceUtilities.cpp
@@ -2562,6 +2562,20 @@ void FBlueprintCompileReinstancer::ReplaceInstancesOfClass_Inner(const TMap<UCla
 						UObject* NewUObject = nullptr;
 						if (OldActor->GetLevel() && OldActor->GetWorld())
 						{
+							// Attached actors will be marked dirty as well when the old instance is destroyed, but
+							// since these attachments will get restored as part of replacement, there's no need to
+							// re-save after reinstancing. In general, this applies only to OFPA-enabled levels, as
+							// in that case, attached actors reside in their own external packages and not the level.
+							TArray<AActor*> AttachedActors;
+							OldActor->GetAttachedActors(AttachedActors);
+							for (AActor* AttachedActor : AttachedActors)
+							{
+								if (AttachedActor)
+								{
+									CheckAndSaveOuterPackageToCleanList(AttachedActor);
+								}
+							}
+
 							ReplaceActorHelper(OldActor, OldClass, NewUObject, NewClass, OldToNewInstanceMap, InOldToNewClassMap, ReinstancedObjectsWeakReferenceMap, ActorAttachmentData, ReplacementActors, bPreserveRootComponent, bSelectionChanged);
 						}
 						else
diff --git a/Engine/Source/Runtime/Engine/Private/ActorConstruction.cpp b/Engine/Source/Runtime/Engine/Private/ActorConstruction.cpp
--- a/Engine/Source/Runtime/Engine/Private/ActorConstruction.cpp
+++ b/Engine/Source/Runtime/Engine/Private/ActorConstruction.cpp
@@ -293,13 +293,20 @@ void AActor::RerunConstructionScripts()
 		ITransaction* CurrentTransaction = GUndo;
 		GUndo = nullptr;
 
-		// Mark package as clean on exit if not transient and not already dirty.
-		bool bMarkPackageClean = false;
-		UPackage* ActorPackage = GetPackage();
-		if (ActorPackage && !ActorPackage->IsDirty() && ActorPackage != GetTransientPackage())
+		// Keep track of non-dirty packages, so we can clear the dirty state after reconstruction.
+		TSet<UPackage*> CleanPackageList;
+		auto CheckAndSaveOuterPackageToCleanList = [&CleanPackageList](const UObject* InObject)
 		{
-			bMarkPackageClean = true;
-		}
+			check(InObject);
+			UPackage* ObjectPackage = InObject->GetPackage();
+			if (ObjectPackage && !ObjectPackage->IsDirty() && ObjectPackage != GetTransientPackage())
+			{
+				CleanPackageList.Add(ObjectPackage);
+			}
+		};
+
+		// Mark package as clean on exit if not transient and not already dirty.
+		CheckAndSaveOuterPackageToCleanList(this);
 		
 		// Create cache to store component data across rerunning construction scripts
 		FComponentInstanceDataCache* InstanceDataCache;
@@ -382,6 +389,10 @@ void AActor::RerunConstructionScripts()
 				AActor* AttachedActor = CachedAttachInfo.Actor.Get();
 				if (AttachedActor)
 				{
+					// Detaching will mark the attached actor's package as dirty, but we don't actually need
+					// it to be re-saved after reconstruction since attachment relationships will be restored.
+					CheckAndSaveOuterPackageToCleanList(AttachedActor);
+
 					FAttachedActorInfo Info;
 					Info.AttachedActor = AttachedActor;
 					Info.AttachedToSocket = CachedAttachInfo.SocketName;
@@ -404,6 +415,10 @@ void AActor::RerunConstructionScripts()
 
 			for (AActor* AttachedActor : AttachedActors)
 			{
+				// Detaching will mark the attached actor's package as dirty, but we don't actually need
+				// it to be re-saved after reconstruction since attachment relationships will be restored.
+				CheckAndSaveOuterPackageToCleanList(AttachedActor);
+
 				// We don't need to detach child actors, that will be handled by component tear down
 				if (!AttachedActor->IsChildActor())
 				{
@@ -630,13 +645,16 @@ void AActor::RerunConstructionScripts()
 			}
 		}
 
-		// If any of the code above caused the package dirty state to change, we reset it back to a "clean" state now. Note that we have to do this
+		// If any of the code above caused a package dirty state to change, we reset it back to a "clean" state now. Note that we have to do this
 		// before we restore the undo buffer below - otherwise, the state change will become part of the current transaction, and we don't want that.
-		if (bMarkPackageClean && ActorPackage->IsDirty())
+		for (UPackage* PackageToMarkAsClean : CleanPackageList)
 		{
-			ActorPackage->SetDirtyFlag(false);
+			check(PackageToMarkAsClean);
+			PackageToMarkAsClean->SetDirtyFlag(false);
 		}
 
+		CleanPackageList.Empty();
+
 		// Restore the undo buffer
 		GUndo = CurrentTransaction;
 
