Ignoring bad set timeline curves to mitigate client crash
#rnx

diff --git a/Engine/Source/Runtime/Engine/Private/Timeline.cpp b/Engine/Source/Runtime/Engine/Private/Timeline.cpp
--- a/Engine/Source/Runtime/Engine/Private/Timeline.cpp
+++ b/Engine/Source/Runtime/Engine/Private/Timeline.cpp
@@ -79,128 +79,191 @@ void FTimeline::AddEvent(float Time, FOnTimelineEvent Event)
 
 void FTimeline::AddInterpVector(UCurveVector* VectorCurve, FOnTimelineVector InterpFunc, FName PropertyName, FName TrackName)
 {
-	FTimelineVectorTrack NewEntry;
-	NewEntry.VectorCurve = VectorCurve;
-	NewEntry.InterpFunc = InterpFunc;
-	NewEntry.TrackName = TrackName;
-	NewEntry.VectorPropertyName = PropertyName;
+	if (VectorCurve)
+	{
+		FTimelineVectorTrack NewEntry;
+		NewEntry.VectorCurve = VectorCurve;
+		NewEntry.InterpFunc = InterpFunc;
+		NewEntry.TrackName = TrackName;
+		NewEntry.VectorPropertyName = PropertyName;
 
-	InterpVectors.Add(NewEntry);
+		InterpVectors.Add(NewEntry);
+	}
+	else
+	{
+		UE_LOG(LogTimeline, Warning, TEXT("FTimeline::AddInterpVector: VectorCurve is null so it will not be added!"));
+	}
 }
 
 void FTimeline::AddInterpVector(UCurveVector* VectorCurve, FOnTimelineVectorStatic InterpFunc)
 {
-	FTimelineVectorTrack NewEntry;
-	NewEntry.VectorCurve = VectorCurve;
-	NewEntry.InterpFuncStatic = InterpFunc;
+	if (VectorCurve)
+	{
+		FTimelineVectorTrack NewEntry;
+		NewEntry.VectorCurve = VectorCurve;
+		NewEntry.InterpFuncStatic = InterpFunc;
 
-	InterpVectors.Add(NewEntry);
+		InterpVectors.Add(NewEntry);
+	}
+	else
+	{
+		UE_LOG(LogTimeline, Warning, TEXT("FTimeline::AddInterpVector: VectorCurve is null so it will not be added!"));
+	}
 }
 
 void FTimeline::AddInterpFloat(UCurveFloat* FloatCurve, FOnTimelineFloat InterpFunc, FName PropertyName, FName TrackName)
 {
-	FTimelineFloatTrack NewEntry;
-	NewEntry.FloatCurve = FloatCurve;
-	NewEntry.InterpFunc = InterpFunc;
-	NewEntry.TrackName = TrackName;
-	NewEntry.FloatPropertyName = PropertyName;
+	if (FloatCurve)
+	{
+		FTimelineFloatTrack NewEntry;
+		NewEntry.FloatCurve = FloatCurve;
+		NewEntry.InterpFunc = InterpFunc;
+		NewEntry.TrackName = TrackName;
+		NewEntry.FloatPropertyName = PropertyName;
 
-	InterpFloats.Add(NewEntry);
+		InterpFloats.Add(NewEntry);
+	}
+	else
+	{
+		UE_LOG(LogTimeline, Warning, TEXT("FTimeline::AddInterpFloat: FloatCurve is null so it will not be added!"));
+	}
 }
 
 void FTimeline::AddInterpFloat(UCurveFloat* FloatCurve, FOnTimelineFloatStatic InterpFunc)
 {
-	FTimelineFloatTrack NewEntry;
-	NewEntry.FloatCurve = FloatCurve;
-	NewEntry.InterpFuncStatic = InterpFunc;
+	if (FloatCurve)
+	{
+		FTimelineFloatTrack NewEntry;
+		NewEntry.FloatCurve = FloatCurve;
+		NewEntry.InterpFuncStatic = InterpFunc;
 
-	InterpFloats.Add(NewEntry);
+		InterpFloats.Add(NewEntry);
+	}
+	else
+	{
+		UE_LOG(LogTimeline, Warning, TEXT("FTimeline::AddInterpFloat: FloatCurve is null so it will not be added!"));
+	}
 }
 
 void FTimeline::AddInterpLinearColor(UCurveLinearColor* LinearColorCurve, FOnTimelineLinearColor InterpFunc, FName PropertyName, FName TrackName)
 {
-	FTimelineLinearColorTrack NewEntry;
-	NewEntry.LinearColorCurve = LinearColorCurve;
-	NewEntry.InterpFunc = InterpFunc;
-	NewEntry.TrackName = TrackName;
-	NewEntry.LinearColorPropertyName = PropertyName;
+	if (LinearColorCurve)
+	{
+		FTimelineLinearColorTrack NewEntry;
+		NewEntry.LinearColorCurve = LinearColorCurve;
+		NewEntry.InterpFunc = InterpFunc;
+		NewEntry.TrackName = TrackName;
+		NewEntry.LinearColorPropertyName = PropertyName;
 
-	InterpLinearColors.Add(NewEntry);
+		InterpLinearColors.Add(NewEntry);
+	}
+	else
+	{
+		UE_LOG(LogTimeline, Warning, TEXT("FTimeline::AddInterpLinearColor: LinearColorCurve is null so it will not be added!"));
+	}
 }
 
 void FTimeline::AddInterpLinearColor(UCurveLinearColor* LinearColorCurve, FOnTimelineLinearColorStatic InterpFunc)
 {
-	FTimelineLinearColorTrack NewEntry;
-	NewEntry.LinearColorCurve = LinearColorCurve;
-	NewEntry.InterpFuncStatic = InterpFunc;
+	if (LinearColorCurve)
+	{
+		FTimelineLinearColorTrack NewEntry;
+		NewEntry.LinearColorCurve = LinearColorCurve;
+		NewEntry.InterpFuncStatic = InterpFunc;
 
-	InterpLinearColors.Add(NewEntry);
+		InterpLinearColors.Add(NewEntry);
+	}
+	else
+	{
+		UE_LOG(LogTimeline, Warning, TEXT("FTimeline::AddInterpLinearColor: LinearColorCurve is null so it will not be added!"));
+	}
 }
 
 void FTimeline::SetFloatCurve(UCurveFloat* NewFloatCurve, FName FloatTrackName)
 {
-	bool bFoundTrack = false;
-	if (FloatTrackName != NAME_None)
+	if (NewFloatCurve)
 	{
-		for (FTimelineFloatTrack& FloatTrack : InterpFloats)
+		bool bFoundTrack = false;
+		if (FloatTrackName != NAME_None)
 		{
-			if (FloatTrack.TrackName == FloatTrackName)
+			for (FTimelineFloatTrack& FloatTrack : InterpFloats)
 			{
-				FloatTrack.FloatCurve = NewFloatCurve;
-				bFoundTrack = true;
-				break;
+				if (FloatTrack.TrackName == FloatTrackName)
+				{
+					FloatTrack.FloatCurve = NewFloatCurve;
+					bFoundTrack = true;
+					break;
+				}
 			}
 		}
-	}
 
-	if(!bFoundTrack)
+		if(!bFoundTrack)
+		{
+			UE_LOG(LogTimeline, Log, TEXT("SetFloatCurve: No float track with name %s!"), *FloatTrackName.ToString());
+		}
+	}
+	else
 	{
-		UE_LOG(LogTimeline, Log, TEXT("SetFloatCurve: No float track with name %s!"), *FloatTrackName.ToString());
+		UE_LOG(LogTimeline, Warning, TEXT("FTimeline::SetFloatCurve: NewFloatCurve is null so it will not be set!"));
 	}
 }
 
 void FTimeline::SetVectorCurve(UCurveVector* NewVectorCurve, FName VectorTrackName)
 {
-	bool bFoundTrack = false;
-	if (VectorTrackName != NAME_None)
+	if (NewVectorCurve)
 	{
-		for (FTimelineVectorTrack& VectorTrack : InterpVectors)
+		bool bFoundTrack = false;
+		if (VectorTrackName != NAME_None)
 		{
-			if (VectorTrack.TrackName == VectorTrackName)
+			for (FTimelineVectorTrack& VectorTrack : InterpVectors)
 			{
-				VectorTrack.VectorCurve = NewVectorCurve;
-				bFoundTrack = true;
-				break;
+				if (VectorTrack.TrackName == VectorTrackName)
+				{
+					VectorTrack.VectorCurve = NewVectorCurve;
+					bFoundTrack = true;
+					break;
+				}
 			}
 		}
-	}
 
-	if (!bFoundTrack)
+		if (!bFoundTrack)
+		{
+			UE_LOG(LogTimeline, Log, TEXT("SetVectorCurve: No vector track with name %s!"), *VectorTrackName.ToString());
+		}
+	}
+	else
 	{
-		UE_LOG(LogTimeline, Log, TEXT("SetVectorCurve: No vector track with name %s!"), *VectorTrackName.ToString());
+		UE_LOG(LogTimeline, Warning, TEXT("FTimeline::SetVectorCurve: NewVectorCurve is null so it will not be set!"));
 	}
 }
 
 void FTimeline::SetLinearColorCurve(UCurveLinearColor* NewLinearColorCurve, FName LinearColorTrackName)
 {
-	bool bFoundTrack = false;
-	if (LinearColorTrackName != NAME_None)
+	if (NewLinearColorCurve)
 	{
-		for (FTimelineLinearColorTrack& ColorTrack : InterpLinearColors)
+		bool bFoundTrack = false;
+		if (LinearColorTrackName != NAME_None)
 		{
-			if (ColorTrack.TrackName == LinearColorTrackName)
+			for (FTimelineLinearColorTrack& ColorTrack : InterpLinearColors)
 			{
-				ColorTrack.LinearColorCurve = NewLinearColorCurve;
-				bFoundTrack = true;
-				break;
+				if (ColorTrack.TrackName == LinearColorTrackName)
+				{
+					ColorTrack.LinearColorCurve = NewLinearColorCurve;
+					bFoundTrack = true;
+					break;
+				}
+
 			}
+		}
 
+		if (!bFoundTrack)
+		{
+			UE_LOG(LogTimeline, Log, TEXT("SetLinearColorCurve: No color track with name %s!"), *LinearColorTrackName.ToString());
 		}
 	}
-
-	if (!bFoundTrack)
+	else
 	{
-		UE_LOG(LogTimeline, Log, TEXT("SetLinearColorCurve: No color track with name %s!"), *LinearColorTrackName.ToString());
+		UE_LOG(LogTimeline, Warning, TEXT("FTimeline::SetLinearColorCurve: NewLinearColorCurve is null so it will not be set!"));
 	}
 }
 
