LightWeightInstanceManager PreSpawnInitalization, OnSpawnedActorDestroyed virtual functions to allow subclasses to override them

diff --git a/Engine/Source/Runtime/Engine/Classes/GameFramework/LightWeightInstanceManager.h b/Engine/Source/Runtime/Engine/Classes/GameFramework/LightWeightInstanceManager.h
--- a/Engine/Source/Runtime/Engine/Classes/GameFramework/LightWeightInstanceManager.h
+++ b/Engine/Source/Runtime/Engine/Classes/GameFramework/LightWeightInstanceManager.h
@@ -169,9 +169,15 @@ protected:
 	// Returns the class to use when spawning a new actor
 	virtual UClass* GetActorClassToSpawn(const FActorInstanceHandle& Handle) const;
 
+	// Called after actor construction but before other systems see this actor spawn
+	virtual void PreSpawnInitalization(const FActorInstanceHandle& Handle, AActor* SpawnedActor);
+
 	// Called after spawning a new actor from a light weight instance
 	virtual void PostActorSpawn(const FActorInstanceHandle& Handle);
 
+	// Called after a spawned actor is destroyed
+	virtual void OnSpawnedActorDestroyed(AActor* DestroyedActor, const int32 DestroyedActorInstanceIndex);
+
 	// Create an object that implements interfaces for the light weight instance specified by Handle
 	UObject* CreateInterfaceObject(const FActorInstanceHandle& Handle)
 	{
diff --git a/Engine/Source/Runtime/Engine/Private/LightWeightInstanceManager.cpp b/Engine/Source/Runtime/Engine/Private/LightWeightInstanceManager.cpp
--- a/Engine/Source/Runtime/Engine/Private/LightWeightInstanceManager.cpp
+++ b/Engine/Source/Runtime/Engine/Private/LightWeightInstanceManager.cpp
@@ -122,8 +122,7 @@ AActor* ALightWeightInstanceManager::ConvertInstanceToActor(const FActorInstance
 		SetSpawnParameters(SpawnParams);
 		SpawnParams.CustomPreSpawnInitalization = [this, &Handle](AActor* SpawnedActor)
 		{
-			Handle.Actor = SpawnedActor;
-			Actors.Add(Handle.GetInstanceIndex(), SpawnedActor);
+			PreSpawnInitalization(Handle, SpawnedActor);
 		};
 
 		NewActor = GetLevel()->GetWorld()->SpawnActor<AActor>(GetActorClassToSpawn(Handle), GetTransform(Handle), SpawnParams);
@@ -143,9 +142,12 @@ AActor* ALightWeightInstanceManager::ConvertInstanceToActor(const FActorInstance
 
 void ALightWeightInstanceManager::OnSpawnedActorDestroyed(AActor* DestroyedActor)
 {
-	check(DestroyedActor);
-	const int32 DestroyedActorInstanceIndex = FindIndexForActor(DestroyedActor);
+	OnSpawnedActorDestroyed(DestroyedActor, FindIndexForActor(DestroyedActor));
+}
 
+void ALightWeightInstanceManager::OnSpawnedActorDestroyed(AActor* DestroyedActor, const int32 DestroyedActorInstanceIndex)
+{
+	check(DestroyedActor);
 	if (!ensure(DestroyedActorInstanceIndex != INDEX_NONE))
 	{
 		UE_LOG(LogLightWeightInstance, Error,
@@ -298,6 +300,12 @@ UClass* ALightWeightInstanceManager::GetActorClassToSpawn(const FActorInstanceHa
 	return RepresentedClass;
 }
 
+void ALightWeightInstanceManager::PreSpawnInitalization(const FActorInstanceHandle& Handle, AActor* SpawnedActor)
+{
+	Handle.Actor = SpawnedActor;
+	Actors.Add(Handle.GetInstanceIndex(), SpawnedActor);
+}
+
 void ALightWeightInstanceManager::PostActorSpawn(const FActorInstanceHandle& Handle)
 {
 	// do nothing
