Add function to LWI manager to query whether it is actually managing any instances or has any hydrated actors.

diff --git a/Engine/Source/Runtime/Engine/Classes/GameFramework/LightWeightInstanceManager.h b/Engine/Source/Runtime/Engine/Classes/GameFramework/LightWeightInstanceManager.h
--- a/Engine/Source/Runtime/Engine/Classes/GameFramework/LightWeightInstanceManager.h
+++ b/Engine/Source/Runtime/Engine/Classes/GameFramework/LightWeightInstanceManager.h
@@ -143,6 +143,8 @@ public:
 
 	AActor* FindActorForInstanceIndex(const int32 InstanceIndex);
 
+	bool HasAnyValidInstancesOrManagedActors() const;
+
 protected:
 	// Creates an actor to replace the instance specified by Handle
 	AActor* ConvertInstanceToActor(const FActorInstanceHandle& Handle);
diff --git a/Engine/Source/Runtime/Engine/Private/LightWeightInstanceManager.cpp b/Engine/Source/Runtime/Engine/Private/LightWeightInstanceManager.cpp
--- a/Engine/Source/Runtime/Engine/Private/LightWeightInstanceManager.cpp
+++ b/Engine/Source/Runtime/Engine/Private/LightWeightInstanceManager.cpp
@@ -455,6 +455,20 @@ void ALightWeightInstanceManager::RemoveInstance(const int32 Index)
 	}
 }
 
+bool ALightWeightInstanceManager::HasAnyValidInstancesOrManagedActors() const
+{
+	bool bHasAnyValidIndex = false;
+	for (const bool bIsIndexValid : ValidIndices)
+	{
+		if (bIsIndexValid)
+		{
+			bHasAnyValidIndex = true;
+			break;
+		}
+	}
+	return bHasAnyValidIndex;
+}
+
 void ALightWeightInstanceManager::GetLifetimeReplicatedProps(TArray<FLifetimeProperty>& OutLifetimeProps) const
 {
 	Super::GetLifetimeReplicatedProps(OutLifetimeProps);
