Remove some logging which did not prove to be useful, and readded some others

diff --git a/Engine/Source/Runtime/Engine/Private/Animation/AnimSequence.cpp b/Engine/Source/Runtime/Engine/Private/Animation/AnimSequence.cpp
--- a/Engine/Source/Runtime/Engine/Private/Animation/AnimSequence.cpp
+++ b/Engine/Source/Runtime/Engine/Private/Animation/AnimSequence.cpp
@@ -1158,18 +1158,19 @@ bool UAnimSequence::IsCachedCookedPlatformDataLoaded(const ITargetPlatform* Targ
 		EndCacheDerivedData(KeyHash);
 	}
 
-	auto RecompressForPlatform = [this, TargetPlatform]()
+	auto RecompressForPlatform = [this, TargetPlatform, KeyHash]()
 	{
 		if (IsInGameThread())
 		{
 			// This is required until we can ensure inter-animsequence dependencies are handled correctly (see UE-184634)
 			BeginCacheForCookedPlatformData(TargetPlatform);
+			UE_LOG(LogAnimation, Warning, TEXT("Missing compressed animation data for %s requested platform %s and hash %s"), *GetName(), *TargetPlatform->PlatformName(), *LexToString(KeyHash));
 		}
 	};
 
 	if (KeyHash == DataKeyHash)
 	{
-		if (!CacheTasksByKeyHash.Contains(KeyHash) && CompressedData.IsValid(this, true))
+		if (!CacheTasksByKeyHash.Contains(KeyHash) && CompressedData.IsValid(this))
 		{
 			return true;
 		}
@@ -1182,7 +1183,7 @@ bool UAnimSequence::IsCachedCookedPlatformDataLoaded(const ITargetPlatform* Targ
 	{
 		if(const TUniquePtr<FCompressedAnimSequence>* CompressedDataPtr = DataByPlatformKeyHash.Find(KeyHash))
 		{
-			if (!CacheTasksByKeyHash.Contains(KeyHash) && (*CompressedDataPtr)->IsValid(this, true))
+			if (!CacheTasksByKeyHash.Contains(KeyHash) && (*CompressedDataPtr)->IsValid(this))
 			{
 				return true;
 			}
