[materials] add additional diagnostics to catch potential corrupted shader maps:
1. check that job results being added to a map are for the correct platform
2. when deserialized key from DDC does not match input key, log messages with both keys as well as the cached data size + hash

diff --git a/Engine/Source/Runtime/Engine/Private/Materials/MaterialShader.cpp b/Engine/Source/Runtime/Engine/Private/Materials/MaterialShader.cpp
--- a/Engine/Source/Runtime/Engine/Private/Materials/MaterialShader.cpp
+++ b/Engine/Source/Runtime/Engine/Private/Materials/MaterialShader.cpp
@@ -1391,6 +1391,7 @@ TSharedRef<FMaterialShaderMap::FAsyncLoadContext> FMaterialShaderMap::BeginLoadF
 	{
 		FString	        DataKey;
 		FSharedBuffer   CachedData;
+		FIoHash         CachedDataHash;
 		EShaderPlatform Platform;
 		FString         AssetName;
 		FRequestOwner   RequestOwner{UE::DerivedData::EPriority::Normal};
@@ -1427,8 +1428,13 @@ TSharedRef<FMaterialShaderMap::FAsyncLoadContext> FMaterialShaderMap::BeginLoadF
 				//InOutShaderMap->RegisterSerializedShaders(false);
 
 				const FString InDataKey = GetMaterialShaderMapKeyString(ShaderMap->GetShaderMapId(), Platform);
-				checkf(InDataKey == DataKey, TEXT("Data deserialized from the DDC would need a different DDC key!"));
-				//checkf(InOutShaderMap->GetShaderMapId() == Context.ShaderMapId, TEXT("Shadermap data deserialized from the DDC does not match the ID we used to build the key!"));
+				if (InDataKey != DataKey)
+				{
+					UE_LOG(LogMaterial, Warning, TEXT("Shader map key recomputed from DDC data: %s"), *InDataKey);
+					UE_LOG(LogMaterial, Warning, TEXT("Shader map key from request: %s"), *DataKey);
+					UE_LOG(LogMaterial, Warning, TEXT("Cached data size %llu, hash %s"), CachedData.GetSize(), *LexToString(CachedDataHash))
+					checkf(false, TEXT("DDC key constructed from deserialized shadermap does not match request key!"));
+				}
 
 				// Register in the global map
 				ShaderMap->Register(Platform);
@@ -1538,6 +1544,7 @@ TSharedRef<FMaterialShaderMap::FAsyncLoadContext> FMaterialShaderMap::BeginLoadF
 					if (bCheckCache)
 					{
 						Result->CachedData = Response.Value.GetData().Decompress();
+						Result->CachedDataHash = Response.Value.GetRawHash();
 						// This callback might hold the last reference to Result, which owns RequestOwner, so
 						// we must not cancel in the Owner's destructor; Canceling in a callback will deadlock.
 						Result->RequestOwner.KeepAlive();
@@ -2278,6 +2285,10 @@ FShader* FMaterialShaderMap::ProcessCompilationResultsForSingleJob(FShaderCompil
 	check(SingleJob);
 	const FShaderCompileJob& CurrentJob = *SingleJob;
 	check(CurrentJob.Id == CompilingId);
+	checkf(CurrentJob.Input.Target.Platform == GetShaderPlatform(), TEXT("Job %s platform %s does not match shader map platform %s"), 
+		*CurrentJob.Input.ShaderName,
+		*LexToString(CurrentJob.Input.Target.GetPlatform()),
+		*LexToString(GetShaderPlatform()));
 
 	GetResourceCode()->AddShaderCompilerOutput(CurrentJob.Output, CurrentJob.Key.ToString());
 
