[UE] Support the creation of Immediate Physics Constraints without copying values from Constraint Instances
- Added a new JointHandle constructor and a CreateJoint function to ImmediatePhysics_Chaos that accept a JointSettings object instead of a FConstraintInstance.

diff --git a/Engine/Source/Runtime/Engine/Private/PhysicsEngine/ImmediatePhysics/ImmediatePhysicsChaos/ImmediatePhysicsJointHandle_Chaos.cpp b/Engine/Source/Runtime/Engine/Private/PhysicsEngine/ImmediatePhysics/ImmediatePhysicsChaos/ImmediatePhysicsJointHandle_Chaos.cpp
--- a/Engine/Source/Runtime/Engine/Private/PhysicsEngine/ImmediatePhysics/ImmediatePhysicsChaos/ImmediatePhysicsJointHandle_Chaos.cpp
+++ b/Engine/Source/Runtime/Engine/Private/PhysicsEngine/ImmediatePhysics/ImmediatePhysicsChaos/ImmediatePhysicsJointHandle_Chaos.cpp
@@ -193,14 +193,16 @@ namespace ImmediatePhysics_Chaos
 
 		ConstraintHandle = Constraints->AddConstraint({ Actor1->ParticleHandle, Actor2->ParticleHandle }, ConstraintSettings);
 
-		if (Actor1->ParticleHandle != nullptr)
-		{
-			FGenericParticleHandle(Actor1->ParticleHandle)->SetInertiaConditioningDirty();
-		}
-		if (Actor2->ParticleHandle != nullptr)
-		{
-			FGenericParticleHandle(Actor2->ParticleHandle)->SetInertiaConditioningDirty();
-		}
+		SetActorInertiaConditioningDirty();
+	}
+
+	FJointHandle::FJointHandle(FChaosConstraintContainer* InConstraints, const FPBDJointSettings& ConstraintSettings, FActorHandle* const Actor1, FActorHandle* const Actor2)
+		: ActorHandles({ Actor1, Actor2 })
+		, Constraints(InConstraints)
+	{
+		ConstraintHandle = Constraints->AddConstraint({ Actor1->ParticleHandle, Actor2->ParticleHandle }, ConstraintSettings);
+
+		SetActorInertiaConditioningDirty();
 	}
 
 	FJointHandle::~FJointHandle()
@@ -238,5 +240,20 @@ namespace ImmediatePhysics_Chaos
 		JointSettings.SoftLinearDamping = bLinearSoft ? LinearDamping : 0.0f;
 		ConstraintHandle->SetSettings(JointSettings);
 	}
+
+	void FJointHandle::SetActorInertiaConditioningDirty()
+	{
+		using namespace Chaos;
+
+		if (ActorHandles[0]->ParticleHandle != nullptr)
+		{
+			FGenericParticleHandle(ActorHandles[0]->ParticleHandle)->SetInertiaConditioningDirty();
+		}
+
+		if (ActorHandles[1]->ParticleHandle != nullptr)
+		{
+			FGenericParticleHandle(ActorHandles[1]->ParticleHandle)->SetInertiaConditioningDirty();
+		}
+	}
 }
 
diff --git a/Engine/Source/Runtime/Engine/Private/PhysicsEngine/ImmediatePhysics/ImmediatePhysicsChaos/ImmediatePhysicsSimulation_Chaos.cpp b/Engine/Source/Runtime/Engine/Private/PhysicsEngine/ImmediatePhysics/ImmediatePhysicsChaos/ImmediatePhysicsSimulation_Chaos.cpp
--- a/Engine/Source/Runtime/Engine/Private/PhysicsEngine/ImmediatePhysics/ImmediatePhysicsChaos/ImmediatePhysicsSimulation_Chaos.cpp
+++ b/Engine/Source/Runtime/Engine/Private/PhysicsEngine/ImmediatePhysics/ImmediatePhysicsChaos/ImmediatePhysicsSimulation_Chaos.cpp
@@ -508,6 +508,13 @@ namespace ImmediatePhysics_Chaos
 		return JointHandle;
 	}
 
+	FJointHandle* FSimulation::CreateJoint(const Chaos::FPBDJointSettings& ConstraintSettings, FActorHandle* const Body1, FActorHandle* const Body2)
+	{
+		FJointHandle* JointHandle = new FJointHandle(&Implementation->Joints, ConstraintSettings, Body1, Body2);
+		Implementation->JointHandles.Add(JointHandle);
+		return JointHandle;
+	}
+
 	void FSimulation::DestroyJoint(FJointHandle* JointHandle)
 	{
 		// @todo(ccaulfield): FJointHandle could remember its index to optimize this
diff --git a/Engine/Source/Runtime/Engine/Public/Physics/ImmediatePhysics/ImmediatePhysicsChaos/ImmediatePhysicsJointHandle_Chaos.h b/Engine/Source/Runtime/Engine/Public/Physics/ImmediatePhysics/ImmediatePhysicsChaos/ImmediatePhysicsJointHandle_Chaos.h
--- a/Engine/Source/Runtime/Engine/Public/Physics/ImmediatePhysics/ImmediatePhysicsChaos/ImmediatePhysicsJointHandle_Chaos.h
+++ b/Engine/Source/Runtime/Engine/Public/Physics/ImmediatePhysics/ImmediatePhysicsChaos/ImmediatePhysicsJointHandle_Chaos.h
@@ -9,6 +9,11 @@
 
 #include "Engine/EngineTypes.h"
 
+namespace Chaos
+{
+	class FPBDJointSettings;
+}
+
 namespace ImmediatePhysics_Chaos
 {
 	/** handle associated with a physics joint. This is the proper way to read/write to the physics simulation */
@@ -17,8 +22,10 @@ namespace ImmediatePhysics_Chaos
 	public:
 		using FChaosConstraintContainer = Chaos::FPBDJointConstraints;
 		using FChaosConstraintHandle = typename Chaos::FPBDJointConstraintHandle;
+		using FPBDJointSettings = Chaos::FPBDJointSettings;
 
 		FJointHandle(FChaosConstraintContainer* InConstraints, FConstraintInstance* ConstraintInstance, FActorHandle* InActor1, FActorHandle* InActor2);
+		FJointHandle(FChaosConstraintContainer* InConstraints, const FPBDJointSettings& ConstraintSettings, FActorHandle* InActor1, FActorHandle* InActor2);
 		~FJointHandle();
 
 		FChaosConstraintHandle* GetConstraint();
@@ -30,6 +37,8 @@ namespace ImmediatePhysics_Chaos
 		void SetSoftLinearSettings(bool bLinearSoft, FReal LinearStiffness, FReal LinearDamping);
 
 	private:
+		void SetActorInertiaConditioningDirty();
+		
 		Chaos::TVec2<FActorHandle*> ActorHandles;
 		FChaosConstraintContainer* Constraints;
 		FChaosConstraintHandle* ConstraintHandle;
diff --git a/Engine/Source/Runtime/Engine/Public/Physics/ImmediatePhysics/ImmediatePhysicsChaos/ImmediatePhysicsSimulation_Chaos.h b/Engine/Source/Runtime/Engine/Public/Physics/ImmediatePhysics/ImmediatePhysicsChaos/ImmediatePhysicsSimulation_Chaos.h
--- a/Engine/Source/Runtime/Engine/Public/Physics/ImmediatePhysics/ImmediatePhysicsChaos/ImmediatePhysicsSimulation_Chaos.h
+++ b/Engine/Source/Runtime/Engine/Public/Physics/ImmediatePhysics/ImmediatePhysicsChaos/ImmediatePhysicsSimulation_Chaos.h
@@ -8,6 +8,11 @@
 #include "Engine/EngineTypes.h"
 #include "Templates/UniquePtr.h"
 
+namespace Chaos
+{
+	class FPBDJointSettings;
+}
+
 namespace ImmediatePhysics_Chaos
 {
 	/** Owns all the data associated with the simulation. Can be considered a single scene or world */
@@ -38,6 +43,10 @@ namespace ImmediatePhysics_Chaos
 
 		/** Create a physical joint and add it to the simulation */
 		FJointHandle* CreateJoint(FConstraintInstance* ConstraintInstance, FActorHandle* Body1, FActorHandle* Body2);
+
+		/** Create a physical joint and add it to the simulation */
+		FJointHandle* CreateJoint(const Chaos::FPBDJointSettings& ConstraintSettings, FActorHandle* const Body1, FActorHandle* const Body2);
+
 		void DestroyJoint(FJointHandle* JointHandle);
 
 		/** Sets the number of active bodies. This number is reset any time a new simulated body is created */
