Clean up scene proxy creation to unify the logic for creating Nanite scene proxies. Fixes issues where editor fallback options wouldn't work on ISM/HISM.

#rnx

diff --git a/Engine/Source/Runtime/Engine/Classes/Components/HierarchicalInstancedStaticMeshComponent.h b/Engine/Source/Runtime/Engine/Classes/Components/HierarchicalInstancedStaticMeshComponent.h
--- a/Engine/Source/Runtime/Engine/Classes/Components/HierarchicalInstancedStaticMeshComponent.h
+++ b/Engine/Source/Runtime/Engine/Classes/Components/HierarchicalInstancedStaticMeshComponent.h
@@ -313,6 +313,8 @@ protected:
 
 	virtual void OnPostLoadPerInstanceData() override;
 
+	virtual FPrimitiveSceneProxy* CreateStaticMeshSceneProxy(Nanite::FMaterialAudit& NaniteMaterials, bool bCreateNanite) override;
+
 	virtual FVector GetTranslatedInstanceSpaceOrigin() const override { return TranslatedInstanceSpaceOrigin; }
 
 	virtual void GetNavigationPerInstanceTransforms(const FBox& AreaBox, TArray<FTransform>& InstanceData) const override;
diff --git a/Engine/Source/Runtime/Engine/Classes/Components/InstancedStaticMeshComponent.h b/Engine/Source/Runtime/Engine/Classes/Components/InstancedStaticMeshComponent.h
--- a/Engine/Source/Runtime/Engine/Classes/Components/InstancedStaticMeshComponent.h
+++ b/Engine/Source/Runtime/Engine/Classes/Components/InstancedStaticMeshComponent.h
@@ -377,11 +377,15 @@ public:
 protected:
 	virtual bool ComponentOverlapComponentImpl(class UPrimitiveComponent* PrimComp, const FVector Pos, const FQuat& Quat, const FCollisionQueryParams& Params) override;
 	virtual bool ComponentOverlapMultiImpl(TArray<struct FOverlapResult>& OutOverlaps, const class UWorld* InWorld, const FVector& Pos, const FQuat& Rot, ECollisionChannel TestChannel, const struct FComponentQueryParams& Params, const struct FCollisionObjectQueryParams& ObjectQueryParams = FCollisionObjectQueryParams::DefaultObjectQueryParam) const override;
-public:
-
 	//~ End UPrimitiveComponent Interface
 
+	//~ Begin UStaticMeshComponentInterface
+protected:
+	virtual FPrimitiveSceneProxy* CreateStaticMeshSceneProxy(Nanite::FMaterialAudit& NaniteMaterials, bool bCreateNanite) override;
+	//~ End UStaticMeshComponentInterface
+
 	//~ Begin UNavRelevantInterface Interface
+public:
 	virtual void GetNavigationData(FNavigationRelevantData& Data) const override;
 	virtual FBox GetNavigationBounds() const override;
 	virtual bool IsNavigationRelevant() const override;
diff --git a/Engine/Source/Runtime/Engine/Classes/Components/StaticMeshComponent.h b/Engine/Source/Runtime/Engine/Classes/Components/StaticMeshComponent.h
--- a/Engine/Source/Runtime/Engine/Classes/Components/StaticMeshComponent.h
+++ b/Engine/Source/Runtime/Engine/Classes/Components/StaticMeshComponent.h
@@ -705,6 +705,9 @@ protected:
 
 	bool ShouldCreateNaniteProxy(Nanite::FMaterialAudit* OutNaniteMaterials = nullptr) const;
 
+	// Overload this in child implementations that wish to extend Static Mesh or Nanite scene proxy implementations
+	virtual FPrimitiveSceneProxy* CreateStaticMeshSceneProxy(Nanite::FMaterialAudit& NaniteMaterials, bool bCreateNanite);
+
 public:
 
 	void ReleaseResources();
diff --git a/Engine/Source/Runtime/Engine/Private/HierarchicalInstancedStaticMesh.cpp b/Engine/Source/Runtime/Engine/Private/HierarchicalInstancedStaticMesh.cpp
--- a/Engine/Source/Runtime/Engine/Private/HierarchicalInstancedStaticMesh.cpp
+++ b/Engine/Source/Runtime/Engine/Private/HierarchicalInstancedStaticMesh.cpp
@@ -3243,12 +3243,23 @@ void UHierarchicalInstancedStaticMeshComponent::SetPerInstanceLightMapAndEditorD
 	}
 }
 
+FPrimitiveSceneProxy* UHierarchicalInstancedStaticMeshComponent::CreateStaticMeshSceneProxy(Nanite::FMaterialAudit& NaniteMaterials, bool bCreateNanite)
+{
+	LLM_SCOPE(ELLMTag::InstancedMesh);
+
+	if (bCreateNanite)
+	{
+		return ::new Nanite::FSceneProxy(NaniteMaterials, this);
+	}
+	
+	return ::new FHierarchicalStaticMeshSceneProxy(this, GetWorld()->GetFeatureLevel());
+}
+
 FPrimitiveSceneProxy* UHierarchicalInstancedStaticMeshComponent::CreateSceneProxy()
 {
 	static const auto NaniteProxyRenderModeVar = IConsoleManager::Get().FindConsoleVariable(TEXT("r.Nanite.ProxyRenderMode"));
 	const int32 NaniteProxyRenderMode = (NaniteProxyRenderModeVar != nullptr) ? (NaniteProxyRenderModeVar->GetInt() != 0) : 0;
 
-	LLM_SCOPE(ELLMTag::InstancedMesh);
 	QUICK_SCOPE_CYCLE_COUNTER(STAT_HierarchicalInstancedStaticMeshComponent_CreateSceneProxy);
 	SCOPE_CYCLE_COUNTER(STAT_FoliageCreateProxy);
 
@@ -3281,28 +3292,10 @@ FPrimitiveSceneProxy* UHierarchicalInstancedStaticMeshComponent::CreateSceneProx
 
 		ProxySize = PerInstanceRenderData->ResourceSize;
 		INC_DWORD_STAT_BY(STAT_FoliageInstanceBuffers, ProxySize);
-
-		Nanite::FMaterialAudit NaniteMaterials{};
-
-		bool bUseNanite = ShouldCreateNaniteProxy(&NaniteMaterials);
-
-		if (bUseNanite)
-		{
-			return ::new Nanite::FSceneProxy(NaniteMaterials, this);
-		}
-		// If we didn't get a proxy, but Nanite was enabled on the asset when it was built, evaluate proxy creation
-		else if (GetStaticMesh()->HasValidNaniteData() && NaniteProxyRenderMode != 0)
-		{
-			// Do not render Nanite proxy
-			return nullptr;
-		}
-		else
-		{
-			return ::new FHierarchicalStaticMeshSceneProxy(this, GetWorld()->GetFeatureLevel());
-		}
 	}
 
-	return nullptr;
+	// NOTE: Purposefully skipping UInstancedStaticMeshComponent implementation
+	return UStaticMeshComponent::CreateSceneProxy();
 }
 
 void UHierarchicalInstancedStaticMeshComponent::UpdateDensityScaling()
diff --git a/Engine/Source/Runtime/Engine/Private/InstancedStaticMesh.cpp b/Engine/Source/Runtime/Engine/Private/InstancedStaticMesh.cpp
--- a/Engine/Source/Runtime/Engine/Private/InstancedStaticMesh.cpp
+++ b/Engine/Source/Runtime/Engine/Private/InstancedStaticMesh.cpp
@@ -2693,13 +2693,22 @@ FBodyInstance* UInstancedStaticMeshComponent::GetBodyInstance(FName BoneName, bo
 	return  const_cast<FBodyInstance*>(&BodyInstance); // If no index is specified we return the primitive component body instance instead
 }
 
-FPrimitiveSceneProxy* UInstancedStaticMeshComponent::CreateSceneProxy()
+FPrimitiveSceneProxy* UInstancedStaticMeshComponent::CreateStaticMeshSceneProxy(Nanite::FMaterialAudit& NaniteMaterials, bool bCreateNanite)
 {
-	static const auto NaniteProxyRenderModeVar = IConsoleManager::Get().FindConsoleVariable(TEXT("r.Nanite.ProxyRenderMode"));
-	const int32 NaniteProxyRenderMode = (NaniteProxyRenderModeVar != nullptr) ? (NaniteProxyRenderModeVar->GetInt() != 0) : 0;
+	// Override which constructor is used for the Nanite scene proxy or create FInstancedStaticMeshSceneProxy
 
 	LLM_SCOPE(ELLMTag::InstancedMesh);
 
+	if (bCreateNanite)
+	{
+		return ::new Nanite::FSceneProxy(NaniteMaterials, this);
+	}
+
+	return ::new FInstancedStaticMeshSceneProxy(this, GetWorld()->GetFeatureLevel());
+}
+
+FPrimitiveSceneProxy* UInstancedStaticMeshComponent::CreateSceneProxy()
+{
 	ProxySize = 0;
 
 	// Verify that the mesh is valid before using it.
@@ -2724,31 +2733,9 @@ FPrimitiveSceneProxy* UInstancedStaticMeshComponent::CreateSceneProxy()
 		}
 		
 		ProxySize = PerInstanceRenderData->ResourceSize;
-
-		Nanite::FMaterialAudit NaniteMaterials{};
-
-		// Is Nanite supported, and is there built Nanite data for this static mesh?
-		bool bUseNanite = ShouldCreateNaniteProxy(&NaniteMaterials);
-
-		if (bUseNanite)
-		{
-			return ::new Nanite::FSceneProxy(NaniteMaterials, this);
-		}
-		// If we didn't get a proxy, but Nanite was enabled on the asset when it was built, evaluate proxy creation
-		else if (GetStaticMesh()->HasValidNaniteData() && NaniteProxyRenderMode != 0)
-		{
-			// Do not render Nanite proxy
-			return nullptr;
-		}
-		else
-		{
-			return ::new FInstancedStaticMeshSceneProxy(this, GetWorld()->GetFeatureLevel());
-		}
-	}
-	else
-	{
-		return nullptr;
 	}
+
+	return Super::CreateSceneProxy();
 }
 
 FMatrix UInstancedStaticMeshComponent::GetRenderMatrix() const
diff --git a/Engine/Source/Runtime/Engine/Private/StaticMeshRender.cpp b/Engine/Source/Runtime/Engine/Private/StaticMeshRender.cpp
--- a/Engine/Source/Runtime/Engine/Private/StaticMeshRender.cpp
+++ b/Engine/Source/Runtime/Engine/Private/StaticMeshRender.cpp
@@ -2555,6 +2555,35 @@ bool UStaticMeshComponent::ShouldCreateNaniteProxy(Nanite::FMaterialAudit* OutNa
 	return true;
 }
 
+FPrimitiveSceneProxy* UStaticMeshComponent::CreateStaticMeshSceneProxy(Nanite::FMaterialAudit& NaniteMaterials, bool bCreateNanite)
+{
+	// Default implementation: Nanite::FSceneProxy or FStaticMeshSceneProxy
+
+	LLM_SCOPE(ELLMTag::StaticMesh);
+
+	if (bCreateNanite)
+	{
+		return ::new Nanite::FSceneProxy(NaniteMaterials, this);
+	}
+
+	// Validate the LOD resources here
+	const FStaticMeshLODResourcesArray& LODResources = GetStaticMesh()->GetRenderData()->LODResources;
+	const int32 SMCurrentMinLOD = GetStaticMesh()->GetMinLODIdx();
+	const int32 EffectiveMinLOD = bOverrideMinLOD ? MinLOD : SMCurrentMinLOD;
+	if (LODResources.Num() == 0	|| LODResources[FMath::Clamp<int32>(EffectiveMinLOD, 0, LODResources.Num()-1)].VertexBuffers.StaticMeshVertexBuffer.GetNumVertices() == 0)
+	{
+		UE_LOG(LogStaticMesh, Verbose, TEXT("Skipping CreateSceneProxy for StaticMeshComponent %s (LOD problems)"), *GetFullName());
+		return nullptr;
+	}
+	
+	auto* Proxy = ::new FStaticMeshSceneProxy(this, false);
+#if STATICMESH_ENABLE_DEBUG_RENDERING
+	SendRenderDebugPhysics(Proxy);
+#endif
+
+	return Proxy;
+}
+
 FPrimitiveSceneProxy* UStaticMeshComponent::CreateSceneProxy()
 {
 	if (GetStaticMesh() == nullptr)
@@ -2597,10 +2626,8 @@ FPrimitiveSceneProxy* UStaticMeshComponent::CreateSceneProxy()
 
 	if (bUseNanite)
 	{
-		LLM_SCOPE(ELLMTag::StaticMesh);
-
 		// Nanite is fully supported
-		return ::new Nanite::FSceneProxy(NaniteMaterials, this);
+		return CreateStaticMeshSceneProxy(NaniteMaterials, true);
 	}
 	// If we didn't get a proxy, but Nanite was enabled on the asset when it was built, evaluate proxy creation
 	else if (HasValidNaniteData() && NaniteMaterials.IsValid(bIsMaskingAllowed))
@@ -2621,23 +2648,7 @@ FPrimitiveSceneProxy* UStaticMeshComponent::CreateSceneProxy()
 		// Fall back to rendering Nanite proxy meshes with traditional static mesh scene proxies
 	}
 
-	const FStaticMeshLODResourcesArray& LODResources = GetStaticMesh()->GetRenderData()->LODResources;
-	const int32 SMCurrentMinLOD = GetStaticMesh()->GetMinLODIdx();
-	const int32 EffectiveMinLOD = bOverrideMinLOD ? MinLOD : SMCurrentMinLOD;
-	if (LODResources.Num() == 0	|| LODResources[FMath::Clamp<int32>(EffectiveMinLOD, 0, LODResources.Num()-1)].VertexBuffers.StaticMeshVertexBuffer.GetNumVertices() == 0)
-	{
-		UE_LOG(LogStaticMesh, Verbose, TEXT("Skipping CreateSceneProxy for StaticMeshComponent %s (LOD problems)"), *GetFullName());
-		return nullptr;
-	}
-
-	LLM_SCOPE(ELLMTag::StaticMesh);
-
-	FPrimitiveSceneProxy* Proxy = ::new FStaticMeshSceneProxy(this, false);
-#if STATICMESH_ENABLE_DEBUG_RENDERING
-	SendRenderDebugPhysics(Proxy);
-#endif
-
-	return Proxy;
+	return CreateStaticMeshSceneProxy(NaniteMaterials, false);
 }
 
 bool UStaticMeshComponent::ShouldRecreateProxyOnUpdateTransform() const
