Skeletal Mesh - Fixed crashes in cloth simulation when setting a new skeletal mesh asset by stopping any ongoing simulations prior to replacing the asset.

diff --git a/Engine/Source/Runtime/Engine/Private/Components/SkeletalMeshComponent.cpp b/Engine/Source/Runtime/Engine/Private/Components/SkeletalMeshComponent.cpp
--- a/Engine/Source/Runtime/Engine/Private/Components/SkeletalMeshComponent.cpp
+++ b/Engine/Source/Runtime/Engine/Private/Components/SkeletalMeshComponent.cpp
@@ -3028,6 +3028,9 @@ void USkeletalMeshComponent::SetSkeletalMesh(USkeletalMesh* InSkelMesh, bool bRe
 		return;
 	}
 
+	// Stop any existing cloth simulation pior to replacing the asset
+	RemoveAllClothingActors();
+
 	// Update property alias
 //#if WITH_EDITORONLY_DATA  // TODO: Re-add these guards once the MovieScene getters/setters are working, so that we can get rid of this redundant pointer in all cooked builds
 	PRAGMA_DISABLE_DEPRECATION_WARNINGS
