Anim Modifiers - Can't add multiple modifiers to multiple animation sequences without the engine crashing
#fix prevent deadlocking when trying to cancel async animation compression task while holding onto evaluation/modification lock on current thread
#misc ensure that apply an animation modifier holds aforementioned lock
#misc add bracket to MotionExtractorModifier

diff --git a/Engine/Plugins/Animation/AnimationData/Source/AnimationData/Private/AnimSequencerDataModel.h b/Engine/Plugins/Animation/AnimationData/Source/AnimationData/Private/AnimSequencerDataModel.h
--- a/Engine/Plugins/Animation/AnimationData/Source/AnimationData/Private/AnimSequencerDataModel.h
+++ b/Engine/Plugins/Animation/AnimationData/Source/AnimationData/Private/AnimSequencerDataModel.h
@@ -119,6 +119,12 @@ protected:
 	{
 		EvaluationLock.Lock();
 	}
+
+	virtual bool TryLockEvaluationAndModification() const override final
+	{
+		return EvaluationLock.TryLock();
+	}
+	
 	virtual void UnlockEvaluationAndModification() const override final
 	{
 		EvaluationLock.Unlock();
diff --git a/Engine/Plugins/Animation/AnimationModifierLibrary/Source/AnimationModifierLibrary/Private/MotionExtractorModifier.cpp b/Engine/Plugins/Animation/AnimationModifierLibrary/Source/AnimationModifierLibrary/Private/MotionExtractorModifier.cpp
--- a/Engine/Plugins/Animation/AnimationModifierLibrary/Source/AnimationModifierLibrary/Private/MotionExtractorModifier.cpp
+++ b/Engine/Plugins/Animation/AnimationModifierLibrary/Source/AnimationModifierLibrary/Private/MotionExtractorModifier.cpp
@@ -10,6 +10,8 @@
 
 #include UE_INLINE_GENERATED_CPP_BY_NAME(MotionExtractorModifier)
 
+#define LOCTEXT_NAMESPACE "UMotionExtractorModifier"
+
 UMotionExtractorModifier::UMotionExtractorModifier()
 	:Super()
 {
@@ -136,11 +138,13 @@ void UMotionExtractorModifier::OnApply_Implementation(UAnimSequence* Animation)
 	}
 
 	IAnimationDataController& Controller = Animation->GetController();
+	Controller.OpenBracket(LOCTEXT("SetMotionCurveBracket", "Setting Motion Curve"));
 	const FAnimationCurveIdentifier CurveId = UAnimationCurveIdentifierExtensions::GetCurveIdentifier(Animation->GetSkeleton(), GetCurveName(), ERawCurveTrackTypes::RCT_Float);
 	if(CurveKeys.Num() && Controller.AddCurve(CurveId))
 	{
 		Controller.SetCurveKeys(CurveId, CurveKeys);
 	}
+	Controller.CloseBracket();
 }
 
 void UMotionExtractorModifier::OnRevert_Implementation(UAnimSequence* Animation)
@@ -186,3 +190,5 @@ float UMotionExtractorModifier::GetDesiredValue(const FTransform& BoneTransform,
 
 	return Value;
 }
+
+#undef LOCTEXT_NAMESPACE
\ No newline at end of file
diff --git a/Engine/Source/Editor/AnimationModifiers/Private/AnimationModifier.cpp b/Engine/Source/Editor/AnimationModifiers/Private/AnimationModifier.cpp
--- a/Engine/Source/Editor/AnimationModifiers/Private/AnimationModifier.cpp
+++ b/Engine/Source/Editor/AnimationModifiers/Private/AnimationModifier.cpp
@@ -75,6 +75,7 @@ void UAnimationModifier::ApplyToAnimationSequence(class UAnimSequence* AnimSeque
 
 	{
 		// Group the OnRevert & OnApply operation into one Bracket to reduce compression request
+		IAnimationDataModel::FEvaluationAndModificationLock EvaluationAndModificationLock(*AnimSequence->GetDataModelInterface());
 		IAnimationDataController::FScopedBracket ScopedBracket(AnimSequence->GetController(), LOCTEXT("ApplyModifierBracket", "Applying Animation Modifier"));
 
 		/** In case this modifier has been previously applied, revert it using the serialised out version at the time */
diff --git a/Engine/Source/Runtime/Engine/Classes/Animation/AnimData/AnimDataModel.h b/Engine/Source/Runtime/Engine/Classes/Animation/AnimData/AnimDataModel.h
--- a/Engine/Source/Runtime/Engine/Classes/Animation/AnimData/AnimDataModel.h
+++ b/Engine/Source/Runtime/Engine/Classes/Animation/AnimData/AnimDataModel.h
@@ -82,6 +82,7 @@ protected:
 	virtual FAnimDataModelModifiedDynamicEvent& GetModifiedDynamicEvent() override { return ModifiedEventDynamic; }
 	virtual void OnNotify(const EAnimDataModelNotifyType& NotifyType, const FAnimDataModelNotifPayload& Payload) override {}
 	virtual void LockEvaluationAndModification() const override final {}
+	virtual bool TryLockEvaluationAndModification() const override final { return true; }
 	virtual void UnlockEvaluationAndModification() const override final {}
 	/** End IAnimationDataModel overrides */
 
diff --git a/Engine/Source/Runtime/Engine/Classes/Animation/AnimData/IAnimationDataModel.h b/Engine/Source/Runtime/Engine/Classes/Animation/AnimData/IAnimationDataModel.h
--- a/Engine/Source/Runtime/Engine/Classes/Animation/AnimData/IAnimationDataModel.h
+++ b/Engine/Source/Runtime/Engine/Classes/Animation/AnimData/IAnimationDataModel.h
@@ -360,16 +360,29 @@ public:
 
 	struct ENGINE_API FEvaluationAndModificationLock
 	{
-		FEvaluationAndModificationLock(IAnimationDataModel& InModel) : Model(InModel) { InModel.LockEvaluationAndModification(); }
-		~FEvaluationAndModificationLock()  { Model.UnlockEvaluationAndModification(); }
+		FEvaluationAndModificationLock(IAnimationDataModel& InModel) : Model(InModel) { InModel.LockEvaluationAndModification(); bLocked = true; }
+		FEvaluationAndModificationLock(IAnimationDataModel& InModel, const TFunction<bool()>& SpinFunc) : Model(InModel)
+		{
+			while(SpinFunc())
+			{
+				bLocked = InModel.TryLockEvaluationAndModification();
+				if (bLocked)
+				{
+					break;
+				}
+			}
+		}
+		~FEvaluationAndModificationLock()  { if(bLocked) {Model.UnlockEvaluationAndModification(); } }
 	private:
 		IAnimationDataModel& Model;
+		bool bLocked = false;
 	};
 protected:
 	virtual FAnimDataModelModifiedDynamicEvent& GetModifiedDynamicEvent() = 0;
 	virtual void OnNotify(const EAnimDataModelNotifyType& NotifyType, const FAnimDataModelNotifPayload& Payload) = 0;
 
 	virtual void LockEvaluationAndModification() const = 0;
+	virtual bool TryLockEvaluationAndModification() const = 0;
 	virtual void UnlockEvaluationAndModification() const = 0;	
 	
 	struct FModelNotifier;
diff --git a/Engine/Source/Runtime/Engine/Private/Animation/AnimCompressionTypes.cpp b/Engine/Source/Runtime/Engine/Private/Animation/AnimCompressionTypes.cpp
--- a/Engine/Source/Runtime/Engine/Private/Animation/AnimCompressionTypes.cpp
+++ b/Engine/Source/Runtime/Engine/Private/Animation/AnimCompressionTypes.cpp
@@ -186,15 +186,26 @@ void FCompressibleAnimData::BakeOutAdditiveIntoRawData(const FFrameRate& SampleR
 	checkf(bSelfAdditiveType || (AdditiveBaseAnimation && !AdditiveBaseAnimation->HasAnyFlags(EObjectFlags::RF_NeedPostLoad)), TEXT("Invalid additive base animation state"));
 
 	// Lock DataModel evaluation
-	IAnimationDataModel::FEvaluationAndModificationLock Lock(*AnimSequence->GetDataModelInterface());
+	IAnimationDataModel::FEvaluationAndModificationLock Lock(*AnimSequence->GetDataModelInterface(), [this]() -> bool
+	{
+		return !IsCancelled();
+	});
 	
 	TUniquePtr<IAnimationDataModel::FEvaluationAndModificationLock> AdditiveBaseLock = nullptr;
 	if (!bSelfAdditiveType)
 	{
 		// Additive is based off another Animation Asset, so lock its evaluation path as well
-		AdditiveBaseLock = MakeUnique<IAnimationDataModel::FEvaluationAndModificationLock>(*AdditiveBaseAnimation->GetDataModelInterface());
+		AdditiveBaseLock = MakeUnique<IAnimationDataModel::FEvaluationAndModificationLock>(*AdditiveBaseAnimation->GetDataModelInterface(), [this]() -> bool
+		{
+			return !IsCancelled();
+		});
 	}
-	
+
+	if (IsCancelled())
+	{
+		return;
+	}
+
 	FMemMark Mark(FMemStack::Get());
 	FByFramePoseEvalContext EvalContext(AnimSequence);
 
@@ -393,7 +404,15 @@ void FCompressibleAnimData::ResampleAnimationTrackData(const FFrameRate& SampleR
 	
 		ensure(DataModelInterface != nullptr);
 		{
-			IAnimationDataModel::FEvaluationAndModificationLock Lock(*DataModelInterface);
+			IAnimationDataModel::FEvaluationAndModificationLock Lock(*DataModelInterface, [this]() -> bool
+			{
+				return !IsCancelled();
+			});
+			
+			if (IsCancelled())
+            {		
+            	return;
+            }
 			
 			FMemMark Mark(FMemStack::Get());
 			FByFramePoseEvalContext EvalContext(AnimSequence);
@@ -767,6 +786,11 @@ void FCompressibleAnimData::FetchData(const ITargetPlatform* InPlatform)
 		RawFloatCurves = AnimSequence->GetDataModelInterface()->GetFloatCurves();
     }
 
+	if (IsCancelled())
+	{		
+		return;
+	}
+
 	for (const FBoneAnimationTrack& AnimTrack : ResampledTrackData)
 	{
 		FRawAnimSequenceTrack& Track = RawAnimationData.Add_GetRef(AnimTrack.InternalTrackData);
diff --git a/Engine/Source/Runtime/Engine/Private/Animation/AnimationCompressionDerivedData.cpp b/Engine/Source/Runtime/Engine/Private/Animation/AnimationCompressionDerivedData.cpp
--- a/Engine/Source/Runtime/Engine/Private/Animation/AnimationCompressionDerivedData.cpp
+++ b/Engine/Source/Runtime/Engine/Private/Animation/AnimationCompressionDerivedData.cpp
@@ -50,6 +50,7 @@ FAnimationSequenceAsyncCacheTask::~FAnimationSequenceAsyncCacheTask()
 
 void FAnimationSequenceAsyncCacheTask::Cancel()
 {
+	CompressibleAnimPtr->IsCancelledSignal.Cancel();		
 	if (BuildTask)
 	{
 		BuildTask->Cancel();
@@ -227,6 +228,11 @@ bool FAnimationSequenceAsyncCacheTask::BuildData() const
 	
 	FCompressibleAnimDataResult CompressionResult;
 	DataToCompress.FetchData(TargetPlatform);
+		
+	if (Owner.IsCanceled())
+	{
+		return false;
+	}
 	DataToCompress.Update(OutData);
 		
 	const bool bBoneCompressionOk = FAnimationUtils::CompressAnimBones(DataToCompress, CompressionResult);
