Skeletal Mesh - Fixed typo in comment.

#rnx

diff --git a/Engine/Source/Runtime/Engine/Private/Components/SkeletalMeshComponent.cpp b/Engine/Source/Runtime/Engine/Private/Components/SkeletalMeshComponent.cpp
--- a/Engine/Source/Runtime/Engine/Private/Components/SkeletalMeshComponent.cpp
+++ b/Engine/Source/Runtime/Engine/Private/Components/SkeletalMeshComponent.cpp
@@ -3028,7 +3028,7 @@ void USkeletalMeshComponent::SetSkeletalMesh(USkeletalMesh* InSkelMesh, bool bRe
 		return;
 	}
 
-	// Stop any existing cloth simulation pior to replacing the asset
+	// Stop any existing cloth simulation prior to replacing the asset
 	RemoveAllClothingActors();
 
 	// Update property alias
