PSO precaching: Added thorough PSO precache data collection for a various particle system

diff --git a/Engine/Plugins/FX/Niagara/Source/Niagara/Private/NiagaraMeshRendererProperties.cpp b/Engine/Plugins/FX/Niagara/Source/Niagara/Private/NiagaraMeshRendererProperties.cpp
--- a/Engine/Plugins/FX/Niagara/Source/Niagara/Private/NiagaraMeshRendererProperties.cpp
+++ b/Engine/Plugins/FX/Niagara/Source/Niagara/Private/NiagaraMeshRendererProperties.cpp
@@ -79,14 +79,22 @@ public:
 		}
 	}
 
+	static void InitVertexFactoryComponents(
+		const FStaticMeshVertexBuffers& VertexBuffers,
+		FNiagaraMeshVertexFactory* VertexFactory,
+		FStaticMeshDataType& OutData)
+	{
+		VertexBuffers.PositionVertexBuffer.BindPositionVertexBuffer(VertexFactory, OutData);
+		VertexBuffers.StaticMeshVertexBuffer.BindTangentVertexBuffer(VertexFactory, OutData);
+		VertexBuffers.StaticMeshVertexBuffer.BindPackedTexCoordVertexBuffer(VertexFactory, OutData, MAX_TEXCOORDS);
+		VertexBuffers.ColorVertexBuffer.BindColorVertexBuffer(VertexFactory, OutData);
+	}
+
 	void SetupVertexFactory(FNiagaraMeshVertexFactory& InVertexFactory, const FLODModelData& LODModelData) const override
 	{
 		FStaticMeshDataType Data;
 		const FStaticMeshLODResources& LODResources = RenderData->LODResources[LODModelData.LODIndex];
-		LODResources.VertexBuffers.PositionVertexBuffer.BindPositionVertexBuffer(&InVertexFactory, Data);
-		LODResources.VertexBuffers.StaticMeshVertexBuffer.BindTangentVertexBuffer(&InVertexFactory, Data);
-		LODResources.VertexBuffers.StaticMeshVertexBuffer.BindTexCoordVertexBuffer(&InVertexFactory, Data, MAX_TEXCOORDS);
-		LODResources.VertexBuffers.ColorVertexBuffer.BindColorVertexBuffer(&InVertexFactory, Data);
+		InitVertexFactoryComponents(LODResources.VertexBuffers, &InVertexFactory, Data);
 		InVertexFactory.SetData(Data);
 	}
 
@@ -662,6 +670,55 @@ void UNiagaraMeshRendererProperties::GetUsedMaterials(const FNiagaraEmitterInsta
 	}
 }
 
+void UNiagaraMeshRendererProperties::CollectPSOPrecacheData(FPSOPrecacheParamsList& OutParams)
+{
+	const FVertexFactoryType* VFType = GetVertexFactoryType();
+	bool bSupportsManualVertexFetch = VFType->SupportsManualVertexFetch(GMaxRHIFeatureLevel);
+		
+	for (int32 MeshIndex = 0; MeshIndex < Meshes.Num(); ++MeshIndex)
+	{
+		// Don't have an instance yet to retrieve the possible material override data from
+		FNiagaraEmitterInstance* EmitterInstance = nullptr;
+		INiagaraRenderableMeshInterface* RenderableMeshInterface = nullptr;
+		UStaticMesh* StaticMesh = nullptr;
+		NiagaraMeshRendererPropertiesInternal::ResolveRenderableMeshInternal(Meshes[MeshIndex], EmitterInstance, RenderableMeshInterface, StaticMesh);
+		if (StaticMesh)
+		{
+			TArray<UMaterialInterface*> OrderedMeshMaterials;
+			FNiagaraRenderableStaticMesh(StaticMesh).GetUsedMaterials(OrderedMeshMaterials);
+			if (OrderedMeshMaterials.Num())
+			{
+				ApplyMaterialOverrides(nullptr, OrderedMeshMaterials);
+				for (UMaterialInterface* MeshMaterial : OrderedMeshMaterials)
+				{
+					if (MeshMaterial)
+					{
+						FPSOPrecacheParams& PSOPrecacheParams = OutParams.AddDefaulted_GetRef();
+						PSOPrecacheParams.MaterialInterface = MeshMaterial;
+						if (!bSupportsManualVertexFetch)
+						{
+							// Assuming here that all LOD use same vertex decl
+							int32 MeshLODIdx = StaticMesh->GetMinLODIdx();
+							if (StaticMesh->GetRenderData()->LODResources.IsValidIndex(MeshLODIdx))
+							{
+								FStaticMeshDataType Data;
+								FVertexDeclarationElementList Elements;
+								FNiagaraRenderableStaticMesh::InitVertexFactoryComponents(StaticMesh->GetRenderData()->LODResources[MeshLODIdx].VertexBuffers, nullptr, Data);
+								FNiagaraMeshVertexFactory::GetVertexElements(GMaxRHIFeatureLevel, bSupportsManualVertexFetch, Data, Elements);
+								PSOPrecacheParams.VertexFactoryDataList.Add(FPSOPrecacheVertexFactoryData(VFType, Elements));
+							}
+						}
+						else
+						{
+							PSOPrecacheParams.VertexFactoryDataList.Add(FPSOPrecacheVertexFactoryData(VFType));
+						}
+					}
+				}
+			}
+		}
+	}
+}
+
 void UNiagaraMeshRendererProperties::GetStreamingMeshInfo(const FBoxSphereBounds& OwnerBounds, const FNiagaraEmitterInstance* InEmitter, TArray<FStreamingRenderAssetPrimitiveInfo>& OutStreamingRenderAssets) const
 {
 	for (const FNiagaraMeshRendererMeshProperties& MeshProperties : Meshes)
diff --git a/Engine/Plugins/FX/Niagara/Source/Niagara/Private/NiagaraRibbonRendererProperties.cpp b/Engine/Plugins/FX/Niagara/Source/Niagara/Private/NiagaraRibbonRendererProperties.cpp
--- a/Engine/Plugins/FX/Niagara/Source/Niagara/Private/NiagaraRibbonRendererProperties.cpp
+++ b/Engine/Plugins/FX/Niagara/Source/Niagara/Private/NiagaraRibbonRendererProperties.cpp
@@ -185,6 +185,20 @@ void UNiagaraRibbonRendererProperties::GetUsedMaterials(const FNiagaraEmitterIns
 	OutMaterials.Add(MaterialInterface ? MaterialInterface : ToRawPtr(Material));
 }
 
+void UNiagaraRibbonRendererProperties::CollectPSOPrecacheData(FPSOPrecacheParamsList& OutParams)
+{
+	const FVertexFactoryType* VFType = GetVertexFactoryType();
+	UMaterialInterface* MaterialInterface = ToRawPtr(Material);
+
+	if (MaterialInterface)
+	{
+		FPSOPrecacheParams& PSOPrecacheParams = OutParams.AddDefaulted_GetRef();
+		PSOPrecacheParams.MaterialInterface = MaterialInterface;
+		// Ribbon VF is the same for MVF and non-MVF cases
+		PSOPrecacheParams.VertexFactoryDataList.Add(FPSOPrecacheVertexFactoryData(VFType));
+	}
+}
+
 const FVertexFactoryType* UNiagaraRibbonRendererProperties::GetVertexFactoryType() const
 {
 	return &FNiagaraRibbonVertexFactory::StaticType;
diff --git a/Engine/Plugins/FX/Niagara/Source/Niagara/Private/NiagaraSpriteRendererProperties.cpp b/Engine/Plugins/FX/Niagara/Source/Niagara/Private/NiagaraSpriteRendererProperties.cpp
--- a/Engine/Plugins/FX/Niagara/Source/Niagara/Private/NiagaraSpriteRendererProperties.cpp
+++ b/Engine/Plugins/FX/Niagara/Source/Niagara/Private/NiagaraSpriteRendererProperties.cpp
@@ -138,6 +138,20 @@ void UNiagaraSpriteRendererProperties::GetUsedMaterials(const FNiagaraEmitterIns
 	OutMaterials.Add(MaterialInterface ? MaterialInterface : ToRawPtr(Material));
 }
 
+void UNiagaraSpriteRendererProperties::CollectPSOPrecacheData(FPSOPrecacheParamsList& OutParams)
+{
+	const FVertexFactoryType* VFType = GetVertexFactoryType();
+	UMaterialInterface* MaterialInterface = ToRawPtr(Material);
+
+	if (MaterialInterface)
+	{
+		FPSOPrecacheParams& PSOPrecacheParams = OutParams.AddDefaulted_GetRef();
+		PSOPrecacheParams.MaterialInterface = MaterialInterface;
+		// Spite VF is the same for MVF and non-MVF cases
+		PSOPrecacheParams.VertexFactoryDataList.Add(FPSOPrecacheVertexFactoryData(VFType));
+	}
+}
+
 const FVertexFactoryType* UNiagaraSpriteRendererProperties::GetVertexFactoryType() const
 {
 	return &FNiagaraSpriteVertexFactory::StaticType;
diff --git a/Engine/Plugins/FX/Niagara/Source/Niagara/Private/NiagaraSystem.cpp b/Engine/Plugins/FX/Niagara/Source/Niagara/Private/NiagaraSystem.cpp
--- a/Engine/Plugins/FX/Niagara/Source/Niagara/Private/NiagaraSystem.cpp
+++ b/Engine/Plugins/FX/Niagara/Source/Niagara/Private/NiagaraSystem.cpp
@@ -1067,13 +1067,12 @@ void UNiagaraSystem::PrecachePSOs()
 					}
 					bool bDisableBackfaceCulling = Properties->IsBackfaceCullingDisabled();
 
-					// Don't have an instance yet to retrieve the possible material override data from
-					const FNiagaraEmitterInstance* EmitterInstance = nullptr;
-					TArray<UMaterialInterface*> Mats;
-					Properties->GetUsedMaterials(EmitterInstance, Mats);
+					UNiagaraRendererProperties::FPSOPrecacheParamsList PSOPrecacheParamsList;
+					Properties->CollectPSOPrecacheData(PSOPrecacheParamsList);
 
-					for (UMaterialInterface* MaterialInterface : Mats)
+					for (UNiagaraRendererProperties::FPSOPrecacheParams& PSOPrecacheParams: PSOPrecacheParamsList)
 					{
+						UMaterialInterface* MaterialInterface = PSOPrecacheParams.MaterialInterface;
 						VFsPerMaterialData* VFsPerMaterial = VFsPerMaterials.FindByPredicate([MaterialInterface, bDisableBackfaceCulling](const VFsPerMaterialData& Other) 
 						{ 
 							return (Other.MaterialInterface == MaterialInterface &&
@@ -1085,7 +1084,11 @@ void UNiagaraSystem::PrecachePSOs()
 							VFsPerMaterial->MaterialInterface = MaterialInterface;
 							VFsPerMaterial->bDisableBackfaceCulling = bDisableBackfaceCulling;
 						}
-						VFsPerMaterial->VertexFactoryTypes.AddUnique(VFType);
+
+						for (FPSOPrecacheVertexFactoryData& VFData : PSOPrecacheParams.VertexFactoryDataList)
+						{
+							VFsPerMaterial->VertexFactoryData.AddUnique(VFData);
+						}
 					}
 				}
 			);
diff --git a/Engine/Plugins/FX/Niagara/Source/Niagara/Public/NiagaraMeshRendererProperties.h b/Engine/Plugins/FX/Niagara/Source/Niagara/Public/NiagaraMeshRendererProperties.h
--- a/Engine/Plugins/FX/Niagara/Source/Niagara/Public/NiagaraMeshRendererProperties.h
+++ b/Engine/Plugins/FX/Niagara/Source/Niagara/Public/NiagaraMeshRendererProperties.h
@@ -206,6 +206,8 @@ public:
 	virtual const FVertexFactoryType* GetVertexFactoryType() const override;
 	virtual bool IsSimTargetSupported(ENiagaraSimTarget InSimTarget) const override { return true; };
 	virtual bool PopulateRequiredBindings(FNiagaraParameterStore& InParameterStore) override;
+	virtual void CollectPSOPrecacheData(FPSOPrecacheParamsList& OutParams) override;
+
 #if WITH_EDITORONLY_DATA
 	virtual const TArray<FNiagaraVariable>& GetOptionalAttributes() override;
 	virtual void GetAdditionalVariables(TArray<FNiagaraVariableBase>& OutArray) const override;
diff --git a/Engine/Plugins/FX/Niagara/Source/Niagara/Public/NiagaraRendererProperties.h b/Engine/Plugins/FX/Niagara/Source/Niagara/Public/NiagaraRendererProperties.h
--- a/Engine/Plugins/FX/Niagara/Source/Niagara/Public/NiagaraRendererProperties.h
+++ b/Engine/Plugins/FX/Niagara/Source/Niagara/Public/NiagaraRendererProperties.h
@@ -315,6 +315,18 @@ public:
 	/** In the case that we need parameters bound in that aren't Particle variables, these should be set up here so that the data is appropriately populated after the simulation.*/
 	virtual bool PopulateRequiredBindings(FNiagaraParameterStore& InParameterStore);
 
+	
+	/**
+	* Collect all the data required for PSO precaching 
+	*/
+	struct FPSOPrecacheParams
+	{
+		UMaterialInterface* MaterialInterface = nullptr;
+		FPSOPrecacheVertexFactoryDataList VertexFactoryDataList;
+	};
+	typedef TArray<FPSOPrecacheParams, TInlineAllocator<2> > FPSOPrecacheParamsList;
+	virtual void CollectPSOPrecacheData(FPSOPrecacheParamsList& OutParams) {};
+
 #if WITH_EDITORONLY_DATA
 
 	virtual bool IsSupportedVariableForBinding(const FNiagaraVariableBase& InSourceForBinding, const FName& InTargetBindingName) const;
diff --git a/Engine/Plugins/FX/Niagara/Source/Niagara/Public/NiagaraRibbonRendererProperties.h b/Engine/Plugins/FX/Niagara/Source/Niagara/Public/NiagaraRibbonRendererProperties.h
--- a/Engine/Plugins/FX/Niagara/Source/Niagara/Public/NiagaraRibbonRendererProperties.h
+++ b/Engine/Plugins/FX/Niagara/Source/Niagara/Public/NiagaraRibbonRendererProperties.h
@@ -219,6 +219,7 @@ public:
 	virtual bool IsBackfaceCullingDisabled() const override;
 	virtual bool IsSimTargetSupported(ENiagaraSimTarget InSimTarget) const override { return true; };
 	virtual bool PopulateRequiredBindings(FNiagaraParameterStore& InParameterStore) override;
+	virtual void CollectPSOPrecacheData(FPSOPrecacheParamsList& OutParams) override;
 
 #if WITH_EDITOR
 	virtual const TArray<FNiagaraVariable>& GetOptionalAttributes() override;
diff --git a/Engine/Plugins/FX/Niagara/Source/Niagara/Public/NiagaraSpriteRendererProperties.h b/Engine/Plugins/FX/Niagara/Source/Niagara/Public/NiagaraSpriteRendererProperties.h
--- a/Engine/Plugins/FX/Niagara/Source/Niagara/Public/NiagaraSpriteRendererProperties.h
+++ b/Engine/Plugins/FX/Niagara/Source/Niagara/Public/NiagaraSpriteRendererProperties.h
@@ -136,6 +136,7 @@ public:
 	virtual const FVertexFactoryType* GetVertexFactoryType() const override;
 	virtual bool IsSimTargetSupported(ENiagaraSimTarget InSimTarget) const override { return true; };
 	virtual bool PopulateRequiredBindings(FNiagaraParameterStore& InParameterStore)  override;
+	virtual void CollectPSOPrecacheData(FPSOPrecacheParamsList& OutParams) override;
 #if WITH_EDITOR
 	virtual const TArray<FNiagaraVariable>& GetOptionalAttributes() override;
 	virtual void GetAdditionalVariables(TArray<FNiagaraVariableBase>& OutArray) const override;
diff --git a/Engine/Plugins/FX/Niagara/Source/NiagaraVertexFactories/Private/NiagaraMeshVertexFactory.cpp b/Engine/Plugins/FX/Niagara/Source/NiagaraVertexFactories/Private/NiagaraMeshVertexFactory.cpp
--- a/Engine/Plugins/FX/Niagara/Source/NiagaraVertexFactories/Private/NiagaraMeshVertexFactory.cpp
+++ b/Engine/Plugins/FX/Niagara/Source/NiagaraVertexFactories/Private/NiagaraMeshVertexFactory.cpp
@@ -70,21 +70,14 @@ public:
 
 IMPLEMENT_TYPE_LAYOUT(FNiagaraMeshVertexFactoryShaderParametersPS);
 
-void FNiagaraMeshVertexFactory::InitRHI()
+void FNiagaraMeshVertexFactory::GetVertexElements(ERHIFeatureLevel::Type FeatureLevel, bool bSupportsManualVertexFetch, FStaticMeshDataType& Data, FVertexDeclarationElementList& Elements, FVertexStreamList& InOutStreams)
 {
-	FVertexDeclarationElementList Elements;
-		
 	if (Data.PositionComponent.VertexBuffer != NULL)
 	{
-		Elements.Add(AccessStreamComponent(Data.PositionComponent, 0));
+		Elements.Add(AccessStreamComponent(Data.PositionComponent, 0, InOutStreams));
 	}
 
-	// Should also work in editor (PIE with ES 3.1 preview seems to work), but we don't precache PSOs for that platform currently, so keep it safe in case there's a flow we missed
-#if !WITH_EDITOR
-	const bool bHasValidFeatureLevel = ensure(HasValidFeatureLevel());
-	const bool bRemoveUnusedVertexElements = bHasValidFeatureLevel ? SupportsManualVertexFetch(GetFeatureLevel()) : false;
-	if (!bRemoveUnusedVertexElements)
-#endif
+	if (!bSupportsManualVertexFetch)
 	{
 		// only tangent,normal are used by the stream. the binormal is derived in the shader
 		uint8 TangentBasisAttributes[2] = { 1, 2 };
@@ -92,7 +85,7 @@ void FNiagaraMeshVertexFactory::InitRHI()
 		{
 			if (Data.TangentBasisComponents[AxisIndex].VertexBuffer != NULL)
 			{
-				Elements.Add(AccessStreamComponent(Data.TangentBasisComponents[AxisIndex], TangentBasisAttributes[AxisIndex]));
+				Elements.Add(AccessStreamComponent(Data.TangentBasisComponents[AxisIndex], TangentBasisAttributes[AxisIndex], InOutStreams));
 			}
 		}
 
@@ -105,14 +98,14 @@ void FNiagaraMeshVertexFactory::InitRHI()
 		// Vertex color
 		if (Data.ColorComponent.VertexBuffer != NULL)
 		{
-			Elements.Add(AccessStreamComponent(Data.ColorComponent, 3));
+			Elements.Add(AccessStreamComponent(Data.ColorComponent, 3, InOutStreams));
 		}
 		else
 		{
 			//If the mesh has no color component, set the null color buffer on a new stream with a stride of 0.
 			//This wastes 4 bytes of bandwidth per vertex, but prevents having to compile out twice the number of vertex factories.
 			FVertexStreamComponent NullColorComponent(&GNullColorVertexBuffer, 0, 0, VET_Color, EVertexStreamUsage::ManualFetch);
-			Elements.Add(AccessStreamComponent(NullColorComponent, 3));
+			Elements.Add(AccessStreamComponent(NullColorComponent, 3, InOutStreams));
 		}
 
 		if (Data.TextureCoordinates.Num())
@@ -122,7 +115,8 @@ void FNiagaraMeshVertexFactory::InitRHI()
 			{
 				Elements.Add(AccessStreamComponent(
 					Data.TextureCoordinates[CoordinateIndex],
-					BaseTexCoordAttribute + CoordinateIndex
+					BaseTexCoordAttribute + CoordinateIndex,
+					InOutStreams
 				));
 			}
 
@@ -130,11 +124,21 @@ void FNiagaraMeshVertexFactory::InitRHI()
 			{
 				Elements.Add(AccessStreamComponent(
 					Data.TextureCoordinates[Data.TextureCoordinates.Num() - 1],
-					BaseTexCoordAttribute + CoordinateIndex
+					BaseTexCoordAttribute + CoordinateIndex,
+					InOutStreams
 				));
 			}
 		}
 	}
+}
+
+void FNiagaraMeshVertexFactory::InitRHI()
+{
+	check(HasValidFeatureLevel());
+	const bool bSupportsManualVertexFetch = SupportsManualVertexFetch(GetFeatureLevel());
+
+	FVertexDeclarationElementList Elements;
+	GetVertexElements(GetFeatureLevel(), bSupportsManualVertexFetch, Data, Elements, Streams);
 
 #if NIAGARA_ENABLE_GPU_SCENE_MESHES
 	if (bAddPrimitiveIDElement)
@@ -201,6 +205,21 @@ void FNiagaraMeshVertexFactory::GetPSOPrecacheVertexFetchElements(EVertexInputSt
 #endif
 }
 
+void FNiagaraMeshVertexFactory::GetVertexElements(ERHIFeatureLevel::Type FeatureLevel, bool bSupportsManualVertexFetch, FStaticMeshDataType& Data, FVertexDeclarationElementList& Elements)
+{
+	FVertexStreamList InOutStreams;
+	GetVertexElements(FeatureLevel, bSupportsManualVertexFetch, Data, Elements, InOutStreams);
+
+#if NIAGARA_ENABLE_GPU_SCENE_MESHES
+	if (UseGPUScene(GMaxRHIShaderPlatform, GMaxRHIFeatureLevel))
+	{
+		// For ES3.1 attribute ID needs to be done differently
+		check(FeatureLevel > ERHIFeatureLevel::ES3_1);
+		Elements.Add(FVertexElement(InOutStreams.Num(), 0, VET_UInt, 13, 0, true));
+	}
+#endif
+}
+
 void FNiagaraMeshVertexFactory::SetData(const FStaticMeshDataType& InData)
 {
 	check(IsInRenderingThread());
diff --git a/Engine/Plugins/FX/Niagara/Source/NiagaraVertexFactories/Public/NiagaraMeshVertexFactory.h b/Engine/Plugins/FX/Niagara/Source/NiagaraVertexFactories/Public/NiagaraMeshVertexFactory.h
--- a/Engine/Plugins/FX/Niagara/Source/NiagaraVertexFactories/Public/NiagaraMeshVertexFactory.h
+++ b/Engine/Plugins/FX/Niagara/Source/NiagaraVertexFactories/Public/NiagaraMeshVertexFactory.h
@@ -157,6 +157,7 @@ public:
 	* Get vertex elements used when during PSO precaching materials using this vertex factory type
 	*/
 	static void GetPSOPrecacheVertexFetchElements(EVertexInputStreamType VertexInputStreamType, FVertexDeclarationElementList& Elements);
+	static void GetVertexElements(ERHIFeatureLevel::Type FeatureLevel, bool bSupportsManualVertexFetch, FStaticMeshDataType& Data, FVertexDeclarationElementList& Elements);
 
 	/**
 	* An implementation of the interface used by TSynchronizedResource to update the resource with new data from the game thread.
@@ -220,6 +221,8 @@ public:
 	{
 		FNiagaraVertexFactoryBase::SetFeatureLevel(InFeatureLevel);
 	}
+protected:
+	static void GetVertexElements(ERHIFeatureLevel::Type FeatureLevel, bool bSupportsManualVertexFetch, FStaticMeshDataType& Data, FVertexDeclarationElementList& Elements, FVertexStreamList& InOutStreams);
 
 protected:
 	FStaticMeshDataType Data;
diff --git a/Engine/Source/Runtime/Engine/Classes/Particles/ParticleSystem.h b/Engine/Source/Runtime/Engine/Classes/Particles/ParticleSystem.h
--- a/Engine/Source/Runtime/Engine/Classes/Particles/ParticleSystem.h
+++ b/Engine/Source/Runtime/Engine/Classes/Particles/ParticleSystem.h
@@ -162,7 +162,7 @@ protected:
 		UMaterialInterface* MaterialInterface = nullptr;
 		EPrimitiveType PrimitiveType = PT_TriangleList; // must match FPSOPrecacheParams::PrimitiveType default value
 		bool bDisableBackfaceCulling = false;  // must match FPSOPrecacheParams::bDisableBackfaceCulling default value
-		TArray<const class FVertexFactoryType*, TInlineAllocator<2>> VertexFactoryTypes;
+		FPSOPrecacheVertexFactoryDataList VertexFactoryData;
 	};
 
 	ENGINE_API void LaunchPSOPrecaching(TArrayView<VFsPerMaterialData> VFsPerMaterials);
diff --git a/Engine/Source/Runtime/Engine/Classes/Particles/TypeData/ParticleModuleTypeDataAnimTrail.h b/Engine/Source/Runtime/Engine/Classes/Particles/TypeData/ParticleModuleTypeDataAnimTrail.h
--- a/Engine/Source/Runtime/Engine/Classes/Particles/TypeData/ParticleModuleTypeDataAnimTrail.h
+++ b/Engine/Source/Runtime/Engine/Classes/Particles/TypeData/ParticleModuleTypeDataAnimTrail.h
@@ -83,6 +83,7 @@ class UParticleModuleTypeDataAnimTrail : public UParticleModuleTypeDataBase
 	virtual FParticleEmitterInstance* CreateInstance(UParticleEmitter* InEmitterParent, UParticleSystemComponent* InComponent) override;
 	virtual const FVertexFactoryType* GetVertexFactoryType() const override;
 	virtual EPrimitiveType GetPrimitiveType() const override { return PT_TriangleStrip; }
+	virtual void CollectPSOPrecacheData(const UParticleEmitter* Emitter, FPSOPrecacheParams& OutParams) override;
 	//~ End UParticleModuleTypeDataBase Interface
 };
 
diff --git a/Engine/Source/Runtime/Engine/Classes/Particles/TypeData/ParticleModuleTypeDataBase.h b/Engine/Source/Runtime/Engine/Classes/Particles/TypeData/ParticleModuleTypeDataBase.h
--- a/Engine/Source/Runtime/Engine/Classes/Particles/TypeData/ParticleModuleTypeDataBase.h
+++ b/Engine/Source/Runtime/Engine/Classes/Particles/TypeData/ParticleModuleTypeDataBase.h
@@ -52,6 +52,16 @@ class UParticleModuleTypeDataBase : public UParticleModule
 	/** Determine if motion blur is enabled for the owning emitter. */
 	virtual bool		IsMotionBlurEnabled() const  { return false; }
 
+
+	/**
+	* Collect all the data required for PSO precaching 
+	*/
+	struct FPSOPrecacheParams
+	{
+		EPrimitiveType PrimitiveType;
+		FPSOPrecacheVertexFactoryDataList VertexFactoryDataList;
+	};
+	virtual void CollectPSOPrecacheData(const UParticleEmitter* Emitter, FPSOPrecacheParams& OutParams) {};
 };
 
 
diff --git a/Engine/Source/Runtime/Engine/Classes/Particles/TypeData/ParticleModuleTypeDataBeam2.h b/Engine/Source/Runtime/Engine/Classes/Particles/TypeData/ParticleModuleTypeDataBeam2.h
--- a/Engine/Source/Runtime/Engine/Classes/Particles/TypeData/ParticleModuleTypeDataBeam2.h
+++ b/Engine/Source/Runtime/Engine/Classes/Particles/TypeData/ParticleModuleTypeDataBeam2.h
@@ -226,6 +226,7 @@ class UParticleModuleTypeDataBeam2 : public UParticleModuleTypeDataBase
 	virtual void CacheModuleInfo(UParticleEmitter* Emitter) override;
 	virtual const FVertexFactoryType* GetVertexFactoryType() const override;
 	virtual EPrimitiveType GetPrimitiveType() const { return PT_TriangleStrip; }
+	virtual void CollectPSOPrecacheData(const UParticleEmitter* Emitter, FPSOPrecacheParams& OutParams) override;
 	//~ End UParticleModuleTypeDataBase Interface
 
 
diff --git a/Engine/Source/Runtime/Engine/Classes/Particles/TypeData/ParticleModuleTypeDataGpu.h b/Engine/Source/Runtime/Engine/Classes/Particles/TypeData/ParticleModuleTypeDataGpu.h
--- a/Engine/Source/Runtime/Engine/Classes/Particles/TypeData/ParticleModuleTypeDataGpu.h
+++ b/Engine/Source/Runtime/Engine/Classes/Particles/TypeData/ParticleModuleTypeDataGpu.h
@@ -484,6 +484,7 @@ class UParticleModuleTypeDataGpu : public UParticleModuleTypeDataBase
 	virtual bool RequiresBuild() const override { return true; }
 	virtual FParticleEmitterInstance* CreateInstance(UParticleEmitter* InEmitterParent, UParticleSystemComponent* InComponent) override;
 	virtual const FVertexFactoryType* GetVertexFactoryType() const override;
+	virtual void CollectPSOPrecacheData(const UParticleEmitter* Emitter, FPSOPrecacheParams& OutParams) override;
 	//~ End UParticleModuleTypeDataBase Interface
 };
 
diff --git a/Engine/Source/Runtime/Engine/Classes/Particles/TypeData/ParticleModuleTypeDataMesh.h b/Engine/Source/Runtime/Engine/Classes/Particles/TypeData/ParticleModuleTypeDataMesh.h
--- a/Engine/Source/Runtime/Engine/Classes/Particles/TypeData/ParticleModuleTypeDataMesh.h
+++ b/Engine/Source/Runtime/Engine/Classes/Particles/TypeData/ParticleModuleTypeDataMesh.h
@@ -226,6 +226,7 @@ class UParticleModuleTypeDataMesh : public UParticleModuleTypeDataBase
 	//~ Begin UParticleModuleTypeDataBase Interface
 	virtual FParticleEmitterInstance* CreateInstance(UParticleEmitter* InEmitterParent, UParticleSystemComponent* InComponent) override;
 	virtual const FVertexFactoryType* GetVertexFactoryType() const override;
+	virtual void CollectPSOPrecacheData(const UParticleEmitter* Emitter, FPSOPrecacheParams& OutParams) override;
 	virtual bool	SupportsSpecificScreenAlignmentFlags() const override {	return true;	}	
 	virtual bool	SupportsSubUV() const override { return true; }
 	virtual bool	IsAMeshEmitter() const override { return true; }
diff --git a/Engine/Source/Runtime/Engine/Classes/Particles/TypeData/ParticleModuleTypeDataRibbon.h b/Engine/Source/Runtime/Engine/Classes/Particles/TypeData/ParticleModuleTypeDataRibbon.h
--- a/Engine/Source/Runtime/Engine/Classes/Particles/TypeData/ParticleModuleTypeDataRibbon.h
+++ b/Engine/Source/Runtime/Engine/Classes/Particles/TypeData/ParticleModuleTypeDataRibbon.h
@@ -180,6 +180,7 @@ class UParticleModuleTypeDataRibbon : public UParticleModuleTypeDataBase
 	virtual FParticleEmitterInstance* CreateInstance(UParticleEmitter* InEmitterParent, UParticleSystemComponent* InComponent) override;
 	virtual const FVertexFactoryType* GetVertexFactoryType() const override;
 	virtual EPrimitiveType GetPrimitiveType() const override { return PT_TriangleStrip; }
+	virtual void CollectPSOPrecacheData(const UParticleEmitter* Emitter, FPSOPrecacheParams& OutParams) override;
 	//~ End UParticleModuleTypeDataBase Interface
 };
 
diff --git a/Engine/Source/Runtime/Engine/Private/Components/PrimitiveComponent.cpp b/Engine/Source/Runtime/Engine/Private/Components/PrimitiveComponent.cpp
--- a/Engine/Source/Runtime/Engine/Private/Components/PrimitiveComponent.cpp
+++ b/Engine/Source/Runtime/Engine/Private/Components/PrimitiveComponent.cpp
@@ -4510,6 +4510,8 @@ void UPrimitiveComponent::PrecachePSOs()
 	FGraphEventArray GraphEvents;
 	for (FComponentPSOPrecacheParams& ComponentPSOPrecacheData : PSOPrecacheDataArray)
 	{
+		// Material has to be fully loaded for a static permutations to work
+		ComponentPSOPrecacheData.MaterialInterface->ConditionalPostLoad();
 		GraphEvents.Append(ComponentPSOPrecacheData.MaterialInterface->PrecachePSOs(ComponentPSOPrecacheData.VertexFactoryDataList, ComponentPSOPrecacheData.PSOPrecacheParams, ComponentPSOPrecacheData.Priority, MaterialPSOPrecacheRequestIDs));
 	}	
 
diff --git a/Engine/Source/Runtime/Engine/Private/Materials/MaterialInstance.cpp b/Engine/Source/Runtime/Engine/Private/Materials/MaterialInstance.cpp
--- a/Engine/Source/Runtime/Engine/Private/Materials/MaterialInstance.cpp
+++ b/Engine/Source/Runtime/Engine/Private/Materials/MaterialInstance.cpp
@@ -2479,6 +2479,9 @@ void UMaterialInstance::CacheShaders(EMaterialShaderPrecompileMode CompileMode)
 
 FGraphEventArray UMaterialInstance::PrecachePSOs(const FPSOPrecacheVertexFactoryDataList& VertexFactoryDataList, const FPSOPrecacheParams& InPreCacheParams, EPSOPrecachePriority Priority, TArray<FMaterialPSOPrecacheRequestID>& OutMaterialPSORequestIDs)
 {
+	// make sure material is initialized
+	ensure(!HasAnyFlags(RF_NeedPostLoad));
+	
 	FGraphEventArray GraphEvents;
 	if (FApp::CanEverRender() && PipelineStateCache::IsPSOPrecachingEnabled() && Parent)
 	{
diff --git a/Engine/Source/Runtime/Engine/Private/MeshParticleVertexFactory.cpp b/Engine/Source/Runtime/Engine/Private/MeshParticleVertexFactory.cpp
--- a/Engine/Source/Runtime/Engine/Private/MeshParticleVertexFactory.cpp
+++ b/Engine/Source/Runtime/Engine/Private/MeshParticleVertexFactory.cpp
@@ -89,58 +89,48 @@ private:
 
 static TGlobalResource<FDummyPrevTransformBuffer> GDummyPrevTransformBuffer;
 
-void FMeshParticleVertexFactory::InitRHI()
+void FMeshParticleVertexFactory::GetVertexElements(ERHIFeatureLevel::Type FeatureLevel, int32 InDynamicVertexStride, int32 InDynamicParameterVertexStride, FDataType& Data, FVertexDeclarationElementList& Elements, FVertexStreamList& InOutStreams)
 {
-	FVertexDeclarationElementList Elements;
-
 	if (Data.bInitialized)
 	{
 		// Stream 0 - Instance data
 		{
-			checkf(DynamicVertexStride != -1, TEXT("FMeshParticleVertexFactory does not have a valid DynamicVertexStride - likely an empty one was made, but SetStrides was not called"));
+			checkf(InDynamicVertexStride != -1, TEXT("FMeshParticleVertexFactory does not have a valid DynamicVertexStride - likely an empty one was made, but SetStrides was not called"));
 			FVertexStream VertexStream;
 			VertexStream.VertexBuffer = NULL;
 			VertexStream.Stride = 0;
 			VertexStream.Offset = 0;
-			Streams.Add(VertexStream);
-	
+			InOutStreams.Add(VertexStream);
+
 			// @todo metal: this will need a valid stride when we get to instanced meshes!
-			Elements.Add(FVertexElement(0, Data.TransformComponent[0].Offset, Data.TransformComponent[0].Type, 8, DynamicVertexStride, EnumHasAnyFlags(EVertexStreamUsage::Instancing, Data.TransformComponent[0].VertexStreamUsage)));
-			Elements.Add(FVertexElement(0, Data.TransformComponent[1].Offset, Data.TransformComponent[1].Type, 9, DynamicVertexStride, EnumHasAnyFlags(EVertexStreamUsage::Instancing, Data.TransformComponent[1].VertexStreamUsage)));
-			Elements.Add(FVertexElement(0, Data.TransformComponent[2].Offset, Data.TransformComponent[2].Type, 10, DynamicVertexStride, EnumHasAnyFlags(EVertexStreamUsage::Instancing, Data.TransformComponent[2].VertexStreamUsage)));
-	
-			Elements.Add(FVertexElement(0, Data.SubUVs.Offset, Data.SubUVs.Type, 11, DynamicVertexStride, EnumHasAnyFlags(EVertexStreamUsage::Instancing, Data.SubUVs.VertexStreamUsage)));
-			Elements.Add(FVertexElement(0, Data.SubUVLerpAndRelTime.Offset, Data.SubUVLerpAndRelTime.Type, 12, DynamicVertexStride, EnumHasAnyFlags(EVertexStreamUsage::Instancing, Data.SubUVLerpAndRelTime.VertexStreamUsage)));
-	
-			Elements.Add(FVertexElement(0, Data.ParticleColorComponent.Offset, Data.ParticleColorComponent.Type, 14, DynamicVertexStride, EnumHasAnyFlags(EVertexStreamUsage::Instancing, Data.ParticleColorComponent.VertexStreamUsage)));
-			Elements.Add(FVertexElement(0, Data.VelocityComponent.Offset, Data.VelocityComponent.Type, 15, DynamicVertexStride, EnumHasAnyFlags(EVertexStreamUsage::Instancing, Data.VelocityComponent.VertexStreamUsage)));
+			Elements.Add(FVertexElement(0, Data.TransformComponent[0].Offset, Data.TransformComponent[0].Type, 8, InDynamicVertexStride, EnumHasAnyFlags(EVertexStreamUsage::Instancing, Data.TransformComponent[0].VertexStreamUsage)));
+			Elements.Add(FVertexElement(0, Data.TransformComponent[1].Offset, Data.TransformComponent[1].Type, 9, InDynamicVertexStride, EnumHasAnyFlags(EVertexStreamUsage::Instancing, Data.TransformComponent[1].VertexStreamUsage)));
+			Elements.Add(FVertexElement(0, Data.TransformComponent[2].Offset, Data.TransformComponent[2].Type, 10, InDynamicVertexStride, EnumHasAnyFlags(EVertexStreamUsage::Instancing, Data.TransformComponent[2].VertexStreamUsage)));
+
+			Elements.Add(FVertexElement(0, Data.SubUVs.Offset, Data.SubUVs.Type, 11, InDynamicVertexStride, EnumHasAnyFlags(EVertexStreamUsage::Instancing, Data.SubUVs.VertexStreamUsage)));
+			Elements.Add(FVertexElement(0, Data.SubUVLerpAndRelTime.Offset, Data.SubUVLerpAndRelTime.Type, 12, InDynamicVertexStride, EnumHasAnyFlags(EVertexStreamUsage::Instancing, Data.SubUVLerpAndRelTime.VertexStreamUsage)));
+
+			Elements.Add(FVertexElement(0, Data.ParticleColorComponent.Offset, Data.ParticleColorComponent.Type, 14, InDynamicVertexStride, EnumHasAnyFlags(EVertexStreamUsage::Instancing, Data.ParticleColorComponent.VertexStreamUsage)));
+			Elements.Add(FVertexElement(0, Data.VelocityComponent.Offset, Data.VelocityComponent.Type, 15, InDynamicVertexStride, EnumHasAnyFlags(EVertexStreamUsage::Instancing, Data.VelocityComponent.VertexStreamUsage)));
 		}
 
 		// Stream 1 - Dynamic parameter
 		{
-			checkf(DynamicParameterVertexStride != -1, TEXT("FMeshParticleVertexFactory does not have a valid DynamicParameterVertexStride - likely an empty one was made, but SetStrides was not called"));
-				
+			checkf(InDynamicParameterVertexStride != -1, TEXT("FMeshParticleVertexFactory does not have a valid DynamicParameterVertexStride - likely an empty one was made, but SetStrides was not called"));
+
 			FVertexStream VertexStream;
 			VertexStream.VertexBuffer = NULL;
 			VertexStream.Stride = 0;
 			VertexStream.Offset = 0;
-			Streams.Add(VertexStream);
-	
-			Elements.Add(FVertexElement(1, 0, VET_Float4, 13, DynamicParameterVertexStride, true));
-		}
+			InOutStreams.Add(VertexStream);
 
-		// Add a dummy resource to avoid crash due to missing resource
-		if (GMaxRHIFeatureLevel >= ERHIFeatureLevel::SM5)
-		{
-			PrevTransformBuffer.NumBytes = 0;
-			PrevTransformBuffer.Buffer = GDummyPrevTransformBuffer.GetVB();
-			PrevTransformBuffer.SRV = GDummyPrevTransformBuffer.GetSRV();
+			Elements.Add(FVertexElement(1, 0, VET_Float4, 13, InDynamicParameterVertexStride, true));
 		}
 	}
 
 	if(Data.PositionComponent.VertexBuffer != NULL)
 	{
-		Elements.Add(AccessStreamComponent(Data.PositionComponent,0));
+		Elements.Add(AccessStreamComponent(Data.PositionComponent, 0, InOutStreams));
 	}
 
 	// only tangent,normal are used by the stream. the binormal is derived in the shader
@@ -149,7 +139,7 @@ void FMeshParticleVertexFactory::InitRHI()
 	{
 		if(Data.TangentBasisComponents[AxisIndex].VertexBuffer != NULL)
 		{
-			Elements.Add(AccessStreamComponent(Data.TangentBasisComponents[AxisIndex],TangentBasisAttributes[AxisIndex]));
+			Elements.Add(AccessStreamComponent(Data.TangentBasisComponents[AxisIndex],TangentBasisAttributes[AxisIndex], InOutStreams));
 		}
 	}
 
@@ -162,16 +152,16 @@ void FMeshParticleVertexFactory::InitRHI()
 	// Vertex color
 	if(Data.ColorComponent.VertexBuffer != NULL)
 	{
-		Elements.Add(AccessStreamComponent(Data.ColorComponent,3));
+		Elements.Add(AccessStreamComponent(Data.ColorComponent,3, InOutStreams));
 	}
 	else
 	{
 		//If the mesh has no color component, set the null color buffer on a new stream with a stride of 0.
 		//This wastes 4 bytes of bandwidth per vertex, but prevents having to compile out twice the number of vertex factories.
 		FVertexStreamComponent NullColorComponent(&GNullColorVertexBuffer, 0, 0, VET_Color, EVertexStreamUsage::ManualFetch);
-		Elements.Add(AccessStreamComponent(NullColorComponent, 3));
+		Elements.Add(AccessStreamComponent(NullColorComponent, 3, InOutStreams));
 	}
-		
+
 	if(Data.TextureCoordinates.Num())
 	{
 		const int32 BaseTexCoordAttribute = 4;
@@ -179,18 +169,34 @@ void FMeshParticleVertexFactory::InitRHI()
 		{
 			Elements.Add(AccessStreamComponent(
 				Data.TextureCoordinates[CoordinateIndex],
-				BaseTexCoordAttribute + CoordinateIndex
-				));
+				BaseTexCoordAttribute + CoordinateIndex,
+				InOutStreams
+			));
 		}
 
 		for(int32 CoordinateIndex = Data.TextureCoordinates.Num();CoordinateIndex < MAX_TEXCOORDS;CoordinateIndex++)
 		{
 			Elements.Add(AccessStreamComponent(
 				Data.TextureCoordinates[Data.TextureCoordinates.Num() - 1],
-				BaseTexCoordAttribute + CoordinateIndex
-				));
+				BaseTexCoordAttribute + CoordinateIndex,
+				InOutStreams
+			));
 		}
 	}
+}
+
+void FMeshParticleVertexFactory::InitRHI()
+{
+	FVertexDeclarationElementList Elements;
+	GetVertexElements(GMaxRHIFeatureLevel, DynamicVertexStride, DynamicParameterVertexStride, Data, Elements, Streams);
+
+	// Add a dummy resource to avoid crash due to missing resource
+	if (GMaxRHIFeatureLevel >= ERHIFeatureLevel::SM5)
+	{
+		PrevTransformBuffer.NumBytes = 0;
+		PrevTransformBuffer.Buffer = GDummyPrevTransformBuffer.GetVB();
+		PrevTransformBuffer.SRV = GDummyPrevTransformBuffer.GetSRV();
+	}
 
 	if(Streams.Num() > 0)
 	{
@@ -199,6 +205,12 @@ void FMeshParticleVertexFactory::InitRHI()
 	}
 }
 
+void FMeshParticleVertexFactory::GetVertexElements(ERHIFeatureLevel::Type FeatureLevel, int32 InDynamicVertexStride, int32 InDynamicParameterVertexStride, FDataType& Data, FVertexDeclarationElementList& Elements)
+{
+	FVertexStreamList InOutStreams;
+	GetVertexElements(FeatureLevel, InDynamicVertexStride, InDynamicParameterVertexStride, Data, Elements, InOutStreams);
+}
+
 void FMeshParticleVertexFactory::SetInstanceBuffer(const FVertexBuffer* InstanceBuffer, uint32 StreamOffset, uint32 Stride)
 {
 	ensure(Stride == DynamicVertexStride);
diff --git a/Engine/Source/Runtime/Engine/Private/Particles/ParticleBeamModules.cpp b/Engine/Source/Runtime/Engine/Private/Particles/ParticleBeamModules.cpp
--- a/Engine/Source/Runtime/Engine/Private/Particles/ParticleBeamModules.cpp
+++ b/Engine/Source/Runtime/Engine/Private/Particles/ParticleBeamModules.cpp
@@ -702,6 +702,17 @@ const FVertexFactoryType* UParticleModuleTypeDataBeam2::GetVertexFactoryType() c
 	return &FParticleBeamTrailVertexFactory::StaticType;
 }
 
+void UParticleModuleTypeDataBeam2::CollectPSOPrecacheData(const UParticleEmitter* Emitter, FPSOPrecacheParams& OutParams)
+{
+	bool bUsesDynamicParameter = (Emitter->DynamicParameterDataOffset > 0);
+	
+	FPSOPrecacheVertexFactoryData VFData;
+	VFData.VertexFactoryType = &FParticleBeamTrailVertexFactory::StaticType;
+	VFData.CustomDefaultVertexDeclaration = FParticleBeamTrailVertexFactory::GetPSOPrecacheVertexDeclaration(bUsesDynamicParameter);
+	OutParams.VertexFactoryDataList.Add(VFData);
+	OutParams.PrimitiveType = GetPrimitiveType();
+}
+
 void UParticleModuleTypeDataBeam2::CacheModuleInfo(UParticleEmitter* Emitter)
 {
 	int32 LODCount = Emitter->LODLevels.Num();
diff --git a/Engine/Source/Runtime/Engine/Private/Particles/ParticleBeamTrailVertexFactory.cpp b/Engine/Source/Runtime/Engine/Private/Particles/ParticleBeamTrailVertexFactory.cpp
--- a/Engine/Source/Runtime/Engine/Private/Particles/ParticleBeamTrailVertexFactory.cpp
+++ b/Engine/Source/Runtime/Engine/Private/Particles/ParticleBeamTrailVertexFactory.cpp
@@ -125,6 +125,11 @@ void FParticleBeamTrailVertexFactory::GetPSOPrecacheVertexFetchElements(EVertexI
 	GParticleBeamTrailVertexDeclaration.VertexDeclarationRHI->GetInitializer(Elements);
 }
 
+FRHIVertexDeclaration* FParticleBeamTrailVertexFactory::GetPSOPrecacheVertexDeclaration(bool bUsesDynamicParameter)
+{
+	return (bUsesDynamicParameter ? GParticleBeamTrailVertexDeclarationDynamic.VertexDeclarationRHI : GParticleBeamTrailVertexDeclaration.VertexDeclarationRHI);
+}
+
 /**
  *	Initialize the Render Hardware Interface for this vertex factory
  */
diff --git a/Engine/Source/Runtime/Engine/Private/Particles/ParticleComponents.cpp b/Engine/Source/Runtime/Engine/Private/Particles/ParticleComponents.cpp
--- a/Engine/Source/Runtime/Engine/Private/Particles/ParticleComponents.cpp
+++ b/Engine/Source/Runtime/Engine/Private/Particles/ParticleComponents.cpp
@@ -154,9 +154,12 @@ void UFXSystemAsset::LaunchPSOPrecaching(TArrayView<VFsPerMaterialData> VFsPerMa
 	{
 		if (VFsPerMaterial.MaterialInterface)
 		{
+			// Material has to be fully loaded for a static permutations to work
+			VFsPerMaterial.MaterialInterface->ConditionalPostLoad();
+			
 			PreCachePSOParams.PrimitiveType = (EPrimitiveType)VFsPerMaterial.PrimitiveType;
 			PreCachePSOParams.bDisableBackFaceCulling = VFsPerMaterial.bDisableBackfaceCulling;
-			PrecachePSOsEvents.Append(VFsPerMaterial.MaterialInterface->PrecachePSOs(VFsPerMaterial.VertexFactoryTypes, PreCachePSOParams,EPSOPrecachePriority::Medium, MaterialPSOPrecacheRequestIDs));
+			PrecachePSOsEvents.Append(VFsPerMaterial.MaterialInterface->PrecachePSOs(VFsPerMaterial.VertexFactoryData, PreCachePSOParams,EPSOPrecachePriority::Medium, MaterialPSOPrecacheRequestIDs));
 		}
 	}
 }
@@ -2654,7 +2657,7 @@ void UParticleSystem::PrecachePSOs()
 		return;
 	}
 
-	TArray<VFsPerMaterialData, TInlineAllocator<2>> VFsPerMaterials;
+	TArray<VFsPerMaterialData, TInlineAllocator<4>> VFsPerMaterials;
 
 	// No per component emitter materials known at this point in time
 	TArray<UMaterialInterface*> EmptyEmitterMaterials;
@@ -2674,15 +2677,28 @@ void UParticleSystem::PrecachePSOs()
 			const UParticleLODLevel* LOD = Emitter->LODLevels[LodIndex];
 			if (LOD && LOD->bEnabled)
 			{
-				const FVertexFactoryType* VFType = LOD->TypeDataModule ? LOD->TypeDataModule->GetVertexFactoryType() : &FParticleSpriteVertexFactory::StaticType;
-				check(VFType);
-				EPrimitiveType PrimitiveType = LOD->TypeDataModule ? LOD->TypeDataModule->GetPrimitiveType() : PT_TriangleList;
+				UParticleModuleTypeDataBase::FPSOPrecacheParams PrecacheParams;
+				if (LOD->TypeDataModule)
+				{
+					LOD->TypeDataModule->CollectPSOPrecacheData(Emitter, PrecacheParams);
+				}
+				else
+				{
+					bool bUsesDynamicParameter = (Emitter->DynamicParameterDataOffset > 0);
+					FPSOPrecacheVertexFactoryData VFData;
+					VFData.VertexFactoryType = &FParticleSpriteVertexFactory::StaticType;
+					VFData.CustomDefaultVertexDeclaration = FParticleSpriteVertexFactory::GetPSOPrecacheVertexDeclaration(bUsesDynamicParameter);
+					PrecacheParams.VertexFactoryDataList.Add(VFData);
+					PrecacheParams.PrimitiveType = PT_TriangleList;
+				}
 
 				Materials.Empty();
 				LOD->GetUsedMaterials(Materials, NamedMaterialSlots, EmptyEmitterMaterials);
 
 				for (UMaterialInterface* MaterialInterface : Materials)
 				{
+					EPrimitiveType PrimitiveType = PrecacheParams.PrimitiveType;
+
 					VFsPerMaterialData* VFsPerMaterial = VFsPerMaterials.FindByPredicate([MaterialInterface, PrimitiveType](const VFsPerMaterialData& Other)
 						{
 							return Other.MaterialInterface == MaterialInterface && Other.PrimitiveType == PrimitiveType;
@@ -2693,7 +2709,11 @@ void UParticleSystem::PrecachePSOs()
 						VFsPerMaterial->MaterialInterface = MaterialInterface;
 						VFsPerMaterial->PrimitiveType = PrimitiveType;
 					}
-					VFsPerMaterial->VertexFactoryTypes.AddUnique(VFType);
+					
+					for (FPSOPrecacheVertexFactoryData VFData : PrecacheParams.VertexFactoryDataList)
+					{
+						VFsPerMaterial->VertexFactoryData.AddUnique(VFData);
+					}
 				}
 			}
 		}
diff --git a/Engine/Source/Runtime/Engine/Private/Particles/ParticleModules.cpp b/Engine/Source/Runtime/Engine/Private/Particles/ParticleModules.cpp
--- a/Engine/Source/Runtime/Engine/Private/Particles/ParticleModules.cpp
+++ b/Engine/Source/Runtime/Engine/Private/Particles/ParticleModules.cpp
@@ -3420,6 +3420,32 @@ const FVertexFactoryType* UParticleModuleTypeDataMesh::GetVertexFactoryType() co
 	return &FMeshParticleVertexFactory::StaticType;
 }
 
+extern void InitMeshParticleVertexFactoryComponents(FMeshParticleVertexFactory* InVertexFactory, const FStaticMeshLODResources& LODResources, FMeshParticleVertexFactory::FDataType& Data);
+
+void UParticleModuleTypeDataMesh::CollectPSOPrecacheData(const UParticleEmitter* Emitter, FPSOPrecacheParams& OutParams)
+{
+	if (Mesh != nullptr)
+	{
+		FStaticMeshRenderData* RenderData = Mesh->GetRenderData();
+		// Assuming here that all LOD use same vertex decl
+		int32 MeshLODIdx = Mesh->GetMinLODIdx();
+		if (RenderData->LODResources.IsValidIndex(MeshLODIdx))
+		{
+			bool bUsesDynamicParameter = (Emitter->DynamicParameterDataOffset > 0);
+			int32 DynamicVertexStride = sizeof(FMeshParticleInstanceVertex);
+			int32 DynamicParameterVertexStride = bUsesDynamicParameter ? sizeof(FMeshParticleInstanceVertexDynamicParameter) : 0;
+						
+			FVertexDeclarationElementList Elements;
+			FMeshParticleVertexFactory::FDataType Data;
+			InitMeshParticleVertexFactoryComponents(nullptr, RenderData->LODResources[MeshLODIdx], Data);
+			FMeshParticleVertexFactory::GetVertexElements(GMaxRHIFeatureLevel, DynamicVertexStride, DynamicParameterVertexStride, Data, Elements);
+			const FVertexFactoryType* VFType = &FMeshParticleVertexFactory::StaticType;
+			OutParams.PrimitiveType = GetPrimitiveType();
+			OutParams.VertexFactoryDataList.Add(FPSOPrecacheVertexFactoryData(VFType, Elements));
+		}
+	}
+}
+
 void UParticleModuleTypeDataMesh::SetToSensibleDefaults(UParticleEmitter* Owner)
 {
 	if ((Mesh == NULL) && GIsEditor )
@@ -4938,6 +4964,12 @@ const FVertexFactoryType* UParticleModuleTypeDataGpu::GetVertexFactoryType() con
 	return &FGPUSpriteVertexFactory::StaticType;
 }
 
+void UParticleModuleTypeDataGpu::CollectPSOPrecacheData(const UParticleEmitter* Emitter, FPSOPrecacheParams& OutParams)
+{
+	OutParams.VertexFactoryDataList.Add(FPSOPrecacheVertexFactoryData(GetVertexFactoryType()));
+	OutParams.PrimitiveType = GetPrimitiveType();
+}
+
 /*-----------------------------------------------------------------------------
 	UParticleModulePivotOffset implementation.
 -----------------------------------------------------------------------------*/
diff --git a/Engine/Source/Runtime/Engine/Private/Particles/ParticleSystemRender.cpp b/Engine/Source/Runtime/Engine/Private/Particles/ParticleSystemRender.cpp
--- a/Engine/Source/Runtime/Engine/Private/Particles/ParticleSystemRender.cpp
+++ b/Engine/Source/Runtime/Engine/Private/Particles/ParticleSystemRender.cpp
@@ -2189,66 +2189,70 @@ void FDynamicMeshEmitterData::GetInstanceData(void* InstanceData, void* DynamicP
 	}
 }
 
-void FDynamicMeshEmitterData::SetupVertexFactory( FMeshParticleVertexFactory* InVertexFactory, const FStaticMeshLODResources& LODResources, uint32 LODIdx) const
+void InitMeshParticleVertexFactoryComponents(FMeshParticleVertexFactory* InVertexFactory, const FStaticMeshLODResources& LODResources, FMeshParticleVertexFactory::FDataType& Data)
 {
-		FMeshParticleVertexFactory::FDataType Data;
-
-		LODResources.VertexBuffers.PositionVertexBuffer.BindPositionVertexBuffer(InVertexFactory, Data);
-		LODResources.VertexBuffers.StaticMeshVertexBuffer.BindTangentVertexBuffer(InVertexFactory, Data);
-		LODResources.VertexBuffers.StaticMeshVertexBuffer.BindTexCoordVertexBuffer(InVertexFactory, Data, MAX_TEXCOORDS);
-		LODResources.VertexBuffers.ColorVertexBuffer.BindColorVertexBuffer(InVertexFactory, Data);
+	LODResources.VertexBuffers.PositionVertexBuffer.BindPositionVertexBuffer(InVertexFactory, Data);
+	LODResources.VertexBuffers.StaticMeshVertexBuffer.BindTangentVertexBuffer(InVertexFactory, Data);
+	LODResources.VertexBuffers.StaticMeshVertexBuffer.BindTexCoordVertexBuffer(InVertexFactory, Data, MAX_TEXCOORDS);
+	LODResources.VertexBuffers.ColorVertexBuffer.BindColorVertexBuffer(InVertexFactory, Data);
+
+	// Initialize instanced data. Vertex buffer and stride are set before render.
+	// Particle color
+	Data.ParticleColorComponent = FVertexStreamComponent(
+		NULL,
+		STRUCT_OFFSET(FMeshParticleInstanceVertex, Color),
+		0,
+		VET_Float4,
+		EVertexStreamUsage::Instancing
+	);
 
-		// Initialize instanced data. Vertex buffer and stride are set before render.
-		// Particle color
-		Data.ParticleColorComponent = FVertexStreamComponent(
+	// Particle transform matrix
+	for (int32 MatrixRow = 0; MatrixRow < 3; MatrixRow++)
+	{
+		Data.TransformComponent[MatrixRow] = FVertexStreamComponent(
 			NULL,
-			STRUCT_OFFSET(FMeshParticleInstanceVertex, Color),
+			STRUCT_OFFSET(FMeshParticleInstanceVertex, Transform) + sizeof(FVector4f) * MatrixRow, 
 			0,
 			VET_Float4,
 			EVertexStreamUsage::Instancing
-			);
+		);
+	}
 
-		// Particle transform matrix
-		for (int32 MatrixRow = 0; MatrixRow < 3; MatrixRow++)
-		{
-			Data.TransformComponent[MatrixRow] = FVertexStreamComponent(
-				NULL,
-				STRUCT_OFFSET(FMeshParticleInstanceVertex, Transform) + sizeof(FVector4f) * MatrixRow, 
-				0,
-				VET_Float4,
-				EVertexStreamUsage::Instancing
-				);
-		}
+	Data.VelocityComponent = FVertexStreamComponent(
+		NULL,
+		STRUCT_OFFSET(FMeshParticleInstanceVertex,Velocity),
+		0,
+		VET_Float4,
+		EVertexStreamUsage::Instancing
+	);
 
-		Data.VelocityComponent = FVertexStreamComponent(
-			NULL,
-			STRUCT_OFFSET(FMeshParticleInstanceVertex,Velocity),
-			0,
-			VET_Float4,
-			EVertexStreamUsage::Instancing
-			);
+	// SubUVs.
+	Data.SubUVs = FVertexStreamComponent(
+		NULL,
+		STRUCT_OFFSET(FMeshParticleInstanceVertex, SubUVParams), 
+		0,
+		VET_Short4,
+		EVertexStreamUsage::Instancing
+	);
 
-		// SubUVs.
-		Data.SubUVs = FVertexStreamComponent(
-			NULL,
-			STRUCT_OFFSET(FMeshParticleInstanceVertex, SubUVParams), 
-			0,
-			VET_Short4,
-			EVertexStreamUsage::Instancing
-			);
+	// Pack SubUV Lerp and the particle's relative time
+	Data.SubUVLerpAndRelTime = FVertexStreamComponent(
+		NULL,
+		STRUCT_OFFSET(FMeshParticleInstanceVertex, SubUVLerp), 
+		0,
+		VET_Float2,
+		EVertexStreamUsage::Instancing
+	);
 
-		// Pack SubUV Lerp and the particle's relative time
-		Data.SubUVLerpAndRelTime = FVertexStreamComponent(
-			NULL,
-			STRUCT_OFFSET(FMeshParticleInstanceVertex, SubUVLerp), 
-			0,
-			VET_Float2,
-			EVertexStreamUsage::Instancing
-			);
+	Data.bInitialized = true;
+}
 
-		Data.bInitialized = true;
-		InVertexFactory->SetData(Data);
-		InVertexFactory->SetLODIdx((uint8)LODIdx);
+void FDynamicMeshEmitterData::SetupVertexFactory( FMeshParticleVertexFactory* InVertexFactory, const FStaticMeshLODResources& LODResources, uint32 LODIdx) const
+{
+	FMeshParticleVertexFactory::FDataType Data;
+	InitMeshParticleVertexFactoryComponents(InVertexFactory, LODResources, Data);
+	InVertexFactory->SetData(Data);
+	InVertexFactory->SetLODIdx((uint8)LODIdx);
 }
 
 ///////////////////////////////////////////////////////////////////////////////
diff --git a/Engine/Source/Runtime/Engine/Private/Particles/ParticleTrailModules.cpp b/Engine/Source/Runtime/Engine/Private/Particles/ParticleTrailModules.cpp
--- a/Engine/Source/Runtime/Engine/Private/Particles/ParticleTrailModules.cpp
+++ b/Engine/Source/Runtime/Engine/Private/Particles/ParticleTrailModules.cpp
@@ -15,6 +15,7 @@
 #include "Particles/TypeData/ParticleModuleTypeDataRibbon.h"
 #include "ParticleBeamTrailVertexFactory.h"
 #include "UObject/UnrealType.h"
+#include "Particles/ParticleEmitter.h"
 
 UParticleModuleTrailBase::UParticleModuleTrailBase(const FObjectInitializer& ObjectInitializer)
 	: Super(ObjectInitializer)
@@ -244,6 +245,17 @@ const FVertexFactoryType* UParticleModuleTypeDataRibbon::GetVertexFactoryType()
 	return &FParticleBeamTrailVertexFactory::StaticType;
 }
 
+void UParticleModuleTypeDataRibbon::CollectPSOPrecacheData(const UParticleEmitter* Emitter, FPSOPrecacheParams& OutParams)
+{
+	bool bUsesDynamicParameter = (Emitter->DynamicParameterDataOffset > 0);
+
+	FPSOPrecacheVertexFactoryData VFData;
+	VFData.VertexFactoryType = &FParticleBeamTrailVertexFactory::StaticType;
+	VFData.CustomDefaultVertexDeclaration = FParticleBeamTrailVertexFactory::GetPSOPrecacheVertexDeclaration(bUsesDynamicParameter);
+	OutParams.VertexFactoryDataList.Add(VFData);
+	OutParams.PrimitiveType = GetPrimitiveType();
+}
+
 /*-----------------------------------------------------------------------------
 	UParticleModuleTypeDataAnimTrail implementation.
 -----------------------------------------------------------------------------*/
@@ -274,3 +286,14 @@ const FVertexFactoryType* UParticleModuleTypeDataAnimTrail::GetVertexFactoryType
 {
 	return &FParticleBeamTrailVertexFactory::StaticType;
 }
+
+void UParticleModuleTypeDataAnimTrail::CollectPSOPrecacheData(const UParticleEmitter* Emitter, FPSOPrecacheParams& OutParams)
+{
+	bool bUsesDynamicParameter = (Emitter->DynamicParameterDataOffset > 0);
+
+	FPSOPrecacheVertexFactoryData VFData;
+	VFData.VertexFactoryType = &FParticleBeamTrailVertexFactory::StaticType;
+	VFData.CustomDefaultVertexDeclaration = FParticleBeamTrailVertexFactory::GetPSOPrecacheVertexDeclaration(bUsesDynamicParameter);
+	OutParams.VertexFactoryDataList.Add(VFData);
+	OutParams.PrimitiveType = GetPrimitiveType();
+}
diff --git a/Engine/Source/Runtime/Engine/Private/Particles/ParticleVertexFactory.cpp b/Engine/Source/Runtime/Engine/Private/Particles/ParticleVertexFactory.cpp
--- a/Engine/Source/Runtime/Engine/Private/Particles/ParticleVertexFactory.cpp
+++ b/Engine/Source/Runtime/Engine/Private/Particles/ParticleVertexFactory.cpp
@@ -243,6 +243,11 @@ void FParticleSpriteVertexFactory::GetPSOPrecacheVertexFetchElements(EVertexInpu
 	GParticleSpriteVertexDeclarationInstanced.VertexDeclarationRHI->GetInitializer(Elements);
 }
 
+FRHIVertexDeclaration* FParticleSpriteVertexFactory::GetPSOPrecacheVertexDeclaration(bool bUsesDynamicParameter)
+{
+	return GetParticleSpriteVertexDeclaration(4, bUsesDynamicParameter).VertexDeclarationRHI;
+}
+
 /**
  *	Initialize the Render Hardware Interface for this vertex factory
  */
diff --git a/Engine/Source/Runtime/Engine/Public/MeshParticleVertexFactory.h b/Engine/Source/Runtime/Engine/Public/MeshParticleVertexFactory.h
--- a/Engine/Source/Runtime/Engine/Public/MeshParticleVertexFactory.h
+++ b/Engine/Source/Runtime/Engine/Public/MeshParticleVertexFactory.h
@@ -118,6 +118,7 @@ public:
 	 * Get vertex elements used when during PSO precaching materials using this vertex factory type
 	 */
 	static void GetPSOPrecacheVertexFetchElements(EVertexInputStreamType VertexInputStreamType, FVertexDeclarationElementList& Elements);
+	static void GetVertexElements(ERHIFeatureLevel::Type FeatureLevel, int32 InDynamicVertexStride, int32 InDynamicParameterVertexStride, FDataType& Data, FVertexDeclarationElementList& Elements);
 
 	/**
 	 * An implementation of the interface used by TSynchronizedResource to update the resource with new data from the game thread.
@@ -187,6 +188,8 @@ public:
 	{
 		return LODIdx;
 	}
+protected:
+	static void GetVertexElements(ERHIFeatureLevel::Type FeatureLevel, int32 InDynamicVertexStride, int32 InDynamicParameterVertexStride, FDataType& Data, FVertexDeclarationElementList& Elements, FVertexStreamList& InOutStreams);
 
 protected:
 	FDataType Data;
diff --git a/Engine/Source/Runtime/Engine/Public/PSOPrecache.h b/Engine/Source/Runtime/Engine/Public/PSOPrecache.h
--- a/Engine/Source/Runtime/Engine/Public/PSOPrecache.h
+++ b/Engine/Source/Runtime/Engine/Public/PSOPrecache.h
@@ -134,7 +134,7 @@ struct FPSOPrecacheVertexFactoryData
 {
 	FPSOPrecacheVertexFactoryData() = default;
 	FPSOPrecacheVertexFactoryData(const FVertexFactoryType* InVertexFactoryType) : VertexFactoryType(InVertexFactoryType), CustomDefaultVertexDeclaration(nullptr) {}
-	FPSOPrecacheVertexFactoryData(const FVertexFactoryType* InVertexFactoryType, const FVertexDeclarationElementList& ElementList);
+	ENGINE_API FPSOPrecacheVertexFactoryData(const FVertexFactoryType* InVertexFactoryType, const FVertexDeclarationElementList& ElementList);
 
 	const FVertexFactoryType* VertexFactoryType = nullptr;
 
diff --git a/Engine/Source/Runtime/Engine/Public/ParticleBeamTrailVertexFactory.h b/Engine/Source/Runtime/Engine/Public/ParticleBeamTrailVertexFactory.h
--- a/Engine/Source/Runtime/Engine/Public/ParticleBeamTrailVertexFactory.h
+++ b/Engine/Source/Runtime/Engine/Public/ParticleBeamTrailVertexFactory.h
@@ -66,6 +66,7 @@ public:
 	 * Get vertex elements used when during PSO precaching materials using this vertex factory type
 	 */
 	static void GetPSOPrecacheVertexFetchElements(EVertexInputStreamType VertexInputStreamType, FVertexDeclarationElementList& Elements);
+	static FRHIVertexDeclaration* GetPSOPrecacheVertexDeclaration(bool bUsesDynamicParameter);
 
 	// FRenderResource interface.
 	virtual void InitRHI() override;
diff --git a/Engine/Source/Runtime/Engine/Public/ParticleVertexFactory.h b/Engine/Source/Runtime/Engine/Public/ParticleVertexFactory.h
--- a/Engine/Source/Runtime/Engine/Public/ParticleVertexFactory.h
+++ b/Engine/Source/Runtime/Engine/Public/ParticleVertexFactory.h
@@ -138,6 +138,7 @@ public:
 	 * Get vertex elements used when during PSO precaching materials using this vertex factory type
 	 */
 	static void GetPSOPrecacheVertexFetchElements(EVertexInputStreamType VertexInputStreamType, FVertexDeclarationElementList& Elements);
+	static FRHIVertexDeclaration* GetPSOPrecacheVertexDeclaration(bool bUsesDynamicParameter);
 
 	/**
 	 * Set the source vertex buffer that contains particle instance data.
