Fixed: Build Error ParticleSystem.h

diff --git a/Engine/Source/Runtime/Engine/Classes/Particles/ParticleSystem.h b/Engine/Source/Runtime/Engine/Classes/Particles/ParticleSystem.h
--- a/Engine/Source/Runtime/Engine/Classes/Particles/ParticleSystem.h
+++ b/Engine/Source/Runtime/Engine/Classes/Particles/ParticleSystem.h
@@ -19,6 +19,8 @@ class UInterpCurveEdSetup;
 class UMaterialInterface;
 class UParticleSystemComponent;
 class UParticleEmitter;
+struct FPSOPrecacheVertexFactoryData;
+typedef TArray<FPSOPrecacheVertexFactoryData, TInlineAllocator<2> > FPSOPrecacheVertexFactoryDataList;
 
 /**
  *	ParticleSystemUpdateMode
