WorldPartitionSubsystem is now responsible of sorting cells and executing cell loading & activation for all registered world partitions.
- World Partition now relies on World's NumStreamingLevelsBeingLoaded instead of computing its own number of loading cells
- Fixed pending cells to be loaded and to be added to the world not updating their cell sorting info
- Fixed UWorld::CanAddLoadedLevelToWorld to work with multiple world partitions
- UWorldPartitionRuntimeCell::IsLoading, UWorldPartition::UpdateStreamingState and UWorldPartition::CanAddLoadedLevelToWorld are now deprecated

diff --git a/Engine/Source/Runtime/Engine/Private/World.cpp b/Engine/Source/Runtime/Engine/Private/World.cpp
--- a/Engine/Source/Runtime/Engine/Private/World.cpp
+++ b/Engine/Source/Runtime/Engine/Private/World.cpp
@@ -4164,9 +4164,13 @@ bool UWorld::CanAddLoadedLevelToWorld(ULevel* Level) const
 	if (CurrentLevelPendingVisibility == nullptr)
 	{
 		// Allow world partition to decide wether a level should be added or not to the world
-		if (const UWorldPartition* WorldPartition = GetWorldPartition())
+		if (const IWorldPartitionCell* Cell = Level->GetWorldPartitionRuntimeCell())
 		{
-			return WorldPartition->CanAddLoadedLevelToWorld(Level);
+			const UWorld* OuterWorld = Cell->GetOuterWorld();
+			if (UWorldPartition* OuterWorldPartition = OuterWorld ? OuterWorld->GetWorldPartition() : nullptr)
+			{
+				return OuterWorldPartition->CanAddCellToWorld(Cell);
+			}
 		}
 	}
 	return true;
diff --git a/Engine/Source/Runtime/Engine/Private/WorldPartition/WorldPartition.cpp b/Engine/Source/Runtime/Engine/Private/WorldPartition/WorldPartition.cpp
--- a/Engine/Source/Runtime/Engine/Private/WorldPartition/WorldPartition.cpp
+++ b/Engine/Source/Runtime/Engine/Private/WorldPartition/WorldPartition.cpp
@@ -11,6 +11,7 @@
 #include "Misc/StringFormatArg.h"
 #include "WorldPartition/WorldPartitionLevelStreamingPolicy.h"
 #include "WorldPartition/WorldPartitionReplay.h"
+#include "WorldPartition/WorldPartitionSubsystem.h"
 #include "WorldPartition/HLOD/HLODSubsystem.h"
 #include "WorldPartition/DataLayer/DataLayerManager.h"
 #include "GameFramework/WorldSettings.h"
@@ -657,6 +658,8 @@ void UWorldPartition::Initialize(UWorld* InWorld, const FTransform& InTransform)
 
 void UWorldPartition::Uninitialize()
 {
+	TRACE_CPUPROFILER_EVENT_SCOPE(UWorldPartition::Uninitialize);
+
 	if (IsInitialized())
 	{
 		check(World);
@@ -673,7 +676,7 @@ void UWorldPartition::Uninitialize()
 		// Unload all loaded cells
 		if (World->IsGameWorld())
 		{
-			UpdateStreamingState();
+			UWorldPartitionSubsystem::UpdateStreamingStateInternal(World, { this });
 		}
 
 #if WITH_EDITOR
@@ -1458,21 +1461,6 @@ void UWorldPartition::Tick(float DeltaSeconds)
 #endif
 }
 
-DECLARE_CYCLE_STAT(TEXT("World Partition Update Streaming"), STAT_WorldPartitionUpdateStreaming, STATGROUP_Engine);
-
-void UWorldPartition::UpdateStreamingState()
-{
-	SCOPE_CYCLE_COUNTER(STAT_WorldPartitionUpdateStreaming);
-
-	if (GetWorld()->IsGameWorld())
-	{
-		if (StreamingPolicy)
-		{
-			StreamingPolicy->UpdateStreamingState();
-		}
-	}
-}
-
 bool UWorldPartition::InjectExternalStreamingObject(URuntimeHashExternalStreamingObjectBase* InExternalStreamingObject)
 {
 	bool bInjected = RuntimeHash->InjectExternalStreamingObject(InExternalStreamingObject);
@@ -1515,11 +1503,14 @@ bool UWorldPartition::GetIntersectingCells(const TArray<FWorldPartitionStreaming
 	return false;
 }
 
-bool UWorldPartition::CanAddLoadedLevelToWorld(class ULevel* InLevel) const
+bool UWorldPartition::CanAddCellToWorld(const IWorldPartitionCell* InCell) const
 {
 	if (GetWorld()->IsGameWorld() && StreamingPolicy)
 	{
-		return StreamingPolicy->CanAddLoadedLevelToWorld(InLevel);
+		if (const UWorldPartitionRuntimeCell* Cell = Cast<const UWorldPartitionRuntimeCell>(InCell))
+		{
+			return StreamingPolicy->CanAddCellToWorld(Cell);
+		}
 	}
 	return true;
 }
diff --git a/Engine/Source/Runtime/Engine/Private/WorldPartition/WorldPartitionLevelStreamingPolicy.cpp b/Engine/Source/Runtime/Engine/Private/WorldPartition/WorldPartitionLevelStreamingPolicy.cpp
--- a/Engine/Source/Runtime/Engine/Private/WorldPartition/WorldPartitionLevelStreamingPolicy.cpp
+++ b/Engine/Source/Runtime/Engine/Private/WorldPartition/WorldPartitionLevelStreamingPolicy.cpp
@@ -19,20 +19,6 @@
 
 #include UE_INLINE_GENERATED_CPP_BY_NAME(WorldPartitionLevelStreamingPolicy)
 
-int32 UWorldPartitionLevelStreamingPolicy::GetCellLoadingCount() const
-{
-	int32 CellLoadingCount = 0;
-
-	ForEachActiveRuntimeCell([&CellLoadingCount](const UWorldPartitionRuntimeCell* Cell)
-	{
-		if (Cell->IsLoading())
-		{
-			++CellLoadingCount;
-		}
-	});
-	return CellLoadingCount;
-}
-
 void UWorldPartitionLevelStreamingPolicy::ForEachActiveRuntimeCell(TFunctionRef<void(const UWorldPartitionRuntimeCell*)> Func) const
 {
 	UWorld* World = WorldPartition->GetWorld();
diff --git a/Engine/Source/Runtime/Engine/Private/WorldPartition/WorldPartitionRuntimeLevelStreamingCell.cpp b/Engine/Source/Runtime/Engine/Private/WorldPartition/WorldPartitionRuntimeLevelStreamingCell.cpp
--- a/Engine/Source/Runtime/Engine/Private/WorldPartition/WorldPartitionRuntimeLevelStreamingCell.cpp
+++ b/Engine/Source/Runtime/Engine/Private/WorldPartition/WorldPartitionRuntimeLevelStreamingCell.cpp
@@ -193,16 +193,6 @@ EStreamingStatus UWorldPartitionRuntimeLevelStreamingCell::GetStreamingStatus()
 	return Super::GetStreamingStatus();
 }
 
-bool UWorldPartitionRuntimeLevelStreamingCell::IsLoading() const
-{
-	if (LevelStreaming)
-	{
-		ELevelStreamingState CurrentState = LevelStreaming->GetLevelStreamingState();
-		return (CurrentState == ELevelStreamingState::Removed || CurrentState == ELevelStreamingState::Unloaded || CurrentState == ELevelStreamingState::Loading);
-	}
-	return Super::IsLoading();
-}
-
 FLinearColor UWorldPartitionRuntimeLevelStreamingCell::GetDebugColor(EWorldPartitionRuntimeCellVisualizeMode VisualizeMode) const
 {
 	switch (VisualizeMode)
diff --git a/Engine/Source/Runtime/Engine/Private/WorldPartition/WorldPartitionStreamingPolicy.cpp b/Engine/Source/Runtime/Engine/Private/WorldPartition/WorldPartitionStreamingPolicy.cpp
--- a/Engine/Source/Runtime/Engine/Private/WorldPartition/WorldPartitionStreamingPolicy.cpp
+++ b/Engine/Source/Runtime/Engine/Private/WorldPartition/WorldPartitionStreamingPolicy.cpp
@@ -25,12 +25,6 @@
 
 #define LOCTEXT_NAMESPACE "WorldPartitionStreamingPolicy"
 
-static int32 GMaxLoadingStreamingCells = 4;
-static FAutoConsoleVariableRef CMaxLoadingStreamingCells(
-	TEXT("wp.Runtime.MaxLoadingStreamingCells"),
-	GMaxLoadingStreamingCells,
-	TEXT("Used to limit the number of concurrent loading world partition streaming cells."));
-
 int32 GBlockOnSlowStreaming = 1;
 static FAutoConsoleVariableRef CVarBlockOnSlowStreaming(
 	TEXT("wp.Runtime.BlockOnSlowStreaming"),
@@ -70,7 +64,9 @@ static void SortStreamingCellsByImportance(TArray<const UWorldPartitionRuntimeCe
 
 UWorldPartitionStreamingPolicy::UWorldPartitionStreamingPolicy(const FObjectInitializer& ObjectInitializer)
 	: Super(ObjectInitializer) 
-	, bLastUpdateCompletedLoadingAndActivation(false)
+	, WorldPartition(nullptr)
+	, ProcessedToActivateCells(0)
+	, ProcessedToLoadCells(0)
 	, bCriticalPerformanceRequestedBlockTillOnWorld(false)
 	, CriticalPerformanceBlockTillLevelStreamingCompletedEpoch(0)
 	, ServerDataLayersStatesEpoch(INT_MIN)
@@ -262,6 +258,12 @@ void UWorldPartitionStreamingPolicy::UpdateStreamingState()
 	check(World);
 	check(World->IsGameWorld());
 
+	const bool bLastUpdateCompletedLoadingAndActivation = ((ProcessedToActivateCells + ProcessedToLoadCells) == (ToActivateCells.Num() + ToLoadCells.Num()));
+	ProcessedToActivateCells = 0;
+	ProcessedToLoadCells = 0;
+	ToActivateCells.Reset();
+	ToLoadCells.Reset();
+
 	// Dermine if the World's BlockTillLevelStreamingCompleted was triggered by WorldPartitionStreamingPolicy
 	if (bCriticalPerformanceRequestedBlockTillOnWorld && IsInBlockTillLevelStreamingCompleted())
 	{
@@ -473,29 +475,42 @@ void UWorldPartitionStreamingPolicy::UpdateStreamingState()
 	FrameLoadCells = FrameLoadCells.Difference(FrameActivateCells);
 
 	// Determine cells to activate (sorted by importance)
-	TArray<const UWorldPartitionRuntimeCell*> ToActivateCells;
 	if (bCanStream)
 	{
 		TRACE_CPUPROFILER_EVENT_SCOPE(UWorldPartitionStreamingPolicy::UpdateStreamingState_ToActivateCells);
 		for (const UWorldPartitionRuntimeCell* Cell : FrameActivateCells)
 		{
-			if (!ActivatedCells.Contains(Cell) && !ShouldSkipCellForPerformance(Cell) && !ShouldSkipDisabledHLODCell(Cell))
+			if (ActivatedCells.Contains(Cell))
+			{
+				// Update streaming source info for pending add to world cells
+				if (ActivatedCells.GetPendingAddToWorldCells().Contains(Cell))
+				{
+					Cell->MergeStreamingSourceInfo();
+				}
+			}
+			else if (!ShouldSkipCellForPerformance(Cell) && !ShouldSkipDisabledHLODCell(Cell))
 			{
 				Cell->MergeStreamingSourceInfo();
 				ToActivateCells.Add(Cell);
 			}
 		}
-		SortStreamingCellsByImportance(ToActivateCells);
 	}
 
 	// Determine cells to load (sorted by importance)
-	TArray<const UWorldPartitionRuntimeCell*> ToLoadCells;
 	if (bCanStream)
 	{
 		TRACE_CPUPROFILER_EVENT_SCOPE(UWorldPartitionStreamingPolicy::UpdateStreamingState_ToLoadCells);
 		for (const UWorldPartitionRuntimeCell* Cell : FrameLoadCells)
 		{
-			if (!LoadedCells.Contains(Cell))
+			if (LoadedCells.Contains(Cell))
+			{
+				// Update streaming source info for pending load cells
+				if (!Cell->GetLevel())
+				{
+					Cell->MergeStreamingSourceInfo();
+				}
+			}
+			else
 			{
 				// Only deactivated server cells need to call ShouldWaitForClientVisibility (those part of ActivatedCells)
 				const bool bIsServerCellToDeactivate = bIsServer && ActivatedCells.Contains(Cell);
@@ -507,7 +522,6 @@ void UWorldPartitionStreamingPolicy::UpdateStreamingState()
 				}
 			}
 		}
-		SortStreamingCellsByImportance(ToLoadCells);
 	}
 
 	// Determine cells to unload
@@ -550,43 +564,6 @@ void UWorldPartitionStreamingPolicy::UpdateStreamingState()
 		SetCellsStateToUnloaded(ToUnloadCells);
 	}
 
-	int32 ToLoadAndActivateCount = 0;
-
-	if (bIsStreamingInEnabled)
-	{
-		// Do Activation State first as it is higher prio than Load State (if we have a limited number of loading cells per frame)
-		if (ToActivateCells.Num() > 0)
-		{
-			ToLoadAndActivateCount += SetCellsStateToActivated(ToActivateCells);
-		}
-
-		if (ToLoadCells.Num() > 0)
-		{
-			ToLoadAndActivateCount += SetCellsStateToLoaded(ToLoadCells);
-		}
-	}
-
-	bLastUpdateCompletedLoadingAndActivation = (ToLoadAndActivateCount == (ToActivateCells.Num() + ToLoadCells.Num()));
-
-	// Sort cells and update streaming priority 
-	{
-		TRACE_CPUPROFILER_EVENT_SCOPE(UWorldPartitionStreamingPolicy::SortCellsAndUpdateStreamingPriority);
-		SortedAddToWorldCells = ActivatedCells.GetPendingAddToWorldCells().Array();
-		SortStreamingCellsByImportance(SortedAddToWorldCells);
-
-		// Update level streaming priority so that UWorld::UpdateLevelStreaming will naturally process the levels in the correct order
-		const int32 MaxPrio = SortedAddToWorldCells.Num();
-		int32 Prio = MaxPrio;
-		const ULevel* CurrentPendingVisibility = GetWorld()->GetCurrentLevelPendingVisibility();
-		for (const UWorldPartitionRuntimeCell* Cell : SortedAddToWorldCells)
-		{
-			// Current pending visibility level is the most important
-			const bool bIsCellPendingVisibility = CurrentPendingVisibility && (Cell->GetLevel() == CurrentPendingVisibility);
-			const int32 SortedPriority = bIsCellPendingVisibility ? MaxPrio + 1 : Prio--;
-			Cell->SetStreamingPriority(SortedPriority);
-		}
-	}
-
 	// Evaluate streaming performance based on cells that should be activated
 	UpdateStreamingPerformance(FrameActivateCells);
 	
@@ -715,91 +692,95 @@ bool UWorldPartitionStreamingPolicy::ShouldSkipCellForPerformance(const UWorldPa
 	return false;
 }
 
-int32 UWorldPartitionStreamingPolicy::GetMaxCellsToLoad() const
+void UWorldPartitionStreamingPolicy::GetCellsToUpdate(TArray<const UWorldPartitionRuntimeCell*>& OutToLoadCells, TArray<const UWorldPartitionRuntimeCell*>& OutToActivateCells)
+{
+	OutToLoadCells.Append(ToLoadCells);
+	OutToActivateCells.Append(ToActivateCells);
+}
+
+void UWorldPartitionStreamingPolicy::GetCellsToReprioritize(TArray<const UWorldPartitionRuntimeCell*>& OutToReprioritizeLoadCells, TArray<const UWorldPartitionRuntimeCell*>& OutToReprioritizeActivateCells)
 {
-	// This policy limits the number of concurrent loading streaming cells, except if match hasn't started
-	if (!WorldPartition->IsServer())
+	for (const UWorldPartitionRuntimeCell* Cell : LoadedCells)
 	{
-		return !IsInBlockTillLevelStreamingCompleted() ? (GMaxLoadingStreamingCells - GetCellLoadingCount()) : MAX_int32;
+		if (!Cell->GetLevel())
+		{
+			OutToReprioritizeLoadCells.Add(Cell);
+		}
 	}
 
-	// Always allow max on server to make sure StreamingLevels are added before clients update the visibility
-	return MAX_int32;
+	for (const UWorldPartitionRuntimeCell* Cell : ActivatedCells.GetPendingAddToWorldCells())
+	{
+		OutToReprioritizeActivateCells.Add(Cell);
+	}
 }
 
-int32 UWorldPartitionStreamingPolicy::SetCellsStateToLoaded(const TArray<const UWorldPartitionRuntimeCell*>& ToLoadCells)
+void UWorldPartitionStreamingPolicy::SetCellStateToLoaded(const UWorldPartitionRuntimeCell* InCell, int32& InOutMaxCellsToLoad)
 {
-	TRACE_CPUPROFILER_EVENT_SCOPE(UWorldPartitionStreamingPolicy::SetCellsStateToLoaded);
-
-	int32 MaxCellsToLoad = GetMaxCellsToLoad();
-
-	int32 LoadedCount = 0;
-
-	// Trigger cell loading. Depending on actual state of cell limit loading.
-	for (const UWorldPartitionRuntimeCell* Cell : ToLoadCells)
+	bool bLoadCell = false;
+	if (ActivatedCells.Contains(InCell))
 	{
-		UE_LOG(LogWorldPartition, Verbose, TEXT("UWorldPartitionStreamingPolicy::LoadCells %s"), *Cell->GetName());
-
-		if (ActivatedCells.Contains(Cell))
-		{
-			Cell->Deactivate();
-			ActivatedCells.Remove(Cell);
-			LoadedCells.Add(Cell);
-			++LoadedCount;
-		}
-		else if (MaxCellsToLoad > 0)
+		InCell->Deactivate();
+		ActivatedCells.Remove(InCell);
+		bLoadCell = true;
+		
+	}
+	else if (WorldPartition->IsStreamingInEnabled())
+	{
+		if (InOutMaxCellsToLoad > 0)
 		{
-			Cell->Load();
-			LoadedCells.Add(Cell);
-			++LoadedCount;
-			if (!Cell->IsAlwaysLoaded())
+			InCell->Load();
+			bLoadCell = true;
+			if (!InCell->IsAlwaysLoaded())
 			{
-				--MaxCellsToLoad;
+				--InOutMaxCellsToLoad;
 			}
 		}
 	}
 
-	return LoadedCount;
+	if (bLoadCell)
+	{
+		UE_LOG(LogWorldPartition, Verbose, TEXT("UWorldPartitionStreamingPolicy::SetCellStateToLoaded %s"), *InCell->GetName());
+		LoadedCells.Add(InCell);
+		++ProcessedToLoadCells;
+	}
 }
 
-int32 UWorldPartitionStreamingPolicy::SetCellsStateToActivated(const TArray<const UWorldPartitionRuntimeCell*>& ToActivateCells)
+void UWorldPartitionStreamingPolicy::SetCellStateToActivated(const UWorldPartitionRuntimeCell* InCell, int32& InOutMaxCellsToLoad)
 {
-	TRACE_CPUPROFILER_EVENT_SCOPE(UWorldPartitionStreamingPolicy::SetCellsStateToActivated);
-
-	int32 MaxCellsToLoad = GetMaxCellsToLoad();
-	
-	int32 ActivatedCount = 0;
-	// Trigger cell activation. Depending on actual state of cell limit loading.
-	for (const UWorldPartitionRuntimeCell* Cell : ToActivateCells)
+	if (!WorldPartition->IsStreamingInEnabled())
 	{
-		UE_LOG(LogWorldPartition, Verbose, TEXT("UWorldPartitionStreamingPolicy::ActivateCells %s"), *Cell->GetName());
+		return;
+	}
 
-		if (LoadedCells.Contains(Cell))
-		{
-			LoadedCells.Remove(Cell);
-			ActivatedCells.Add(Cell);
-			Cell->Activate();
-			++ActivatedCount;
-		}
-		else if (MaxCellsToLoad > 0)
+	bool bActivateCell = false;
+	if (LoadedCells.Contains(InCell))
+	{
+		LoadedCells.Remove(InCell);
+		bActivateCell = true;
+	}
+	else if (InOutMaxCellsToLoad > 0)
+	{
+		if (!InCell->IsAlwaysLoaded())
 		{
-			if (!Cell->IsAlwaysLoaded())
-			{
-				--MaxCellsToLoad;
-			}
-			ActivatedCells.Add(Cell);
-			Cell->Activate();
-			++ActivatedCount;
+			--InOutMaxCellsToLoad;
 		}
+		bActivateCell = true;
+	}
+
+	if (bActivateCell)
+	{
+		UE_LOG(LogWorldPartition, Verbose, TEXT("UWorldPartitionStreamingPolicy::SetCellStateToActivated %s"), *InCell->GetName());
+		ActivatedCells.Add(InCell);
+		InCell->Activate();
+		++ProcessedToActivateCells;
 	}
-	return ActivatedCount;
 }
 
-void UWorldPartitionStreamingPolicy::SetCellsStateToUnloaded(const TArray<const UWorldPartitionRuntimeCell*>& ToUnloadCells)
+void UWorldPartitionStreamingPolicy::SetCellsStateToUnloaded(const TArray<const UWorldPartitionRuntimeCell*>& InToUnloadCells)
 {
 	TRACE_CPUPROFILER_EVENT_SCOPE(UWorldPartitionStreamingPolicy::SetCellsStateToUnloaded);
 
-	for (const UWorldPartitionRuntimeCell* Cell : ToUnloadCells)
+	for (const UWorldPartitionRuntimeCell* Cell : InToUnloadCells)
 	{
 		if (Cell->CanUnload())
 		{
@@ -811,10 +792,11 @@ void UWorldPartitionStreamingPolicy::SetCellsStateToUnloaded(const TArray<const
 	}
 }
 
-bool UWorldPartitionStreamingPolicy::CanAddLoadedLevelToWorld(ULevel* InLevel) const
+bool UWorldPartitionStreamingPolicy::CanAddCellToWorld(const UWorldPartitionRuntimeCell* InCell) const
 {
-	TRACE_CPUPROFILER_EVENT_SCOPE(UWorldPartitionStreamingPolicy::CanAddLoadedLevelToWorld);
+	TRACE_CPUPROFILER_EVENT_SCOPE(UWorldPartitionStreamingPolicy::CanAddCellToWorld);
 
+	check(InCell);
 	check(WorldPartition->IsInitialized());
 
 	// Always allow AddToWorld in Dedicated server and Listen server
@@ -829,13 +811,7 @@ bool UWorldPartitionStreamingPolicy::CanAddLoadedLevelToWorld(ULevel* InLevel) c
 		return true;
 	}
 
-	const UWorldPartitionRuntimeCell** Cell = Algo::FindByPredicate(SortedAddToWorldCells, [InLevel](const UWorldPartitionRuntimeCell* ItCell) { return ItCell->GetLevel() == InLevel; });
-	if (Cell && ShouldSkipCellForPerformance(*Cell))
-	{
-		return false;
-	}
- 
-	return true;
+	return !ShouldSkipCellForPerformance(InCell);
 }
 
 bool UWorldPartitionStreamingPolicy::IsStreamingCompleted(const TArray<FWorldPartitionStreamingSource>* InStreamingSources) const
diff --git a/Engine/Source/Runtime/Engine/Private/WorldPartition/WorldPartitionSubsystem.cpp b/Engine/Source/Runtime/Engine/Private/WorldPartition/WorldPartitionSubsystem.cpp
--- a/Engine/Source/Runtime/Engine/Private/WorldPartition/WorldPartitionSubsystem.cpp
+++ b/Engine/Source/Runtime/Engine/Private/WorldPartition/WorldPartitionSubsystem.cpp
@@ -11,6 +11,7 @@
 #include "WorldPartition/WorldPartitionStreamingSource.h"
 #include "WorldPartition/WorldPartitionReplay.h"
 #include "WorldPartition/WorldPartitionDraw2DContext.h"
+#include "WorldPartition/WorldPartitionStreamingPolicy.h"
 #include "WorldPartition/DataLayer/WorldDataLayersActorDesc.h"
 #include "WorldPartition/DataLayer/DataLayerManager.h"
 #include "WorldPartition/DataLayer/DataLayerAsset.h"
@@ -125,6 +126,12 @@ static FAutoConsoleCommandWithOutputDevice GDumpStreamingSourcesCmd(
 	})
 );
 
+static int32 GMaxLoadingStreamingCells = 4;
+static FAutoConsoleVariableRef CMaxLoadingStreamingCells(
+	TEXT("wp.Runtime.MaxLoadingStreamingCells"),
+	GMaxLoadingStreamingCells,
+	TEXT("Used to limit the number of concurrent loading world partition streaming cells."));
+
 UWorldPartitionSubsystem::UWorldPartitionSubsystem()
 : StreamingSourcesHash(0)
 {}
@@ -418,7 +425,10 @@ void UWorldPartitionSubsystem::Initialize(FSubsystemCollectionBase& Collection)
 
 	GetWorld()->OnWorldPartitionInitialized().AddUObject(this, &UWorldPartitionSubsystem::OnWorldPartitionInitialized);
 	GetWorld()->OnWorldPartitionUninitialized().AddUObject(this, &UWorldPartitionSubsystem::OnWorldPartitionUninitialized);
-	FLevelStreamingDelegates::OnLevelStreamingTargetStateChanged.AddUObject(this, &UWorldPartitionSubsystem::OnLevelStreamingTargetStateChanged);
+	if (GetWorld()->IsGameWorld())
+	{
+		FLevelStreamingDelegates::OnLevelStreamingTargetStateChanged.AddUObject(this, &UWorldPartitionSubsystem::OnLevelStreamingTargetStateChanged);
+	}
 }
 
 void UWorldPartitionSubsystem::Deinitialize()
@@ -433,7 +443,10 @@ void UWorldPartitionSubsystem::Deinitialize()
 
 	GetWorld()->OnWorldPartitionInitialized().RemoveAll(this);
 	GetWorld()->OnWorldPartitionUninitialized().RemoveAll(this);
-	FLevelStreamingDelegates::OnLevelStreamingTargetStateChanged.RemoveAll(this);
+	if (GetWorld()->IsGameWorld())
+	{
+		FLevelStreamingDelegates::OnLevelStreamingTargetStateChanged.RemoveAll(this);
+	}
 
 	// At this point World Partition should be uninitialized
 	check(!GetWorldPartition() || !GetWorldPartition()->IsInitialized());
@@ -514,7 +527,8 @@ void UWorldPartitionSubsystem::OnLevelStreamingTargetStateChanged(UWorld* World,
 		UWorldPartition* WorldPartition = LevelIfLoaded->GetTypedOuter<UWorld>()->GetWorldPartition();
 		if (WorldPartition && WorldPartition->IsInitialized() && !WorldPartition->CanStream())
 		{
-			WorldPartition->UpdateStreamingState();
+			TRACE_CPUPROFILER_EVENT_SCOPE(UWorldPartitionSubsystem::OnLevelStreamingTargetStateChanged);
+			UWorldPartitionSubsystem::UpdateStreamingStateInternal(World, { WorldPartition });
 		}
 	}
 }
@@ -764,15 +778,118 @@ void UWorldPartitionSubsystem::GetStreamingSources(const UWorldPartition* InWorl
 	}
 }
 
+DECLARE_CYCLE_STAT(TEXT("World Partition Update Streaming"), STAT_WorldPartitionUpdateStreaming, STATGROUP_Engine);
+
 void UWorldPartitionSubsystem::UpdateStreamingState()
 {
-	UpdateStreamingSources();
+	SCOPE_CYCLE_COUNTER(STAT_WorldPartitionUpdateStreaming);
+
+	UWorldPartitionSubsystem::UpdateStreamingStateInternal(GetWorld(), RegisteredWorldPartitions);
+}
+
+void UWorldPartitionSubsystem::UpdateStreamingStateInternal(const UWorld* InWorld, const TArray<TObjectPtr<UWorldPartition>>& InWorldPartitions)
+{
+	TRACE_CPUPROFILER_EVENT_SCOPE(UWorldPartitionSubsystem::UpdateStreamingState);
+
+	const UWorld* World = InWorld;
+	if (!World->IsGameWorld())
+	{
+		return;
+	}
+
+	// Subsystem can be null during EndPlayMap. WorldPartition::Uninitialize will still call UpdateStreamingStateInternal to cleanup it's streaming levels
+	if (UWorldPartitionSubsystem* WorldPartitionSubsystem = UWorld::GetSubsystem<UWorldPartitionSubsystem>(World))
+	{
+		// Update streaming sources
+		WorldPartitionSubsystem->UpdateStreamingSources();
+	}
+	
+	// Make temp copy of array as UpdateStreamingState may FlushAsyncLoading, which may add a new world partition to RegisteredWorldPartitions while iterating
+	const TArray<UWorldPartition*> RegisteredWorldPartitionsCopy = InWorldPartitions;
+
+	TArray<const UWorldPartitionRuntimeCell*> ToActivateCells;
+	TArray<const UWorldPartitionRuntimeCell*> ToLoadCells;
+
+	// Update streaming state of all registered world partitions and cumulate cells to activate and load
+	TMap<const UWorldPartitionRuntimeCell*, UWorldPartition*> CellToWorldPartition;
+	for (UWorldPartition* RegisteredWorldPartition : RegisteredWorldPartitionsCopy)
+	{
+		if (RegisteredWorldPartition->StreamingPolicy)
+		{
+			RegisteredWorldPartition->StreamingPolicy->UpdateStreamingState();
+			RegisteredWorldPartition->StreamingPolicy->GetCellsToUpdate(ToLoadCells, ToActivateCells);
+		}
+	}
+
+	// Sort all cells to activate and load
+	{
+		TRACE_CPUPROFILER_EVENT_SCOPE(SortStreamingCellsByImportance);
+		Algo::Sort(ToActivateCells, [](const UWorldPartitionRuntimeCell* CellA, const UWorldPartitionRuntimeCell* CellB) { return CellA->SortCompare(CellB) < 0; });
+		Algo::Sort(ToLoadCells, [](const UWorldPartitionRuntimeCell* CellA, const UWorldPartitionRuntimeCell* CellB) { return CellA->SortCompare(CellB) < 0; });
+	}
 
-	//make temp copy of array as UpdateStreamingState may FlushAsyncLoading, which may add a new world partition to RegisteredWorldPartitions while iterating
-	const TArray<UWorldPartition*> RegisteredWorldPartitionsCopy = RegisteredWorldPartitions;
-	for (UWorldPartition* RegisteredWorldParition : RegisteredWorldPartitionsCopy)
+	// Compute maximum number of cells to load
+	const ENetMode NetMode = World->GetNetMode();
+	const bool bIsServer = (NetMode == NM_DedicatedServer || NetMode == NM_ListenServer);
+	int32 MaxCellsToLoad = (!bIsServer && !World->GetIsInBlockTillLevelStreamingCompleted() && (GMaxLoadingStreamingCells > 0)) ? (GMaxLoadingStreamingCells - (int32)World->GetNumStreamingLevelsBeingLoaded()) : MAX_int32;
+
+	// Process cells to activate
+	{
+		TRACE_CPUPROFILER_EVENT_SCOPE(ToActivateCells);
+		// Trigger cell activation. Depending on actual state of cell limit loading.
+		for (const UWorldPartitionRuntimeCell* Cell : ToActivateCells)
+		{
+			Cell->GetOuterWorld()->GetWorldPartition()->StreamingPolicy->SetCellStateToActivated(Cell, MaxCellsToLoad);
+		}
+	}
+	// Process cells to load
 	{
-		RegisteredWorldParition->UpdateStreamingState();
+		TRACE_CPUPROFILER_EVENT_SCOPE(ToLoadCells);
+		for (const UWorldPartitionRuntimeCell* Cell : ToLoadCells)
+		{
+			Cell->GetOuterWorld()->GetWorldPartition()->StreamingPolicy->SetCellStateToLoaded(Cell, MaxCellsToLoad);
+		}
+	}
+
+	// Build and reprioritize cells that are still in a pending state (waiting to be loaded or activated)
+	TArray<const UWorldPartitionRuntimeCell*> PendingToLoadCells;
+	TArray<const UWorldPartitionRuntimeCell*> PendingToActivateCells;
+	{
+		TRACE_CPUPROFILER_EVENT_SCOPE(BuildPendingCells);
+		PendingToLoadCells.Reserve(ToLoadCells.Num());
+		PendingToActivateCells.Reserve(ToActivateCells.Num());
+		for (UWorldPartition* RegisteredWorldPartition : RegisteredWorldPartitionsCopy)
+		{
+			RegisteredWorldPartition->StreamingPolicy->GetCellsToReprioritize(PendingToLoadCells, PendingToActivateCells);
+		}
+	}
+	{
+		{
+			TRACE_CPUPROFILER_EVENT_SCOPE(SortPendingCells);
+			Algo::Sort(PendingToLoadCells, [](const UWorldPartitionRuntimeCell* CellA, const UWorldPartitionRuntimeCell* CellB) { return CellA->SortCompare(CellB) < 0; });
+			Algo::Sort(PendingToActivateCells, [](const UWorldPartitionRuntimeCell* CellA, const UWorldPartitionRuntimeCell* CellB) { return CellA->SortCompare(CellB) < 0; });
+		}
+
+		const ULevel* CurrentPendingVisibility = World->GetCurrentLevelPendingVisibility();
+		const UWorldPartitionRuntimeCell* CurrentPendingVisibilityCell = CurrentPendingVisibility ? Cast<const UWorldPartitionRuntimeCell>(CurrentPendingVisibility->GetWorldPartitionRuntimeCell()) : nullptr;
+		auto SetCellsStreamingPriority = [CurrentPendingVisibilityCell](const TArray<const UWorldPartitionRuntimeCell*>& InCells, int32 InPriorityBias = 0)
+		{
+			const int32 MaxPrio = InCells.Num() + InPriorityBias;
+			int32 Prio = MaxPrio;
+			for (const UWorldPartitionRuntimeCell* Cell : InCells)
+			{
+				// Make sure that the current pending visibility level is the most important so that level streaming update will process it right away
+				const int32 SortedPriority = (Cell == CurrentPendingVisibilityCell) ? MaxPrio + 1 : Prio--;
+				Cell->SetStreamingPriority(SortedPriority);
+			}
+		};
+
+		// Update level streaming priority so that UWorld::UpdateLevelStreaming will naturally process the levels in the correct order
+		{
+			TRACE_CPUPROFILER_EVENT_SCOPE(SetPendingCellsStreamingPriority);
+			SetCellsStreamingPriority(PendingToLoadCells);
+			SetCellsStreamingPriority(PendingToActivateCells, PendingToLoadCells.Num());
+		}
 	}
 }
 
diff --git a/Engine/Source/Runtime/Engine/Public/WorldPartition/WorldPartition.h b/Engine/Source/Runtime/Engine/Public/WorldPartition/WorldPartition.h
--- a/Engine/Source/Runtime/Engine/Public/WorldPartition/WorldPartition.h
+++ b/Engine/Source/Runtime/Engine/Public/WorldPartition/WorldPartition.h
@@ -332,8 +332,7 @@ public:
 	bool IsMainWorldPartition() const;
 
 	void Tick(float DeltaSeconds);
-	void UpdateStreamingState();
-	bool CanAddLoadedLevelToWorld(ULevel* InLevel) const;
+	bool CanAddCellToWorld(const IWorldPartitionCell* InCell) const;
 	bool IsStreamingCompleted(const TArray<FWorldPartitionStreamingSource>* InStreamingSources) const;
 	bool IsStreamingCompleted(EWorldPartitionRuntimeCellState QueryState, const TArray<FWorldPartitionStreamingQuerySource>& QuerySources, bool bExactState) const;
 	bool GetIntersectingCells(const TArray<FWorldPartitionStreamingQuerySource>& InSources, TArray<const IWorldPartitionCell*>& OutCells) const;
@@ -359,6 +358,11 @@ public:
 
 	UDataLayerManager* GetDataLayerManager() const;
 
+	UE_DEPRECATED(5.3, "UpdateStreamingState is deprecated, use UWorldPartitionSubsystem::UpdateStreamingState instead.")
+	void UpdateStreamingState() {}
+	UE_DEPRECATED(5.3, "CanAddLoadedLevelToWorld is deprecated, use CanAddCellToWorld instead.")
+	bool CanAddLoadedLevelToWorld(ULevel* InLevel) const { return true; }
+
 #if WITH_EDITORONLY_DATA
 	UPROPERTY(DuplicateTransient)
 	TObjectPtr<UWorldPartitionEditorHash> EditorHash;
diff --git a/Engine/Source/Runtime/Engine/Public/WorldPartition/WorldPartitionLevelStreamingPolicy.h b/Engine/Source/Runtime/Engine/Public/WorldPartition/WorldPartitionLevelStreamingPolicy.h
--- a/Engine/Source/Runtime/Engine/Public/WorldPartition/WorldPartitionLevelStreamingPolicy.h
+++ b/Engine/Source/Runtime/Engine/Public/WorldPartition/WorldPartitionLevelStreamingPolicy.h
@@ -42,8 +42,6 @@ public:
 	virtual bool RemoveExternalStreamingObject(URuntimeHashExternalStreamingObjectBase* ExternalStreamingObject) override;
 
 protected:
-	virtual int32 GetCellLoadingCount() const override;
-
 	void ForEachActiveRuntimeCell(TFunctionRef<void(const UWorldPartitionRuntimeCell*)> Func) const;
 
 private:
diff --git a/Engine/Source/Runtime/Engine/Public/WorldPartition/WorldPartitionRuntimeCell.h b/Engine/Source/Runtime/Engine/Public/WorldPartition/WorldPartitionRuntimeCell.h
--- a/Engine/Source/Runtime/Engine/Public/WorldPartition/WorldPartitionRuntimeCell.h
+++ b/Engine/Source/Runtime/Engine/Public/WorldPartition/WorldPartitionRuntimeCell.h
@@ -194,6 +194,7 @@ class ENGINE_API UWorldPartitionRuntimeCell : public UObject, public IWorldParti
 	virtual void SetPriority(int32 InPriority) { Priority = InPriority; }
 	virtual void SetStreamingPriority(int32 InStreamingPriority) const PURE_VIRTUAL(UWorldPartitionRuntimeCell::SetStreamingPriority,);
 	virtual EStreamingStatus GetStreamingStatus() const { return LEVEL_Unloaded; }
+	UE_DEPRECATED(5.3, "IsLoading is deprecated.")
 	virtual bool IsLoading() const { return false; }
 	void SetClientOnlyVisible(bool bInClientOnlyVisible) { bClientOnlyVisible = bInClientOnlyVisible; }
 	bool GetClientOnlyVisible() const { return bClientOnlyVisible; }
diff --git a/Engine/Source/Runtime/Engine/Public/WorldPartition/WorldPartitionRuntimeLevelStreamingCell.h b/Engine/Source/Runtime/Engine/Public/WorldPartition/WorldPartitionRuntimeLevelStreamingCell.h
--- a/Engine/Source/Runtime/Engine/Public/WorldPartition/WorldPartitionRuntimeLevelStreamingCell.h
+++ b/Engine/Source/Runtime/Engine/Public/WorldPartition/WorldPartitionRuntimeLevelStreamingCell.h
@@ -29,7 +29,6 @@ class ENGINE_API UWorldPartitionRuntimeLevelStreamingCell : public UWorldPartiti
 	virtual FLinearColor GetDebugColor(EWorldPartitionRuntimeCellVisualizeMode VisualizeMode) const override;
 	virtual void SetIsAlwaysLoaded(bool bInIsAlwaysLoaded) override;
 	virtual EStreamingStatus GetStreamingStatus() const override;
-	virtual bool IsLoading() const override;
 	//~End UWorldPartitionRuntimeCell Interface
 
 	//~Begin IWorldPartitionCell Interface
diff --git a/Engine/Source/Runtime/Engine/Public/WorldPartition/WorldPartitionStreamingPolicy.h b/Engine/Source/Runtime/Engine/Public/WorldPartition/WorldPartitionStreamingPolicy.h
--- a/Engine/Source/Runtime/Engine/Public/WorldPartition/WorldPartitionStreamingPolicy.h
+++ b/Engine/Source/Runtime/Engine/Public/WorldPartition/WorldPartitionStreamingPolicy.h
@@ -27,7 +27,7 @@ class UWorldPartitionStreamingPolicy : public UObject
 public:
 	virtual bool GetIntersectingCells(const TArray<FWorldPartitionStreamingQuerySource>& InSources, TArray<const IWorldPartitionCell*>& OutCells) const;
 	virtual void UpdateStreamingState();
-	virtual bool CanAddLoadedLevelToWorld(class ULevel* InLevel) const;
+	virtual bool CanAddCellToWorld(const UWorldPartitionRuntimeCell* InCell) const;
 	virtual bool DrawRuntimeHash2D(FWorldPartitionDraw2DContext& DrawContext);
 	virtual void DrawRuntimeHash3D();
 	virtual void DrawRuntimeCellsDetails(class UCanvas* Canvas, FVector2D& Offset) {}
@@ -38,6 +38,9 @@ public:
 	virtual void OnCellShown(const UWorldPartitionRuntimeCell* InCell);
 	virtual void OnCellHidden(const UWorldPartitionRuntimeCell* InCell);
 
+	UE_DEPRECATED(5.3, "CanAddLoadedLevelToWorld is deprecated.")
+	virtual bool CanAddLoadedLevelToWorld(class ULevel* InLevel) const { return true; }
+
 #if WITH_EDITOR
 	virtual TSubclassOf<class UWorldPartitionRuntimeCell> GetRuntimeCellClass() const PURE_VIRTUAL(UWorldPartitionStreamingPolicy::GetRuntimeCellClass, return UWorldPartitionRuntimeCell::StaticClass(); );
 
@@ -66,11 +69,11 @@ public:
 	virtual bool RemoveExternalStreamingObject(URuntimeHashExternalStreamingObjectBase* ExternalStreamingObject) { return true; }
 
 protected:
-	virtual int32 SetCellsStateToLoaded(const TArray<const UWorldPartitionRuntimeCell*>& ToLoadCells);
-	virtual int32 SetCellsStateToActivated(const TArray<const UWorldPartitionRuntimeCell*>& ToActivateCells);
+	virtual void SetCellStateToLoaded(const UWorldPartitionRuntimeCell* InCell, int32& InOutMaxCellsToLoad);
+	virtual void SetCellStateToActivated(const UWorldPartitionRuntimeCell* InCell, int32& InOutMaxCellsToLoad);
 	virtual void SetCellsStateToUnloaded(const TArray<const UWorldPartitionRuntimeCell*>& ToUnloadCells);
-	virtual int32 GetCellLoadingCount() const { return 0; }
-	virtual int32 GetMaxCellsToLoad() const;
+	virtual void GetCellsToUpdate(TArray<const UWorldPartitionRuntimeCell*>& OutToLoadCells, TArray<const UWorldPartitionRuntimeCell*>& OutToActivateCells);
+	virtual void GetCellsToReprioritize(TArray<const UWorldPartitionRuntimeCell*>& OutToLoadCells, TArray<const UWorldPartitionRuntimeCell*>& OutToActivateCells);
 	virtual void UpdateStreamingSources(bool bCanOptimizeUpdate);
 	void UpdateStreamingPerformance(const TSet<const UWorldPartitionRuntimeCell*>& InCells);
 	bool ShouldSkipCellForPerformance(const UWorldPartitionRuntimeCell* Cell) const;
@@ -105,6 +108,12 @@ protected:
 	TSet<const UWorldPartitionRuntimeCell*> FrameActivateCells;
 	TSet<const UWorldPartitionRuntimeCell*> FrameLoadCells;
 
+	// Used by UWorldPartitionSubsystem
+	TArray<const UWorldPartitionRuntimeCell*> ToActivateCells;
+	TArray<const UWorldPartitionRuntimeCell*> ToLoadCells;
+	int32 ProcessedToActivateCells;
+	int32 ProcessedToLoadCells;
+
 private:
 	// Update optimization
 	uint32 ComputeUpdateStreamingHash(bool bCanOptimizeUpdate) const;
@@ -118,7 +127,6 @@ private:
 	static FAutoConsoleVariableRef CVarUpdateOptimEnabled;
 	static FAutoConsoleVariableRef CVarForceUpdateFrameCount;
 
-	bool bLastUpdateCompletedLoadingAndActivation;
 	bool bCriticalPerformanceRequestedBlockTillOnWorld;
 	int32 CriticalPerformanceBlockTillLevelStreamingCompletedEpoch;
 	int32 ServerDataLayersStatesEpoch;
@@ -137,4 +145,6 @@ private:
 	double OnScreenMessageStartTime;
 	EWorldPartitionStreamingPerformance  OnScreenMessageStreamingPerformance;
 #endif
+
+	friend class UWorldPartitionSubsystem;
 };
\ No newline at end of file
diff --git a/Engine/Source/Runtime/Engine/Public/WorldPartition/WorldPartitionSubsystem.h b/Engine/Source/Runtime/Engine/Public/WorldPartition/WorldPartitionSubsystem.h
--- a/Engine/Source/Runtime/Engine/Public/WorldPartition/WorldPartitionSubsystem.h
+++ b/Engine/Source/Runtime/Engine/Public/WorldPartition/WorldPartitionSubsystem.h
@@ -167,6 +167,8 @@ private:
 	void OnWorldPartitionUninitialized(UWorldPartition* InWorldPartition);
 	void OnLevelStreamingTargetStateChanged(UWorld* World, const ULevelStreaming* StreamingLevel, ULevel* LevelIfLoaded, ELevelStreamingState CurrentState, ELevelStreamingTargetState PrevTarget, ELevelStreamingTargetState NewTarget);
 
+	static void UpdateStreamingStateInternal(const UWorld* InWorld, const TArray<TObjectPtr<UWorldPartition>>& InWorldPartitions);
+
 	UWorldPartition* GetWorldPartition();
 	const UWorldPartition* GetWorldPartition() const;
 	bool CanDebugDraw() const;
