Update SuggestProjectileVelocity_MovingTarget to only adjust for slope in the Walking MovementMode.

diff --git a/Engine/Source/Runtime/Engine/Private/GameplayStatics.cpp b/Engine/Source/Runtime/Engine/Private/GameplayStatics.cpp
--- a/Engine/Source/Runtime/Engine/Private/GameplayStatics.cpp
+++ b/Engine/Source/Runtime/Engine/Private/GameplayStatics.cpp
@@ -3063,18 +3063,24 @@ bool UGameplayStatics::SuggestProjectileVelocity_MovingTarget(const UObject* Wor
 	TimeToTarget = FMath::Clamp(TimeToTarget, 0.1, TimeToTarget);
 
 	// Find the target's velocity
-	// If the target is a character, GetVelocity() will return a vector with a Z value of 0, so we need to adjust for slope
 	FVector TargetVelocity = TargetActor->GetVelocity();
+
+	// If the target is a character moving on ground or the floor result is walkable (and therefore used for movement),
+	// GetVelocity() will return a vector with a Z value of 0, so we need to adjust for slope.
 	if (ACharacter* TargetCharacter = Cast<ACharacter>(TargetActor))
 	{
 		if (UCharacterMovementComponent* TargetMovementComp = TargetCharacter->GetCharacterMovement())
 		{
-			FHitResult FloorHit = TargetMovementComp->CurrentFloor.HitResult;
-			const double FloorAngleRadians = FMath::Acos(FloorHit.ImpactNormal.Z);
+			FFindFloorResult FloorResult = TargetMovementComp->CurrentFloor;
+			if (TargetMovementComp->IsMovingOnGround() || FloorResult.IsWalkableFloor())
+			{
+				FHitResult FloorHit = FloorResult.HitResult;
+				const double FloorAngleRadians = FMath::Acos(FloorHit.ImpactNormal.Z);
 
-			// z = x * tan(theta)
-			const double ZMagnitude = TargetVelocity.Size2D() * FMath::Tan(FloorAngleRadians);
-			TargetVelocity.Z = TargetVelocity.Dot(FloorHit.ImpactNormal) < 0.0 ? ZMagnitude : -ZMagnitude;
+				// z = x * tan(theta)
+				const double ZMagnitude = TargetVelocity.Size2D() * FMath::Tan(FloorAngleRadians);
+				TargetVelocity.Z = TargetVelocity.Dot(FloorHit.ImpactNormal) < 0.0 ? ZMagnitude : -ZMagnitude;
+			}
 		}
 	}
 
