Speculative include fix in animation code

#rnx

diff --git a/Engine/Source/Runtime/Engine/Public/Animation/AnimCompressionTypes.h b/Engine/Source/Runtime/Engine/Public/Animation/AnimCompressionTypes.h
--- a/Engine/Source/Runtime/Engine/Public/Animation/AnimCompressionTypes.h
+++ b/Engine/Source/Runtime/Engine/Public/Animation/AnimCompressionTypes.h
@@ -14,6 +14,7 @@
 #include "Animation/AnimationAsset.h"
 
 #include "Async/MappedFileHandle.h"
+#include "HAL/ThreadSafeBool.h"
 #include "HAL/PlatformFileManager.h"
 #include "GenericPlatform/GenericPlatformFile.h"
 #include "Misc/Paths.h"
