*Add Experimental UCLASS/USTRUCT specifiers to LightWeightInstance classes.
*Remove some UPROPERTY/UFUNCTIONS from not uclasses.

diff --git a/Engine/Source/Runtime/Engine/Classes/GameFramework/LightWeightInstanceManager.h b/Engine/Source/Runtime/Engine/Classes/GameFramework/LightWeightInstanceManager.h
--- a/Engine/Source/Runtime/Engine/Classes/GameFramework/LightWeightInstanceManager.h
+++ b/Engine/Source/Runtime/Engine/Classes/GameFramework/LightWeightInstanceManager.h
@@ -27,7 +27,7 @@ struct ENGINE_API FLWIData
 
 
 // Base class for interfaces for each handle
-UCLASS()
+UCLASS(Experimental)
 class ENGINE_API UActorInstanceHandleInterface : public UObject
 {
 	GENERATED_UCLASS_BODY()
@@ -49,7 +49,7 @@ protected:
 
 
 
-UCLASS(BlueprintType, Blueprintable)
+UCLASS(BlueprintType, Blueprintable, Experimental)
 class ENGINE_API ALightWeightInstanceManager : public AActor
 {
 	GENERATED_UCLASS_BODY()
diff --git a/Engine/Source/Runtime/Engine/Classes/GameFramework/LightWeightInstanceStaticMeshManager.h b/Engine/Source/Runtime/Engine/Classes/GameFramework/LightWeightInstanceStaticMeshManager.h
--- a/Engine/Source/Runtime/Engine/Classes/GameFramework/LightWeightInstanceStaticMeshManager.h
+++ b/Engine/Source/Runtime/Engine/Classes/GameFramework/LightWeightInstanceStaticMeshManager.h
@@ -11,7 +11,7 @@
 
 DECLARE_DYNAMIC_DELEGATE_RetVal_OneParam(FActorInstanceHandle, FOnActorReady, FActorInstanceHandle, InHandle);
 
-UCLASS(BlueprintType, Blueprintable)
+UCLASS(BlueprintType, Blueprintable, Experimental)
 class ENGINE_API ALightWeightInstanceStaticMeshManager : public ALightWeightInstanceManager, public ISMInstanceManager
 {
 	GENERATED_UCLASS_BODY()
diff --git a/Engine/Source/Runtime/Engine/Classes/GameFramework/LightWeightInstanceSubsystem.h b/Engine/Source/Runtime/Engine/Classes/GameFramework/LightWeightInstanceSubsystem.h
--- a/Engine/Source/Runtime/Engine/Classes/GameFramework/LightWeightInstanceSubsystem.h
+++ b/Engine/Source/Runtime/Engine/Classes/GameFramework/LightWeightInstanceSubsystem.h
@@ -41,14 +41,12 @@ struct ENGINE_API FLightWeightInstanceSubsystem
 
 	// Returns the instance manager that handles instances of type Class that live in Level
 	UE_DEPRECATED(5.3, "Use the version that takes in a position.")
-	UFUNCTION(Server, Unreliable)
 	ALightWeightInstanceManager* FindOrAddLightWeightInstanceManager(UClass* ActorClass, const UDataLayerInstance* DataLayer, UWorld* World);
 
 	// Returns the instance manager that handles actors of type ActorClass in level Level
 	ALightWeightInstanceManager* FindLightWeightInstanceManager(UClass& ActorClass, UWorld& World, const FVector& InPos, const UDataLayerInstance* DataLayer = nullptr) const;
 
 	// Returns the instance manager that handles instances of type Class that live in Level
-	UFUNCTION(Server, Unreliable)
 	ALightWeightInstanceManager* FindOrAddLightWeightInstanceManager(UClass& ActorClass, UWorld& World, const FVector& InPos, const UDataLayerInstance* DataLayer = nullptr);
 
 	// Returns the actor specified by Handle. This may require loading and creating the actor object.
@@ -124,7 +122,6 @@ private:
 	static TSharedPtr<FLightWeightInstanceSubsystem, ESPMode::ThreadSafe> LWISubsystem;
 
 	// TODO: preallocate the size of this based on a config variable
-	UPROPERTY()
 	TArray<ALightWeightInstanceManager*> LWInstanceManagers;
 
 	// Mutex to make sure we don't change the LWInstanceManagers array while reading/writing it.
