Level Instance: avoid creating default HLOD Layer assets (they are not used)

#rnx

diff --git a/Engine/Source/Runtime/Engine/Private/LevelInstance/LevelInstanceSubsystem.cpp b/Engine/Source/Runtime/Engine/Private/LevelInstance/LevelInstanceSubsystem.cpp
--- a/Engine/Source/Runtime/Engine/Private/LevelInstance/LevelInstanceSubsystem.cpp
+++ b/Engine/Source/Runtime/Engine/Private/LevelInstance/LevelInstanceSubsystem.cpp
@@ -1083,6 +1083,9 @@ ILevelInstanceInterface* ULevelInstanceSubsystem::CreateLevelInstanceFrom(const
 				check(!WorldPartition->IsStreamingEnabled());
 				check(InLevel->IsUsingActorFolders());
 
+				// Reset HLOD Layer (no defaults needed for Level Instances)
+				WorldPartition->SetDefaultHLODLayer(nullptr);
+
 				// Make sure new level's AWorldDataLayers contains all the necessary Data Layer Instances before moving actors
 				TSet<TObjectPtr<const UDataLayerAsset>> SourceDataLayerAssets;
 				for (AActor* ActorToMove : ActorsToMove)
