Fixed local height fog non unity build issues.

diff --git a/Engine/Source/Runtime/Engine/Private/Components/LocalHeightFogComponent.cpp b/Engine/Source/Runtime/Engine/Private/Components/LocalHeightFogComponent.cpp
--- a/Engine/Source/Runtime/Engine/Private/Components/LocalHeightFogComponent.cpp
+++ b/Engine/Source/Runtime/Engine/Private/Components/LocalHeightFogComponent.cpp
@@ -6,6 +6,11 @@
 #include "Components/ArrowComponent.h"
 #include "Components/DrawSphereComponent.h"
 #include "Components/BillboardComponent.h"
+#include "SceneInterface.h"
+#include "SceneManagement.h"
+#include "Engine/World.h"
+#include "UObject/ConstructorHelpers.h"
+#include "Engine/Texture2D.h"
 
 #include UE_INLINE_GENERATED_CPP_BY_NAME(LocalHeightFogComponent)
 
