#UE Fix for moving a blueprint on top of a redirector. This often happens when moving a BP back to a place it once was.

diff --git a/Engine/Source/Runtime/Engine/Private/Blueprint.cpp b/Engine/Source/Runtime/Engine/Private/Blueprint.cpp
--- a/Engine/Source/Runtime/Engine/Private/Blueprint.cpp
+++ b/Engine/Source/Runtime/Engine/Private/Blueprint.cpp
@@ -529,6 +529,19 @@ bool UBlueprint::RenameGeneratedClasses( const TCHAR* InName, UObject* NewOuter,
 			}
 		};
 
+		const auto CheckRedirectors = [](FName ClassName, UClass* ForClass, UObject* NewOuter)
+		{
+			if (UObjectRedirector* Redirector = FindObjectFast<UObjectRedirector>(NewOuter, ClassName))
+			{
+				// If we found a redirector, check that the object it points to is of the same class.
+				if (Redirector->DestinationObject
+					&& Redirector->DestinationObject->GetClass() == ForClass->GetClass())
+				{
+					Redirector->Rename(nullptr, GetTransientPackage(), REN_ForceNoResetLoaders | REN_DontCreateRedirectors);
+				}
+			}
+		};
+
 		FName SkelClassName, GenClassName;
 		GetBlueprintClassNames(GenClassName, SkelClassName, FName(InName));
 
@@ -537,6 +550,7 @@ bool UBlueprint::RenameGeneratedClasses( const TCHAR* InName, UObject* NewOuter,
 		{
 			// check for collision of CDO name, move aside if necessary:
 			TryFreeCDOName(GeneratedClass, NewTopLevelObjectOuter, Flags);
+			CheckRedirectors(GenClassName, GeneratedClass, NewTopLevelObjectOuter);
 			bool bMovedOK = GeneratedClass->Rename(*GenClassName.ToString(), NewTopLevelObjectOuter, Flags);
 			if (!bMovedOK)
 			{
@@ -548,6 +562,7 @@ bool UBlueprint::RenameGeneratedClasses( const TCHAR* InName, UObject* NewOuter,
 		if (SkeletonGeneratedClass != NULL && SkeletonGeneratedClass != GeneratedClass)
 		{
 			TryFreeCDOName(SkeletonGeneratedClass, NewTopLevelObjectOuter, Flags);
+			CheckRedirectors(SkelClassName, SkeletonGeneratedClass, NewTopLevelObjectOuter);
 			bool bMovedOK = SkeletonGeneratedClass->Rename(*SkelClassName.ToString(), NewTopLevelObjectOuter, Flags);
 			if (!bMovedOK)
 			{
