Add runtime capacity checks to Blueprint-facing script container library APIs to mitigate crashes due to integer type overflow.

diff --git a/Engine/Source/Runtime/Engine/Classes/Kismet/BlueprintMapLibrary.h b/Engine/Source/Runtime/Engine/Classes/Kismet/BlueprintMapLibrary.h
--- a/Engine/Source/Runtime/Engine/Classes/Kismet/BlueprintMapLibrary.h
+++ b/Engine/Source/Runtime/Engine/Classes/Kismet/BlueprintMapLibrary.h
@@ -426,5 +426,8 @@ public:
 	static bool GenericMap_IsNotEmpty(const void* TargetMap, const FMapProperty* MapProperty);
 	static void GenericMap_Clear(const void* TargetMap, const FMapProperty* MapProperty);
 	static void GenericMap_SetMapPropertyByName(UObject* OwnerObject, FName MapPropertyName, const void* SrcMapAddr);
+
+private:
+	static constexpr int32 MaxSupportedMapSize = TNumericLimits<int32>::Max();
 };
 
diff --git a/Engine/Source/Runtime/Engine/Classes/Kismet/BlueprintSetLibrary.h b/Engine/Source/Runtime/Engine/Classes/Kismet/BlueprintSetLibrary.h
--- a/Engine/Source/Runtime/Engine/Classes/Kismet/BlueprintSetLibrary.h
+++ b/Engine/Source/Runtime/Engine/Classes/Kismet/BlueprintSetLibrary.h
@@ -559,5 +559,8 @@ public:
 	static void GenericSet_Union(const void* SetA, const FSetProperty* SetPropertyA, const void* SetB, const FSetProperty* SetPropertyB, const void* SetResult, const FSetProperty* SetPropertyResult);
 	static void GenericSet_Difference(const void* SetA, const FSetProperty* SetPropertyA, const void* SetB, const FSetProperty* SetPropertyB, const void* SetResult, const FSetProperty* SetPropertyResult);
 	static void GenericSet_SetSetPropertyByName(UObject* OwnerObject, FName SetPropertyName, const void* SrcSetAddr);
+
+private:
+	static constexpr int32 MaxSupportedSetSize = TNumericLimits<int32>::Max();
 };
 
diff --git a/Engine/Source/Runtime/Engine/Classes/Kismet/KismetArrayLibrary.h b/Engine/Source/Runtime/Engine/Classes/Kismet/KismetArrayLibrary.h
--- a/Engine/Source/Runtime/Engine/Classes/Kismet/KismetArrayLibrary.h
+++ b/Engine/Source/Runtime/Engine/Classes/Kismet/KismetArrayLibrary.h
@@ -283,9 +283,14 @@ class ENGINE_API UKismetArrayLibrary : public UBlueprintFunctionLibrary
 	static bool GenericArray_IsValidIndex(const void* TargetArray, const FArrayProperty* ArrayProp, int32 IndexToTest);
 	
 private:
+	static constexpr int32 MaxSupportedArraySize = TNumericLimits<int32>::Max();
+
 	static void GenericArray_HandleBool(const FProperty* Property, void* ItemPtr);
 
 public:
+	// Shared warning IDs for use by other container libraries.
+	static const FName ReachedMaximumContainerSizeWarning;
+
 	// Helper function to get the last valid index of the array for error reporting, or 0 if the array is empty
 	static int32 GetLastIndex(const FScriptArrayHelper& ArrayHelper)
 	{
diff --git a/Engine/Source/Runtime/Engine/Private/BlueprintMapLibrary.cpp b/Engine/Source/Runtime/Engine/Private/BlueprintMapLibrary.cpp
--- a/Engine/Source/Runtime/Engine/Private/BlueprintMapLibrary.cpp
+++ b/Engine/Source/Runtime/Engine/Private/BlueprintMapLibrary.cpp
@@ -1,6 +1,7 @@
 // Copyright Epic Games, Inc. All Rights Reserved.
 
 #include "Kismet/BlueprintMapLibrary.h"
+#include "Kismet/KismetArrayLibrary.h"
 
 #include UE_INLINE_GENERATED_CPP_BY_NAME(BlueprintMapLibrary)
 
@@ -9,7 +10,14 @@ void UBlueprintMapLibrary::GenericMap_Add(const void* TargetMap, const FMapPrope
 	if (TargetMap)
 	{
 		FScriptMapHelper MapHelper(MapProperty, TargetMap);
-		MapHelper.AddPair(KeyPtr, ValuePtr);
+		if (MapHelper.Num() < MaxSupportedMapSize)
+		{
+			MapHelper.AddPair(KeyPtr, ValuePtr);
+		}
+		else if (!MapHelper.FindValueFromHash(KeyPtr))
+		{
+			FFrame::KismetExecutionMessage(*FString::Printf(TEXT("Attempted add to map '%s' beyond the maximum supported capacity!"), *MapProperty->GetName()), ELogVerbosity::Warning, UKismetArrayLibrary::ReachedMaximumContainerSizeWarning);
+		}
 	}
 }
 
@@ -140,5 +148,4 @@ void UBlueprintMapLibrary::GenericMap_SetMapPropertyByName(UObject* OwnerObject,
 			MapProp->CopyValuesInternal(Dest, SrcMapAddr, 1);
 		}
 	}
-}
-
+}
\ No newline at end of file
diff --git a/Engine/Source/Runtime/Engine/Private/BlueprintSetLibrary.cpp b/Engine/Source/Runtime/Engine/Private/BlueprintSetLibrary.cpp
--- a/Engine/Source/Runtime/Engine/Private/BlueprintSetLibrary.cpp
+++ b/Engine/Source/Runtime/Engine/Private/BlueprintSetLibrary.cpp
@@ -10,7 +10,14 @@ void UBlueprintSetLibrary::GenericSet_Add(const void* TargetSet, const FSetPrope
 	if (TargetSet)
 	{
 		FScriptSetHelper SetHelper(SetProperty, TargetSet);
-		SetHelper.AddElement(ItemPtr);
+		if (SetHelper.Num() < MaxSupportedSetSize)
+		{
+			SetHelper.AddElement(ItemPtr);
+		}
+		else if (SetHelper.FindElementIndexFromHash(ItemPtr) == INDEX_NONE)
+		{
+			FFrame::KismetExecutionMessage(*FString::Printf(TEXT("Attempted add to set '%s' beyond the maximum supported capacity!"), *SetProperty->GetName()), ELogVerbosity::Warning, UKismetArrayLibrary::ReachedMaximumContainerSizeWarning);
+		}
 	}
 }
 
@@ -221,4 +228,3 @@ void UBlueprintSetLibrary::GenericSet_SetSetPropertyByName(UObject* OwnerObject,
 		}
 	}
 }
-
diff --git a/Engine/Source/Runtime/Engine/Private/KismetArrayLibrary.cpp b/Engine/Source/Runtime/Engine/Private/KismetArrayLibrary.cpp
--- a/Engine/Source/Runtime/Engine/Private/KismetArrayLibrary.cpp
+++ b/Engine/Source/Runtime/Engine/Private/KismetArrayLibrary.cpp
@@ -22,6 +22,8 @@ const FName ResizeArrayNegativeWarning = FName("ResizeArrayNegativeWarning");
 const FName SwapElementsInArrayWarning = FName("SwapElementsInArrayWarning");
 const FName RandomAccessToEmptyArrayWarning = FName("RandomAccessToEmptyArrayWarning");
 
+const FName UKismetArrayLibrary::ReachedMaximumContainerSizeWarning = FName("ReachedMaximumContainerSizeWarning");
+
 UKismetArrayLibrary::UKismetArrayLibrary(const FObjectInitializer& ObjectInitializer)
 	: Super(ObjectInitializer)
 {
@@ -67,6 +69,12 @@ UKismetArrayLibrary::UKismetArrayLibrary(const FObjectInitializer& ObjectInitial
 			LOCTEXT("RandomAccessToEmptyArrayWarning", "Random access to empty array")
 		)
 	);
+	FBlueprintSupport::RegisterBlueprintWarning(
+		FBlueprintWarningDeclaration(
+			ReachedMaximumContainerSizeWarning,
+			LOCTEXT("ReachedMaximumContainerSizeWarning", "Container has reached maximum size")
+		)
+	);
 }
 
 void UKismetArrayLibrary::FilterArray(const TArray<AActor*>& TargetArray, TSubclassOf<class AActor> FilterClass, TArray<AActor*>& FilteredArray)
@@ -88,10 +96,17 @@ int32 UKismetArrayLibrary::GenericArray_Add(void* TargetArray, const FArrayPrope
 	if( TargetArray )
 	{
 		FScriptArrayHelper ArrayHelper(ArrayProp, TargetArray);
-		FProperty* InnerProp = ArrayProp->Inner;
+		if (ArrayHelper.Num() < MaxSupportedArraySize)
+		{
+			FProperty* InnerProp = ArrayProp->Inner;
 
-		NewIndex = ArrayHelper.AddValue();
-		InnerProp->CopySingleValueToScriptVM(ArrayHelper.GetRawPtr(NewIndex), NewItem);
+			NewIndex = ArrayHelper.AddValue();
+			InnerProp->CopySingleValueToScriptVM(ArrayHelper.GetRawPtr(NewIndex), NewItem);
+		}
+		else
+		{
+			FFrame::KismetExecutionMessage(*FString::Printf(TEXT("Attempted add to array '%s' beyond the maximum supported capacity!"), *ArrayProp->GetName()), ELogVerbosity::Warning, ReachedMaximumContainerSizeWarning);
+		}
 	}
 	return NewIndex;
 }
@@ -106,8 +121,15 @@ int32 UKismetArrayLibrary::GenericArray_AddUnique(void* TargetArray, const FArra
 
 		if (GenericArray_Find(TargetArray, ArrayProp, NewItem) == INDEX_NONE)
 		{
-			NewIndex = ArrayHelper.AddValue();
-			InnerProp->CopySingleValueToScriptVM(ArrayHelper.GetRawPtr(NewIndex), NewItem);
+			if (ArrayHelper.Num() < MaxSupportedArraySize)
+			{
+				NewIndex = ArrayHelper.AddValue();
+				InnerProp->CopySingleValueToScriptVM(ArrayHelper.GetRawPtr(NewIndex), NewItem);
+			}
+			else
+			{
+				FFrame::KismetExecutionMessage(*FString::Printf(TEXT("Attempted add to array '%s' beyond the maximum supported capacity!"), *ArrayProp->GetName()), ELogVerbosity::Warning, ReachedMaximumContainerSizeWarning);
+			}
 		}
 	}
 	return NewIndex;
@@ -156,11 +178,24 @@ void UKismetArrayLibrary::GenericArray_Append(void* TargetArray, const FArrayPro
 		{
 			FProperty* InnerProp = TargetArrayProp->Inner;
 
-			const int32 EndIdx = SourceArrayHelper.Num();
-			int32 StartIdx = TargetArrayHelper.AddValues(SourceArrayHelper.Num());
-			for(int32 x = 0; x < EndIdx; ++x, ++StartIdx)
+			// Determine the maximum index value that we can store, and compute the difference between that and the index of the
+			// last element in the target array. If the difference is smaller than the size of the source array, then we'll only
+			// add that number of values (i.e. until we reach the maximum supported size). Otherwise, we append the full source array.
+			const int32 MaxIdx = MaxSupportedArraySize - 1;
+			const int32 AddNum = FMath::Min(SourceArrayHelper.Num(), MaxIdx - GetLastIndex(TargetArrayHelper));
+			if (AddNum > 0)
 			{
-				InnerProp->CopySingleValueToScriptVM(TargetArrayHelper.GetRawPtr(StartIdx), SourceArrayHelper.GetRawPtr(x));
+				int32 StartIdx = TargetArrayHelper.AddValues(AddNum);
+				for (int32 x = 0; x < AddNum; ++x, ++StartIdx)
+				{
+					InnerProp->CopySingleValueToScriptVM(TargetArrayHelper.GetRawPtr(StartIdx), SourceArrayHelper.GetRawPtr(x));
+				}
+			}
+
+			// Warn if we didn't add all the values from the source array (due to having reached our maximum supported size).
+			if (AddNum < SourceArrayHelper.Num())
+			{
+				FFrame::KismetExecutionMessage(*FString::Printf(TEXT("Attempted append to array '%s' beyond the maximum supported capacity!"), *TargetArrayProp->GetName()), ELogVerbosity::Warning, ReachedMaximumContainerSizeWarning);
 			}
 		}
 	}
@@ -361,7 +396,7 @@ void UKismetArrayLibrary::GenericArray_Set(void* TargetArray, const FArrayProper
 		FProperty* InnerProp = ArrayProp->Inner;
 
 		// Expand the array, if desired
-		if (!ArrayHelper.IsValidIndex(Index) && bSizeToFit && (Index >= 0))
+		if (!ArrayHelper.IsValidIndex(Index) && bSizeToFit && (Index >= 0) && (Index < MaxSupportedArraySize))
 		{
 			ArrayHelper.ExpandForIndex(Index);
 		}
@@ -370,6 +405,10 @@ void UKismetArrayLibrary::GenericArray_Set(void* TargetArray, const FArrayProper
 		{
 			InnerProp->CopySingleValueToScriptVM(ArrayHelper.GetRawPtr(Index), NewItem);
 		}
+		else if (Index >= MaxSupportedArraySize)
+		{
+			FFrame::KismetExecutionMessage(*FString::Printf(TEXT("Attempted to set an index on array %s beyond its maximum supported capacity!"), *ArrayProp->GetName()), ELogVerbosity::Warning, ReachedMaximumContainerSizeWarning);
+		}
 		else
 		{
 			FFrame::KismetExecutionMessage(*FString::Printf(TEXT("Attempted to set an invalid index on array %s [%d/%d]!"), *ArrayProp->GetName(), Index, GetLastIndex(ArrayHelper)), ELogVerbosity::Warning, SetOutOfBoundsWarning);
