World Partition: move cell debug name to base class, so all hash implementations can use it.

diff --git a/Engine/Source/Runtime/Engine/Private/WorldPartition/WorldPartitionRuntimeCellData.cpp b/Engine/Source/Runtime/Engine/Private/WorldPartition/WorldPartitionRuntimeCellData.cpp
--- a/Engine/Source/Runtime/Engine/Private/WorldPartition/WorldPartitionRuntimeCellData.cpp
+++ b/Engine/Source/Runtime/Engine/Private/WorldPartition/WorldPartitionRuntimeCellData.cpp
@@ -11,6 +11,12 @@ UWorldPartitionRuntimeCellData::UWorldPartitionRuntimeCellData(const FObjectInit
 	, ContentBounds(ForceInit)
 {}
 
+void UWorldPartitionRuntimeCellData::Serialize(FArchive& Ar)
+{
+	Super::Serialize(Ar);
+	Ar << DebugName;
+}
+
 void UWorldPartitionRuntimeCellData::ResetStreamingSourceInfo() const
 {
 	check(ShouldResetStreamingSourceInfo());
@@ -44,3 +50,8 @@ FBox UWorldPartitionRuntimeCellData::GetCellBounds() const
 {
 	return ContentBounds;
 }
+
+FString UWorldPartitionRuntimeCellData::GetDebugName() const
+{
+	return DebugName.GetString();
+}
\ No newline at end of file
diff --git a/Engine/Source/Runtime/Engine/Private/WorldPartition/WorldPartitionRuntimeCellDataSpatialHash.cpp b/Engine/Source/Runtime/Engine/Private/WorldPartition/WorldPartitionRuntimeCellDataSpatialHash.cpp
--- a/Engine/Source/Runtime/Engine/Private/WorldPartition/WorldPartitionRuntimeCellDataSpatialHash.cpp
+++ b/Engine/Source/Runtime/Engine/Private/WorldPartition/WorldPartitionRuntimeCellDataSpatialHash.cpp
@@ -27,12 +27,6 @@ UWorldPartitionRuntimeCellDataSpatialHash::UWorldPartitionRuntimeCellDataSpatial
 	, CachedMinSquareDistanceToSource(MAX_dbl)
 {}
 
-void UWorldPartitionRuntimeCellDataSpatialHash::Serialize(FArchive& Ar)
-{
-	Super::Serialize(Ar);
-	Ar << DebugName;
-}
-
 void UWorldPartitionRuntimeCellDataSpatialHash::ResetStreamingSourceInfo() const
 {
 	Super::ResetStreamingSourceInfo();
@@ -182,9 +176,4 @@ bool UWorldPartitionRuntimeCellDataSpatialHash::IsDebugShown() const
 	return Super::IsDebugShown() && 
 		   FWorldPartitionDebugHelper::IsDebugRuntimeHashGridShown(GridName) &&
 		   FWorldPartitionDebugHelper::IsDebugCellNameShown(GetName());
-}
-
-FString UWorldPartitionRuntimeCellDataSpatialHash::GetDebugName() const
-{
-	return DebugName.GetString();
 }
\ No newline at end of file
diff --git a/Engine/Source/Runtime/Engine/Public/WorldPartition/WorldPartitionRuntimeCellData.h b/Engine/Source/Runtime/Engine/Public/WorldPartition/WorldPartitionRuntimeCellData.h
--- a/Engine/Source/Runtime/Engine/Public/WorldPartition/WorldPartitionRuntimeCellData.h
+++ b/Engine/Source/Runtime/Engine/Public/WorldPartition/WorldPartitionRuntimeCellData.h
@@ -2,6 +2,7 @@
 #pragma once
 
 #include "CoreMinimal.h"
+#include "StringDev.h"
 #include "WorldPartition/WorldPartitionStreamingSource.h"
 #include "WorldPartitionRuntimeCellData.generated.h"
 
@@ -15,6 +16,10 @@ class ENGINE_API UWorldPartitionRuntimeCellData : public UObject
 
 	virtual ~UWorldPartitionRuntimeCellData() {}
 
+	//~Begin UObject Interface
+	void Serialize(FArchive& Ar);
+	//~End UObject Interface
+
 	inline bool ShouldResetStreamingSourceInfo() const { return CachedSourceInfoEpoch != StreamingSourceCacheEpoch; }
 
 	virtual void ResetStreamingSourceInfo() const;
@@ -26,7 +31,7 @@ class ENGINE_API UWorldPartitionRuntimeCellData : public UObject
 	virtual FBox GetCellBounds() const;
 
 	virtual bool IsDebugShown() const { return true; }
-	virtual FString GetDebugName() const { return GetName(); }
+	virtual FString GetDebugName() const;
 
 	static int32 StreamingSourceCacheEpoch;
 	static inline void DirtyStreamingSourceCacheEpoch() { ++StreamingSourceCacheEpoch; }
@@ -42,4 +47,6 @@ class ENGINE_API UWorldPartitionRuntimeCellData : public UObject
 
 	UPROPERTY()
 	FBox ContentBounds;
+
+	FStringTest DebugName;
 };
\ No newline at end of file
diff --git a/Engine/Source/Runtime/Engine/Public/WorldPartition/WorldPartitionRuntimeCellDataSpatialHash.h b/Engine/Source/Runtime/Engine/Public/WorldPartition/WorldPartitionRuntimeCellDataSpatialHash.h
--- a/Engine/Source/Runtime/Engine/Public/WorldPartition/WorldPartitionRuntimeCellDataSpatialHash.h
+++ b/Engine/Source/Runtime/Engine/Public/WorldPartition/WorldPartitionRuntimeCellDataSpatialHash.h
@@ -2,7 +2,6 @@
 #pragma once
 
 #include "CoreMinimal.h"
-#include "StringDev.h"
 #include "WorldPartition/WorldPartitionRuntimeCellData.h"
 #include "WorldPartitionRuntimeCellDataSpatialHash.generated.h"
 
@@ -14,10 +13,6 @@ class ENGINE_API UWorldPartitionRuntimeCellDataSpatialHash : public UWorldPartit
 {
 	GENERATED_UCLASS_BODY()
 
-	//~Begin UObject Interface
-	void Serialize(FArchive& Ar);
-	//~End UObject Interface
-
 	//~Begin UWorldPartitionRuntimeCellData
 	virtual void ResetStreamingSourceInfo() const override;
 	virtual void AppendStreamingSourceInfo(const FWorldPartitionStreamingSource& Source, const FSphericalSector& SourceShape) const override;
@@ -25,7 +20,6 @@ class ENGINE_API UWorldPartitionRuntimeCellDataSpatialHash : public UWorldPartit
 	virtual int32 SortCompare(const UWorldPartitionRuntimeCellData* InOther, bool bCanUseSortingCache = true) const override;
 	virtual FBox GetCellBounds() const override;
 	virtual bool IsDebugShown() const override;
-	virtual FString GetDebugName() const override;
 	//~End UWorldPartitionRuntimeCellData
 
 	bool IsBlockingSource() const { return bCachedIsBlockingSource; }
@@ -43,8 +37,6 @@ class ENGINE_API UWorldPartitionRuntimeCellDataSpatialHash : public UWorldPartit
 	UPROPERTY()
 	FName GridName;
 
-	FStringTest DebugName;
-
 private:
 	float ComputeSourceToCellAngleFactor(const FSphericalSector& SourceShape) const;
 
