Level Instances: fix regression when editing Level Instances that are instanced multiple times.

diff --git a/Engine/Source/Runtime/Engine/Private/WorldPartition/WorldPartition.cpp b/Engine/Source/Runtime/Engine/Private/WorldPartition/WorldPartition.cpp
--- a/Engine/Source/Runtime/Engine/Private/WorldPartition/WorldPartition.cpp
+++ b/Engine/Source/Runtime/Engine/Private/WorldPartition/WorldPartition.cpp
@@ -332,8 +332,18 @@ void UWorldPartition::OnGCPostReachabilityAnalysis()
 }
 
 void UWorldPartition::OnPackageDirtyStateChanged(UPackage* Package)
-{
-	auto ShouldHandleActor = [this](AActor* Actor) {return Actor && Actor->IsMainPackageActor() && (Actor->GetLevel() != nullptr) && IsActorDescHandled(Actor);};
+{	
+	auto ShouldHandleActor = [this](AActor* Actor)
+	{
+		if (Actor && Actor->IsMainPackageActor())
+		{
+			if (ULevel* OuterLevel = Actor->GetTypedOuter<ULevel>())
+			{
+				return OuterLevel->GetWorldPartition() == this;
+			}
+		}
+		return false;
+	};
 
 	if (AActor* Actor = AActor::FindActorInPackage(Package); ShouldHandleActor(Actor))
 	{
