- Sum inner objects for none Cascade particle systems
- WIthout this we are only including the size of the base object for Niagara, not all the Scripts, etc

#rnx

diff --git a/Engine/Source/Runtime/Engine/Private/UnrealEngine.cpp b/Engine/Source/Runtime/Engine/Private/UnrealEngine.cpp
--- a/Engine/Source/Runtime/Engine/Private/UnrealEngine.cpp
+++ b/Engine/Source/Runtime/Engine/Private/UnrealEngine.cpp
@@ -7105,8 +7105,18 @@ bool UEngine::HandleListParticleSystemsCommand( const TCHAR* Cmd, FOutputDevice&
 	{			
 		UFXSystemAsset* Tree = *SystemIt;
 		Description = FString::Printf(TEXT("%s"), *Tree->GetPathName());
-		FArchiveCountMem Count( Tree );
-		int32 RootSize = Count.GetMax();
+
+		int32 RootSize = 0;
+		if (Tree->IsA<UParticleSystem>())
+		{
+			RootSize = FArchiveCountMem(Tree).GetMax();
+		}
+		else
+		{
+			FResourceSizeEx ResSize = FResourceSizeEx(EResourceSizeMode::EstimatedTotal);
+			Tree->GetResourceSizeEx(ResSize);
+			RootSize = ResSize.GetTotalMemoryBytes();
+		}
 
 		SortedSets.Add(FSortedParticleSet(Description, RootSize, RootSize, 0, 0, 0, FResourceSizeEx(EResourceSizeMode::EstimatedTotal), FResourceSizeEx(EResourceSizeMode::Exclusive)));
 		SortMap.Add(Tree,SortedSets.Num() - 1);
