Removed ApplicationCore public include. It is technically correct but there are targets that have issues with this dependency.

diff --git a/Engine/Source/Runtime/Engine/Engine.Build.cs b/Engine/Source/Runtime/Engine/Engine.Build.cs
--- a/Engine/Source/Runtime/Engine/Engine.Build.cs
+++ b/Engine/Source/Runtime/Engine/Engine.Build.cs
@@ -18,7 +18,6 @@ public class Engine : ModuleRules
 		PublicIncludePathModuleNames.AddRange(
 			new string[] {
 				"AnimationCore",
-				"ApplicationCore",
 				"AudioExtensions",
 				"AudioMixer", 
 				"AudioMixerCore",
