Skeletal Mesh: Do not flush the render pipe if not overriding skin weights or vertex colors.

#rnx

diff --git a/Engine/Source/Runtime/Engine/Classes/Components/SkinnedMeshComponent.h b/Engine/Source/Runtime/Engine/Classes/Components/SkinnedMeshComponent.h
--- a/Engine/Source/Runtime/Engine/Classes/Components/SkinnedMeshComponent.h
+++ b/Engine/Source/Runtime/Engine/Classes/Components/SkinnedMeshComponent.h
@@ -158,14 +158,14 @@ struct ENGINE_API FSkelMeshComponentLODInfo
 	~FSkelMeshComponentLODInfo();
 
 	void ReleaseOverrideVertexColorsAndBlock();
-	void BeginReleaseOverrideVertexColors();
+	bool BeginReleaseOverrideVertexColors();
 	void EndReleaseOverrideVertexColors();
 private:
 	void CleanUpOverrideVertexColors();
 
 public:
 	void ReleaseOverrideSkinWeightsAndBlock();
-	void BeginReleaseOverrideSkinWeights();
+	bool BeginReleaseOverrideSkinWeights();
 	void EndReleaseOverrideSkinWeights();
 private:
 	void CleanUpOverrideSkinWeights();
diff --git a/Engine/Source/Runtime/Engine/Private/Components/SkinnedMeshComponent.cpp b/Engine/Source/Runtime/Engine/Private/Components/SkinnedMeshComponent.cpp
--- a/Engine/Source/Runtime/Engine/Private/Components/SkinnedMeshComponent.cpp
+++ b/Engine/Source/Runtime/Engine/Private/Components/SkinnedMeshComponent.cpp
@@ -1220,16 +1220,29 @@ void USkinnedMeshComponent::InitLODInfos()
 				// Batch release all resources of LOD infos we're about to destruct.
 				// This is relevant when overrides have been set but LODInfos are
 				// re-initialized, for example by changing the mesh at runtime.
+				bool bNeedRenderFlush = false;
+				
 				for (int32 Idx = 0; Idx < LODInfo.Num(); ++Idx)
 				{
-					LODInfo[Idx].BeginReleaseOverrideSkinWeights();
-					LODInfo[Idx].BeginReleaseOverrideVertexColors();
+					if (LODInfo[Idx].BeginReleaseOverrideSkinWeights())
+					{
+						bNeedRenderFlush = true;
+					}
+					
+					if (LODInfo[Idx].BeginReleaseOverrideVertexColors())
+					{
+						bNeedRenderFlush = true;
+					}
 				}
-				FlushRenderingCommands();
-				for (int32 Idx = 0; Idx < LODInfo.Num(); ++Idx)
+
+				if (bNeedRenderFlush)
 				{
-					LODInfo[Idx].EndReleaseOverrideSkinWeights();
-					LODInfo[Idx].EndReleaseOverrideVertexColors();
+					FlushRenderingCommands();
+					for (int32 Idx = 0; Idx < LODInfo.Num(); ++Idx)
+					{
+						LODInfo[Idx].EndReleaseOverrideSkinWeights();
+						LODInfo[Idx].EndReleaseOverrideVertexColors();
+					}
 				}
 			}
 			
@@ -4023,13 +4036,16 @@ void FSkelMeshComponentLODInfo::ReleaseOverrideVertexColorsAndBlock()
 	}
 }
 
-void FSkelMeshComponentLODInfo::BeginReleaseOverrideVertexColors()
+bool FSkelMeshComponentLODInfo::BeginReleaseOverrideVertexColors()
 {
 	if (OverrideVertexColors)
 	{
 		// enqueue a rendering command to release
 		BeginReleaseResource(OverrideVertexColors);
+		return true;
 	}
+
+	return false;
 }
 
 void FSkelMeshComponentLODInfo::EndReleaseOverrideVertexColors()
@@ -4059,13 +4075,16 @@ void FSkelMeshComponentLODInfo::ReleaseOverrideSkinWeightsAndBlock()
 	}
 }
 
-void FSkelMeshComponentLODInfo::BeginReleaseOverrideSkinWeights()
+bool FSkelMeshComponentLODInfo::BeginReleaseOverrideSkinWeights()
 {
 	if (OverrideSkinWeights)
 	{
 		// enqueue a rendering command to release
 		OverrideSkinWeights->BeginReleaseResources();
+		return true;
 	}
+
+	return false;
 }
 
 void FSkelMeshComponentLODInfo::EndReleaseOverrideSkinWeights()
