Substrate - bring back support for CustomUV on root node.

diff --git a/Engine/Source/Runtime/Engine/Private/Materials/Material.cpp b/Engine/Source/Runtime/Engine/Private/Materials/Material.cpp
--- a/Engine/Source/Runtime/Engine/Private/Materials/Material.cpp
+++ b/Engine/Source/Runtime/Engine/Private/Materials/Material.cpp
@@ -6656,6 +6656,11 @@ bool UMaterial::IsPropertySupported(EMaterialProperty InProperty) const
 			bSupported = true;
 			break;
 		}
+
+		if (InProperty >= MP_CustomizedUVs0 && InProperty <= MP_CustomizedUVs7)
+		{
+			bSupported = (InProperty - MP_CustomizedUVs0) < NumCustomizedUVs;
+		}
 	}
 	return bSupported;
 }
@@ -6849,6 +6854,11 @@ static bool IsPropertyActive_Internal(EMaterialProperty InProperty,
 				Active = true;
 				break;
 			}
+
+			if (InProperty >= MP_CustomizedUVs0 && InProperty <= MP_CustomizedUVs7)
+			{
+				Active = true;
+			}
 		}
 	}
 	else
