Iris - Use ActorReplicationBridge when replication is needed on an Actor.

#rnx

diff --git a/Engine/Source/Runtime/Engine/Private/Net/Iris/ReplicationSystem/ReplicationSystemUtil.cpp b/Engine/Source/Runtime/Engine/Private/Net/Iris/ReplicationSystem/ReplicationSystemUtil.cpp
--- a/Engine/Source/Runtime/Engine/Private/Net/Iris/ReplicationSystem/ReplicationSystemUtil.cpp
+++ b/Engine/Source/Runtime/Engine/Private/Net/Iris/ReplicationSystem/ReplicationSystemUtil.cpp
@@ -328,12 +328,17 @@ void FReplicationSystemUtil::AddDependentActor(const AActor* Parent, AActor* Chi
 		{
 			if (ReplicationSystem->IsServer())
 			{
-				if (UObjectReplicationBridge* Bridge = ReplicationSystem->GetReplicationBridgeAs<UObjectReplicationBridge>())
+				if (UActorReplicationBridge* Bridge = ReplicationSystem->GetReplicationBridgeAs<UActorReplicationBridge>())
 				{
 					const FNetRefHandle ParentRefHandle = Bridge->GetReplicatedRefHandle(ParentHandle);
 					if (ParentRefHandle.IsValid())
 					{
-						const FNetRefHandle ChildRefHandle = Bridge->BeginReplication(Child);
+						FNetRefHandle ChildRefHandle = Bridge->GetReplicatedRefHandle(Child);
+						if (!ChildRefHandle.IsValid())
+						{
+							const FActorBeginReplicationParams BeginReplicationParams;
+							ChildRefHandle = Bridge->BeginReplication(Child, BeginReplicationParams);
+						}
 						if (ensureMsgf(ChildRefHandle.IsValid(), TEXT("FReplicationSystemUtil::AddDependentActor Child %s must be replicated"), *GetPathNameSafe(Child)))
 						{
 							Bridge->AddDependentObject(ParentRefHandle, ChildRefHandle, SchedulingHint);
