Optimize UTexture::PostEditChange to use reverse lookups instead of costly object iteration
  - Add CompositeTexture -> Texture reverse lookup
  - Add mask for reverse lookup validation so we can validate only selected tables to improve performance

diff --git a/Engine/Plugins/Experimental/Mutable/Source/CustomizableObjectEditor/Private/MuCOE/GenerateMutableSource/GenerateMutableSourceSurface.cpp b/Engine/Plugins/Experimental/Mutable/Source/CustomizableObjectEditor/Private/MuCOE/GenerateMutableSource/GenerateMutableSourceSurface.cpp
--- a/Engine/Plugins/Experimental/Mutable/Source/CustomizableObjectEditor/Private/MuCOE/GenerateMutableSource/GenerateMutableSourceSurface.cpp
+++ b/Engine/Plugins/Experimental/Mutable/Source/CustomizableObjectEditor/Private/MuCOE/GenerateMutableSource/GenerateMutableSourceSurface.cpp
@@ -670,7 +670,7 @@ mu::NodeSurfacePtr GenerateMutableSourceSurface(const UEdGraphPin * Pin, FMutabl
 							LastImage = MipmapImage;
 
 							// Apply composite image. This needs to be computed after mipmaps generation. 	
-							if (ReferenceTexture && ReferenceTexture->CompositeTexture && ReferenceTexture->CompositeTextureMode != CTM_Disabled)
+							if (ReferenceTexture && ReferenceTexture->GetCompositeTexture() && ReferenceTexture->CompositeTextureMode != CTM_Disabled)
 							{
 								mu::NodeImageNormalCompositePtr CompositedImage = new mu::NodeImageNormalComposite();
 								CompositedImage->SetBase(LastImage.get());
@@ -693,7 +693,7 @@ mu::NodeSurfacePtr GenerateMutableSourceSurface(const UEdGraphPin * Pin, FMutabl
 
 								mu::NodeImageConstantPtr CompositeNormalImage = new mu::NodeImageConstant();
 
-								UTexture2D* ReferenceCompositeNormalTexture = Cast<UTexture2D>(ReferenceTexture->CompositeTexture);
+								UTexture2D* ReferenceCompositeNormalTexture = Cast<UTexture2D>(ReferenceTexture->GetCompositeTexture());
 								if (ReferenceCompositeNormalTexture)
 								{
 									GenerationContext.ArrayTextureUnrealToMutableTask.Add(FTextureUnrealToMutableTask(CompositeNormalImage, ReferenceCompositeNormalTexture, Node, true));
diff --git a/Engine/Source/Runtime/Engine/Classes/Engine/Texture.h b/Engine/Source/Runtime/Engine/Classes/Engine/Texture.h
--- a/Engine/Source/Runtime/Engine/Classes/Engine/Texture.h
+++ b/Engine/Source/Runtime/Engine/Classes/Engine/Texture.h
@@ -25,6 +25,7 @@
 #include "ImageCore.h"
 #if WITH_EDITORONLY_DATA
 #include "Misc/TVariant.h"
+#include "ObjectCacheEventSink.h"
 #include "DerivedDataCacheKeyProxy.h"
 #endif
 
@@ -1231,15 +1232,49 @@ public:
 	/** Per asset specific setting to define the mip-map generation properties like sharpening and kernel size. */
 	UPROPERTY(EditAnywhere, BlueprintReadWrite, Category=LevelOfDetail)
 	TEnumAsByte<enum TextureMipGenSettings> MipGenSettings;
-	
+
 	/**
 	 * Can be defined to modify the roughness based on the normal map variation (mostly from mip maps).
 	 * MaxAlpha comes in handy to define a base roughness if no source alpha was there.
 	 * Make sure the normal map has at least as many mips as this texture.
 	 */
-	UPROPERTY(EditAnywhere, BlueprintReadWrite, Category=Compositing)
+	UE_DEPRECATED(5.3, "Use GetCompositeTexture() and SetCompositeTexture() instead.")
+	UPROPERTY(EditAnywhere, BlueprintReadWrite, Category=Compositing, meta = (AllowPrivateAccess), Setter = SetCompositeTexture, Getter = GetCompositeTexture)
 	TObjectPtr<class UTexture> CompositeTexture;
 
+private:
+	/** Used to track down when CompositeTexture has been modified for notification purpose.
+	*	NOTE: Do not make a TObjectPtr or UPROPERTY as it would defeat the purpose.
+	*/
+	class UTexture* KnownCompositeTexture = nullptr;
+
+	/** Called when the CompositeTexture property gets overwritten without us knowing about it */
+	ENGINE_API void OutdatedKnownCompositeTextureDetected() const;
+	ENGINE_API void NotifyIfCompositeTextureChanged();
+
+public:
+	void SetCompositeTexture(UTexture* InCompositeTexture)
+	{
+PRAGMA_DISABLE_DEPRECATION_WARNINGS
+		CompositeTexture = InCompositeTexture;
+PRAGMA_ENABLE_DEPRECATION_WARNINGS
+		NotifyIfCompositeTextureChanged();
+	}
+
+	UTexture* GetCompositeTexture() const
+	{
+PRAGMA_DISABLE_DEPRECATION_WARNINGS
+		// This should never happen and is a last resort, we should have caught the property overwrite well before we reach this code
+		// but this can happen from legacy code that use reflection to set the property without using the new _InContainer functions
+		// which will bypass the setter we put in place.
+		if (KnownCompositeTexture != CompositeTexture)
+		{
+			OutdatedKnownCompositeTextureDetected();
+		}
+		return CompositeTexture;
+PRAGMA_ENABLE_DEPRECATION_WARNINGS
+	}
+
 	/* defines how the CompositeTexture is applied, e.g. CTM_RoughnessFromNormalAlpha */
 	UPROPERTY(EditAnywhere, BlueprintReadWrite, Category=Compositing, AdvancedDisplay)
 	TEnumAsByte<enum ECompositeTextureMode> CompositeTextureMode;
diff --git a/Engine/Source/Runtime/Engine/Private/ObjectCacheContext.cpp b/Engine/Source/Runtime/Engine/Private/ObjectCacheContext.cpp
--- a/Engine/Source/Runtime/Engine/Private/ObjectCacheContext.cpp
+++ b/Engine/Source/Runtime/Engine/Private/ObjectCacheContext.cpp
@@ -43,6 +43,19 @@ static TAutoConsoleVariable<int32> CVarObjectReverseLookupMode(
 	ECVF_ReadOnly /* Can only be enabled from the command-line */
 );
 
+static TAutoConsoleVariable<int32> CVarObjectReverseLookupMask(
+	TEXT("Editor.ObjectReverseLookupMask"),
+	-1,
+	TEXT("-1 - Does validation on all objects type (slowest) \n")
+	TEXT("0  - Skip validation on all objects type \n")
+	TEXT("1  - Perform validation on Texture to Textures lookups \n")
+	TEXT("2  - Perform validation on Texture to Materials lookups \n")
+	TEXT("4  - Perform validation on Material to Primitives lookups \n")
+	TEXT("8  - Perform validation on StaticMesh To Components lookups \n")
+	TEXT("16 - Perform validation on Material to Materials lookups \n"),
+	ECVF_Default /* Can only be enabled from the command-line */
+);
+
 namespace ObjectCacheContextImpl { void Validate(); }
 
 static FAutoConsoleCommand CVarObjectReverseLookupValidate(
@@ -58,6 +71,19 @@ static FAutoConsoleCommand CVarObjectReverseLookupValidate(
 
 namespace ObjectCacheContextImpl {
 
+enum class EObjectReverseLookupMask : int32
+{
+	None = 0,
+	TextureToTextureLookupCache = 1,
+	TextureToMaterialLookupCache = 2,
+	MaterialToPrimitiveLookupCache = 4,
+	StaticMeshToComponentLookupCache = 8,
+	MaterialToMaterialsLookupCache = 16,
+	All = -1
+};
+
+ENUM_CLASS_FLAGS(EObjectReverseLookupMask);
+
 enum class EObjectReverseLookupMode
 {
 	Temporary = 0,
@@ -65,10 +91,10 @@ enum class EObjectReverseLookupMode
 	Comparison = 2
 };
 
+static bool GIsGetObjectReverseLookupModeInitialized = false;
 EObjectReverseLookupMode GetObjectReverseLookupMode()
 {
-	static bool bInitialized = false;
-	if (!bInitialized)
+	if (!GIsGetObjectReverseLookupModeInitialized)
 	{
 		check(IsInGameThread());
 		int32 Mode = 0;
@@ -76,12 +102,48 @@ EObjectReverseLookupMode GetObjectReverseLookupMode()
 		{
 			CVarObjectReverseLookupMode->Set(Mode, ECVF_SetByCommandline);
 		}
-		bInitialized = true;
+		GIsGetObjectReverseLookupModeInitialized = true;
 	}
 
 	return (EObjectReverseLookupMode)CVarObjectReverseLookupMode.GetValueOnAnyThread();
 }
 
+static bool GIsGetObjectReverseLookupMaskInitialized = false;
+EObjectReverseLookupMask GetObjectReverseLookupMask()
+{
+	if (!GIsGetObjectReverseLookupMaskInitialized)
+	{
+		check(IsInGameThread());
+		int32 Mask = 0;
+		if (FParse::Value(FCommandLine::Get(), TEXT("-ObjectReverseLookupMask="), Mask))
+		{
+			CVarObjectReverseLookupMask->Set(Mask, ECVF_SetByCommandline);
+		}
+		GIsGetObjectReverseLookupMaskInitialized = true;
+	}
+
+	return (EObjectReverseLookupMask)CVarObjectReverseLookupMask.GetValueOnAnyThread();
+}
+
+bool IsReverseLookupComparisonActive(EObjectReverseLookupMask Mask)
+{
+	return GetObjectReverseLookupMode() == EObjectReverseLookupMode::Comparison && (GetObjectReverseLookupMask() & Mask) != EObjectReverseLookupMask::None;
+}
+
+bool ShouldUseLookupTable(EObjectReverseLookupMask Mask)
+{
+	const EObjectReverseLookupMode Mode = GetObjectReverseLookupMode();
+	switch (Mode)
+	{
+		case EObjectReverseLookupMode::Comparison:
+			return !IsReverseLookupComparisonActive(Mask);
+		case EObjectReverseLookupMode::Permanent:
+			return true;
+		default:
+			return false;
+	}
+}
+
 struct FDefaultObjectSearchPredicate
 {
 	EInternalObjectFlags ExcludedInternalFlags;
@@ -283,6 +345,7 @@ private:
 
 FObjectReverseLookupCache<UMaterialInterface, UMaterialInstance>   GMaterialToMaterialInstanceLookupCache; // Parent -> Children
 FObjectReverseLookupCache<UTexture, UMaterialInterface>            GTextureToMaterialLookupCache;
+FObjectReverseLookupCache<UTexture, UTexture>                      GTextureToTextureLookupCache; // CompositeTexture -> Texture
 FObjectReverseLookupCache<UStaticMesh, UStaticMeshComponent>       GStaticMeshToComponentLookupCache;
 FObjectReverseLookupCache<UMaterialInterface, UPrimitiveComponent> GMaterialToPrimitiveLookupCache;
 
@@ -315,6 +378,20 @@ void Validate()
 		ErrorCount += TempLookup.Compare(GTextureToMaterialLookupCache);
 	}
 
+	// Scan and compare UTexture -> UTexture
+	{
+		FObjectReverseLookupCache<UTexture, UTexture> TempLookup;
+		for (TObjectIterator<UTexture> It(RF_ClassDefaultObject, true /*bIncludeDerivedClasses*/, GetObjectCacheInternalFlagsExclusion()); It; ++It)
+		{
+			UTexture* Texture = *It;
+			if (Texture->GetCompositeTexture() != nullptr && Texture->GetCompositeTexture() != Texture)
+			{
+				TempLookup.Update(*It, { It->GetCompositeTexture() });
+			}
+		}
+		ErrorCount += TempLookup.Compare(GTextureToTextureLookupCache);
+	}
+
 	// Scan and compare UMaterialInterface -> UPrimitiveComponent
 	{
 		FObjectReverseLookupCache<UMaterialInterface, UPrimitiveComponent> TempLookup;
@@ -502,6 +579,22 @@ namespace ObjectCacheContextImpl {
 			}
 		}
 	}
+
+	void GetTexturesAffectedByTexture_Iteration(UTexture* InTexture, TSet<UTexture*>& OutAffectedTextures)
+	{
+		TRACE_CPUPROFILER_EVENT_SCOPE(GetTextureAffectedByTexture_Iteration);
+
+		for (TObjectIterator<UTexture> It(RF_ClassDefaultObject, true /*bIncludeDerivedClasses*/, GetObjectCacheInternalFlagsExclusion()); It; ++It)
+		{
+			UTexture* Texture = *It;
+
+			if (Texture != InTexture && Texture->GetCompositeTexture() == InTexture && Texture->CompositeTextureMode != CTM_Disabled)
+			{
+				OutAffectedTextures.Add(Texture);
+			}
+		}
+	}
+
 #endif // #if WITH_EDITOR
 
 } // namespace ObjectCacheContextImpl
@@ -549,8 +642,7 @@ TObjectCacheIterator<UTexture> FObjectCacheContext::GetUsedTextures(UMaterialInt
 {
 	using namespace ObjectCacheContextImpl;
 #if WITH_EDITOR
-	const EObjectReverseLookupMode ObjectCacheContextMode = GetObjectReverseLookupMode();
-	if (ObjectCacheContextMode == EObjectReverseLookupMode::Permanent)
+	if (ShouldUseLookupTable(EObjectReverseLookupMask::TextureToMaterialLookupCache))
 	{
 		return GTextureToMaterialLookupCache.GetTo(MaterialInterface);
 	}
@@ -564,13 +656,13 @@ TObjectCacheIterator<UTexture> FObjectCacheContext::GetUsedTextures(UMaterialInt
 	}
 
 #if WITH_EDITOR
-	if (ObjectCacheContextMode == EObjectReverseLookupMode::Comparison)
+	if (IsReverseLookupComparisonActive(EObjectReverseLookupMask::TextureToMaterialLookupCache))
 	{
 		TSet<UTexture*> LookupResult;
 		for (UTexture* Texture : GTextureToMaterialLookupCache.GetTo(MaterialInterface))
 		{
 			LookupResult.FindOrAdd(Texture);
-			checkf(Textures->Contains(Texture), TEXT("Permanent Object Cache has an additionnal texture %s on material %s"), *Texture->GetFullName(), *MaterialInterface->GetFullName());
+			checkf(Textures->Contains(Texture), TEXT("Permanent Object Cache has an additional texture %s on material %s"), *Texture->GetFullName(), *MaterialInterface->GetFullName());
 		}
 
 		for (UTexture* Texture : *Textures)
@@ -607,8 +699,7 @@ TObjectCacheIterator<UMaterialInterface> FObjectCacheContext::GetUsedMaterials(U
 {
 	using namespace ObjectCacheContextImpl;
 #if WITH_EDITOR
-	const EObjectReverseLookupMode ObjectCacheContextMode = GetObjectReverseLookupMode();
-	if (ObjectCacheContextMode == EObjectReverseLookupMode::Permanent)
+	if (ShouldUseLookupTable(EObjectReverseLookupMask::MaterialToPrimitiveLookupCache))
 	{
 		return GMaterialToPrimitiveLookupCache.GetTo(Component);
 	}
@@ -634,13 +725,13 @@ TObjectCacheIterator<UMaterialInterface> FObjectCacheContext::GetUsedMaterials(U
 	// The only case where reverse lookup is allowed to diverge is when the render state is dirty meaning
 	// UsedMaterials might have been modified but the SceneProxy hasn't been recreated yet.
 	// This is not a problem as we're using the reverse proxy to know which component needs to be dirtied anyway.
-	if (ObjectCacheContextMode == EObjectReverseLookupMode::Comparison && Component->IsRenderStateCreated() && Component->SceneProxy && !Component->IsRenderStateDirty())
+	if (IsReverseLookupComparisonActive(EObjectReverseLookupMask::MaterialToPrimitiveLookupCache) && Component->IsRenderStateCreated() && Component->SceneProxy && !Component->IsRenderStateDirty())
 	{
 		TSet<UMaterialInterface*> LookupResult;
 		for (UMaterialInterface* Material : GMaterialToPrimitiveLookupCache.GetTo(Component))
 		{
 			LookupResult.FindOrAdd(Material);
-			checkf(Materials->Contains(Material), TEXT("Permanent Object Cache has an additionnal material %s for component %"), *Material->GetFullName(), *Component->GetFullName());
+			checkf(Materials->Contains(Material), TEXT("Permanent Object Cache has an additional material %s for component %s"), *Material->GetFullName(), *Component->GetFullName());
 		}
 
 		for (UMaterialInterface* Material : *Materials)
@@ -657,8 +748,7 @@ TObjectCacheIterator<UStaticMeshComponent> FObjectCacheContext::GetStaticMeshCom
 {
 	using namespace ObjectCacheContextImpl;
 #if WITH_EDITOR
-	const EObjectReverseLookupMode ObjectCacheContextMode = GetObjectReverseLookupMode();
-	if (ObjectCacheContextMode == EObjectReverseLookupMode::Permanent)
+	if (ShouldUseLookupTable(EObjectReverseLookupMask::StaticMeshToComponentLookupCache))
 	{
 		return GStaticMeshToComponentLookupCache.GetFrom(InStaticMesh);
 	}
@@ -682,13 +772,13 @@ TObjectCacheIterator<UStaticMeshComponent> FObjectCacheContext::GetStaticMeshCom
 	const TSet<UStaticMeshComponent*>& ComputeResult = Set ? *Set : EmptySet;
 
 #if WITH_EDITOR
-	if (ObjectCacheContextMode == EObjectReverseLookupMode::Comparison)
+	if (IsReverseLookupComparisonActive(EObjectReverseLookupMask::StaticMeshToComponentLookupCache))
 	{
 		TSet<UStaticMeshComponent*> LookupResult;
 		for (UStaticMeshComponent* Component : GStaticMeshToComponentLookupCache.GetFrom(InStaticMesh))
 		{
 			LookupResult.FindOrAdd(Component);
-			checkf(ComputeResult.Contains(Component), TEXT("Permanent Object Cache has an additionnal component %s for staticmesh %s"), *Component->GetFullName(), *InStaticMesh->GetFullName());
+			checkf(ComputeResult.Contains(Component), TEXT("Permanent Object Cache has an additional component %s for staticmesh %s"), *Component->GetFullName(), *InStaticMesh->GetFullName());
 		}
 
 		for (UStaticMeshComponent* Component : ComputeResult)
@@ -705,8 +795,7 @@ TObjectCacheIterator<UMaterialInterface> FObjectCacheContext::GetMaterialsAffect
 {
 	using namespace ObjectCacheContextImpl;
 #if WITH_EDITOR
-	const EObjectReverseLookupMode ObjectCacheContextMode = GetObjectReverseLookupMode();
-	if (ObjectCacheContextMode == EObjectReverseLookupMode::Permanent)
+	if (ShouldUseLookupTable(EObjectReverseLookupMask::TextureToMaterialLookupCache))
 	{
 		return GTextureToMaterialLookupCache.GetFrom(InTexture);
 	}
@@ -734,13 +823,13 @@ TObjectCacheIterator<UMaterialInterface> FObjectCacheContext::GetMaterialsAffect
 	const TSet<UMaterialInterface*>& ComputeResult = Set ? *Set : EmptySet;
 
 #if WITH_EDITOR
-	if (ObjectCacheContextMode == EObjectReverseLookupMode::Comparison)
+	if (IsReverseLookupComparisonActive(EObjectReverseLookupMask::TextureToMaterialLookupCache))
 	{
 		TSet<UMaterialInterface*> LookupResult;
 		for (UMaterialInterface* Material : GTextureToMaterialLookupCache.GetFrom(InTexture))
 		{
 			LookupResult.FindOrAdd(Material);
-			checkf(ComputeResult.Contains(Material), TEXT("Permanent Object Cache has an additionnal material %s for texture %s"), *Material->GetFullName(), *InTexture->GetFullName());
+			checkf(ComputeResult.Contains(Material), TEXT("Permanent Object Cache has an additional material %s for texture %s"), *Material->GetFullName(), *InTexture->GetFullName());
 		}
 
 		for (UMaterialInterface* Material : ComputeResult)
@@ -753,6 +842,44 @@ TObjectCacheIterator<UMaterialInterface> FObjectCacheContext::GetMaterialsAffect
 	return TObjectCacheIterator<UMaterialInterface>(ComputeResult.Array());
 }
 
+#if WITH_EDITOR
+
+TObjectCacheIterator<UTexture> FObjectCacheContext::GetTexturesAffectedByTexture(UTexture* InTexture)
+{
+	using namespace ObjectCacheContextImpl;
+	TRACE_CPUPROFILER_EVENT_SCOPE(FObjectCacheContext::GetTexturesAffectedByTexture);
+
+	// Filter out if CompositeTextureMode is not active
+	auto Predicate = [](UTexture* Texture) { return Texture->CompositeTextureMode != CTM_Disabled; };
+
+	if (ShouldUseLookupTable(EObjectReverseLookupMask::TextureToTextureLookupCache))
+	{
+		return GTextureToTextureLookupCache.GetFrom(InTexture, Predicate);
+	}
+
+	TSet<UTexture*> IterationResults;
+	ObjectCacheContextImpl::GetTexturesAffectedByTexture_Iteration(InTexture, IterationResults);
+
+	if (IsReverseLookupComparisonActive(EObjectReverseLookupMask::TextureToTextureLookupCache))
+	{
+		TSet<UTexture*> LookupResults;
+		for (UTexture* Texture : GTextureToTextureLookupCache.GetFrom(InTexture, Predicate))
+		{
+			LookupResults.Add(Texture);
+			checkf(IterationResults.Contains(Texture), TEXT("Permanent Object Cache has an additional element %s"), *Texture->GetFullName());
+		}
+
+		for (UTexture* Texture : IterationResults)
+		{
+			checkf(LookupResults.Contains(Texture), TEXT("Permanent Object Cache is missing element %s"), *Texture->GetFullName());
+		}
+	}
+
+	return TObjectCacheIterator<UTexture>(IterationResults.Array());
+}
+
+#endif
+
 TObjectCacheIterator<USkinnedMeshComponent> FObjectCacheContext::GetSkinnedMeshComponents(USkinnedAsset* InSkinnedAsset)
 {
 	if (!SkinnedAssetToComponents.IsSet())
@@ -781,8 +908,7 @@ TObjectCacheIterator<UMaterialInterface> FObjectCacheContext::GetMaterialsAffect
 	TRACE_CPUPROFILER_EVENT_SCOPE(FObjectCacheContext::GetMaterialsAffectedByMaterials);
 
 #if WITH_EDITOR
-	const EObjectReverseLookupMode ObjectCacheContextMode = GetObjectReverseLookupMode();
-	if (ObjectCacheContextMode == EObjectReverseLookupMode::Permanent)
+	if (ShouldUseLookupTable(EObjectReverseLookupMask::MaterialToMaterialsLookupCache))
 	{
 		TSet<UMaterialInterface*> LookupResults;
 		ObjectCacheContextImpl::GetMaterialsAffectedByMaterials_Lookup(InMaterials, LookupResults);
@@ -794,19 +920,19 @@ TObjectCacheIterator<UMaterialInterface> FObjectCacheContext::GetMaterialsAffect
 	ObjectCacheContextImpl::GetMaterialsAffectedByMaterials_Iteration(InMaterials, IterationResults);
 
 #if WITH_EDITOR
-	if (ObjectCacheContextMode == EObjectReverseLookupMode::Comparison)
+	if (IsReverseLookupComparisonActive(EObjectReverseLookupMask::MaterialToMaterialsLookupCache))
 	{
 		TSet<UMaterialInterface*> LookupResults;
 		ObjectCacheContextImpl::GetMaterialsAffectedByMaterials_Lookup(InMaterials, LookupResults);
 		
 		for (UMaterialInterface* MaterialInterface : LookupResults)
 		{
-			checkf(IterationResults.Contains(MaterialInterface), TEXT("Permanent Object Cache has an additionnal %s"), *MaterialInterface->GetFullName());
+			checkf(IterationResults.Contains(MaterialInterface), TEXT("Permanent Object Cache has an additional element %s"), *MaterialInterface->GetFullName());
 		}
 
 		for (UMaterialInterface* MaterialInterface : IterationResults)
 		{
-			checkf(LookupResults.Contains(MaterialInterface), TEXT("Permanent Object Cache is missing a %s"), *MaterialInterface->GetFullName());
+			checkf(LookupResults.Contains(MaterialInterface), TEXT("Permanent Object Cache is missing element %s"), *MaterialInterface->GetFullName());
 		}
 	}
 #endif
@@ -820,8 +946,7 @@ TObjectCacheIterator<UPrimitiveComponent> FObjectCacheContext::GetPrimitivesAffe
 	TRACE_CPUPROFILER_EVENT_SCOPE(FObjectCacheContext::GetPrimitivesAffectedByMaterials);
 
 #if WITH_EDITOR
-	const EObjectReverseLookupMode ObjectCacheContextMode = GetObjectReverseLookupMode();
-	if (ObjectCacheContextMode == EObjectReverseLookupMode::Permanent)
+	if (ShouldUseLookupTable(EObjectReverseLookupMask::MaterialToPrimitiveLookupCache))
 	{
 		TSet<UPrimitiveComponent*> AffectedPrimitives;
 		ObjectCacheContextImpl::GetPrimitivesAffectedByMaterials_Lookup(InMaterials, AffectedPrimitives);
@@ -833,7 +958,7 @@ TObjectCacheIterator<UPrimitiveComponent> FObjectCacheContext::GetPrimitivesAffe
 	ObjectCacheContextImpl::GetPrimitivesAffectedByMaterials_Iteration(*this, InMaterials, IterationResults);
 
 #if WITH_EDITOR
-	if (ObjectCacheContextMode == EObjectReverseLookupMode::Comparison)
+	if (IsReverseLookupComparisonActive(EObjectReverseLookupMask::MaterialToPrimitiveLookupCache))
 	{
 		TSet<UPrimitiveComponent*> LookupResults;
 		ObjectCacheContextImpl::GetPrimitivesAffectedByMaterials_Lookup(InMaterials, LookupResults);
@@ -845,7 +970,7 @@ TObjectCacheIterator<UPrimitiveComponent> FObjectCacheContext::GetPrimitivesAffe
 		{
 			if (Component->IsRenderStateCreated() && Component->SceneProxy && !Component->IsRenderStateDirty())
 			{
-				checkf(IterationResults.Contains(Component), TEXT("Permanent Object Cache has an additionnal %s"), *Component->GetFullName());
+				checkf(IterationResults.Contains(Component), TEXT("Permanent Object Cache has an additional element %s"), *Component->GetFullName());
 			}
 		}
 
@@ -853,7 +978,7 @@ TObjectCacheIterator<UPrimitiveComponent> FObjectCacheContext::GetPrimitivesAffe
 		{
 			if (Component->IsRenderStateCreated() && Component->SceneProxy && !Component->IsRenderStateDirty())
 			{
-				checkf(LookupResults.Contains(Component), TEXT("Permanent Object Cache is missing a %s"), *Component->GetFullName());
+				checkf(LookupResults.Contains(Component), TEXT("Permanent Object Cache is missing element %s"), *Component->GetFullName());
 			}
 		}
 	}
@@ -913,6 +1038,7 @@ struct FObjectCacheEventSinkPrivate
 	static void NotifyReferencedTextureChanged_Concurrent(UMaterialInterface*);
 	static void NotifyStaticMeshChanged_Concurrent(UStaticMeshComponent*);
 	static void NotifyMaterialDestroyed_Concurrent(UMaterialInterface*);
+	static void NotifyCompositeTextureChanged_Concurrent(UTexture*);
 
 	enum ECacheEventType
 	{
@@ -920,7 +1046,8 @@ struct FObjectCacheEventSinkPrivate
 		EMaterialsChanged,
 		ERenderStateChanged,
 		EReferencedTextureChanged,
-		EStaticMeshChanged
+		EStaticMeshChanged,
+		ECompositeTextureChanged,
 	};
 
 	struct FNotifyEvent
@@ -930,6 +1057,7 @@ struct FObjectCacheEventSinkPrivate
 		const UPrimitiveComponent* PrimitiveComponent;
 		UStaticMeshComponent* StaticMeshComponent;
 		TArray<UMaterialInterface*> UsedMaterials;
+		UTexture* Texture;
 	};
 
 	static std::atomic<bool> bShouldQueueSinkEvents;
@@ -938,7 +1066,7 @@ struct FObjectCacheEventSinkPrivate
 	using TEventList = TLockFreePointerListFIFO<FNotifyEvent, PLATFORM_CACHE_LINE_SIZE>;
 	static TEventAllocator& GetAllocator();
 	static TEventList& GetNotifyEvents();
-	static void AddNotifyEvent(FObjectCacheEventSinkPrivate::ECacheEventType EventType, UMaterialInterface* MaterialInterface, const UPrimitiveComponent* PrimitiveComponent, UStaticMeshComponent* StaticMeshComponent, const TArray<UMaterialInterface*>* UsedMaterials);
+	static void AddNotifyEvent(FObjectCacheEventSinkPrivate::ECacheEventType EventType, UMaterialInterface* MaterialInterface, const UPrimitiveComponent* PrimitiveComponent, UStaticMeshComponent* StaticMeshComponent, const TArray<UMaterialInterface*>* UsedMaterials, UTexture* Texture);
 };
 
 } // namespace ObjectCacheContextImpl
@@ -964,6 +1092,22 @@ void FObjectCacheEventSink::EndQueueNotifyEvents()
 	FObjectCacheEventSinkPrivate::EndQueueNotifyEvents();
 }
 
+void FObjectCacheEventSink::NotifyCompositeTextureChanged_Concurrent(UTexture* InTexture)
+{
+	using namespace ObjectCacheContextImpl;
+
+	if (GetObjectReverseLookupMode() != EObjectReverseLookupMode::Temporary)
+	{
+		if (FObjectCacheEventSinkPrivate::bShouldQueueSinkEvents)
+		{
+			FObjectCacheEventSinkPrivate::AddNotifyEvent(FObjectCacheEventSinkPrivate::ECompositeTextureChanged, nullptr, nullptr, nullptr, nullptr, InTexture);
+		}
+		else
+		{
+			FObjectCacheEventSinkPrivate::NotifyCompositeTextureChanged_Concurrent(InTexture);
+		}
+	}
+}
 void FObjectCacheEventSink::NotifyMaterialDestroyed_Concurrent(UMaterialInterface* MaterialInterface)
 {
 	using namespace ObjectCacheContextImpl;
@@ -972,7 +1116,7 @@ void FObjectCacheEventSink::NotifyMaterialDestroyed_Concurrent(UMaterialInterfac
 	{
 		if (FObjectCacheEventSinkPrivate::bShouldQueueSinkEvents)
 		{
-			FObjectCacheEventSinkPrivate::AddNotifyEvent(FObjectCacheEventSinkPrivate::EMaterialDestroyed, MaterialInterface, nullptr, nullptr, nullptr);
+			FObjectCacheEventSinkPrivate::AddNotifyEvent(FObjectCacheEventSinkPrivate::EMaterialDestroyed, MaterialInterface, nullptr, nullptr, nullptr, nullptr);
 		}
 		else
 		{
@@ -989,7 +1133,7 @@ void FObjectCacheEventSink::NotifyUsedMaterialsChanged_Concurrent(const UPrimiti
 	{
 		if (FObjectCacheEventSinkPrivate::bShouldQueueSinkEvents)
 		{
-			FObjectCacheEventSinkPrivate::AddNotifyEvent(FObjectCacheEventSinkPrivate::EMaterialsChanged, nullptr, PrimitiveComponent, nullptr, &UsedMaterials);
+			FObjectCacheEventSinkPrivate::AddNotifyEvent(FObjectCacheEventSinkPrivate::EMaterialsChanged, nullptr, PrimitiveComponent, nullptr, &UsedMaterials, nullptr);
 		}
 		else
 		{
@@ -1008,7 +1152,7 @@ void FObjectCacheEventSink::NotifyRenderStateChanged_Concurrent(UActorComponent*
 		{
 			if (FObjectCacheEventSinkPrivate::bShouldQueueSinkEvents)
 			{
-				FObjectCacheEventSinkPrivate::AddNotifyEvent(FObjectCacheEventSinkPrivate::ERenderStateChanged, nullptr, Primitive, nullptr, nullptr);
+				FObjectCacheEventSinkPrivate::AddNotifyEvent(FObjectCacheEventSinkPrivate::ERenderStateChanged, nullptr, Primitive, nullptr, nullptr, nullptr);
 			}
 			else
 			{
@@ -1026,7 +1170,7 @@ void FObjectCacheEventSink::NotifyReferencedTextureChanged_Concurrent(UMaterialI
 	{
 		if (FObjectCacheEventSinkPrivate::bShouldQueueSinkEvents)
 		{
-			FObjectCacheEventSinkPrivate::AddNotifyEvent(FObjectCacheEventSinkPrivate::EReferencedTextureChanged, MaterialInterface, nullptr, nullptr, nullptr);
+			FObjectCacheEventSinkPrivate::AddNotifyEvent(FObjectCacheEventSinkPrivate::EReferencedTextureChanged, MaterialInterface, nullptr, nullptr, nullptr, nullptr);
 		}
 		else
 		{
@@ -1043,7 +1187,7 @@ void FObjectCacheEventSink::NotifyStaticMeshChanged_Concurrent(UStaticMeshCompon
 	{
 		if (FObjectCacheEventSinkPrivate::bShouldQueueSinkEvents)
 		{
-			FObjectCacheEventSinkPrivate::AddNotifyEvent(FObjectCacheEventSinkPrivate::EStaticMeshChanged, nullptr, nullptr, StaticMeshComponent, nullptr);
+			FObjectCacheEventSinkPrivate::AddNotifyEvent(FObjectCacheEventSinkPrivate::EStaticMeshChanged, nullptr, nullptr, StaticMeshComponent, nullptr, nullptr);
 		}
 		else
 		{
@@ -1068,7 +1212,7 @@ FObjectCacheEventSinkPrivate::TEventList& FObjectCacheEventSinkPrivate::GetNotif
 	return NotifyEvents;
 }
 
-void FObjectCacheEventSinkPrivate::AddNotifyEvent(FObjectCacheEventSinkPrivate::ECacheEventType EventType, UMaterialInterface* MaterialInterface, const UPrimitiveComponent* PrimitiveComponent, UStaticMeshComponent* StaticMeshComponent, const TArray<UMaterialInterface*>* UsedMaterials)
+void FObjectCacheEventSinkPrivate::AddNotifyEvent(FObjectCacheEventSinkPrivate::ECacheEventType EventType, UMaterialInterface* MaterialInterface, const UPrimitiveComponent* PrimitiveComponent, UStaticMeshComponent* StaticMeshComponent, const TArray<UMaterialInterface*>* UsedMaterials, UTexture* Texture)
 {
 
 	FNotifyEvent* Event = GetAllocator().New();
@@ -1076,6 +1220,7 @@ void FObjectCacheEventSinkPrivate::AddNotifyEvent(FObjectCacheEventSinkPrivate::
 	Event->MaterialInterface = MaterialInterface;
 	Event->PrimitiveComponent = PrimitiveComponent;
 	Event->StaticMeshComponent = StaticMeshComponent;
+	Event->Texture = Texture;
 
 	if (EventType == EMaterialsChanged)
 	{
@@ -1118,6 +1263,9 @@ void FObjectCacheEventSinkPrivate::ProcessQueuedNotifyEvents()
 		case EStaticMeshChanged:
 			NotifyStaticMeshChanged_Concurrent(Event->StaticMeshComponent);
 			break;
+		case ECompositeTextureChanged:
+			NotifyCompositeTextureChanged_Concurrent(Event->Texture);
+			break;
 		}
 
 		Allocator.Free(Event);
@@ -1131,6 +1279,24 @@ void FObjectCacheEventSinkPrivate::EndQueueNotifyEvents()
 	bShouldQueueSinkEvents = false;
 }
 
+
+void FObjectCacheEventSinkPrivate::NotifyCompositeTextureChanged_Concurrent(UTexture* Texture)
+{
+	using namespace ObjectCacheContextImpl;
+
+#if WITH_EDITOR
+	// Stop tracking textures that are being destroyed
+	if (Texture->HasAnyFlags(RF_BeginDestroyed) || Texture->GetCompositeTexture() == nullptr || Texture->GetCompositeTexture() == Texture)
+	{
+		GTextureToTextureLookupCache.Update(Texture, {});
+	}
+	else
+	{
+		GTextureToTextureLookupCache.Update(Texture, { Texture->GetCompositeTexture() });
+	}
+#endif
+}
+
 void FObjectCacheEventSinkPrivate::NotifyMaterialDestroyed_Concurrent(UMaterialInterface* MaterialInterface)
 {
 	using namespace ObjectCacheContextImpl;
diff --git a/Engine/Source/Runtime/Engine/Private/Texture.cpp b/Engine/Source/Runtime/Engine/Private/Texture.cpp
--- a/Engine/Source/Runtime/Engine/Private/Texture.cpp
+++ b/Engine/Source/Runtime/Engine/Private/Texture.cpp
@@ -21,6 +21,7 @@
 #include "Interfaces/ITargetPlatform.h"
 #include "Engine/TextureLODSettings.h"
 #include "RenderUtils.h"
+#include "ObjectCacheContext.h"
 #include "Rendering/StreamableTextureResource.h"
 #include "RenderingThread.h"
 #include "Interfaces/ITextureFormat.h"
@@ -785,24 +786,21 @@ void UTexture::PostEditChangeProperty(FPropertyChangedEvent& PropertyChangedEven
 		TRACE_CPUPROFILER_EVENT_SCOPE(UpdateDependentMaterials);
 
 		// Update any material that uses this texture and must force a recompile of cache resource
+		FObjectCacheContextScope ObjectCache;
+
 		TArray<UMaterial*> MaterialsToUpdate;
 		TSet<UMaterial*> BaseMaterialsThatUseThisTexture;
-		// this walks all Materials in the world
-		for (TObjectIterator<UMaterialInterface> It; It; ++It)
+		for (UMaterialInterface* MaterialInterface : ObjectCache.GetContext().GetMaterialsAffectedByTexture(this))
 		{
-			UMaterialInterface* MaterialInterface = *It;
-			if (DoesMaterialUseTexture(MaterialInterface, this))
+			UMaterial* Material = MaterialInterface->GetMaterial();
+			bool MaterialAlreadyCompute = false;
+			BaseMaterialsThatUseThisTexture.Add(Material, &MaterialAlreadyCompute);
+			if (!MaterialAlreadyCompute)
 			{
-				UMaterial* Material = MaterialInterface->GetMaterial();
-				bool MaterialAlreadyCompute = false;
-				BaseMaterialsThatUseThisTexture.Add(Material, &MaterialAlreadyCompute);
-				if (!MaterialAlreadyCompute)
+				if (Material->IsTextureForceRecompileCacheRessource(this))
 				{
-					if (Material->IsTextureForceRecompileCacheRessource(this))
-					{
-						MaterialsToUpdate.Add(Material);
-						Material->UpdateMaterialShaderCacheAndTextureReferences();
-					}
+					MaterialsToUpdate.Add(Material);
+					Material->UpdateMaterialShaderCacheAndTextureReferences();
 				}
 			}
 		}
@@ -834,25 +832,14 @@ void UTexture::PostEditChangeProperty(FPropertyChangedEvent& PropertyChangedEven
 	}
 		
 #if WITH_EDITORONLY_DATA
-	// any texture that is referencing this texture as AssociatedNormalMap needs to be informed
+	// any texture that is referencing this texture as CompositeTexture needs to be informed
 	{
 		TRACE_CPUPROFILER_EVENT_SCOPE(UpdateDependentTextures);
 
-		TArray<UTexture*> TexturesThatUseThisTexture;
-
-		for (TObjectIterator<UTexture> It; It; ++It) // walk all textures in the world
-		{
-			UTexture* Tex = *It;
-			
-			if(Tex != this && Tex->CompositeTexture == this && Tex->CompositeTextureMode != CTM_Disabled)
-			{
-				TexturesThatUseThisTexture.Add(Tex);
-			}
-		}
-		// there is a potential infinite loop here if two textures depend on each other (or in a ring)
-		for (int32 i = 0; i < TexturesThatUseThisTexture.Num(); ++i)
+		FObjectCacheContextScope ObjectCache;
+		for (UTexture* Texture : ObjectCache.GetContext().GetTexturesAffectedByTexture(this))
 		{
-			TexturesThatUseThisTexture[i]->PostEditChange();
+			Texture->PostEditChange();
 		}
 	}
 #endif
@@ -880,6 +867,10 @@ void UTexture::Serialize(FArchive& Ar)
 {
 	Ar.UsingCustomVersion(FUE5MainStreamObjectVersion::GUID);
 
+#if WITH_EDITORONLY_DATA
+	NotifyIfCompositeTextureChanged();
+#endif
+
 	Super::Serialize(Ar);
 	
 	FStripDataFlags StripFlags(Ar);
@@ -1046,6 +1037,8 @@ void UTexture::Serialize(FArchive& Ar)
 			CompositeTextureMode = CTM_Disabled;
 		}
 	}
+
+	NotifyIfCompositeTextureChanged();
 #endif // #if WITH_EDITORONLY_DATA
 }
 
@@ -1058,8 +1051,37 @@ void UTexture::AppendToClassSchema(FAppendToClassSchemaContext& Context)
 	uint8 LegacyScalingBool = IsEnableLegacyAlphaCoverageThresholdScaling();
 	Context.Update(&LegacyScalingBool, sizeof(LegacyScalingBool));
 }
-#endif
 
+void UTexture::OutdatedKnownCompositeTextureDetected() const
+{
+PRAGMA_DISABLE_DEPRECATION_WARNINGS
+	ensureMsgf(
+		KnownCompositeTexture == CompositeTexture,
+		TEXT("CompositeTexture property overwritten for texture %s without a call to NotifyIfCompositeTextureChanged(). KnownCompositeTexture (%p) != CompositeTexture (%p - %s)"),
+		*GetFullName(),
+		KnownCompositeTexture,
+		CompositeTexture,
+		CompositeTexture ? *CompositeTexture->GetFullName() : TEXT("nullptr")
+	);
+
+	// This is a last resort, call the notification now
+	UTexture* MutableThis = const_cast<UTexture*>(this);
+	MutableThis->NotifyIfCompositeTextureChanged();
+PRAGMA_ENABLE_DEPRECATION_WARNINGS
+}
+
+void UTexture::NotifyIfCompositeTextureChanged()
+{
+PRAGMA_DISABLE_DEPRECATION_WARNINGS
+	if (KnownCompositeTexture != CompositeTexture)
+	{
+		KnownCompositeTexture = CompositeTexture;
+		FObjectCacheEventSink::NotifyCompositeTextureChanged_Concurrent(this);
+	}
+PRAGMA_ENABLE_DEPRECATION_WARNINGS
+}
+
+#endif // #if WITH_EDITORONLY_DATA
 
 void UTexture::PostInitProperties()
 {
@@ -1146,6 +1168,11 @@ void UTexture::BeginDestroy()
 	{
 		BeginFinalReleaseResource();
 	}
+
+#if WITH_EDITOR
+	// The object cache needs to be notified when we're getting destroyed
+	FObjectCacheEventSink::NotifyCompositeTextureChanged_Concurrent(this);
+#endif
 }
 
 bool UTexture::IsReadyForFinishDestroy()
@@ -3778,22 +3805,19 @@ void UTexture::NotifyMaterials(const ENotifyMaterialsEffectOnShaders EffectOnSha
 	// Create a material update context to safely update materials.
 	{
 		FMaterialUpdateContext UpdateContext;
+		FObjectCacheContextScope ObjectCache;
 
 		// Notify any material that uses this texture
 		TSet<UMaterial*> BaseMaterialsThatUseThisTexture;
-		for (TObjectIterator<UMaterialInterface> It; It; ++It)
+		for (UMaterialInterface* MaterialInterface : ObjectCache.GetContext().GetMaterialsAffectedByTexture(this))
 		{
-			UMaterialInterface* MaterialInterface = *It;
-			if (DoesMaterialUseTexture(MaterialInterface, this))
-			{
-				UpdateContext.AddMaterialInterface(MaterialInterface);
-				// This is a bit tricky. We want to make sure all materials using this texture are
-				// updated. Materials are always updated. Material instances may also have to be
-				// updated and if they have static permutations their children must be updated
-				// whether they use the texture or not! The safe thing to do is to add the instance's
-				// base material to the update context causing all materials in the tree to update.
-				BaseMaterialsThatUseThisTexture.Add(MaterialInterface->GetMaterial());
-			}
+			UpdateContext.AddMaterialInterface(MaterialInterface);
+			// This is a bit tricky. We want to make sure all materials using this texture are
+			// updated. Materials are always updated. Material instances may also have to be
+			// updated and if they have static permutations their children must be updated
+			// whether they use the texture or not! The safe thing to do is to add the instance's
+			// base material to the update context causing all materials in the tree to update.
+			BaseMaterialsThatUseThisTexture.Add(MaterialInterface->GetMaterial());
 		}
 
 		// Go ahead and update any base materials that need to be.
diff --git a/Engine/Source/Runtime/Engine/Private/TextureDerivedData.cpp b/Engine/Source/Runtime/Engine/Private/TextureDerivedData.cpp
--- a/Engine/Source/Runtime/Engine/Private/TextureDerivedData.cpp
+++ b/Engine/Source/Runtime/Engine/Private/TextureDerivedData.cpp
@@ -367,11 +367,11 @@ void GetTextureDerivedDataKeySuffix(const UTexture& Texture, const FTextureBuild
 	
 	FString CompositeTextureStr;
 
-	if(IsValid(Texture.CompositeTexture) && Texture.CompositeTextureMode != CTM_Disabled && Texture.CompositeTexture->Source.IsValid())
+	if(IsValid(Texture.GetCompositeTexture()) && Texture.CompositeTextureMode != CTM_Disabled && Texture.GetCompositeTexture()->Source.IsValid())
 	{
 		// CompositeTextureMode output changed so force a new DDC key value :
 		CompositeTextureStr += TEXT("_Composite090802022_");
-		CompositeTextureStr += Texture.CompositeTexture->Source.GetIdString();
+		CompositeTextureStr += Texture.GetCompositeTexture()->Source.GetIdString();
 	}
 
 	// child texture formats may need to know the mip dimensions in order to generate the ddc
@@ -943,7 +943,7 @@ static void GetTextureBuildSettings(
 	OutBuildSettings.CompositeTextureMode = Texture.CompositeTextureMode;
 	OutBuildSettings.CompositePower = Texture.CompositePower;
 
-	if ( Texture.CompositeTexture && ! Texture.CompositeTexture->Source.IsValid() )
+	if ( Texture.GetCompositeTexture() && !Texture.GetCompositeTexture()->Source.IsValid() )
 	{
 		// have a CompositeTexture but it has no source, don't use it :
 		OutBuildSettings.CompositeTextureMode = CTM_Disabled;
diff --git a/Engine/Source/Runtime/Engine/Private/TextureDerivedDataBuildUtils.cpp b/Engine/Source/Runtime/Engine/Private/TextureDerivedDataBuildUtils.cpp
--- a/Engine/Source/Runtime/Engine/Private/TextureDerivedDataBuildUtils.cpp
+++ b/Engine/Source/Runtime/Engine/Private/TextureDerivedDataBuildUtils.cpp
@@ -327,12 +327,12 @@ FCbObject SaveTextureBuildSettings(const UTexture& Texture, const FTextureBuildS
 	Writer.SetName("Source");
 	WriteSource(Writer, Texture, LayerIndex, BuildSettings);
 
-	if (bUseCompositeTexture && Texture.CompositeTexture)
+	if (bUseCompositeTexture && Texture.GetCompositeTexture())
 	{
-		check( Texture.CompositeTexture->Source.IsValid() ); // should have been checked to set bUseCompositeTexture
+		check( Texture.GetCompositeTexture()->Source.IsValid() ); // should have been checked to set bUseCompositeTexture
 
 		Writer.SetName("CompositeSource");
-		WriteSource(Writer, *Texture.CompositeTexture, LayerIndex, BuildSettings);
+		WriteSource(Writer, *Texture.GetCompositeTexture(), LayerIndex, BuildSettings);
 	}
 
 	Writer.EndObject();
diff --git a/Engine/Source/Runtime/Engine/Private/TextureDerivedDataTask.cpp b/Engine/Source/Runtime/Engine/Private/TextureDerivedDataTask.cpp
--- a/Engine/Source/Runtime/Engine/Private/TextureDerivedDataTask.cpp
+++ b/Engine/Source/Runtime/Engine/Private/TextureDerivedDataTask.cpp
@@ -483,7 +483,7 @@ public:
 		{
 			const FCompressedBuffer& Buffer = Key == UTF8TEXTVIEW("Source")
 				? FindSource(SourceBuffer, Texture.Source, BulkDataId)
-				: FindSource(CompositeSourceBuffer, Texture.CompositeTexture->Source, BulkDataId);
+				: FindSource(CompositeSourceBuffer, Texture.GetCompositeTexture()->Source, BulkDataId);
 			if (Buffer)
 			{
 				Inputs.Add({Key, Buffer.GetRawHash(), Buffer.GetRawSize()});
@@ -510,7 +510,7 @@ public:
 			{
 				const FCompressedBuffer& Buffer = Key == UTF8TEXTVIEW("Source")
 					? FindSource(SourceBuffer, Texture.Source, BulkDataId)
-					: FindSource(CompositeSourceBuffer, Texture.CompositeTexture->Source, BulkDataId);
+					: FindSource(CompositeSourceBuffer, Texture.GetCompositeTexture()->Source, BulkDataId);
 				if (Buffer)
 				{
 					Inputs.Add({Key, Buffer});
@@ -757,7 +757,7 @@ static void DDC1_BuildTexture(
 		// Compress the texture by calling texture compressor directly.
 		TArray<FCompressedImage2D> CompressedMips;
 		if (Compressor->BuildTexture(TextureData.Blocks[0].MipsPerLayer[0],
-			((bool)Texture.CompositeTexture && CompositeTextureData.Blocks.Num() && CompositeTextureData.Blocks[0].MipsPerLayer.Num()) ? CompositeTextureData.Blocks[0].MipsPerLayer[0] : TArray<FImage>(),
+			((bool)Texture.GetCompositeTexture() && CompositeTextureData.Blocks.Num() && CompositeTextureData.Blocks[0].MipsPerLayer.Num()) ? CompositeTextureData.Blocks[0].MipsPerLayer[0] : TArray<FImage>(),
 			InBuildSettingsPerLayer[0],
 			TexturePathName,
 			CompressedMips,
@@ -1287,9 +1287,9 @@ FTextureCacheDerivedDataWorker::FTextureCacheDerivedDataWorker(
 	// All of these settings are fixed across build settings and are derived directly from the texture.
 	// So we can just use layer 0 of whatever we have.
 	TextureData.Init(Texture, (TextureMipGenSettings)BuildSettingsPerLayerFetchOrBuild[0].MipGenSettings, BuildSettingsPerLayerFetchOrBuild[0].bCubemap, BuildSettingsPerLayerFetchOrBuild[0].bTextureArray, BuildSettingsPerLayerFetchOrBuild[0].bVolume, bAllowAsyncLoading);
-	if (Texture.CompositeTexture && Texture.CompositeTextureMode != CTM_Disabled && Texture.CompositeTexture->Source.IsValid())
+	if (Texture.GetCompositeTexture() && Texture.CompositeTextureMode != CTM_Disabled && Texture.GetCompositeTexture()->Source.IsValid())
 	{
-		bool bMatchingBlocks = Texture.CompositeTexture->Source.GetNumBlocks() == Texture.Source.GetNumBlocks();
+		bool bMatchingBlocks = Texture.GetCompositeTexture()->Source.GetNumBlocks() == Texture.Source.GetNumBlocks();
 		
 		if (!bMatchingBlocks)
 		{
@@ -1299,7 +1299,7 @@ FTextureCacheDerivedDataWorker::FTextureCacheDerivedDataWorker(
 
 		if ( bMatchingBlocks )
 		{
-			CompositeTextureData.Init(*Texture.CompositeTexture, (TextureMipGenSettings)BuildSettingsPerLayerFetchOrBuild[0].MipGenSettings, BuildSettingsPerLayerFetchOrBuild[0].bCubemap, BuildSettingsPerLayerFetchOrBuild[0].bTextureArray, BuildSettingsPerLayerFetchOrBuild[0].bVolume, bAllowAsyncLoading);
+			CompositeTextureData.Init(*Texture.GetCompositeTexture(), (TextureMipGenSettings)BuildSettingsPerLayerFetchOrBuild[0].MipGenSettings, BuildSettingsPerLayerFetchOrBuild[0].bCubemap, BuildSettingsPerLayerFetchOrBuild[0].bTextureArray, BuildSettingsPerLayerFetchOrBuild[0].bVolume, bAllowAsyncLoading);
 		}
 	}
 }
@@ -1656,11 +1656,11 @@ static bool DDC1_LoadAndValidateTextureData(
 	}
 
 	bool bHasCompositeTextureSourceMips = false;
-	if (CompositeTextureData.IsValid() && Texture.CompositeTexture && Texture.CompositeTexture->Source.IsBulkDataLoaded())
+	if (CompositeTextureData.IsValid() && Texture.GetCompositeTexture() && Texture.GetCompositeTexture()->Source.IsBulkDataLoaded())
 	{
 		TRACE_CPUPROFILER_EVENT_SCOPE(GetCompositeSourceMips);
-		check( Texture.CompositeTexture->Source.IsValid() );
-		CompositeTextureData.GetSourceMips(Texture.CompositeTexture->Source, ImageWrapper);
+		check( Texture.GetCompositeTexture()->Source.IsValid() );
+		CompositeTextureData.GetSourceMips(Texture.GetCompositeTexture()->Source, ImageWrapper);
 		bHasCompositeTextureSourceMips = true;
 	}
 
@@ -2027,9 +2027,9 @@ void FTextureCacheDerivedDataWorker::Finalize()
 		}
 
 		TextureData.GetSourceMips(Texture.Source, ImageWrapper);
-		if (Texture.CompositeTexture && Texture.CompositeTexture->Source.IsValid())
+		if (Texture.GetCompositeTexture() && Texture.GetCompositeTexture()->Source.IsValid())
 		{
-			CompositeTextureData.GetSourceMips(Texture.CompositeTexture->Source, ImageWrapper);
+			CompositeTextureData.GetSourceMips(Texture.GetCompositeTexture()->Source, ImageWrapper);
 		}
 
 		if (DDC1_IsTextureDataValid(TextureData, CompositeTextureData) == false)
@@ -2661,9 +2661,9 @@ public:
 		DefinitionBuilder.AddConstant(UTF8TEXTVIEW("EngineParameters"), UE::TextureBuildUtilities::TextureEngineParameters::ToCompactBinaryWithDefaults(GenerateTextureEngineParameters()));
 		DefinitionBuilder.AddConstant(UTF8TEXTVIEW("Settings"), SaveTextureBuildSettings(Texture, Settings, 0, bUseCompositeTexture, RequiredMemoryEstimate));
 		DefinitionBuilder.AddInputBulkData(UTF8TEXTVIEW("Source"), Texture.Source.GetPersistentId());
-		if (Texture.CompositeTexture && bUseCompositeTexture)
+		if (Texture.GetCompositeTexture() && bUseCompositeTexture)
 		{
-			DefinitionBuilder.AddInputBulkData(UTF8TEXTVIEW("CompositeSource"), Texture.CompositeTexture->Source.GetPersistentId());
+			DefinitionBuilder.AddInputBulkData(UTF8TEXTVIEW("CompositeSource"), Texture.GetCompositeTexture()->Source.GetPersistentId());
 		}
 		return DefinitionBuilder.Build();
 	}
@@ -2823,8 +2823,8 @@ public:
 			}
 		}
 		
-		const bool bCompositeTextureViable = Texture.CompositeTexture && Texture.CompositeTextureMode != CTM_Disabled && Texture.CompositeTexture->Source.IsValid();
-		bool bMatchingBlocks = bCompositeTextureViable && (Texture.CompositeTexture->Source.GetNumBlocks() == Texture.Source.GetNumBlocks());
+		const bool bCompositeTextureViable = Texture.GetCompositeTexture() && Texture.CompositeTextureMode != CTM_Disabled && Texture.GetCompositeTexture()->Source.IsValid();
+		bool bMatchingBlocks = bCompositeTextureViable && (Texture.GetCompositeTexture()->Source.GetNumBlocks() == Texture.Source.GetNumBlocks());
 		
 		if (bCompositeTextureViable)
 		{
@@ -2868,9 +2868,9 @@ public:
 		}
 		Key.BuildDefinitionKey = BuildDefinition.GetKey().Hash;
 		Key.SourceGuid = Texture.Source.GetId();
-		if (bUseCompositeTexture && Texture.CompositeTexture)
+		if (bUseCompositeTexture && Texture.GetCompositeTexture())
 		{
-			Key.CompositeSourceGuid = Texture.CompositeTexture->Source.GetId();
+			Key.CompositeSourceGuid = Texture.GetCompositeTexture()->Source.GetId();
 		}
 		return Key;
 	}
diff --git a/Engine/Source/Runtime/Engine/Public/ObjectCacheContext.h b/Engine/Source/Runtime/Engine/Public/ObjectCacheContext.h
--- a/Engine/Source/Runtime/Engine/Public/ObjectCacheContext.h
+++ b/Engine/Source/Runtime/Engine/Public/ObjectCacheContext.h
@@ -77,6 +77,9 @@ public:
 	TObjectCacheIterator<UTexture>              GetUsedTextures(UMaterialInterface* InMaterial);
 	TObjectCacheIterator<UMaterialInterface>    GetUsedMaterials(UPrimitiveComponent* InComponent);
 	TObjectCacheIterator<UMaterialInterface>    GetMaterialsAffectedByMaterials(TArrayView<UMaterialInterface*> InMaterials);
+#if WITH_EDITOR
+	TObjectCacheIterator<UTexture>              GetTexturesAffectedByTexture(UTexture* InTexture);
+#endif
 
 private:
 	friend class FObjectCacheContextScope;
@@ -84,6 +87,9 @@ private:
 	TMap<TObjectKey<UPrimitiveComponent>, TSet<UMaterialInterface*>> PrimitiveComponentToMaterial;
 	TMap<TObjectKey<UMaterialInterface>, TSet<UTexture*>> MaterialUsedTextures;
 	TOptional<TMap<TObjectKey<UTexture>, TSet<UMaterialInterface*>>> TextureToMaterials;
+#if WITH_EDITOR
+	TOptional<TMap<TObjectKey<UTexture>, TSet<UTexture*>>> TextureToTextures;
+#endif
 	TOptional<TMap<TObjectKey<UMaterialInterface>, TSet<UPrimitiveComponent*>>> MaterialToPrimitives;
 	TOptional<TMap<TObjectKey<UStaticMesh>, TSet<UStaticMeshComponent*>>> StaticMeshToComponents;
 	TOptional<TMap<TObjectKey<USkinnedAsset>, TSet<USkinnedMeshComponent*>>> SkinnedAssetToComponents;
diff --git a/Engine/Source/Runtime/Engine/Public/ObjectCacheEventSink.h b/Engine/Source/Runtime/Engine/Public/ObjectCacheEventSink.h
--- a/Engine/Source/Runtime/Engine/Public/ObjectCacheEventSink.h
+++ b/Engine/Source/Runtime/Engine/Public/ObjectCacheEventSink.h
@@ -24,6 +24,7 @@ struct ENGINE_API FObjectCacheEventSink
 	static void NotifyReferencedTextureChanged_Concurrent(UMaterialInterface*);
 	static void NotifyStaticMeshChanged_Concurrent(UStaticMeshComponent*);
 	static void NotifyMaterialDestroyed_Concurrent(UMaterialInterface*);
+	static void NotifyCompositeTextureChanged_Concurrent(UTexture*);
 };
 
 #endif // #if WITH_EDITOR
\ No newline at end of file
