World Partition Class Descriptors: remove extra cautious check as we found a case where the object class path might have changed and the actor meta data wasn't updated by the plugin fast rename path.

diff --git a/Engine/Source/Runtime/Engine/Private/WorldPartition/WorldPartitionActorDesc.cpp b/Engine/Source/Runtime/Engine/Private/WorldPartition/WorldPartitionActorDesc.cpp
--- a/Engine/Source/Runtime/Engine/Private/WorldPartition/WorldPartitionActorDesc.cpp
+++ b/Engine/Source/Runtime/Engine/Private/WorldPartition/WorldPartitionActorDesc.cpp
@@ -489,7 +489,6 @@ void FWorldPartitionActorDesc::Serialize(FArchive& Ar)
 			{
 				FGuid ClassDescGuid;
 				Ar << ClassDescGuid;
-				check(Guid == ClassDescGuid);
 			}
 		}
 	}
