World Partition Streaming Generation:
- Added per-instance container data when generating streaming, to avoid having instance logic during streaming context generation pass.
- Unified per-instance container data inheritance into a single function to centralize how we inherit these and make code easier to understand.
- Added per-instance container validation and moved per-instance data validation into it since it can be different from the template container.

diff --git a/Engine/Source/Runtime/Engine/Private/WorldPartition/WorldPartitionStreamingGeneration.cpp b/Engine/Source/Runtime/Engine/Private/WorldPartition/WorldPartitionStreamingGeneration.cpp
--- a/Engine/Source/Runtime/Engine/Private/WorldPartition/WorldPartitionStreamingGeneration.cpp
+++ b/Engine/Source/Runtime/Engine/Private/WorldPartition/WorldPartitionStreamingGeneration.cpp
@@ -160,33 +160,16 @@ class FWorldPartitionStreamingGenerator
 					}
 
 					FActorSetInstance& ActorSetInstance = ActorSetInstances.Emplace_GetRef();
+					const FContainerInstanceDescriptor::FPerInstanceData& PerInstanceData = ContainerInstanceDescriptor.PerInstanceData.FindChecked(ReferenceActorDescView.GetGuid());
 				
 					ActorSetInstance.ContainerInstance = &ActorSetContainer;
 					ActorSetInstance.ActorSet = &ActorSet;
 					ActorSetInstance.ContainerID = ContainerInstanceDescriptor.ID;
 					ActorSetInstance.Transform = ContainerInstanceDescriptor.Transform;
-
-					// Apply AND logic on spatially loaded flag
-					ActorSetInstance.bIsSpatiallyLoaded = ReferenceActorDescView.GetIsSpatiallyLoaded() && ContainerInstanceDescriptor.bIsSpatiallyLoaded;
-
-					// Since Content Bundles streaming generation happens in its own context, all actor set instances must have the same content bundle GUID for now, so Level Instances
-					// placed inside a Content Bundle will propagate their Content Bundle GUID to child instances.
-					ActorSetInstance.ContentBundleID = MainWorldContainer->GetContentBundleGuid();
-
-					if (ContainerInstanceDescriptor.ID.IsMainContainer())
-					{
-						// Main container will get inherited properties from the actor descriptors
-						ActorSetInstance.RuntimeGrid = ReferenceActorDescView.GetRuntimeGrid();
-						ActorSetInstance.DataLayers = StreamingGenerator->GetRuntimeDataLayerInstances(ReferenceActorDescView.GetRuntimeDataLayerInstanceNames());
-					}
-					else
-					{
-						// Sub containers will inherit some properties from the parent actor descriptors
-						ActorSetInstance.RuntimeGrid = ContainerInstanceDescriptor.RuntimeGrid;
-						TSet<FName> CombinedDataLayers = ContainerInstanceDescriptor.RuntimeDataLayers;
-						CombinedDataLayers.Append(ReferenceActorDescView.GetRuntimeDataLayerInstanceNames());
-						ActorSetInstance.DataLayers = StreamingGenerator->GetRuntimeDataLayerInstances(CombinedDataLayers.Array());
-					}
+					ActorSetInstance.bIsSpatiallyLoaded = PerInstanceData.bIsSpatiallyLoaded;
+					ActorSetInstance.ContentBundleID = ContainerInstanceDescriptor.ContentBundleID;
+					ActorSetInstance.RuntimeGrid = PerInstanceData.RuntimeGrid;
+					ActorSetInstance.DataLayers = StreamingGenerator->GetRuntimeDataLayerInstances(PerInstanceData.DataLayers);
 
 					ActorSetInstance.Bounds.Init();
 					for (const FGuid& ActorGuid : ActorSetInstance.ActorSet->Actors)
@@ -242,19 +225,34 @@ class FWorldPartitionStreamingGenerator
 		TArray<FActorSetInstance> ActorSetInstances;
 	};
 
+	/** 
+	 * An actor container descriptor, one for the main world and one for every unique level used by level instances.
+	 */
 	struct FContainerDescriptor
 	{
 		FContainerDescriptor()
 			: Container(nullptr)
 		{}
 
+		/** The unique actor descriptor container for this descriptor */
 		const UActorDescContainer* Container;
+
+		/** The actor descriptor views for for this descriptor */
 		FActorDescViewMap ActorDescViewMap;
+
+		/** Set of editor-only actors that are not part of the actor descriptor views */
 		TSet<FGuid> EditorOnlyActorDescMap;
+
+		/** List of actor descriptor views that are containers (mainly level instances) */
 		TArray<FWorldPartitionActorDescView> ContainerInstanceViews;
+
+		/** List of actor clusters for this descriptor */
 		TArray<TArray<FGuid>> Clusters;
 	};
 
+	/** 
+	 * An actor container instance descriptor, one for the main world and one for every actor container instance
+	 */
 	struct FContainerInstanceDescriptor
 	{
 		FContainerInstanceDescriptor()
@@ -266,13 +264,23 @@ class FWorldPartitionStreamingGenerator
 		FTransform Transform;
 		const UActorDescContainer* Container;
 		EContainerClusterMode ClusterMode;
-		TSet<FName> RuntimeDataLayers;
+		TArray<FName> RuntimeDataLayers;
 		FName RuntimeGrid;
 		bool bIsSpatiallyLoaded;
 		FString OwnerName;
 		FActorContainerID ID;
 		FActorContainerID ParentID;
+		FGuid ContentBundleID;
 		TMap<FActorContainerID, TSet<FGuid>> FilteredActors;
+
+		struct FPerInstanceData
+		{
+			bool bIsSpatiallyLoaded;
+			FName RuntimeGrid;
+			TArray<FName> DataLayers;
+		};
+
+		TMap<FGuid, FPerInstanceData> PerInstanceData;
 	};
 
 	void ResolveRuntimeSpatiallyLoaded(FWorldPartitionActorDescView& ActorDescView)
@@ -424,132 +432,173 @@ class FWorldPartitionStreamingGenerator
 
 	void CreateActorDescriptorViewsRecursive(const FContainerInstanceDescriptor& InContainerInstanceDescriptor)
 	{
-		// Create or resolve container descriptor
-		FContainerDescriptor& ContainerDescriptor = ContainerDescriptorsMap.FindOrAdd(InContainerInstanceDescriptor.Container);
-		
-		// Scope ContainerInstanceDescriptor so that it doesn't get used mistakenly after the ContainerInstanceViews loop
+		// Inherited parent properties logic
+		auto InheritParentContainerProperties = [](const FActorContainerID& InParentContainerID, bool bInParentIsSpatiallyLoaded, FName InParentRuntimeGrid, const TArray<FName>& InParentRuntimeDataLayers, bool& bInOutIsSpatiallyLoaded, FName& InOutRuntimeGrid, TArray<FName>& InOutRuntimeDataLayers)
 		{
-			// Create container instance descriptor
-			check(!ContainerInstanceDescriptorsMap.Contains(InContainerInstanceDescriptor.ID));
-
-			FContainerInstanceDescriptor& ContainerInstanceDescriptor = ContainerInstanceDescriptorsMap.Add(InContainerInstanceDescriptor.ID, InContainerInstanceDescriptor);
+			// Apply AND logic on spatially loaded flag
+			bInOutIsSpatiallyLoaded &= bInParentIsSpatiallyLoaded;
 
-			if (!ContainerDescriptor.Container)
+			// Runtime grid is only inherited from the main world, since level instance doesn't support setting this value on actors
+			if (InParentContainerID.IsMainContainer())
 			{
-				ContainerDescriptor.Container = ContainerInstanceDescriptor.Container;
-
-				// Gather actor descriptor views for this container
-				CreateActorDescViewMap(ContainerInstanceDescriptor.Container, ContainerDescriptor.ActorDescViewMap, ContainerDescriptor.EditorOnlyActorDescMap, ContainerInstanceDescriptor.ID, ContainerDescriptor.ContainerInstanceViews);
-
-				// Resolve actor descriptor views before validation
-				ResolveContainerDescriptor(ContainerDescriptor);
-
-				// Validate container, fixing anything illegal, etc.
-				ValidateContainerDescriptor(ContainerDescriptor, ContainerInstanceDescriptor.ID.IsMainContainer());
+				InOutRuntimeGrid = InParentRuntimeGrid;
+			}
 
-				// Update container, computing cluster, bounds, etc.
-				UpdateContainerDescriptor(ContainerDescriptor);
+			// Data layers are accumulated down the hierarchy chain, since level instances supports data layers assignation on actors
+			if (InParentRuntimeDataLayers.Num())
+			{
+				TSet<FName> RuntimeDataLayers(InOutRuntimeDataLayers);
+				RuntimeDataLayers.Append(InParentRuntimeDataLayers);
+				InOutRuntimeDataLayers = RuntimeDataLayers.Array();
 			}
+		};
 
-			// Calculate Bounds of Non-container ActorDescViews
-			check(!ContainerInstanceDescriptor.Bounds.IsValid);
-			ContainerDescriptor.ActorDescViewMap.ForEachActorDescView([&ContainerInstanceDescriptor](const FWorldPartitionActorDescView& ActorDescView)
+		// ContainerDescriptor may be reallocated after this scope
+		{
+			// Create or resolve container descriptor
+			FContainerDescriptor& ContainerDescriptor = ContainerDescriptorsMap.FindOrAdd(InContainerInstanceDescriptor.Container);
+		
+			// ContainerInstanceDescriptor may be reallocated after this scope
 			{
-				if (ActorDescView.GetIsSpatiallyLoaded())
+				// Create container instance descriptor
+				check(!ContainerInstanceDescriptorsMap.Contains(InContainerInstanceDescriptor.ID));
+
+				FContainerInstanceDescriptor& ContainerInstanceDescriptor = ContainerInstanceDescriptorsMap.Add(InContainerInstanceDescriptor.ID, InContainerInstanceDescriptor);
+
+				if (!ContainerDescriptor.Container)
 				{
-					const FBox RuntimeBounds = ActorDescView.GetRuntimeBounds();
-					check(RuntimeBounds.IsValid);
+					ContainerDescriptor.Container = ContainerInstanceDescriptor.Container;
 
-					ContainerInstanceDescriptor.Bounds += RuntimeBounds.TransformBy(ContainerInstanceDescriptor.Transform);
-				}
-			});
-		}
+					// Gather actor descriptor views for this container
+					CreateActorDescViewMap(ContainerInstanceDescriptor.Container, ContainerDescriptor.ActorDescViewMap, ContainerDescriptor.EditorOnlyActorDescMap, ContainerInstanceDescriptor.ID, ContainerDescriptor.ContainerInstanceViews);
 
-		// Inherited parent properties logic
-		auto InheritParentContainerProperties = [](const FActorContainerID& InParentContainerID, const FWorldPartitionActorDescView& InParentActorDescView, TSet<FName>& InOutRuntimeDataLayers, FName& InOutRuntimeGrid, bool& bInOutIsSpatiallyLoaded)
-		{
-			// Runtime grid is only inherited from the main world, since level instance doesn't support setting this value on actors
-			if (InParentContainerID.IsMainContainer())
-			{
-				InOutRuntimeGrid = InParentActorDescView.GetRuntimeGrid();
-			}
+					// Resolve actor descriptor views before validation
+					ResolveContainerDescriptor(ContainerDescriptor);
 
-			// Apply AND logic on spatially loaded flag
-			bInOutIsSpatiallyLoaded = bInOutIsSpatiallyLoaded && InParentActorDescView.GetIsSpatiallyLoaded();
+					// Validate container, fixing anything illegal, etc.
+					ValidateContainerDescriptor(ContainerDescriptor, ContainerInstanceDescriptor.ID.IsMainContainer());
 
-			// Data layers are accumulated down the hierarchy chain, since level instances supports data layers assignation on actors
-			InOutRuntimeDataLayers.Append(InParentActorDescView.GetRuntimeDataLayerInstanceNames());
-		};
+					// Update container, computing cluster, bounds, etc.
+					UpdateContainerDescriptor(ContainerDescriptor);
+				}
 
-		// Parse actor containers
-		TArray<FWorldPartitionActorDescView> ContainerInstanceViews = ContainerDescriptor.ContainerInstanceViews;
-		for (const FWorldPartitionActorDescView& ContainerInstanceView : ContainerInstanceViews)
-		{
-			FWorldPartitionActorDesc::FContainerInstance SubContainerInstance;
-			if (!ContainerInstanceView.GetContainerInstance(SubContainerInstance))
-			{
-				ErrorHandler->OnLevelInstanceInvalidWorldAsset(ContainerInstanceView, ContainerInstanceView.GetContainerPackage(), IStreamingGenerationErrorHandler::ELevelInstanceInvalidReason::WorldAssetHasInvalidContainer);
-				continue;
+				// Calculate Bounds of Non-container ActorDescViews
+				check(!ContainerInstanceDescriptor.Bounds.IsValid);
+				ContainerDescriptor.ActorDescViewMap.ForEachActorDescView([&ContainerInstanceDescriptor](const FWorldPartitionActorDescView& ActorDescView)
+				{
+					if (ActorDescView.GetIsSpatiallyLoaded())
+					{
+						const FBox RuntimeBounds = ActorDescView.GetRuntimeBounds();
+						check(RuntimeBounds.IsValid);
+
+						ContainerInstanceDescriptor.Bounds += RuntimeBounds.TransformBy(ContainerInstanceDescriptor.Transform);
+					}
+				});
 			}
 
-			if (ContainerInstancesStack.Contains(SubContainerInstance.Container->ContainerPackageName))
+			// Parse actor containers
+			TArray<FWorldPartitionActorDescView> ContainerInstanceViews = ContainerDescriptor.ContainerInstanceViews;
+			for (const FWorldPartitionActorDescView& ContainerInstanceView : ContainerInstanceViews)
 			{
-				ErrorHandler->OnLevelInstanceInvalidWorldAsset(ContainerInstanceView, ContainerInstanceView.GetContainerPackage(), IStreamingGenerationErrorHandler::ELevelInstanceInvalidReason::CirculalReference);
-				continue;
-			}
+				FWorldPartitionActorDesc::FContainerInstance SubContainerInstance;
+				if (!ContainerInstanceView.GetContainerInstance(SubContainerInstance))
+				{
+					ErrorHandler->OnLevelInstanceInvalidWorldAsset(ContainerInstanceView, ContainerInstanceView.GetContainerPackage(), IStreamingGenerationErrorHandler::ELevelInstanceInvalidReason::WorldAssetHasInvalidContainer);
+					continue;
+				}
+
+				if (ContainerInstancesStack.Contains(SubContainerInstance.Container->ContainerPackageName))
+				{
+					ErrorHandler->OnLevelInstanceInvalidWorldAsset(ContainerInstanceView, ContainerInstanceView.GetContainerPackage(), IStreamingGenerationErrorHandler::ELevelInstanceInvalidReason::CirculalReference);
+					continue;
+				}
 
-			const FGuid ActorGuid = ContainerInstanceView.GetGuid();
-			const FActorContainerID ContainerID(InContainerInstanceDescriptor.ID, ActorGuid);
+				const FGuid ActorGuid = ContainerInstanceView.GetGuid();
+				const FActorContainerID ContainerID(InContainerInstanceDescriptor.ID, ActorGuid);
 						
-			TMap<FActorContainerID, TSet<FGuid>> ContainerInstanceFilter;
-			if (InContainerInstanceDescriptor.ID.IsMainContainer() && ContainerInstanceView.IsContainerFilter())
-			{
-				if (const FWorldPartitionActorFilter* ContainerFilter = ContainerInstanceView.GetContainerFilter())
+				TMap<FActorContainerID, TSet<FGuid>> ContainerInstanceFilter;
+				if (InContainerInstanceDescriptor.ID.IsMainContainer() && ContainerInstanceView.IsContainerFilter())
 				{
-					UWorld* OwningWorld = WorldPartitionContext ? WorldPartitionContext->GetWorld() : nullptr;
-					if (UWorldPartitionSubsystem* WorldPartitionSubsystem = UWorld::GetSubsystem<UWorldPartitionSubsystem>(OwningWorld))
+					if (const FWorldPartitionActorFilter* ContainerFilter = ContainerInstanceView.GetContainerFilter())
 					{
-						ContainerInstanceFilter = WorldPartitionSubsystem->GetFilteredActorsPerContainer(ContainerID, ContainerInstanceView.GetContainerPackage().ToString(), *ContainerFilter);
+						UWorld* OwningWorld = WorldPartitionContext ? WorldPartitionContext->GetWorld() : nullptr;
+						if (UWorldPartitionSubsystem* WorldPartitionSubsystem = UWorld::GetSubsystem<UWorldPartitionSubsystem>(OwningWorld))
+						{
+							ContainerInstanceFilter = WorldPartitionSubsystem->GetFilteredActorsPerContainer(ContainerID, ContainerInstanceView.GetContainerPackage().ToString(), *ContainerFilter);
+						}
 					}
 				}
-			}
-
-			ContainerInstancesStack.Add(SubContainerInstance.Container->ContainerPackageName);
-
-			check(SubContainerInstance.Container);
-			FContainerInstanceDescriptor SubContainerInstanceDescriptor;
-
-			// Inherit fields from FContainerInstance
-			SubContainerInstanceDescriptor.ID = ContainerID;
-			SubContainerInstanceDescriptor.Container = SubContainerInstance.Container;
-			SubContainerInstanceDescriptor.Transform = SubContainerInstance.Transform * InContainerInstanceDescriptor.Transform;
-			SubContainerInstanceDescriptor.FilteredActors = MoveTemp(ContainerInstanceFilter);
-			SubContainerInstanceDescriptor.ParentID = InContainerInstanceDescriptor.ID;
-			SubContainerInstanceDescriptor.OwnerName = *ContainerInstanceView.GetActorLabelOrName().ToString();
-
-			// Inherit parent container properties
-			FName InheritedRuntimeGrid = InContainerInstanceDescriptor.RuntimeGrid;
-			bool bInheritedbIsSpatiallyLoaded = InContainerInstanceDescriptor.bIsSpatiallyLoaded;
-			TSet<FName> InheritedRuntimeDataLayers = InContainerInstanceDescriptor.RuntimeDataLayers;
-			InheritParentContainerProperties(InContainerInstanceDescriptor.ID, ContainerInstanceView, InheritedRuntimeDataLayers, InheritedRuntimeGrid, bInheritedbIsSpatiallyLoaded);
 
-			SubContainerInstanceDescriptor.bIsSpatiallyLoaded = bInheritedbIsSpatiallyLoaded;
-			SubContainerInstanceDescriptor.RuntimeGrid = InheritedRuntimeGrid;
-			SubContainerInstanceDescriptor.RuntimeDataLayers = InheritedRuntimeDataLayers;
+				ContainerInstancesStack.Add(SubContainerInstance.Container->ContainerPackageName);
+
+				check(SubContainerInstance.Container);
+				FContainerInstanceDescriptor SubContainerInstanceDescriptor;
+
+				// Inherit fields from FContainerInstance
+				SubContainerInstanceDescriptor.ID = ContainerID;
+				SubContainerInstanceDescriptor.Container = SubContainerInstance.Container;
+				SubContainerInstanceDescriptor.Transform = SubContainerInstance.Transform * InContainerInstanceDescriptor.Transform;
+				SubContainerInstanceDescriptor.FilteredActors = MoveTemp(ContainerInstanceFilter);
+				SubContainerInstanceDescriptor.ParentID = InContainerInstanceDescriptor.ID;
+				SubContainerInstanceDescriptor.OwnerName = *ContainerInstanceView.GetActorLabelOrName().ToString();
+				// Since Content Bundles streaming generation happens in its own context, all actor set instances must have the same content bundle GUID for now, so Level Instances
+				// placed inside a Content Bundle will propagate their Content Bundle GUID to child instances.
+				SubContainerInstanceDescriptor.ContentBundleID = InContainerInstanceDescriptor.ContentBundleID;
+				SubContainerInstanceDescriptor.bIsSpatiallyLoaded = InContainerInstanceDescriptor.bIsSpatiallyLoaded;
+				SubContainerInstanceDescriptor.RuntimeGrid = InContainerInstanceDescriptor.RuntimeGrid;
+				SubContainerInstanceDescriptor.RuntimeDataLayers = InContainerInstanceDescriptor.RuntimeDataLayers;
+
+				// Inherit parent container properties
+				InheritParentContainerProperties(
+					InContainerInstanceDescriptor.ID,
+					ContainerInstanceView.GetIsSpatiallyLoaded(),
+					ContainerInstanceView.GetRuntimeGrid(),
+					ContainerInstanceView.GetRuntimeDataLayerInstanceNames(),
+					SubContainerInstanceDescriptor.bIsSpatiallyLoaded,
+					SubContainerInstanceDescriptor.RuntimeGrid,
+					SubContainerInstanceDescriptor.RuntimeDataLayers
+				);
+
+				CreateActorDescriptorViewsRecursive(SubContainerInstanceDescriptor);
+
+				verify(ContainerInstancesStack.Remove(SubContainerInstance.Container->ContainerPackageName));
+			}
+		}
 
-			CreateActorDescriptorViewsRecursive(SubContainerInstanceDescriptor);
+		// Fetch the versions stored in the map as it can have been reallocated during recursion
+		FContainerDescriptor& ContainerDescriptor = ContainerDescriptorsMap.FindOrAdd(InContainerInstanceDescriptor.Container);
+		FContainerInstanceDescriptor& ContainerInstanceDescriptor = ContainerInstanceDescriptorsMap.FindChecked(InContainerInstanceDescriptor.ID);
 
-			verify(ContainerInstancesStack.Remove(SubContainerInstance.Container->ContainerPackageName));
-		}
-							
 		if (!InContainerInstanceDescriptor.ID.IsMainContainer())
 		{
 			FContainerInstanceDescriptor& ParentContainer = ContainerInstanceDescriptorsMap.FindChecked(InContainerInstanceDescriptor.ParentID);
-			
-			// Fetch version stored in map as it will be the only one that will have its bounds updated by its Children
-			const FContainerInstanceDescriptor& ContainerInstanceWithUpdatedBounds = ContainerInstanceDescriptorsMap.FindChecked(InContainerInstanceDescriptor.ID);
-			ParentContainer.Bounds += ContainerInstanceWithUpdatedBounds.Bounds;
+			ParentContainer.Bounds += ContainerInstanceDescriptor.Bounds;
 		}
+
+		// Apply per-instance data
+		ContainerInstanceDescriptor.PerInstanceData.Reserve(ContainerDescriptor.ActorDescViewMap.Num());
+		ContainerDescriptor.ActorDescViewMap.ForEachActorDescView([&ContainerInstanceDescriptor, &InheritParentContainerProperties](FWorldPartitionActorDescView& ActorDescView)
+		{
+			FContainerInstanceDescriptor::FPerInstanceData& PerInstanceData = ContainerInstanceDescriptor.PerInstanceData.Add(ActorDescView.GetGuid());
+
+			PerInstanceData.bIsSpatiallyLoaded = ContainerInstanceDescriptor.bIsSpatiallyLoaded;
+			PerInstanceData.RuntimeGrid = ContainerInstanceDescriptor.RuntimeGrid;
+			PerInstanceData.DataLayers = ContainerInstanceDescriptor.RuntimeDataLayers;
+
+			// Inherit parent container properties
+			InheritParentContainerProperties(
+				ContainerInstanceDescriptor.ID,
+				ActorDescView.GetIsSpatiallyLoaded(),
+				ActorDescView.GetRuntimeGrid(),
+				ActorDescView.GetRuntimeDataLayerInstanceNames(),
+				PerInstanceData.bIsSpatiallyLoaded,
+				PerInstanceData.RuntimeGrid,
+				PerInstanceData.DataLayers
+			);
+		});
+
+		// Validate container instance, fixing anything illegal, etc.
+		ValidateContainerInstanceDescriptor(ContainerInstanceDescriptor, ContainerInstanceDescriptor.ID.IsMainContainer());
 	}
 
 	/** 
@@ -559,12 +608,12 @@ class FWorldPartitionStreamingGenerator
 	{
 		TRACE_CPUPROFILER_EVENT_SCOPE(FWorldPartitionStreamingGenerator::CreateActorContainers);
 
-		// Since we apply AND logic on spatially loaded flag recursively, startup value must be true
 		FContainerInstanceDescriptor MainContainerInstance;
 		MainContainerInstance.Container = InContainer;
-		MainContainerInstance.bIsSpatiallyLoaded = true;
+		MainContainerInstance.bIsSpatiallyLoaded = true; // Since we apply AND logic on spatially loaded flag recursively, startup value must be true
 		MainContainerInstance.ClusterMode = EContainerClusterMode::Partitioned;
 		MainContainerInstance.OwnerName = TEXT("MainContainer");
+		MainContainerInstance.ContentBundleID = InContainer->GetContentBundleGuid();
 		CreateActorDescriptorViewsRecursive(MainContainerInstance);
 	}
 
@@ -700,37 +749,12 @@ class FWorldPartitionStreamingGenerator
 					return false;
 				};
 
-				// Validate runtime grid
-				auto IsRuntimeGridValid = [this](const FWorldPartitionActorDescView& ActorDescView)
-				{
-PRAGMA_DISABLE_DEPRECATION_WARNINGS
-					if (!ActorDescView.ShouldValidateRuntimeGrid())
-					{
-						return true;
-					}
-PRAGMA_ENABLE_DEPRECATION_WARNINGS
-
-					return IsValidGrid(ActorDescView.GetRuntimeGrid());
-				};
-
 				// Validate runtime grid references
 				auto IsReferenceRuntimeGridValid = [](const FWorldPartitionActorDescView& RefererActorDescView, const FWorldPartitionActorDescView& ReferenceActorDescView)
 				{
 					return RefererActorDescView.GetRuntimeGrid() == ReferenceActorDescView.GetRuntimeGrid();
 				};
 
-				if (!IsRuntimeGridValid(ActorDescView))
-				{
-					if (PassType == EPassType::ErrorReporting)
-					{
-						ErrorHandler->OnInvalidRuntimeGrid(ActorDescView, ActorDescView.GetRuntimeGrid());
-					}
-					else
-					{
-						ActorDescView.SetForcedNoRuntimeGrid();
-					}
-				}
-
 				// Build references List
 				struct FActorReferenceInfo
 				{
@@ -911,6 +935,59 @@ PRAGMA_ENABLE_DEPRECATION_WARNINGS
 		}
 	}
 
+	/** 
+	 * Perform various validations on the container descriptor instance, and adjust it based on different requirements. This needs to happen before updating
+	 * containers bounds because some actor descriptor views might change grid placement, etc.
+	 */
+	void ValidateContainerInstanceDescriptor(FContainerInstanceDescriptor& ContainerInstanceDescriptor, bool bIsMainContainer)
+	{
+		FContainerDescriptor& ContainerDescriptor = ContainerDescriptorsMap.FindChecked(ContainerInstanceDescriptor.Container);
+
+		// Perform various adjustements based on validations and report errors
+		//
+		// The first validation pass is used to report errors, subsequent passes are used to make corrections to the actor descriptor views.
+		// Since the references can form cycles/long chains in the data fixes might need to be propagated in multiple passes.
+		// 
+		// This works because fixes are deterministic and always apply the same way to both Actors being modified, so there's no ordering issues possible.
+		int32 NbErrorsDetected = INDEX_NONE;
+		for(uint32 NbValidationPasses = 0; NbErrorsDetected; NbValidationPasses++)
+		{
+			// Type of work performed in this pass, for clarity
+			enum class EPassType
+			{
+				ErrorReporting,
+				Fixup
+			};
+			const EPassType PassType = NbValidationPasses == 0 ? EPassType::ErrorReporting : EPassType::Fixup;
+
+			NbErrorsDetected = 0;
+
+			ContainerDescriptor.ActorDescViewMap.ForEachActorDescView([this, &ContainerInstanceDescriptor, &NbErrorsDetected, PassType](FWorldPartitionActorDescView& ActorDescView)
+			{
+				FContainerInstanceDescriptor::FPerInstanceData& PerInstanceData = ContainerInstanceDescriptor.PerInstanceData.FindChecked(ActorDescView.GetGuid());
+
+PRAGMA_DISABLE_DEPRECATION_WARNINGS
+				if (ActorDescView.ShouldValidateRuntimeGrid())
+PRAGMA_ENABLE_DEPRECATION_WARNINGS
+				{
+					if (!IsValidGrid(PerInstanceData.RuntimeGrid))
+					{
+						if (PassType == EPassType::ErrorReporting)
+						{
+							ErrorHandler->OnInvalidRuntimeGrid(ActorDescView, ActorDescView.GetRuntimeGrid());
+						}
+						else
+						{
+							PerInstanceData.RuntimeGrid = NAME_None;
+						}
+
+						NbErrorsDetected++;
+					}
+				}
+			});
+		}
+	}
+
 	/** 
 	 * Update the container descriptor containers to adjust their bounds from actor descriptor views.
 	 */
@@ -940,6 +1017,14 @@ public:
 		TArray<TSubclassOf<AActor>> FilteredClasses;
 		TFunction<bool(FName)> IsValidGrid;
 
+		FWorldPartitionStreamingGeneratorParams& SetWorldPartitionContext(const UWorldPartition* InWorldPartitionContext) { WorldPartitionContext = InWorldPartitionContext; return *this; }
+		FWorldPartitionStreamingGeneratorParams& SetModifiedActorsDescList(FActorDescList* InModifiedActorsDescList) { ModifiedActorsDescList = InModifiedActorsDescList; return *this; }
+		FWorldPartitionStreamingGeneratorParams& SetErrorHandler(IStreamingGenerationErrorHandler* InErrorHandler) { ErrorHandler = InErrorHandler; return *this; }
+		FWorldPartitionStreamingGeneratorParams& SetEnableStreaming(bool bInEnableStreaming) { bEnableStreaming = bInEnableStreaming; return *this; }
+		FWorldPartitionStreamingGeneratorParams& SetActorGuidsToContainerMap(const TMap<FGuid, const UActorDescContainer*>& InActorGuidsToContainerMap) { ActorGuidsToContainerMap = InActorGuidsToContainerMap; return *this; }
+		FWorldPartitionStreamingGeneratorParams& SetFilteredClasses(const TArray<TSubclassOf<AActor>>& InFilteredClasses) { FilteredClasses = InFilteredClasses; return *this; }
+		FWorldPartitionStreamingGeneratorParams& SetIsValidGrid(TFunction<bool(FName)> InIsValidGrid) { IsValidGrid = InIsValidGrid; return *this; }
+
 		inline static FStreamingGenerationNullErrorHandler NullErrorHandler;
 	};
 
@@ -1169,12 +1254,12 @@ bool UWorldPartition::GenerateContainerStreaming(const FGenerateStreamingParams&
 		HierarchicalLogAr = MakeUnique<FHierarchicalLogArchive>(*LogFileAr);
 	}
 
-	FWorldPartitionStreamingGenerator::FWorldPartitionStreamingGeneratorParams StreamingGeneratorParams;
-	StreamingGeneratorParams.WorldPartitionContext = this;
-	StreamingGeneratorParams.ModifiedActorsDescList = ModifiedActorsDescList;
-	StreamingGeneratorParams.IsValidGrid = [this](FName GridName) { return RuntimeHash->IsValidGrid(GridName); };
-	StreamingGeneratorParams.ErrorHandler = StreamingGenerationErrorHandlerOverride ? (*StreamingGenerationErrorHandlerOverride)(ErrorHandler) : ErrorHandler;
-	StreamingGeneratorParams.bEnableStreaming = IsStreamingEnabled();
+	FWorldPartitionStreamingGenerator::FWorldPartitionStreamingGeneratorParams StreamingGeneratorParams = FWorldPartitionStreamingGenerator::FWorldPartitionStreamingGeneratorParams()
+		.SetWorldPartitionContext(this)
+		.SetModifiedActorsDescList( ModifiedActorsDescList)
+		.SetIsValidGrid([this](FName GridName) { return RuntimeHash->IsValidGrid(GridName); })
+		.SetErrorHandler(StreamingGenerationErrorHandlerOverride ? (*StreamingGenerationErrorHandlerOverride)(ErrorHandler) : ErrorHandler)
+		.SetEnableStreaming(IsStreamingEnabled());
 
 	FWorldPartitionStreamingGenerator StreamingGenerator(StreamingGeneratorParams);
 
@@ -1226,12 +1311,12 @@ void UWorldPartition::GenerateHLOD(ISourceControlHelper* SourceControlHelper, bo
 	ForEachActorDescContainer([this, &SourceControlHelper, bCreateActorsOnly](UActorDescContainer* InActorDescContainer)
 	{
 		FStreamingGenerationLogErrorHandler LogErrorHandler;
-		FWorldPartitionStreamingGenerator::FWorldPartitionStreamingGeneratorParams StreamingGeneratorParams;
-		StreamingGeneratorParams.WorldPartitionContext = this;
-		StreamingGeneratorParams.ErrorHandler = StreamingGenerationErrorHandlerOverride ? (*StreamingGenerationErrorHandlerOverride)(&LogErrorHandler) : &LogErrorHandler;
-		StreamingGeneratorParams.bEnableStreaming = IsStreamingEnabled();
-		StreamingGeneratorParams.FilteredClasses.Add(AWorldPartitionHLOD::StaticClass());
-		StreamingGeneratorParams.IsValidGrid = [this](FName GridName) { return RuntimeHash->IsValidGrid(GridName); };
+		FWorldPartitionStreamingGenerator::FWorldPartitionStreamingGeneratorParams StreamingGeneratorParams = FWorldPartitionStreamingGenerator::FWorldPartitionStreamingGeneratorParams()
+			.SetWorldPartitionContext(this)
+			.SetErrorHandler(StreamingGenerationErrorHandlerOverride ? (*StreamingGenerationErrorHandlerOverride)(&LogErrorHandler) : &LogErrorHandler)
+			.SetEnableStreaming(IsStreamingEnabled())
+			.SetFilteredClasses({ AWorldPartitionHLOD::StaticClass() })
+			.SetIsValidGrid([this](FName GridName) { return RuntimeHash->IsValidGrid(GridName); });
 
 		FWorldPartitionStreamingGenerator StreamingGenerator(StreamingGeneratorParams);
 		StreamingGenerator.PreparationPhase(InActorDescContainer);
@@ -1248,12 +1333,12 @@ void UWorldPartition::CheckForErrors(IStreamingGenerationErrorHandler* ErrorHand
 {
 	FActorDescList ModifiedActorDescList;
 
-	FWorldPartitionStreamingGenerator::FWorldPartitionStreamingGeneratorParams StreamingGeneratorParams;
-	StreamingGeneratorParams.WorldPartitionContext = this;
-	StreamingGeneratorParams.ModifiedActorsDescList = &ModifiedActorDescList;
-	StreamingGeneratorParams.ErrorHandler = StreamingGenerationErrorHandlerOverride ? (*StreamingGenerationErrorHandlerOverride)(ErrorHandler) : ErrorHandler;
-	StreamingGeneratorParams.IsValidGrid = [this](FName GridName) { return RuntimeHash->IsValidGrid(GridName); };
-	StreamingGeneratorParams.bEnableStreaming = IsStreamingEnabled();
+	FWorldPartitionStreamingGenerator::FWorldPartitionStreamingGeneratorParams StreamingGeneratorParams = FWorldPartitionStreamingGenerator::FWorldPartitionStreamingGeneratorParams()
+		.SetWorldPartitionContext(this)
+		.SetModifiedActorsDescList(&ModifiedActorDescList)
+		.SetErrorHandler(StreamingGenerationErrorHandlerOverride ? (*StreamingGenerationErrorHandlerOverride)(ErrorHandler) : ErrorHandler)
+		.SetIsValidGrid([this](FName GridName) { return RuntimeHash->IsValidGrid(GridName); })
+		.SetEnableStreaming(IsStreamingEnabled());
 
 	ForEachActorDescContainer([&StreamingGeneratorParams](const UActorDescContainer* InActorDescContainer)
 	{
@@ -1287,13 +1372,13 @@ void UWorldPartition::CheckForErrors(const FCheckForErrorsParams& Params)
 {
 	FActorDescList ModifiedActorDescList;
 
-	FWorldPartitionStreamingGenerator::FWorldPartitionStreamingGeneratorParams StreamingGeneratorParams;
-	StreamingGeneratorParams.WorldPartitionContext = Params.ActorDescContainer->GetWorldPartition();
-	StreamingGeneratorParams.ModifiedActorsDescList = !Params.ActorDescContainer->IsTemplateContainer() ? &ModifiedActorDescList : nullptr;
-	StreamingGeneratorParams.ErrorHandler = StreamingGenerationErrorHandlerOverride ? (*StreamingGenerationErrorHandlerOverride)(Params.ErrorHandler) : Params.ErrorHandler;
-	StreamingGeneratorParams.IsValidGrid = [](FName GridName) { return true; };
-	StreamingGeneratorParams.bEnableStreaming = Params.bEnableStreaming;
-	StreamingGeneratorParams.ActorGuidsToContainerMap = Params.ActorGuidsToContainerMap;
+	FWorldPartitionStreamingGenerator::FWorldPartitionStreamingGeneratorParams StreamingGeneratorParams = FWorldPartitionStreamingGenerator::FWorldPartitionStreamingGeneratorParams()
+		.SetWorldPartitionContext(Params.ActorDescContainer->GetWorldPartition())
+		.SetModifiedActorsDescList(!Params.ActorDescContainer->IsTemplateContainer() ? &ModifiedActorDescList : nullptr)
+		.SetErrorHandler(StreamingGenerationErrorHandlerOverride ? (*StreamingGenerationErrorHandlerOverride)(Params.ErrorHandler) : Params.ErrorHandler)
+		.SetIsValidGrid([](FName GridName) { return true; })
+		.SetEnableStreaming(Params.bEnableStreaming)
+		.SetActorGuidsToContainerMap(Params.ActorGuidsToContainerMap);
 
 	FWorldPartitionStreamingGenerator StreamingGenerator(StreamingGeneratorParams);
 
