World Partition: fix actor clusters not being sorted in streaming generation output logs.

#prefllight none

diff --git a/Engine/Source/Runtime/Engine/Private/WorldPartition/WorldPartitionStreamingGeneration.cpp b/Engine/Source/Runtime/Engine/Private/WorldPartition/WorldPartitionStreamingGeneration.cpp
--- a/Engine/Source/Runtime/Engine/Private/WorldPartition/WorldPartitionStreamingGeneration.cpp
+++ b/Engine/Source/Runtime/Engine/Private/WorldPartition/WorldPartitionStreamingGeneration.cpp
@@ -1095,8 +1095,15 @@ public:
 			{
 				UE_SCOPED_INDENT_LOG_ARCHIVE(Ar.PrintfIndent(TEXT("Clusters:")));
 
+				TArray<TArray<FGuid>> SortedClusters = ContainerDescriptor.Clusters;
+				for (TArray<FGuid>& ActorGuids : SortedClusters)
+				{
+					ActorGuids.Sort([](const FGuid& GuidA, const FGuid& GuidB) { return GuidA < GuidB; });
+				}
+				SortedClusters.Sort([](const TArray<FGuid>& GuidA, const TArray<FGuid>& GuidB) { return GuidA[0] < GuidB[0]; });
+
 				int ClusterIndex = 0;
-				for (TArray<FGuid>& ActorGuids : ContainerDescriptor.Clusters)
+				for (TArray<FGuid>& ActorGuids : SortedClusters)
 				{
 					UE_SCOPED_INDENT_LOG_ARCHIVE(Ar.PrintfIndent(TEXT("[%3d]"), ClusterIndex++));
 					for (const FGuid& ActorGuid : ActorGuids)
