WorldPartitionSubsystem no longer updates streaming sources on the server when none of its registered world partition has server streaming enabled.

diff --git a/Engine/Source/Runtime/Engine/Private/WorldPartition/WorldPartitionSubsystem.cpp b/Engine/Source/Runtime/Engine/Private/WorldPartition/WorldPartitionSubsystem.cpp
--- a/Engine/Source/Runtime/Engine/Private/WorldPartition/WorldPartitionSubsystem.cpp
+++ b/Engine/Source/Runtime/Engine/Private/WorldPartition/WorldPartitionSubsystem.cpp
@@ -134,6 +134,7 @@ static FAutoConsoleVariableRef CMaxLoadingStreamingCells(
 
 UWorldPartitionSubsystem::UWorldPartitionSubsystem()
 : StreamingSourcesHash(0)
+, NumWorldPartitionServerStreamingEnabled(0)
 {}
 
 UWorldPartition* UWorldPartitionSubsystem::GetWorldPartition()
@@ -491,12 +492,16 @@ void UWorldPartitionSubsystem::OnWorldPartitionInitialized(UWorldPartition* InWo
 
 	check(!RegisteredWorldPartitions.Contains(InWorldPartition));
 	RegisteredWorldPartitions.Add(InWorldPartition);
+	NumWorldPartitionServerStreamingEnabled += InWorldPartition->IsServerStreamingEnabled() ? 1 : 0;
+	check(NumWorldPartitionServerStreamingEnabled >= 0);
 }
 
 void UWorldPartitionSubsystem::OnWorldPartitionUninitialized(UWorldPartition* InWorldPartition)
 {
 	check(RegisteredWorldPartitions.Contains(InWorldPartition));
 	RegisteredWorldPartitions.Remove(InWorldPartition);
+	NumWorldPartitionServerStreamingEnabled -= InWorldPartition->IsServerStreamingEnabled() ? 1 : 0;
+	check(NumWorldPartitionServerStreamingEnabled >= 0);
 
 	if (RegisteredWorldPartitions.IsEmpty())
 	{
@@ -663,6 +668,18 @@ static FAutoConsoleVariableRef CVarUpdateStreamingSources(
 static const FName NAME_SIEStreamingSource(TEXT("SIE"));
 #endif
 
+bool UWorldPartitionSubsystem::IsServer() const
+{
+	const ENetMode NetMode = GetWorld()->GetNetMode();
+	const bool bIsServer = (NetMode == NM_DedicatedServer || NetMode == NM_ListenServer);
+	return bIsServer;
+}
+
+bool UWorldPartitionSubsystem::HasAnyWorldPartitionServerStreamingEnabled() const
+{
+	return (NumWorldPartitionServerStreamingEnabled > 0);
+}
+
 void UWorldPartitionSubsystem::UpdateStreamingSources()
 {
 	if (!GUpdateStreamingSources)
@@ -686,28 +703,31 @@ void UWorldPartitionSubsystem::UpdateStreamingSources()
 
 	if (!bIsUsingReplayStreamingSources)
 	{
-		bool bAllowPlayerControllerStreamingSources = true;
-#if WITH_EDITOR
-		if (UWorldPartition::IsSimulating())
+		if (!IsServer() || HasAnyWorldPartitionServerStreamingEnabled() || AWorldPartitionReplay::IsRecordingEnabled(World))
 		{
-			// We are in the SIE
-			const FVector ViewLocation = GCurrentLevelEditingViewportClient->GetViewLocation();
-			const FRotator ViewRotation = GCurrentLevelEditingViewportClient->GetViewRotation();
-			StreamingSources.Add(FWorldPartitionStreamingSource(NAME_SIEStreamingSource, ViewLocation, ViewRotation, EStreamingSourceTargetState::Activated, /*bBlockOnSlowLoading=*/false, EStreamingSourcePriority::Default, false));
-			bAllowPlayerControllerStreamingSources = false;
-		}
+			bool bAllowPlayerControllerStreamingSources = true;
+#if WITH_EDITOR
+			if (UWorldPartition::IsSimulating())
+			{
+				// We are in the SIE
+				const FVector ViewLocation = GCurrentLevelEditingViewportClient->GetViewLocation();
+				const FRotator ViewRotation = GCurrentLevelEditingViewportClient->GetViewRotation();
+				StreamingSources.Add(FWorldPartitionStreamingSource(NAME_SIEStreamingSource, ViewLocation, ViewRotation, EStreamingSourceTargetState::Activated, /*bBlockOnSlowLoading=*/false, EStreamingSourcePriority::Default, false));
+				bAllowPlayerControllerStreamingSources = false;
+			}
 #endif
-		TArray<FWorldPartitionStreamingSource> ProviderStreamingSources;
-		for (IWorldPartitionStreamingSourceProvider* StreamingSourceProvider : GetStreamingSourceProviders())
-		{
-			if (bAllowPlayerControllerStreamingSources || !Cast<APlayerController>(StreamingSourceProvider->GetStreamingSourceOwner()))
+			TArray<FWorldPartitionStreamingSource> ProviderStreamingSources;
+			for (IWorldPartitionStreamingSourceProvider* StreamingSourceProvider : GetStreamingSourceProviders())
 			{
-				ProviderStreamingSources.Reset();
-				if (StreamingSourceProvider->GetStreamingSources(ProviderStreamingSources))
+				if (bAllowPlayerControllerStreamingSources || !Cast<APlayerController>(StreamingSourceProvider->GetStreamingSourceOwner()))
 				{
-					for (FWorldPartitionStreamingSource& ProviderStreamingSource : ProviderStreamingSources)
+					ProviderStreamingSources.Reset();
+					if (StreamingSourceProvider->GetStreamingSources(ProviderStreamingSources))
 					{
-						StreamingSources.Add(MoveTemp(ProviderStreamingSource));
+						for (FWorldPartitionStreamingSource& ProviderStreamingSource : ProviderStreamingSources)
+						{
+							StreamingSources.Add(MoveTemp(ProviderStreamingSource));
+						}
 					}
 				}
 			}
diff --git a/Engine/Source/Runtime/Engine/Public/WorldPartition/WorldPartitionSubsystem.h b/Engine/Source/Runtime/Engine/Public/WorldPartition/WorldPartitionSubsystem.h
--- a/Engine/Source/Runtime/Engine/Public/WorldPartition/WorldPartitionSubsystem.h
+++ b/Engine/Source/Runtime/Engine/Public/WorldPartition/WorldPartitionSubsystem.h
@@ -169,6 +169,9 @@ private:
 
 	static void UpdateStreamingStateInternal(const UWorld* InWorld, const TArray<TObjectPtr<UWorldPartition>>& InWorldPartitions);
 
+	bool IsServer() const;
+	bool HasAnyWorldPartitionServerStreamingEnabled() const;
+
 	UWorldPartition* GetWorldPartition();
 	const UWorldPartition* GetWorldPartition() const;
 	bool CanDebugDraw() const;
@@ -187,6 +190,7 @@ private:
 	TArray<FWorldPartitionStreamingSource> StreamingSources;
 	TMap<FName, FStreamingSourceVelocity> StreamingSourcesVelocity;
 	uint32 StreamingSourcesHash;
+	int32 NumWorldPartitionServerStreamingEnabled;
 
 	TArray<FWorldPartitionDraw2DContext> WorldPartitionsDraw2DContext;
 	FDelegateHandle	DrawHandle;
