Fixed transform order issue in UWorldPartitionRuntimeLevelStreamingCell::CreateAndSetLevelStreaming

diff --git a/Engine/Source/Runtime/Engine/Private/WorldPartition/WorldPartitionRuntimeLevelStreamingCell.cpp b/Engine/Source/Runtime/Engine/Private/WorldPartition/WorldPartitionRuntimeLevelStreamingCell.cpp
--- a/Engine/Source/Runtime/Engine/Private/WorldPartition/WorldPartitionRuntimeLevelStreamingCell.cpp
+++ b/Engine/Source/Runtime/Engine/Private/WorldPartition/WorldPartitionRuntimeLevelStreamingCell.cpp
@@ -124,7 +124,7 @@ bool UWorldPartitionRuntimeLevelStreamingCell::CreateAndSetLevelStreaming(const
 	LevelStreaming->SetWorldAsset(TSoftObjectPtr<UWorld>(FSoftObjectPath(LevelPackageNameStrBuilder.ToString())));
 
 	// Include WorldPartition's transform to Level
-	LevelStreaming->SetLevelTransform(GetOuterWorld()->GetWorldPartition()->GetInstanceTransform() * InInstanceTransform);
+	LevelStreaming->SetLevelTransform(InInstanceTransform * GetOuterWorld()->GetWorldPartition()->GetInstanceTransform());
 	LevelStreaming->bClientOnlyVisible = GetClientOnlyVisible();
 	LevelStreaming->Initialize(*this);
 	LevelStreaming->PackageNameToLoad = FName(LongPackageName);
