- Make UWorldPartition::IsStreamingInEnabled() always return true on the server if server streaming is disabled.
- Protected world partition server's UpdateStreamingState not to use the optimized code path if the last update didn't complete for any reason

diff --git a/Engine/Source/Runtime/Engine/Private/WorldPartition/WorldPartition.cpp b/Engine/Source/Runtime/Engine/Private/WorldPartition/WorldPartition.cpp
--- a/Engine/Source/Runtime/Engine/Private/WorldPartition/WorldPartition.cpp
+++ b/Engine/Source/Runtime/Engine/Private/WorldPartition/WorldPartition.cpp
@@ -1594,6 +1594,10 @@ EWorldPartitionStreamingPerformance UWorldPartition::GetStreamingPerformance() c
 
 bool UWorldPartition::IsStreamingInEnabled() const
 {
+	if (IsServer() && !IsServerStreamingEnabled())
+	{
+		return true;
+	}
 	return bStreamingInEnabled;
 }
 
diff --git a/Engine/Source/Runtime/Engine/Private/WorldPartition/WorldPartitionStreamingPolicy.cpp b/Engine/Source/Runtime/Engine/Private/WorldPartition/WorldPartitionStreamingPolicy.cpp
--- a/Engine/Source/Runtime/Engine/Private/WorldPartition/WorldPartitionStreamingPolicy.cpp
+++ b/Engine/Source/Runtime/Engine/Private/WorldPartition/WorldPartitionStreamingPolicy.cpp
@@ -353,6 +353,7 @@ void UWorldPartitionStreamingPolicy::UpdateStreamingState()
 
 			// Server will activate all non data layer cells at first and then load/activate/unload data layer cells only when the data layer states change
 			if (!bIsServerStreamingEnabled && 
+				bLastUpdateCompletedLoadingAndActivation && 
 				(ServerStreamingEnabledEpoch == NewServerStreamingEnabledEpoch) &&
 				(ServerStreamingStateEpoch == WorldPartition->GetStreamingStateEpoch()) &&
 				(ServerDataLayersStatesEpoch == AWorldDataLayers::GetDataLayersStateEpoch()))
