Make sure materials are 'PostLoaded' on PSO precaching requests

diff --git a/Engine/Source/Runtime/Engine/Private/Components/PrimitiveComponent.cpp b/Engine/Source/Runtime/Engine/Private/Components/PrimitiveComponent.cpp
--- a/Engine/Source/Runtime/Engine/Private/Components/PrimitiveComponent.cpp
+++ b/Engine/Source/Runtime/Engine/Private/Components/PrimitiveComponent.cpp
@@ -4510,8 +4510,6 @@ void UPrimitiveComponent::PrecachePSOs()
 	FGraphEventArray GraphEvents;
 	for (FComponentPSOPrecacheParams& ComponentPSOPrecacheData : PSOPrecacheDataArray)
 	{
-		// Material has to be fully loaded for a static permutations to work
-		ComponentPSOPrecacheData.MaterialInterface->ConditionalPostLoad();
 		GraphEvents.Append(ComponentPSOPrecacheData.MaterialInterface->PrecachePSOs(ComponentPSOPrecacheData.VertexFactoryDataList, ComponentPSOPrecacheData.PSOPrecacheParams, ComponentPSOPrecacheData.Priority, MaterialPSOPrecacheRequestIDs));
 	}	
 
diff --git a/Engine/Source/Runtime/Engine/Private/Materials/MaterialInstance.cpp b/Engine/Source/Runtime/Engine/Private/Materials/MaterialInstance.cpp
--- a/Engine/Source/Runtime/Engine/Private/Materials/MaterialInstance.cpp
+++ b/Engine/Source/Runtime/Engine/Private/Materials/MaterialInstance.cpp
@@ -2479,12 +2479,12 @@ void UMaterialInstance::CacheShaders(EMaterialShaderPrecompileMode CompileMode)
 
 FGraphEventArray UMaterialInstance::PrecachePSOs(const FPSOPrecacheVertexFactoryDataList& VertexFactoryDataList, const FPSOPrecacheParams& InPreCacheParams, EPSOPrecachePriority Priority, TArray<FMaterialPSOPrecacheRequestID>& OutMaterialPSORequestIDs)
 {
-	// make sure material is initialized
-	ensure(!HasAnyFlags(RF_NeedPostLoad));
-	
 	FGraphEventArray GraphEvents;
 	if (FApp::CanEverRender() && PipelineStateCache::IsPSOPrecachingEnabled() && Parent)
 	{
+		// make sure material is initialized
+		ConditionalPostLoad();
+
 		if (bHasStaticPermutationResource)
 		{			
 			EMaterialQualityLevel::Type ActiveQualityLevel = GetCachedScalabilityCVars().MaterialQualityLevel;
diff --git a/Engine/Source/Runtime/Engine/Private/Particles/ParticleComponents.cpp b/Engine/Source/Runtime/Engine/Private/Particles/ParticleComponents.cpp
--- a/Engine/Source/Runtime/Engine/Private/Particles/ParticleComponents.cpp
+++ b/Engine/Source/Runtime/Engine/Private/Particles/ParticleComponents.cpp
@@ -154,9 +154,6 @@ void UFXSystemAsset::LaunchPSOPrecaching(TArrayView<VFsPerMaterialData> VFsPerMa
 	{
 		if (VFsPerMaterial.MaterialInterface)
 		{
-			// Material has to be fully loaded for a static permutations to work
-			VFsPerMaterial.MaterialInterface->ConditionalPostLoad();
-			
 			PreCachePSOParams.PrimitiveType = (EPrimitiveType)VFsPerMaterial.PrimitiveType;
 			PreCachePSOParams.bDisableBackFaceCulling = VFsPerMaterial.bDisableBackfaceCulling;
 			PrecachePSOsEvents.Append(VFsPerMaterial.MaterialInterface->PrecachePSOs(VFsPerMaterial.VertexFactoryData, PreCachePSOParams,EPSOPrecachePriority::Medium, MaterialPSOPrecacheRequestIDs));
