Report MIC name on missing usage if any.

- This should help artists identify a specific asset that is assigned to an actor which requires new usage, both during the cook and and in the editor.  Printing just base material name is unhelpful as it doesn't allow to find where this usage was needed.

diff --git a/Engine/Source/Runtime/Engine/Classes/Materials/Material.h b/Engine/Source/Runtime/Engine/Classes/Materials/Material.h
--- a/Engine/Source/Runtime/Engine/Classes/Materials/Material.h
+++ b/Engine/Source/Runtime/Engine/Classes/Materials/Material.h
@@ -1324,9 +1324,10 @@ public:
 	 * Set the given usage flag.
 	 * @param bNeedsRecompile - true if the material was recompiled for the usage change
 	 * @param Usage - The usage flag to set
+	 * @param MaterialInstance - MI requesting the usage (as of now usage flags are only stored in the parent)
 	 * @return bool - true if the material can be used for rendering with the given type.
 	 */
-	ENGINE_API bool SetMaterialUsage(bool &bNeedsRecompile, const EMaterialUsage Usage);
+	ENGINE_API bool SetMaterialUsage(bool &bNeedsRecompile, const EMaterialUsage Usage, UMaterialInterface* MaterialInstance = nullptr);
 
 	/**
 	 * Tests to see if this material needs a usage flag update
diff --git a/Engine/Source/Runtime/Engine/Private/Materials/Material.cpp b/Engine/Source/Runtime/Engine/Private/Materials/Material.cpp
--- a/Engine/Source/Runtime/Engine/Private/Materials/Material.cpp
+++ b/Engine/Source/Runtime/Engine/Private/Materials/Material.cpp
@@ -1762,7 +1762,7 @@ bool UMaterial::NeedsSetMaterialUsage_Concurrent(bool &bOutHasUsage, EMaterialUs
 	return false;
 }
 
-bool UMaterial::SetMaterialUsage(bool &bNeedsRecompile, EMaterialUsage Usage)
+bool UMaterial::SetMaterialUsage(bool &bNeedsRecompile, EMaterialUsage Usage, UMaterialInterface* MaterialInstance)
 {
 	bNeedsRecompile = false;
 
@@ -1772,6 +1772,11 @@ bool UMaterial::SetMaterialUsage(bool &bNeedsRecompile, EMaterialUsage Usage)
 		return false;
 	}
 
+	auto GetPathNameOfAssetNeedingUsage = [MaterialInstance, this]() -> FString
+	{
+		return MaterialInstance ? MaterialInstance->GetPathName() : GetPathName();
+	};
+
 	// Check that the material has been flagged for use with the given usage flag.
 	if(!GetUsageByFlag(Usage) && !bUsedAsSpecialEngineMaterial)
 	{
@@ -1783,7 +1788,7 @@ bool UMaterial::SetMaterialUsage(bool &bNeedsRecompile, EMaterialUsage Usage)
 			//Do not warn the user during automation testing
 			if (!GIsAutomationTesting)
 			{
-				UE_LOG(LogMaterial, Display, TEXT("Material %s needed to have new flag set %s !"), *GetPathName(), *GetUsageName(Usage));
+				UE_LOG(LogMaterial, Display, TEXT("Material %s needed to have new flag set %s !"), *GetPathNameOfAssetNeedingUsage(), *GetUsageName(Usage));
 			}
 
 			// Open a material update context so this material can be modified safely.
@@ -1808,7 +1813,7 @@ bool UMaterial::SetMaterialUsage(bool &bNeedsRecompile, EMaterialUsage Usage)
 #if WITH_EDITOR
 				// The package could not be marked as dirty as we're loading content in the editor. Add a Map Check error to notify the user.
 				FFormatNamedArguments Arguments;
-				Arguments.Add(TEXT("Material"), FText::FromString(*GetPathName()));
+				Arguments.Add(TEXT("Material"), FText::FromString(GetPathNameOfAssetNeedingUsage()));
 				Arguments.Add(TEXT("Usage"), FText::FromString(*GetUsageName(Usage)));
 				FMessageLog("MapCheck").Warning()
 					->AddToken(FUObjectToken::Create(this))
@@ -1823,7 +1828,7 @@ bool UMaterial::SetMaterialUsage(bool &bNeedsRecompile, EMaterialUsage Usage)
 			uint32 UsageFlagBit = (1 << (uint32)Usage);
 			if ((UsageFlagWarnings & UsageFlagBit) == 0)
 			{
-				UE_LOG(LogMaterial, Warning, TEXT("Material %s missing %s=True! Default Material will be used in game."), *GetPathName(), *GetUsageName(Usage));
+				UE_LOG(LogMaterial, Warning, TEXT("Material %s missing %s=True! Default Material will be used in game."), *GetPathNameOfAssetNeedingUsage(), *GetUsageName(Usage));
 				
 				if (bAutomaticallySetUsageInEditor)
 				{
diff --git a/Engine/Source/Runtime/Engine/Private/Materials/MaterialInstance.cpp b/Engine/Source/Runtime/Engine/Private/Materials/MaterialInstance.cpp
--- a/Engine/Source/Runtime/Engine/Private/Materials/MaterialInstance.cpp
+++ b/Engine/Source/Runtime/Engine/Private/Materials/MaterialInstance.cpp
@@ -1472,7 +1472,7 @@ bool UMaterialInstance::CheckMaterialUsage(const EMaterialUsage Usage)
 	if(Material)
 	{
 		bool bNeedsRecompile = false;
-		bool bUsageSetSuccessfully = Material->SetMaterialUsage(bNeedsRecompile, Usage);
+		bool bUsageSetSuccessfully = Material->SetMaterialUsage(bNeedsRecompile, Usage, this);
 		if (bNeedsRecompile)
 		{
 			CacheResourceShadersForRendering(EMaterialShaderPrecompileMode::None);
