Skeletal Mesh - Moved the USkeletalMesh::UpdateUVChannelData() to the USkinnedAsset base class, so that it can also be used by other SkinnedAsset derived classes.

#rnx

diff --git a/Engine/Source/Runtime/Engine/Classes/Engine/SkeletalMesh.h b/Engine/Source/Runtime/Engine/Classes/Engine/SkeletalMesh.h
--- a/Engine/Source/Runtime/Engine/Classes/Engine/SkeletalMesh.h
+++ b/Engine/Source/Runtime/Engine/Classes/Engine/SkeletalMesh.h
@@ -2353,13 +2353,6 @@ public:
 	/** Allocate a new FSkeletalMeshRenderData and assign to SkeletalMeshRenderData member.  */
 	void AllocateResourceForRendering();
 
-	/** 
-	 * Update the material UV channel data used by the texture streamer. 
-	 *
-	 * @param bResetOverrides		True if overridden values should be reset.
-	 */
-	void UpdateUVChannelData(bool bResetOverrides);
-
 	/**
 	 * Returns the UV channel data for a given material index. Used by the texture streamer.
 	 * This data applies to all lod-section using the same material. USkinnedAsset interface.
diff --git a/Engine/Source/Runtime/Engine/Classes/Engine/SkinnedAsset.h b/Engine/Source/Runtime/Engine/Classes/Engine/SkinnedAsset.h
--- a/Engine/Source/Runtime/Engine/Classes/Engine/SkinnedAsset.h
+++ b/Engine/Source/Runtime/Engine/Classes/Engine/SkinnedAsset.h
@@ -276,6 +276,13 @@ public:
 	PURE_VIRTUAL(USkinnedAsset::GetImportedModel, return nullptr;);
 #endif // WITH_EDITORONLY_DATA
 
+	/**
+	 * Update the material UV channel data used by the texture streamer.
+	 *
+	 * @param bResetOverrides True if overridden values should be reset.
+	 */
+	void UpdateUVChannelData(bool bResetOverrides);
+
 protected:	
 	/** Lock properties that should not be modified/accessed during async build. */
 	void AcquireAsyncProperty(uint64 AsyncProperties = MAX_uint64, ESkinnedAssetAsyncPropertyLockType LockType = ESkinnedAssetAsyncPropertyLockType::ReadWrite);
diff --git a/Engine/Source/Runtime/Engine/Private/SkeletalMesh.cpp b/Engine/Source/Runtime/Engine/Private/SkeletalMesh.cpp
--- a/Engine/Source/Runtime/Engine/Private/SkeletalMesh.cpp
+++ b/Engine/Source/Runtime/Engine/Private/SkeletalMesh.cpp
@@ -75,7 +75,6 @@
 #include "EditorFramework/AssetImportData.h"
 #include "Engine/SkeletalMeshSocket.h"
 #include "Components/BrushComponent.h"
-#include "Streaming/UVChannelDensity.h"
 
 #include "ClothingAssetBase.h"
 
@@ -968,111 +967,6 @@ int32 USkeletalMesh::GetNumImportedVertices() const
 }
 #endif
 
-#if WITH_EDITORONLY_DATA
-static void AccumulateUVDensities(float* OutWeightedUVDensities, float* OutWeights, const FSkeletalMeshLODRenderData& LODData, const FSkelMeshRenderSection& Section)
-{
-	const int32 NumTotalTriangles = LODData.GetTotalFaces();
-	const int32 NumCoordinateIndex = FMath::Min<int32>(LODData.GetNumTexCoords(), TEXSTREAM_MAX_NUM_UVCHANNELS);
-
-	FUVDensityAccumulator UVDensityAccs[TEXSTREAM_MAX_NUM_UVCHANNELS];
-	for (int32 CoordinateIndex = 0; CoordinateIndex < NumCoordinateIndex; ++CoordinateIndex)
-	{
-		UVDensityAccs[CoordinateIndex].Reserve(NumTotalTriangles);
-	}
-
-	TArray<uint32> Indices;
-	LODData.MultiSizeIndexContainer.GetIndexBuffer( Indices );
-	if (!Indices.Num()) return;
-
-	const uint32* SrcIndices = Indices.GetData() + Section.BaseIndex;
-	uint32 NumTriangles = Section.NumTriangles;
-
-	// Figure out Unreal unit per texel ratios.
-	for (uint32 TriangleIndex=0; TriangleIndex < NumTriangles; TriangleIndex++ )
-	{
-		//retrieve indices
-		uint32 Index0 = SrcIndices[TriangleIndex*3];
-		uint32 Index1 = SrcIndices[TriangleIndex*3+1];
-		uint32 Index2 = SrcIndices[TriangleIndex*3+2];
-
-		const float Aera = FUVDensityAccumulator::GetTriangleAera(
-			LODData.StaticVertexBuffers.PositionVertexBuffer.VertexPosition(Index0),
-			LODData.StaticVertexBuffers.PositionVertexBuffer.VertexPosition(Index1),
-			LODData.StaticVertexBuffers.PositionVertexBuffer.VertexPosition(Index2));
-
-		if (Aera > UE_SMALL_NUMBER)
-		{
-			for (int32 CoordinateIndex = 0; CoordinateIndex < NumCoordinateIndex; ++CoordinateIndex)
-			{
-				const float UVAera = FUVDensityAccumulator::GetUVChannelAera(
-					FVector2D(LODData.StaticVertexBuffers.StaticMeshVertexBuffer.GetVertexUV(Index0, CoordinateIndex)),
-					FVector2D(LODData.StaticVertexBuffers.StaticMeshVertexBuffer.GetVertexUV(Index1, CoordinateIndex)),
-					FVector2D(LODData.StaticVertexBuffers.StaticMeshVertexBuffer.GetVertexUV(Index2, CoordinateIndex)));
-
-				UVDensityAccs[CoordinateIndex].PushTriangle(Aera, UVAera);
-			}
-		}
-	}
-
-	for (int32 CoordinateIndex = 0; CoordinateIndex < NumCoordinateIndex; ++CoordinateIndex)
-	{
-		UVDensityAccs[CoordinateIndex].AccumulateDensity(OutWeightedUVDensities[CoordinateIndex], OutWeights[CoordinateIndex]);
-	}
-}
-#endif
-
-void USkeletalMesh::UpdateUVChannelData(bool bRebuildAll)
-{
-#if WITH_EDITORONLY_DATA
-	// Once cooked, the data requires to compute the scales will not be CPU accessible.
-	FSkeletalMeshRenderData* Resource = GetResourceForRendering();
-	if (FPlatformProperties::HasEditorOnlyData() && Resource)
-	{
-		TArray<FSkeletalMaterial>& MeshMaterials = GetMaterials();
-		for (int32 MaterialIndex = 0; MaterialIndex < MeshMaterials.Num(); ++MaterialIndex)
-		{
-			FMeshUVChannelInfo& UVChannelData = MeshMaterials[MaterialIndex].UVChannelData;
-
-			// Skip it if we want to keep it.
-			if (UVChannelData.IsInitialized() && (!bRebuildAll || UVChannelData.bOverrideDensities))
-				continue;
-
-			float WeightedUVDensities[TEXSTREAM_MAX_NUM_UVCHANNELS] = {0, 0, 0, 0};
-			float Weights[TEXSTREAM_MAX_NUM_UVCHANNELS] = {0, 0, 0, 0};
-
-			for (int32 LODIndex = 0; LODIndex < Resource->LODRenderData.Num(); ++LODIndex)
-			{
-				const FSkeletalMeshLODRenderData& LODData = Resource->LODRenderData[LODIndex];
-				const TArray<int32>& RemappedMaterialIndices = GetLODInfoArray()[LODIndex].LODMaterialMap;
-
-				for (int32 SectionIndex = 0; SectionIndex < LODData.RenderSections.Num(); ++SectionIndex)
-				{
-					const FSkelMeshRenderSection& SectionInfo = LODData.RenderSections[SectionIndex];
-					const int32 UsedMaterialIndex =
-						SectionIndex < RemappedMaterialIndices.Num() && MeshMaterials.IsValidIndex(RemappedMaterialIndices[SectionIndex]) ?
-						RemappedMaterialIndices[SectionIndex] :
-						SectionInfo.MaterialIndex;
-
-					if (UsedMaterialIndex != MaterialIndex)
-							continue;
-
-					AccumulateUVDensities(WeightedUVDensities, Weights, LODData, SectionInfo);
-				}
-			}
-
-			UVChannelData.bInitialized = true;
-			UVChannelData.bOverrideDensities = false;
-			for (int32 CoordinateIndex = 0; CoordinateIndex < TEXSTREAM_MAX_NUM_UVCHANNELS; ++CoordinateIndex)
-			{
-				UVChannelData.LocalUVDensities[CoordinateIndex] = (Weights[CoordinateIndex] > UE_KINDA_SMALL_NUMBER) ? (WeightedUVDensities[CoordinateIndex] / Weights[CoordinateIndex]) : 0;
-			}
-		}
-
-		Resource->SyncUVChannelData(GetMaterials());
-	}
-#endif
-}
-
 const FMeshUVChannelInfo* USkeletalMesh::GetUVChannelData(int32 MaterialIndex) const
 {
 	if (GetMaterials().IsValidIndex(MaterialIndex))
diff --git a/Engine/Source/Runtime/Engine/Private/SkinnedAsset.cpp b/Engine/Source/Runtime/Engine/Private/SkinnedAsset.cpp
--- a/Engine/Source/Runtime/Engine/Private/SkinnedAsset.cpp
+++ b/Engine/Source/Runtime/Engine/Private/SkinnedAsset.cpp
@@ -9,6 +9,9 @@
 #include "Materials/MaterialInterface.h"
 #include "SkeletalRenderGPUSkin.h"
 #include "PSOPrecache.h"
+#if WITH_EDITORONLY_DATA
+#include "Streaming/UVChannelDensity.h"
+#endif
 
 #include UE_INLINE_GENERATED_CPP_BY_NAME(SkinnedAsset)
 
@@ -31,6 +34,59 @@ static FAutoConsoleVariableRef CVarAnimSkinnedAssetISPCEnabled(TEXT("a.SkinnedAs
 
 DEFINE_LOG_CATEGORY_STATIC(LogSkinnedAsset, Log, All);
 
+#if WITH_EDITORONLY_DATA
+static void AccumulateUVDensities(float* OutWeightedUVDensities, float* OutWeights, const FSkeletalMeshLODRenderData& LODData, const FSkelMeshRenderSection& Section)
+{
+	const int32 NumTotalTriangles = LODData.GetTotalFaces();
+	const int32 NumCoordinateIndex = FMath::Min<int32>(LODData.GetNumTexCoords(), TEXSTREAM_MAX_NUM_UVCHANNELS);
+
+	FUVDensityAccumulator UVDensityAccs[TEXSTREAM_MAX_NUM_UVCHANNELS];
+	for (int32 CoordinateIndex = 0; CoordinateIndex < NumCoordinateIndex; ++CoordinateIndex)
+	{
+		UVDensityAccs[CoordinateIndex].Reserve(NumTotalTriangles);
+	}
+
+	TArray<uint32> Indices;
+	LODData.MultiSizeIndexContainer.GetIndexBuffer(Indices);
+	if (!Indices.Num()) return;
+
+	const uint32* SrcIndices = Indices.GetData() + Section.BaseIndex;
+	uint32 NumTriangles = Section.NumTriangles;
+
+	// Figure out Unreal unit per texel ratios.
+	for (uint32 TriangleIndex = 0; TriangleIndex < NumTriangles; TriangleIndex++)
+	{
+		//retrieve indices
+		uint32 Index0 = SrcIndices[TriangleIndex * 3];
+		uint32 Index1 = SrcIndices[TriangleIndex * 3 + 1];
+		uint32 Index2 = SrcIndices[TriangleIndex * 3 + 2];
+
+		const float Aera = FUVDensityAccumulator::GetTriangleAera(
+			LODData.StaticVertexBuffers.PositionVertexBuffer.VertexPosition(Index0),
+			LODData.StaticVertexBuffers.PositionVertexBuffer.VertexPosition(Index1),
+			LODData.StaticVertexBuffers.PositionVertexBuffer.VertexPosition(Index2));
+
+		if (Aera > UE_SMALL_NUMBER)
+		{
+			for (int32 CoordinateIndex = 0; CoordinateIndex < NumCoordinateIndex; ++CoordinateIndex)
+			{
+				const float UVAera = FUVDensityAccumulator::GetUVChannelAera(
+					FVector2D(LODData.StaticVertexBuffers.StaticMeshVertexBuffer.GetVertexUV(Index0, CoordinateIndex)),
+					FVector2D(LODData.StaticVertexBuffers.StaticMeshVertexBuffer.GetVertexUV(Index1, CoordinateIndex)),
+					FVector2D(LODData.StaticVertexBuffers.StaticMeshVertexBuffer.GetVertexUV(Index2, CoordinateIndex)));
+
+				UVDensityAccs[CoordinateIndex].PushTriangle(Aera, UVAera);
+			}
+		}
+	}
+
+	for (int32 CoordinateIndex = 0; CoordinateIndex < NumCoordinateIndex; ++CoordinateIndex)
+	{
+		UVDensityAccs[CoordinateIndex].AccumulateDensity(OutWeightedUVDensities[CoordinateIndex], OutWeights[CoordinateIndex]);
+	}
+}
+#endif
+
 USkinnedAsset::USkinnedAsset(const FObjectInitializer& ObjectInitializer)
 	: Super(ObjectInitializer)
 {}
@@ -192,6 +248,58 @@ bool USkinnedAsset::IsCompiling() const
 }
 #endif // WITH_EDITOR
 
+void USkinnedAsset::UpdateUVChannelData(bool bRebuildAll)
+{
+#if WITH_EDITORONLY_DATA
+	// Once cooked, the data requires to compute the scales will not be CPU accessible.
+	FSkeletalMeshRenderData* Resource = GetResourceForRendering();
+	if (FPlatformProperties::HasEditorOnlyData() && Resource)
+	{
+		TArray<FSkeletalMaterial>& MeshMaterials = GetMaterials();
+		for (int32 MaterialIndex = 0; MaterialIndex < MeshMaterials.Num(); ++MaterialIndex)
+		{
+			FMeshUVChannelInfo& UVChannelData = MeshMaterials[MaterialIndex].UVChannelData;
+
+			// Skip it if we want to keep it.
+			if (UVChannelData.IsInitialized() && (!bRebuildAll || UVChannelData.bOverrideDensities))
+				continue;
+
+			float WeightedUVDensities[TEXSTREAM_MAX_NUM_UVCHANNELS] = { 0, 0, 0, 0 };
+			float Weights[TEXSTREAM_MAX_NUM_UVCHANNELS] = { 0, 0, 0, 0 };
+
+			for (int32 LODIndex = 0; LODIndex < Resource->LODRenderData.Num(); ++LODIndex)
+			{
+				const FSkeletalMeshLODRenderData& LODData = Resource->LODRenderData[LODIndex];
+				const TArray<int32>& RemappedMaterialIndices = GetLODInfoArray()[LODIndex].LODMaterialMap;
+
+				for (int32 SectionIndex = 0; SectionIndex < LODData.RenderSections.Num(); ++SectionIndex)
+				{
+					const FSkelMeshRenderSection& SectionInfo = LODData.RenderSections[SectionIndex];
+					const int32 UsedMaterialIndex =
+						SectionIndex < RemappedMaterialIndices.Num() && MeshMaterials.IsValidIndex(RemappedMaterialIndices[SectionIndex]) ?
+						RemappedMaterialIndices[SectionIndex] :
+						SectionInfo.MaterialIndex;
+
+					if (UsedMaterialIndex != MaterialIndex)
+						continue;
+
+					AccumulateUVDensities(WeightedUVDensities, Weights, LODData, SectionInfo);
+				}
+			}
+
+			UVChannelData.bInitialized = true;
+			UVChannelData.bOverrideDensities = false;
+			for (int32 CoordinateIndex = 0; CoordinateIndex < TEXSTREAM_MAX_NUM_UVCHANNELS; ++CoordinateIndex)
+			{
+				UVChannelData.LocalUVDensities[CoordinateIndex] = (Weights[CoordinateIndex] > UE_KINDA_SMALL_NUMBER) ? (WeightedUVDensities[CoordinateIndex] / Weights[CoordinateIndex]) : 0;
+			}
+		}
+
+		Resource->SyncUVChannelData(GetMaterials());
+	}
+#endif
+}
+
 void USkinnedAsset::AcquireAsyncProperty(uint64 AsyncProperties, ESkinnedAssetAsyncPropertyLockType LockType)
 {
 #if WITH_EDITOR
