Include grid loading range (and override loading range) in streaming source hash to force update when overriding the loading range at runtime.

diff --git a/Engine/Source/Runtime/Engine/Private/WorldPartition/WorldPartitionRuntimeSpatialHash.cpp b/Engine/Source/Runtime/Engine/Private/WorldPartition/WorldPartitionRuntimeSpatialHash.cpp
--- a/Engine/Source/Runtime/Engine/Private/WorldPartition/WorldPartitionRuntimeSpatialHash.cpp
+++ b/Engine/Source/Runtime/Engine/Private/WorldPartition/WorldPartitionRuntimeSpatialHash.cpp
@@ -34,6 +34,7 @@
 #include "Misc/HierarchicalLogArchive.h"
 #include "UObject/UnrealType.h"
 #include "UObject/ObjectSaveContext.h"
+#include "Misc/HashBuilder.h"
 
 #include UE_INLINE_GENERATED_CPP_BY_NAME(WorldPartitionRuntimeSpatialHash)
 
@@ -1549,6 +1550,16 @@ void UWorldPartitionRuntimeSpatialHash::ForEachStreamingCellsQuery(const FWorldP
 	});
 }
 
+uint32 UWorldPartitionRuntimeSpatialHash::ComputeUpdateStreamingHash() const
+{
+	FHashBuilder HashBuilder;
+	ForEachStreamingGrid([&](const FSpatialHashStreamingGrid& StreamingGrid)
+	{
+		HashBuilder << StreamingGrid.GetLoadingRange();
+	});
+	return HashBuilder.GetHash();
+}
+
 void UWorldPartitionRuntimeSpatialHash::ForEachStreamingCellsSources(const TArray<FWorldPartitionStreamingSource>& Sources, TFunctionRef<bool(const UWorldPartitionRuntimeCell*, EStreamingSourceTargetState)> Func) const
 {
 	FStreamingSourceCells ActivateStreamingSourceCells;
diff --git a/Engine/Source/Runtime/Engine/Private/WorldPartition/WorldPartitionStreamingPolicy.cpp b/Engine/Source/Runtime/Engine/Private/WorldPartition/WorldPartitionStreamingPolicy.cpp
--- a/Engine/Source/Runtime/Engine/Private/WorldPartition/WorldPartitionStreamingPolicy.cpp
+++ b/Engine/Source/Runtime/Engine/Private/WorldPartition/WorldPartitionStreamingPolicy.cpp
@@ -158,6 +158,10 @@ uint32 UWorldPartitionStreamingPolicy::ComputeUpdateStreamingHash(bool bCanOptim
 
 		// Build hash that will be used to detect relevant changes
 		FHashBuilder HashBuilder;
+		if (WorldPartition->RuntimeHash)
+		{
+			HashBuilder << WorldPartition->RuntimeHash->ComputeUpdateStreamingHash();
+		}
 		HashBuilder << ComputeServerStreamingEnabledEpoch();
 		HashBuilder << GetOuterUWorldPartition()->GetStreamingStateEpoch();
 		HashBuilder << AWorldDataLayers::GetDataLayersStateEpoch();
diff --git a/Engine/Source/Runtime/Engine/Public/WorldPartition/WorldPartitionRuntimeHash.h b/Engine/Source/Runtime/Engine/Public/WorldPartition/WorldPartitionRuntimeHash.h
--- a/Engine/Source/Runtime/Engine/Public/WorldPartition/WorldPartitionRuntimeHash.h
+++ b/Engine/Source/Runtime/Engine/Public/WorldPartition/WorldPartitionRuntimeHash.h
@@ -136,6 +136,8 @@ public:
 	virtual void ForEachStreamingCells(TFunctionRef<bool(const UWorldPartitionRuntimeCell*)> Func) const {}
 	virtual void ForEachStreamingCellsQuery(const FWorldPartitionStreamingQuerySource& QuerySource, TFunctionRef<bool(const UWorldPartitionRuntimeCell*)> Func, FWorldPartitionQueryCache* QueryCache = nullptr) const {}
 	virtual void ForEachStreamingCellsSources(const TArray<FWorldPartitionStreamingSource>& Sources, TFunctionRef<bool(const UWorldPartitionRuntimeCell*, EStreamingSourceTargetState)> Func) const {}
+	// Computes a hash value of all runtime hash specific dependencies that affects the update of the streaming
+	virtual uint32 ComputeUpdateStreamingHash() const { return 0; }
 
 	bool IsCellRelevantFor(bool bClientOnlyVisible) const;
 	EWorldPartitionStreamingPerformance GetStreamingPerformance(const TSet<const UWorldPartitionRuntimeCell*>& CellToActivate) const;
diff --git a/Engine/Source/Runtime/Engine/Public/WorldPartition/WorldPartitionRuntimeSpatialHash.h b/Engine/Source/Runtime/Engine/Public/WorldPartition/WorldPartitionRuntimeSpatialHash.h
--- a/Engine/Source/Runtime/Engine/Public/WorldPartition/WorldPartitionRuntimeSpatialHash.h
+++ b/Engine/Source/Runtime/Engine/Public/WorldPartition/WorldPartitionRuntimeSpatialHash.h
@@ -284,6 +284,7 @@ public:
 	virtual void ForEachStreamingCells(TFunctionRef<bool(const UWorldPartitionRuntimeCell*)> Func) const override;
 	virtual void ForEachStreamingCellsQuery(const FWorldPartitionStreamingQuerySource& QuerySource, TFunctionRef<bool(const UWorldPartitionRuntimeCell*)> Func, FWorldPartitionQueryCache* QueryCache = nullptr) const override;
 	virtual void ForEachStreamingCellsSources(const TArray<FWorldPartitionStreamingSource>& Sources, TFunctionRef<bool(const UWorldPartitionRuntimeCell*, EStreamingSourceTargetState)> Func) const override;
+	virtual uint32 ComputeUpdateStreamingHash() const override;
 
 	virtual bool InjectExternalStreamingObject(URuntimeHashExternalStreamingObjectBase* ExternalStreamingObject) override;
 	virtual bool RemoveExternalStreamingObject(URuntimeHashExternalStreamingObjectBase* ExternalStreamingObject) override;
