Release Graphevent when material PSO precache request is done (not needed anymore)

diff --git a/Engine/Source/Runtime/Engine/Classes/Particles/ParticleSystem.h b/Engine/Source/Runtime/Engine/Classes/Particles/ParticleSystem.h
--- a/Engine/Source/Runtime/Engine/Classes/Particles/ParticleSystem.h
+++ b/Engine/Source/Runtime/Engine/Classes/Particles/ParticleSystem.h
@@ -122,6 +122,8 @@ public:
 
 	ENGINE_API virtual void PostInitProperties() override;
 
+	ENGINE_API virtual bool IsReadyForFinishDestroy() override;
+
 	/** Max number of components of this system to keep resident in the world component pool. */
 	UPROPERTY(EditAnywhere, Category = Performance)
 	uint32 MaxPoolSize;
@@ -154,7 +156,7 @@ public:
 #endif
 #endif
 
-	const FGraphEventArray& GetPrecachePSOsEvents() const { return PrecachePSOsEvents; }
+	const FGraphEventRef& GetPrecachePSOsEvent() const { return PrecachePSOsEvent; }
 	const TArray<FMaterialPSOPrecacheRequestID>& GetMaterialPSOPrecacheRequestIDs() const { return MaterialPSOPrecacheRequestIDs; }
 
 protected:
@@ -168,7 +170,7 @@ protected:
 
 	ENGINE_API void LaunchPSOPrecaching(TArrayView<VFsPerMaterialData> VFsPerMaterials);
 
-	FGraphEventArray PrecachePSOsEvents;
+	FGraphEventRef PrecachePSOsEvent;
 	TArray<FMaterialPSOPrecacheRequestID> MaterialPSOPrecacheRequestIDs;
 };
 
diff --git a/Engine/Source/Runtime/Engine/Private/PSOPrecache.cpp b/Engine/Source/Runtime/Engine/Private/PSOPrecache.cpp
--- a/Engine/Source/Runtime/Engine/Private/PSOPrecache.cpp
+++ b/Engine/Source/Runtime/Engine/Private/PSOPrecache.cpp
@@ -238,6 +238,15 @@ public:
 				{
 					CollectionGraphEvent = FGraphEvent::CreateGraphEvent();
 					PrecacheData.CollectionGraphEvent = CollectionGraphEvent;
+
+					// Create task the clear mark fully complete in the cache when done
+					FFunctionGraphTask::CreateAndDispatchWhenReady(
+						[this, Params]
+						{
+							MarkCompilationComplete(Params);
+						},
+						TStatId{}, CollectionGraphEvent
+						);
 				}
 				MaterialPSORequestData.Add(Params, PrecacheData);
 			}
@@ -284,6 +293,12 @@ public:
 		// update the state
 		FindResult->State = FindResult->ActivePSOPrecacheRequests.IsEmpty() ? EState::Completed : EState::Compiling;
 
+		// Release the graph event when done
+		if (FindResult->State == EState::Completed)
+		{
+			FindResult->CollectionGraphEvent = nullptr;
+		}
+
 		// Boost priority if requested already
 		if (FindResult->Priority == EPSOPrecachePriority::High)
 		{
@@ -374,11 +389,22 @@ private:
 			if (PrecacheData.ActivePSOPrecacheRequests.IsEmpty())
 			{
 				PrecacheData.State = EState::Completed;
+				PrecacheData.CollectionGraphEvent = nullptr;
 			}
 		}
 
 		// Not done yet?
 		return (PrecacheData.State != EState::Completed);
+	}	
+
+	void MarkCompilationComplete(const FMaterialPSOPrecacheParams& Params)
+	{
+		FRWScopeLock WriteLock(RWLock, SLT_Write);
+		FPrecacheData* FindResult = MaterialPSORequestData.Find(Params);
+		if (FindResult)
+		{
+			verify(!CheckCompilingPSOs(*FindResult, false /*bBoostPriority*/));
+		}
 	}
 
 	FRWLock RWLock;
diff --git a/Engine/Source/Runtime/Engine/Private/Particles/ParticleComponents.cpp b/Engine/Source/Runtime/Engine/Private/Particles/ParticleComponents.cpp
--- a/Engine/Source/Runtime/Engine/Private/Particles/ParticleComponents.cpp
+++ b/Engine/Source/Runtime/Engine/Private/Particles/ParticleComponents.cpp
@@ -145,11 +145,23 @@ void UFXSystemAsset::PostInitProperties()
 #endif
 }
 
+bool UFXSystemAsset::IsReadyForFinishDestroy()
+{
+	// Don't touch PrecachePSOsEvent directly because the cleanup task could set to a nullptr
+	if (PrecachePSOsEvent)
+	{
+		return false;
+	}
+
+	return Super::IsReadyForFinishDestroy();
+}
+
 void UFXSystemAsset::LaunchPSOPrecaching(TArrayView<VFsPerMaterialData> VFsPerMaterials)
 {
 	FPSOPrecacheParams PreCachePSOParams;
 	PreCachePSOParams.SetMobility(EComponentMobility::Movable);
 
+	FGraphEventArray PrecachePSOsEvents;
 	for (VFsPerMaterialData& VFsPerMaterial : VFsPerMaterials)
 	{
 		if (VFsPerMaterial.MaterialInterface)
@@ -159,6 +171,34 @@ void UFXSystemAsset::LaunchPSOPrecaching(TArrayView<VFsPerMaterialData> VFsPerMa
 			PrecachePSOsEvents.Append(VFsPerMaterial.MaterialInterface->PrecachePSOs(VFsPerMaterial.VertexFactoryData, PreCachePSOParams,EPSOPrecachePriority::Medium, MaterialPSOPrecacheRequestIDs));
 		}
 	}
+
+	// Create task the clear cached events when done
+	if (PrecachePSOsEvents.Num() > 0)
+	{
+		struct FReleasePrecachePSOsEventTask
+		{
+			explicit FReleasePrecachePSOsEventTask(FGraphEventRef& InPrecachePSOsEvent)
+				: PrecachePSOsEvent(&InPrecachePSOsEvent)
+			{
+			}
+
+			static TStatId GetStatId() { return TStatId(); }
+			static ENamedThreads::Type GetDesiredThread() { return ENamedThreads::AnyThread; }
+			static ESubsequentsMode::Type GetSubsequentsMode() { return ESubsequentsMode::TrackSubsequents; }
+
+			void DoTask(ENamedThreads::Type CurrentThread, const FGraphEventRef& MyCompletionGraphEvent)
+			{
+				*PrecachePSOsEvent = nullptr;
+			}
+
+			FGraphEventRef* PrecachePSOsEvent;
+		};
+
+		// need to set `PrecachePSOsEvent` before the task is launched to not race with its execution
+		TGraphTask<FReleasePrecachePSOsEventTask>* ReleasePrecachePSOsEventTask = TGraphTask<FReleasePrecachePSOsEventTask>::CreateTask().ConstructAndHold(PrecachePSOsEvent);
+		PrecachePSOsEvent = ReleasePrecachePSOsEventTask->GetCompletionEvent();
+		ReleasePrecachePSOsEventTask->Unlock();
+	}
 }
 
 //////////////////////////////////////////////////////////////////////////
@@ -3508,7 +3548,7 @@ void UFXSystemComponent::PrecacheAssetPSOs(UFXSystemAsset* FXSystemAsset)
 		return;
 	}
 
-	const FGraphEventArray& GraphEvents = FXSystemAsset->GetPrecachePSOsEvents();
+	FGraphEventRef GraphEvent = FXSystemAsset->GetPrecachePSOsEvent();
 
 	check(IsInGameThread() || IsInParallelGameThread());
 
@@ -3517,20 +3557,14 @@ void UFXSystemComponent::PrecacheAssetPSOs(UFXSystemAsset* FXSystemAsset)
 	bPSOPrecacheRequestBoosted = false;
 
 	// The asset will keep the Precache events alive, but these might be over. Avoid delaying scene proxy creation if everything is finished
-	bool bAllEventsDone = true;
-	for (FGraphEventRef GraphEvent : GraphEvents)
-	{
-		if (!GraphEvent->IsComplete())
-		{
-			bAllEventsDone = false;
-			break;
-		}
-	}
-
+	bool bAllEventsDone = GraphEvent == nullptr || GraphEvent->IsComplete();
 	if (!bAllEventsDone)
 	{
 		MaterialPSOPrecacheRequestIDs.Append(FXSystemAsset->GetMaterialPSOPrecacheRequestIDs());
-		RequestRecreateRenderStateWhenPSOPrecacheFinished(GraphEvents);
+
+		FGraphEventArray Events;
+		Events.Add(GraphEvent);
+		RequestRecreateRenderStateWhenPSOPrecacheFinished(Events);
 	}
 
 	bPSOPrecacheCalled = true;
