World Partition Streaming Generation: sort actor instances by their guid instead of their paths as having multiple instances of the same actor results in non-deterministic output log.

diff --git a/Engine/Source/Runtime/Engine/Private/WorldPartition/WorldPartitionRuntimeLevelStreamingCell.cpp b/Engine/Source/Runtime/Engine/Private/WorldPartition/WorldPartitionRuntimeLevelStreamingCell.cpp
--- a/Engine/Source/Runtime/Engine/Private/WorldPartition/WorldPartitionRuntimeLevelStreamingCell.cpp
+++ b/Engine/Source/Runtime/Engine/Private/WorldPartition/WorldPartitionRuntimeLevelStreamingCell.cpp
@@ -363,7 +363,7 @@ void UWorldPartitionRuntimeLevelStreamingCell::DumpStateLog(FHierarchicalLogArch
 
 	TArray<FWorldPartitionRuntimeCellObjectMapping> SortedPackages(Packages);
 
-	SortedPackages.Sort([](const FWorldPartitionRuntimeCellObjectMapping& A, const FWorldPartitionRuntimeCellObjectMapping& B) { return A.Path.LexicalLess(B.Path); });
+	SortedPackages.Sort([](const FWorldPartitionRuntimeCellObjectMapping& A, const FWorldPartitionRuntimeCellObjectMapping& B) { return A.ActorInstanceGuid < B.ActorInstanceGuid; });
 
 	for (const FWorldPartitionRuntimeCellObjectMapping& Mapping : SortedPackages)
 	{
