Fix compilation error UnrealEditor NonUnity Win64

diff --git a/Engine/Source/Runtime/Engine/Private/PhysicsEngine/Experimental/ChaosBlueprintLibrary.cpp b/Engine/Source/Runtime/Engine/Private/PhysicsEngine/Experimental/ChaosBlueprintLibrary.cpp
--- a/Engine/Source/Runtime/Engine/Private/PhysicsEngine/Experimental/ChaosBlueprintLibrary.cpp
+++ b/Engine/Source/Runtime/Engine/Private/PhysicsEngine/Experimental/ChaosBlueprintLibrary.cpp
@@ -1,10 +1,11 @@
 // Copyright Epic Games, Inc. All Rights Reserved.
 #include "PhysicsEngine/ChaosBlueprintLibrary.h"
+#include "Engine/Engine.h"
 #include "Physics/Experimental/ChaosEventRelay.h"
 
 const UChaosEventRelay* UChaosBlueprintLibrary::GetEventRelayFromContext(UObject* ContextObject)
 {
-	if (ContextObject)
+	if (GEngine && ContextObject)
 	{
 		UWorld* World = GEngine->GetWorldFromContextObject(ContextObject, EGetWorldErrorMode::ReturnNull);
 		if (World)
