Removed duplicate line initializing Grid origin in UWorldPartitionRuntimeSpatialHash::CreateStreamingGrid

diff --git a/Engine/Source/Runtime/Engine/Private/WorldPartition/WorldPartitionRuntimeSpatialHash.cpp b/Engine/Source/Runtime/Engine/Private/WorldPartition/WorldPartitionRuntimeSpatialHash.cpp
--- a/Engine/Source/Runtime/Engine/Private/WorldPartition/WorldPartitionRuntimeSpatialHash.cpp
+++ b/Engine/Source/Runtime/Engine/Private/WorldPartition/WorldPartitionRuntimeSpatialHash.cpp
@@ -1364,7 +1364,6 @@ bool UWorldPartitionRuntimeSpatialHash::CreateStreamingGrid(const FSpatialHashRu
 
 	FSpatialHashStreamingGrid& CurrentStreamingGrid = StreamingGrids.AddDefaulted_GetRef();
 	CurrentStreamingGrid.GridName = RuntimeGrid.GridName;
-	CurrentStreamingGrid.Origin = PartionedActors.Origin;
 	CurrentStreamingGrid.CellSize = PartionedActors.CellSize;
 	CurrentStreamingGrid.WorldBounds = PartionedActors.WorldBounds;
 	CurrentStreamingGrid.LoadingRange = RuntimeGrid.LoadingRange;
