Bump animation compression DDC key to purge corrupted data
Improve corrupted data handling

diff --git a/Engine/Source/Runtime/Engine/Classes/Animation/AnimCompress.h b/Engine/Source/Runtime/Engine/Classes/Animation/AnimCompress.h
--- a/Engine/Source/Runtime/Engine/Classes/Animation/AnimCompress.h
+++ b/Engine/Source/Runtime/Engine/Classes/Animation/AnimCompress.h
@@ -272,7 +272,7 @@ namespace UE
 			// This is a version string that mimics the old versioning scheme. If you
 			// want to bump this version, generate a new guid using VS->Tools->Create GUID and
 			// return it here. Ex.
-			static FString AnimationCompressionVersionString = TEXT("E64909A2AA124539905EBA5557DDEE58");
+			static FString AnimationCompressionVersionString = TEXT("75162A70E2074707ACF3DB9CC1B19597");
 			
 			struct ENGINE_API FAnimationCompressionMemorySummaryScope
 			{
diff --git a/Engine/Source/Runtime/Engine/Private/Animation/AnimationCompressionDerivedData.cpp b/Engine/Source/Runtime/Engine/Private/Animation/AnimationCompressionDerivedData.cpp
--- a/Engine/Source/Runtime/Engine/Private/Animation/AnimationCompressionDerivedData.cpp
+++ b/Engine/Source/Runtime/Engine/Private/Animation/AnimationCompressionDerivedData.cpp
@@ -126,83 +126,57 @@
 	
 	if (Response.Status == EStatus::Ok)
 	{
-		Owner.LaunchTask(TEXT("AnimationSequenceSerialize"), [this, Value = MoveTemp(Response.Value)]
+		Owner.LaunchTask(TEXT("AnimationSequenceSerialize"), [this, Name = Response.Name, Key = Response.Key, Value = MoveTemp(Response.Value)]
 		{
-			// Release execution resource as soon as the task is done
-			ON_SCOPE_EXIT { ExecutionResource = nullptr; };
-			
-			if (UAnimSequence* AnimSequence = WeakAnimSequence.Get())
+			bool bIsDataValid = true;
+
 			{
-				COOK_STAT(auto Timer = AnimSequenceCookStats::UsageStats.TimeSyncWork());
-				
-				const FSharedBuffer RecordData = Value.GetData().Decompress();
-				FMemoryReaderView Ar(RecordData, /*bIsPersistent*/ true);
-				CompressedData->SerializeCompressedData(Ar, true, AnimSequence, AnimSequence->GetSkeleton(), CompressibleAnimPtr->BoneCompressionSettings, CompressibleAnimPtr->CurveCompressionSettings);
+				// Release execution resource as soon as the task is done
+				ON_SCOPE_EXIT{ ExecutionResource = nullptr; };
 
-				if (!CompressedData->IsValid(AnimSequence, true))
+				if (UAnimSequence* AnimSequence = WeakAnimSequence.Get())
 				{
-					UE_LOG(LogAnimationCompression, Warning, TEXT("Fetched invalid compressed animation data for %s"), *CompressibleAnimPtr->FullName);
-					CompressedData->Reset();
+					COOK_STAT(auto Timer = AnimSequenceCookStats::UsageStats.TimeSyncWork());
+
+					const FSharedBuffer RecordData = Value.GetData().Decompress();
+					FMemoryReaderView Ar(RecordData, /*bIsPersistent*/ true);
+					CompressedData->SerializeCompressedData(Ar, true, AnimSequence, AnimSequence->GetSkeleton(), CompressibleAnimPtr->BoneCompressionSettings, CompressibleAnimPtr->CurveCompressionSettings);
+
+					if (!CompressedData->IsValid(AnimSequence, true))
+					{
+						UE_LOG(LogAnimationCompression, Warning, TEXT("Fetched invalid compressed animation data for %s"), *CompressibleAnimPtr->FullName);
+						CompressedData->Reset();
+
+						bIsDataValid = false;
+					}
+					else
+					{
+						UE_LOG(LogAnimationCompression, Verbose, TEXT("Fetched compressed animation data for %s"), *CompressibleAnimPtr->FullName);
+						COOK_STAT(Timer.AddHit(int64(Ar.TotalSize())));
+					}
+
+					if (Compression::FAnimationCompressionMemorySummaryScope::ShouldStoreCompressionResults())
+					{
+						const double CompressionEndTime = FPlatformTime::Seconds();
+						const double CompressionTime = CompressionEndTime - CompressionStartTime;
+
+						TArray<FBoneData> BoneData;
+						FAnimationUtils::BuildSkeletonMetaData(AnimSequence->GetSkeleton(), BoneData);
+						Compression::FAnimationCompressionMemorySummaryScope::CompressionResultSummary().GatherPostCompressionStats(*CompressedData, BoneData, AnimSequence->GetFName(), CompressionTime, false);
+					}
 				}
-				else
-				{
-					UE_LOG(LogAnimationCompression, Verbose, TEXT("Fetched compressed animation data for %s"), *CompressibleAnimPtr->FullName);
-				}
+			}
 
-				if (Compression::FAnimationCompressionMemorySummaryScope::ShouldStoreCompressionResults())
-				{
-					const double CompressionEndTime = FPlatformTime::Seconds();
-					const double CompressionTime = CompressionEndTime - CompressionStartTime;
-
-					TArray<FBoneData> BoneData;
-					FAnimationUtils::BuildSkeletonMetaData(AnimSequence->GetSkeleton(), BoneData);
-					Compression::FAnimationCompressionMemorySummaryScope::CompressionResultSummary().GatherPostCompressionStats(*CompressedData, BoneData, AnimSequence->GetFName(), CompressionTime,false);
-				}
-				
-				COOK_STAT(Timer.AddHit(int64(Ar.TotalSize())));
+			if (!bIsDataValid)
+			{
+				// Our DDC data appears to be corrupted, launch a new compression task to refresh it
+				LaunchCompressionTask(Name, Key);
 			}
 		});
 	}
 	else if (Response.Status == EStatus::Error)
 	{
-		Owner.LaunchTask(TEXT("AnimationSequenceCompression"), [this, Name = Response.Name, Key = Response.Key]
-		{
-			COOK_STAT(auto Timer = AnimSequenceCookStats::UsageStats.TimeSyncWork());
-			
-			// Release execution resource as soon as the task is done
-			ON_SCOPE_EXIT { ExecutionResource = nullptr; };
-
-			if (!BuildData())
-			{
-				return;
-			}
-			
-			if (const UAnimSequence* AnimSequence = WeakAnimSequence.Get())
-			{
-				if (!CompressedData->IsValid(AnimSequence, true))
-				{
-					UE_LOG(LogAnimationCompression, Warning, TEXT("Generated invalid compressed animation data for %s"), *CompressibleAnimPtr->FullName);
-				}
-				else
-				{
-					TArray64<uint8> RecordData;
-					FMemoryWriter64 Ar(RecordData, /*bIsPersistent*/ true);
-					CompressedData->SerializeCompressedData(Ar, true, nullptr, nullptr, CompressibleAnimPtr->BoneCompressionSettings, CompressibleAnimPtr->CurveCompressionSettings);
-					UE_LOG(LogAnimationCompression, Display, TEXT("Storing compressed animation data for %s, at %s/%s"), *Name, *FString(Key.Bucket.ToString()), *LexToString(Key.Hash));
-					GetCache().PutValue({ {Name, Key, FValue::Compress(MakeSharedBufferFromArray(MoveTemp(RecordData)))} }, Owner);
-
-					COOK_STAT(Timer.AddMiss(int64(Ar.Tell())));
-				}
-				
-				if (Compression::FAnimationCompressionMemorySummaryScope::ShouldStoreCompressionResults())
-				{
-					const double CompressionEndTime = FPlatformTime::Seconds();
-					const double CompressionTime = CompressionEndTime - CompressionStartTime;
-					Compression::FAnimationCompressionMemorySummaryScope::CompressionResultSummary().GatherPostCompressionStats(*CompressedData, CompressibleAnimPtr->BoneData, AnimSequence->GetFName(), CompressionTime, true);
-				}
-				
-			}
-		});
+		LaunchCompressionTask(Response.Name, Response.Key);
 	}
 	else
 	{
@@ -262,7 +236,48 @@
 	
 	return false;
 }
-	
+
+void FAnimationSequenceAsyncCacheTask::LaunchCompressionTask(const UE::DerivedData::FSharedString& Name, const UE::DerivedData::FCacheKey& Key)
+{
+	Owner.LaunchTask(TEXT("AnimationSequenceCompression"), [this, Name, Key]
+		{
+			COOK_STAT(auto Timer = AnimSequenceCookStats::UsageStats.TimeSyncWork());
+
+			// Release execution resource as soon as the task is done
+			ON_SCOPE_EXIT{ ExecutionResource = nullptr; };
+
+			if (!BuildData())
+			{
+				return;
+			}
+
+			if (const UAnimSequence* AnimSequence = WeakAnimSequence.Get())
+			{
+				if (!CompressedData->IsValid(AnimSequence, true))
+				{
+					UE_LOG(LogAnimationCompression, Warning, TEXT("Generated invalid compressed animation data for %s"), *CompressibleAnimPtr->FullName);
+				}
+				else
+				{
+					TArray64<uint8> RecordData;
+					FMemoryWriter64 Ar(RecordData, /*bIsPersistent*/ true);
+					CompressedData->SerializeCompressedData(Ar, true, nullptr, nullptr, CompressibleAnimPtr->BoneCompressionSettings, CompressibleAnimPtr->CurveCompressionSettings);
+					UE_LOG(LogAnimationCompression, Display, TEXT("Storing compressed animation data for %s, at %s/%s"), *Name, *FString(Key.Bucket.ToString()), *LexToString(Key.Hash));
+					UE::DerivedData::GetCache().PutValue({ {Name, Key, UE::DerivedData::FValue::Compress(MakeSharedBufferFromArray(MoveTemp(RecordData)))} }, Owner);
+
+					COOK_STAT(Timer.AddMiss(int64(Ar.Tell())));
+				}
+
+				if (Compression::FAnimationCompressionMemorySummaryScope::ShouldStoreCompressionResults())
+				{
+					const double CompressionEndTime = FPlatformTime::Seconds();
+					const double CompressionTime = CompressionEndTime - CompressionStartTime;
+					Compression::FAnimationCompressionMemorySummaryScope::CompressionResultSummary().GatherPostCompressionStats(*CompressedData, CompressibleAnimPtr->BoneData, AnimSequence->GetFName(), CompressionTime, true);
+				}
+			}
+		});
+}
+
 int32 GSkipDDC = 0;
 static FAutoConsoleVariableRef CVarSkipDDC(
 	TEXT("a.SkipDDC"),
diff --git a/Engine/Source/Runtime/Engine/Private/Animation/AnimationCompressionDerivedData.h b/Engine/Source/Runtime/Engine/Private/Animation/AnimationCompressionDerivedData.h
--- a/Engine/Source/Runtime/Engine/Private/Animation/AnimationCompressionDerivedData.h
+++ b/Engine/Source/Runtime/Engine/Private/Animation/AnimationCompressionDerivedData.h
@@ -16,7 +16,14 @@
 
 struct FCompressedAnimSequence;
 class UAnimSequence;
-namespace UE::DerivedData { struct FCacheGetValueResponse; }
+namespace UE::DerivedData
+{
+	struct FCacheGetValueResponse;
+	struct FCacheKey;
+
+	template <typename CharType> class TSharedString;
+	using FSharedString = TSharedString<TCHAR>;
+}
 
 namespace UE::Anim
 {
@@ -80,6 +87,7 @@ namespace UE::Anim
 		void BeginCache(const FIoHash& KeyHash);
 		void EndCache(UE::DerivedData::FCacheGetValueResponse&& Response);
 		bool BuildData() const;
+		void LaunchCompressionTask(const UE::DerivedData::FSharedString& Name, const UE::DerivedData::FCacheKey& Key);
 
 	private:
 		friend class FAnimationSequenceAsyncBuildWorker;
