EnsureDependenciesAreLoaded can't be called from worker threads
#fix remove EnsureAnimSequenceLoaded from non-GT code path. It's still called in UAnimSequence::BeginCache so they are validated at starting point of compression.

diff --git a/Engine/Source/Runtime/Engine/Private/Animation/AnimCompressionTypes.cpp b/Engine/Source/Runtime/Engine/Private/Animation/AnimCompressionTypes.cpp
--- a/Engine/Source/Runtime/Engine/Private/Animation/AnimCompressionTypes.cpp
+++ b/Engine/Source/Runtime/Engine/Private/Animation/AnimCompressionTypes.cpp
@@ -394,8 +394,6 @@ void FCompressibleAnimData::ResampleAnimationTrackData(const FFrameRate& SampleR
 	
 	if (AnimSequence->ShouldDataModelBeValid())
 	{
-		FAnimationUtils::EnsureAnimSequenceLoaded(*AnimSequence);
-
 		TScriptInterface<IAnimationDataModel> DataModelInterface = AnimSequence->GetDataModelInterface();
 
 		// Make a copy, deals with bone name and index
