World Partition : Added missing REN_ForceNoResetLoaders on rename called on world partition replay actor

diff --git a/Engine/Source/Runtime/Engine/Private/WorldPartition/WorldPartitionReplay.cpp b/Engine/Source/Runtime/Engine/Private/WorldPartition/WorldPartitionReplay.cpp
--- a/Engine/Source/Runtime/Engine/Private/WorldPartition/WorldPartitionReplay.cpp
+++ b/Engine/Source/Runtime/Engine/Private/WorldPartition/WorldPartitionReplay.cpp
@@ -117,7 +117,7 @@ void AWorldPartitionReplay::Uninitialize(UWorld* World)
 		if (ensure(WorldPartition && WorldPartition->Replay))
 		{
 			World->DestroyActor(WorldPartition->Replay);
-			WorldPartition->Replay->Rename(nullptr, GetTransientPackage(), REN_DontCreateRedirectors | REN_NonTransactional);
+			WorldPartition->Replay->Rename(nullptr, GetTransientPackage(), REN_DontCreateRedirectors | REN_NonTransactional | REN_ForceNoResetLoaders);
 			WorldPartition->Replay = nullptr;
 		}
 	}
