Level Instance: Handle case where Level Instance was loading, then requested to be unloaded and requested to be loaded again all while it was still loading

#rnx

diff --git a/Engine/Source/Runtime/Engine/Private/LevelInstance/LevelInstanceSubsystem.cpp b/Engine/Source/Runtime/Engine/Private/LevelInstance/LevelInstanceSubsystem.cpp
--- a/Engine/Source/Runtime/Engine/Private/LevelInstance/LevelInstanceSubsystem.cpp
+++ b/Engine/Source/Runtime/Engine/Private/LevelInstance/LevelInstanceSubsystem.cpp
@@ -176,6 +176,11 @@ void ULevelInstanceSubsystem::RequestLoadLevelInstance(ILevelInstanceInterface*
 		{
 			LevelInstancesToUnload.Remove(LevelInstance->GetLevelInstanceID());
 
+			if (IsLoading(LevelInstance))
+			{
+				return;
+			}
+
 			bool* bForcePtr = LevelInstancesToLoadOrUpdate.Find(LevelInstance);
 
 			// Avoid loading if already loaded. Can happen if actor requests unload/load in same frame. Without the force it means its not necessary.
@@ -213,6 +218,11 @@ bool ULevelInstanceSubsystem::IsLoaded(const ILevelInstanceInterface* LevelInsta
 	return LevelInstance->HasValidLevelInstanceID() && LoadedLevelInstances.Contains(LevelInstance->GetLevelInstanceID());
 }
 
+bool ULevelInstanceSubsystem::IsLoading(const ILevelInstanceInterface* LevelInstance) const
+{
+	return LevelInstance->HasValidLevelInstanceID() && LoadingLevelInstances.Contains(LevelInstance->GetLevelInstanceID());
+}
+
 void ULevelInstanceSubsystem::UpdateStreamingState()
 {
 	if (!LevelInstancesToUnload.Num() && !LevelInstancesToLoadOrUpdate.Num())
diff --git a/Engine/Source/Runtime/Engine/Public/LevelInstance/LevelInstanceSubsystem.h b/Engine/Source/Runtime/Engine/Public/LevelInstance/LevelInstanceSubsystem.h
--- a/Engine/Source/Runtime/Engine/Public/LevelInstance/LevelInstanceSubsystem.h
+++ b/Engine/Source/Runtime/Engine/Public/LevelInstance/LevelInstanceSubsystem.h
@@ -55,6 +55,7 @@ public:
 	void RequestLoadLevelInstance(ILevelInstanceInterface* LevelInstance, bool bUpdate);
 	void RequestUnloadLevelInstance(ILevelInstanceInterface* LevelInstance);
 	bool IsLoaded(const ILevelInstanceInterface* LevelInstance) const;
+	bool IsLoading(const ILevelInstanceInterface* LevelInstance) const;
 	void ForEachLevelInstanceAncestorsAndSelf(AActor* Actor, TFunctionRef<bool(ILevelInstanceInterface*)> Operation) const;
 	ULevelStreamingLevelInstance* GetLevelInstanceLevelStreaming(const ILevelInstanceInterface* LevelInstance) const;
 	void ForEachActorInLevelInstance(const ILevelInstanceInterface* LevelInstance, TFunctionRef<bool(AActor* LevelActor)> Operation) const;
