Add missing initialization of ShadowCacheInvalidationBehavior

diff --git a/Engine/Source/Runtime/Engine/Private/Components/PrimitiveComponent.cpp b/Engine/Source/Runtime/Engine/Private/Components/PrimitiveComponent.cpp
--- a/Engine/Source/Runtime/Engine/Private/Components/PrimitiveComponent.cpp
+++ b/Engine/Source/Runtime/Engine/Private/Components/PrimitiveComponent.cpp
@@ -300,6 +300,7 @@ UPrimitiveComponent::UPrimitiveComponent(const FObjectInitializer& ObjectInitial
 	bAffectDynamicIndirectLighting = true;
 	bAffectDistanceFieldLighting = true;
 	bCastStaticShadow = true;
+	ShadowCacheInvalidationBehavior = EShadowCacheInvalidationBehavior::Auto;
 	bCastVolumetricTranslucentShadow = false;
 	bCastContactShadow = true;
 	IndirectLightingCacheQuality = ILCQ_Point;
diff --git a/Engine/Source/Runtime/Foliage/Private/InstancedFoliage.cpp b/Engine/Source/Runtime/Foliage/Private/InstancedFoliage.cpp
--- a/Engine/Source/Runtime/Foliage/Private/InstancedFoliage.cpp
+++ b/Engine/Source/Runtime/Foliage/Private/InstancedFoliage.cpp
@@ -629,6 +629,8 @@ UFoliageType::UFoliageType(const FObjectInitializer& ObjectInitializer)
 	bEvaluateWorldPositionOffset = true;
 	WorldPositionOffsetDisableDistance = 0;
 
+	ShadowCacheInvalidationBehavior = EShadowCacheInvalidationBehavior::Auto;
+
 	BodyInstance.SetCollisionProfileName(UCollisionProfile::NoCollision_ProfileName);
 
 	/** Ecosystem settings*/
diff --git a/Engine/Source/Runtime/Landscape/Private/Landscape.cpp b/Engine/Source/Runtime/Landscape/Private/Landscape.cpp
--- a/Engine/Source/Runtime/Landscape/Private/Landscape.cpp
+++ b/Engine/Source/Runtime/Landscape/Private/Landscape.cpp
@@ -1373,6 +1373,7 @@ ALandscapeProxy::ALandscapeProxy(const FObjectInitializer& ObjectInitializer)
 	LOD0DistributionSetting = 1.25f;
 	LODDistributionSetting = 3.0f;
 	bCastStaticShadow = true;
+	ShadowCacheInvalidationBehavior = EShadowCacheInvalidationBehavior::Auto;
 	bUsedForNavigation = true;
 	bFillCollisionUnderLandscapeForNavmesh = false;
 	BodyInstance.SetCollisionProfileName(UCollisionProfile::BlockAll_ProfileName);
diff --git a/Engine/Source/Runtime/Landscape/Private/LandscapeGrass.cpp b/Engine/Source/Runtime/Landscape/Private/LandscapeGrass.cpp
--- a/Engine/Source/Runtime/Landscape/Private/LandscapeGrass.cpp
+++ b/Engine/Source/Runtime/Landscape/Private/LandscapeGrass.cpp
@@ -1948,6 +1948,7 @@ FGrassVariety::FGrassVariety()
 	, bCastContactShadow(true)
 	, bKeepInstanceBufferCPUCopy(false)
 	, InstanceWorldPositionOffsetDisableDistance(0)
+	, ShadowCacheInvalidationBehavior(EShadowCacheInvalidationBehavior::Auto)
 {
 	GrassDensityQuality.Init(GGrassQualityLevelCVarName, GGrassQualityLevelScalabilitySection);
 	StartCullDistanceQuality.Init(GGrassQualityLevelCVarName, GGrassQualityLevelScalabilitySection);
