FWorldPartitionActorFilter: Serialize DataLayer AssetPath as FString (Fixes BP Compilation clearing CDO values)

#rnx

diff --git a/Engine/Source/Runtime/Core/Public/UObject/FortniteMainBranchObjectVersion.h b/Engine/Source/Runtime/Core/Public/UObject/FortniteMainBranchObjectVersion.h
--- a/Engine/Source/Runtime/Core/Public/UObject/FortniteMainBranchObjectVersion.h
+++ b/Engine/Source/Runtime/Core/Public/UObject/FortniteMainBranchObjectVersion.h
@@ -319,6 +319,9 @@ struct CORE_API FFortniteMainBranchObjectVersion
 		// Serialize ActorDesc bIsMainWorldOnly
 		WorldPartitionActorDescIsMainWorldOnly,
 
+		// FWorldPartitionActorFilter go back to FString serialize of AssetPaths to avoid FArchiveReplaceOrClearExternalReferences clearing CDO references on BP Compile
+		WorldPartitionActorFilterStringAssetPath,
+
 		// -----<new versions can be added above this line>-------------------------------------------------
 		VersionPlusOne,
 		LatestVersion = VersionPlusOne - 1
diff --git a/Engine/Source/Runtime/Engine/Private/WorldPartition/Filter/WorldPartitionActorFilter.cpp b/Engine/Source/Runtime/Engine/Private/WorldPartition/Filter/WorldPartitionActorFilter.cpp
--- a/Engine/Source/Runtime/Engine/Private/WorldPartition/Filter/WorldPartitionActorFilter.cpp
+++ b/Engine/Source/Runtime/Engine/Private/WorldPartition/Filter/WorldPartitionActorFilter.cpp
@@ -160,17 +160,18 @@ bool FWorldPartitionActorFilter::Serialize(FArchive& Ar)
 		{
 			FSoftObjectPath AssetPath;
 
-			if (Ar.CustomVer(FFortniteMainBranchObjectVersion::GUID) < FFortniteMainBranchObjectVersion::WorldPartitionActorDescSerializeSoftObjectPathSupport)
+			if (Ar.CustomVer(FFortniteMainBranchObjectVersion::GUID) < FFortniteMainBranchObjectVersion::WorldPartitionActorDescSerializeSoftObjectPathSupport || 
+				Ar.CustomVer(FFortniteMainBranchObjectVersion::GUID) >= FFortniteMainBranchObjectVersion::WorldPartitionActorFilterStringAssetPath)
 			{
 				FString AssetPathStr;
 				Ar << AssetPathStr;
 				AssetPath = FSoftObjectPath(AssetPathStr);
-				FWorldPartitionHelpers::FixupRedirectedAssetPath(AssetPath);
 			}
 			else
 			{
 				Ar << AssetPath;
 			}
+			FWorldPartitionHelpers::FixupRedirectedAssetPath(AssetPath);
 
 			bool bIncluded;
 			Ar << bIncluded;
@@ -182,7 +183,8 @@ bool FWorldPartitionActorFilter::Serialize(FArchive& Ar)
 	{
 		for (auto& [AssetPath, DataLayerFilter] : DataLayerFilters)
 		{
-			Ar << AssetPath << DataLayerFilter.bIncluded;
+			FString AssetPathStr = AssetPath.ToString();
+			Ar << AssetPathStr << DataLayerFilter.bIncluded;
 		}
 	}
 
