Disable replication of WorldSettings on levels created for world partition streaming cells.

diff --git a/Engine/Source/Runtime/Engine/Classes/GameFramework/Actor.h b/Engine/Source/Runtime/Engine/Classes/GameFramework/Actor.h
--- a/Engine/Source/Runtime/Engine/Classes/GameFramework/Actor.h
+++ b/Engine/Source/Runtime/Engine/Classes/GameFramework/Actor.h
@@ -4077,6 +4077,7 @@ private:
 
 	friend struct FSetActorHiddenInSceneOutliner;
 	friend struct FSetActorGuid;
+	friend struct FSetActorReplicates;
 	friend struct FSetActorInstanceGuid;
 	friend struct FSetActorContentBundleGuid;
 	friend struct FAssignActorDataLayer;
@@ -4337,6 +4338,23 @@ private:
 	friend class UWorldPartitionConvertCommandlet;
 };
 
+#if WITH_EDITOR
+struct FSetActorReplicates
+{
+private:
+	FSetActorReplicates(AActor* InActor, bool bInReplicates)
+	{
+		if (InActor->bReplicates != bInReplicates)
+		{
+			check(!InActor->bActorInitialized);
+			InActor->bReplicates = bInReplicates;
+			InActor->RemoteRole = (bInReplicates ? ROLE_SimulatedProxy : ROLE_None);
+		}
+	}
+	friend class FWorldPartitionLevelHelper;
+};
+#endif
+
 struct FSetActorInstanceGuid
 {
 private:
diff --git a/Engine/Source/Runtime/Engine/Private/WorldPartition/WorldPartitionLevelHelper.cpp b/Engine/Source/Runtime/Engine/Private/WorldPartition/WorldPartitionLevelHelper.cpp
--- a/Engine/Source/Runtime/Engine/Private/WorldPartition/WorldPartitionLevelHelper.cpp
+++ b/Engine/Source/Runtime/Engine/Private/WorldPartition/WorldPartitionLevelHelper.cpp
@@ -22,6 +22,7 @@
 #include "WorldPartition/WorldPartitionLog.h"
 #include "WorldPartition/WorldPartitionPackageHelper.h"
 #include "WorldPartition/ContentBundle/ContentBundleEditor.h"
+#include "GameFramework/WorldSettings.h"
 #include "LevelUtils.h"
 #include "ActorFolder.h"
 
@@ -334,6 +335,8 @@ ULevel* FWorldPartitionLevelHelper::CreateEmptyLevelForRuntimeCell(const UWorldP
 	check(NewWorld->GetWorldSettings());
 	check(UWorld::FindWorldInPackage(CellPackage) == NewWorld);
 	check(InPackage || (NewWorld->GetPathName() == InWorldAssetName));
+	// We don't need the cell level's world setting to replicate
+	FSetActorReplicates SetActorReplicates(NewWorld->GetWorldSettings(), false);
 	
 	// Setup of streaming cell Runtime Level
 	ULevel* NewLevel = NewWorld->PersistentLevel;
