Favor sparse data properties if they match the name of a deprecated object property

This allows BPs to still function when object data is moved to sparse data, and the original object property deprecated

diff --git a/Engine/Source/Editor/KismetCompiler/Private/KismetCompilerMisc.cpp b/Engine/Source/Editor/KismetCompiler/Private/KismetCompilerMisc.cpp
--- a/Engine/Source/Editor/KismetCompiler/Private/KismetCompilerMisc.cpp
+++ b/Engine/Source/Editor/KismetCompiler/Private/KismetCompilerMisc.cpp
@@ -585,26 +585,42 @@ FProperty* FKismetCompilerUtilities::FindNamedPropertyInScope(UStruct* Scope, FN
 		return nullptr;
 	};
 
+	auto FindSparseClassDataProperty = [&FindProperty](UStruct* CurrentScope) -> FProperty*
+	{
+		if (UClass* Class = Cast<UClass>(CurrentScope))
+		{
+			if (UStruct* SparseData = Class->GetSparseClassDataStruct())
+			{
+				return FindProperty(SparseData);
+			}
+		}
+		return nullptr;
+	};
+
 	bIsSparseProperty = false;
 	while (Scope)
 	{
 		// Check the given scope first
 		if (FProperty* Property = FindProperty(Scope))
 		{
-			return Property;
-		}
-
-		// If this is a class, check the sparse data for the property
-		if (UClass* Class = Cast<UClass>(Scope))
-		{
-			if (UStruct* SparseData = Class->GetSparseClassDataStruct())
+			if (Property->HasAllPropertyFlags(CPF_Deprecated))
 			{
-				if (FProperty* Property = FindProperty(SparseData))
+				// If this property is deprecated, check to see if the sparse data has a property that 
+				// we should use instead (eg, when migrating data from an object into the sparse data)
+				if (FProperty* SparseProperty = FindSparseClassDataProperty(Scope))
 				{
 					bIsSparseProperty = true;
-					return Property;
+					return SparseProperty;
 				}
 			}
+			return Property;
+		}
+
+		// Check the sparse data for the property
+		if (FProperty* SparseProperty = FindSparseClassDataProperty(Scope))
+		{
+			bIsSparseProperty = true;
+			return SparseProperty;
 		}
 
 		// Functions don't automatically check their class when using a field iterator
diff --git a/Engine/Source/Runtime/Engine/Private/MemberReference.cpp b/Engine/Source/Runtime/Engine/Private/MemberReference.cpp
--- a/Engine/Source/Runtime/Engine/Private/MemberReference.cpp
+++ b/Engine/Source/Runtime/Engine/Private/MemberReference.cpp
@@ -416,14 +416,35 @@ TFieldType* FMemberReference::ResolveMemberImpl(UClass* SelfScope, TFieldTypeCla
 			if (TargetScope)
 #endif
 			{
-				// Find in target scope or in the sparse class data
-				ReturnField = FindUFieldOrFProperty<TFieldType>(TargetScope, MemberName, EFieldIterationFlags::IncludeAll);
-				if (!ReturnField)
+				auto FindSparseClassDataField = [this, TargetScope]() -> TFieldType*
 				{
 					if (UScriptStruct* SparseClassDataStruct = TargetScope->GetSparseClassDataStruct())
 					{
-						ReturnField = FindUFieldOrFProperty<TFieldType>(SparseClassDataStruct, MemberName, EFieldIterationFlags::IncludeAll);
+						return FindUFieldOrFProperty<TFieldType>(SparseClassDataStruct, MemberName, EFieldIterationFlags::IncludeAll);
+					}
+					return nullptr;
+				};
+
+				// Check the target scope first
+				if (TFieldType* TargetField = FindUFieldOrFProperty<TFieldType>(TargetScope, MemberName, EFieldIterationFlags::IncludeAll))
+				{
+					if (FProperty* TargetProperty = FFieldVariant(TargetField).Get<FProperty>();
+						TargetProperty && TargetProperty->HasAllPropertyFlags(CPF_Deprecated))
+					{
+						// If this property is deprecated, check to see if the sparse data has a property that 
+						// we should use instead (eg, when migrating data from an object into the sparse data)
+						ReturnField = FindSparseClassDataField();
 					}
+					if (!ReturnField)
+					{
+						ReturnField = TargetField;
+					}
+				}
+
+				// Check the sparse data for the field
+				if (!ReturnField)
+				{
+					ReturnField = FindSparseClassDataField();
 				}
 			}
 
