#ue5 - Fix assert in debug when changing maps or during server travel. Network smoothing should just snap to the target in this case while net mode is temporarily Standalone.

diff --git a/Engine/Source/Runtime/Engine/Private/Components/CharacterMovementComponent.cpp b/Engine/Source/Runtime/Engine/Private/Components/CharacterMovementComponent.cpp
--- a/Engine/Source/Runtime/Engine/Private/Components/CharacterMovementComponent.cpp
+++ b/Engine/Source/Runtime/Engine/Private/Components/CharacterMovementComponent.cpp
@@ -7633,7 +7633,6 @@ void UCharacterMovementComponent::SmoothCorrection(const FVector& OldLocation, c
 
 	// We shouldn't be running this on a server that is not a listen server.
 	checkSlow(GetNetMode() != NM_DedicatedServer);
-	checkSlow(GetNetMode() != NM_Standalone);
 
 	// Only client proxies or remote clients on a listen server should run this code.
 	const bool bIsSimulatedProxy = (CharacterOwner->GetLocalRole() == ROLE_SimulatedProxy);
@@ -7644,7 +7643,7 @@ void UCharacterMovementComponent::SmoothCorrection(const FVector& OldLocation, c
 	bNetworkSmoothingComplete = false;
 
 	// Handle selected smoothing mode.
-	if (NetworkSmoothingMode == ENetworkSmoothingMode::Disabled)
+	if (NetworkSmoothingMode == ENetworkSmoothingMode::Disabled || GetNetMode() == NM_Standalone)
 	{
 		UpdatedComponent->SetWorldLocationAndRotation(NewLocation, NewRotation, false, nullptr, ETeleportType::TeleportPhysics);
 		bNetworkSmoothingComplete = true;
