Fix unitialized material uniform data when setting an RVT texture that isn't bound to a volume in the scene.

diff --git a/Engine/Source/Runtime/Engine/Private/Materials/MaterialUniformExpressions.cpp b/Engine/Source/Runtime/Engine/Private/Materials/MaterialUniformExpressions.cpp
--- a/Engine/Source/Runtime/Engine/Private/Materials/MaterialUniformExpressions.cpp
+++ b/Engine/Source/Runtime/Engine/Private/Materials/MaterialUniformExpressions.cpp
@@ -811,6 +811,10 @@ void FUniformExpressionSet::FillUniformBuffer(const FMaterialRenderContext& Mate
 					{
 						AllocatedVT->GetPackedUniform(VTPackedUniform, Parameter.VirtualTextureLayerIndex);
 					}
+					else
+					{
+						*VTPackedUniform = FUintVector4(0, 0, 0, 0);
+					}
 				}
 			}
 		}
