add cvar (ClusterUnion.UseAccelerationStructure, default false for now) to add an acceleration structure to the cluster union component that'll be used for component level sweeps/raycasts/overlaps

- Modified all the low-level overlap/raycast/sweep functions to take in an optional acceleration structure override pointer.
- Converted GeomSweepSingle into a template function similar to GeomSweepMulti and GeomOverlapMulti.
- Added template specializations to take in FPhysicsGeometry for GeomSweepSingle, GeomSweepMulti, and GeomOverlapMulti.
- Added override for FChaosEngineInterface::GetGeometryCollection to take in a FPhysicsGeometry.
- Also changed the cluster union to remove the primitive component from it if the primitive component's physics state is destroyed

diff --git a/Engine/Source/Runtime/Engine/Classes/PhysicsEngine/ClusterUnionComponent.h b/Engine/Source/Runtime/Engine/Classes/PhysicsEngine/ClusterUnionComponent.h
--- a/Engine/Source/Runtime/Engine/Classes/PhysicsEngine/ClusterUnionComponent.h
+++ b/Engine/Source/Runtime/Engine/Classes/PhysicsEngine/ClusterUnionComponent.h
@@ -24,6 +24,9 @@ DECLARE_LOG_CATEGORY_EXTERN(LogClusterUnion, Log, All);
 namespace Chaos
 {
 	class FPBDRigidsSolver;
+
+	template <typename TPayload, typename T, int d>
+	class ISpatialAccelerationCollection;
 }
 
 USTRUCT()
@@ -235,12 +238,19 @@ private:
 	// User data to be able to tie the cluster particle back to this component.
 	FChaosUserData PhysicsUserData;
 
+	// An acceleration structure of all children components managed by the cluster union itself.
+	TUniquePtr<Chaos::ISpatialAcceleration<Chaos::FAccelerationStructureHandle, Chaos::FReal, 3>> AccelerationStructure;
+
 	// Need to handle the fact that this component may or may not be initialized prior to the components referenced in
 	// ClusteredComponentsReferences. This function lets us listen to OnComponentPhysicsStateChanged on the incoming
 	// primitive component so that once the physics state is properly created we can begin the process of adding it.
 	UFUNCTION()
 	void HandleComponentPhysicsStateChange(UPrimitiveComponent* ChangedComponent, EComponentPhysicsStateChange StateChange);
 
+	// Once in the cluster union, if the component's physics state is destroyed, we should remove it from the cluster union.
+	UFUNCTION()
+	void HandleComponentPhysicsStateChangePostAddIntoClusterUnion(UPrimitiveComponent* ChangedComponent, EComponentPhysicsStateChange StateChange);
+
 	// These functions only get called when the physics thread syncs to the game thread thereby enforcing a physics thread authoritative view of
 	// what particles are currently contained within the cluster union.
 	void HandleAddOrModifiedClusteredComponent(UPrimitiveComponent* ChangedComponent, const TMap<int32, FTransform>& PerBoneChildToParent);
diff --git a/Engine/Source/Runtime/Engine/Private/Collision/SceneQuery.cpp b/Engine/Source/Runtime/Engine/Private/Collision/SceneQuery.cpp
--- a/Engine/Source/Runtime/Engine/Private/Collision/SceneQuery.cpp
+++ b/Engine/Source/Runtime/Engine/Private/Collision/SceneQuery.cpp
@@ -106,10 +106,6 @@ struct FRaycastSQAdditionalInputs
 	}
 };
 
-void LowLevelRaycast(FPhysScene& Scene, const FVector& Start, const FVector& Dir, float DeltaMag, FPhysicsHitCallback<FHitRaycast>& HitBuffer, EHitFlags OutputFlags, FQueryFlags QueryFlags, const FCollisionFilterData& Filter, const ChaosInterface::FQueryFilterData& QueryFilterData, ICollisionQueryFilterCallbackBase* QueryCallback, const ChaosInterface::FQueryDebugParams& DebugParams);
-void LowLevelSweep(FPhysScene& Scene, const FPhysicsGeometry& Geom, const FTransform& StartTM, const FVector& Dir, float DeltaMag, FPhysicsHitCallback<FHitSweep>& HitBuffer, EHitFlags OutputFlags, FQueryFlags QueryFlags, const FCollisionFilterData& Filter, const ChaosInterface::FQueryFilterData& QueryFilterData, ICollisionQueryFilterCallbackBase* QueryCallback, const ChaosInterface::FQueryDebugParams& DebugParams);
-void LowLevelOverlap(FPhysScene& Scene, const FPhysicsGeometry& Geom, const FTransform& GeomPose, FPhysicsHitCallback<FHitOverlap>& HitBuffer, FQueryFlags QueryFlags, const FCollisionFilterData& Filter, const ChaosInterface::FQueryFilterData& QueryFilterData, ICollisionQueryFilterCallbackBase* QueryCallback, const ChaosInterface::FQueryDebugParams& DebugParams);
-
 template <typename InHitType, ESweepOrRay InGeometryQuery, ESingleMultiOrTest InSingleMultiOrTest>
 struct TSQTraits
 {
@@ -148,8 +144,8 @@ struct TSQTraits
 	}
 
 	//SceneTrace - ray
-	template <typename TGeomInputs, ESweepOrRay T = GeometryQuery>
-	static typename TEnableIf<T == ESweepOrRay::Raycast, void>::Type SceneTrace(FPhysScene& Scene, const TGeomInputs& GeomInputs, const FVector& Dir, float DeltaMag, const FTransform& StartTM, THitBuffer& HitBuffer, EHitFlags OutputFlags, EQueryFlags QueryFlags, const FCollisionFilterData& FilterData, const FCollisionQueryParams& Params, ICollisionQueryFilterCallbackBase* QueryCallback)
+	template <typename TAccelContainer, typename TGeomInputs, ESweepOrRay T = GeometryQuery>
+	static typename TEnableIf<T == ESweepOrRay::Raycast, void>::Type SceneTrace(const TAccelContainer& Container, const TGeomInputs& GeomInputs, const FVector& Dir, float DeltaMag, const FTransform& StartTM, THitBuffer& HitBuffer, EHitFlags OutputFlags, EQueryFlags QueryFlags, const FCollisionFilterData& FilterData, const FCollisionQueryParams& Params, ICollisionQueryFilterCallbackBase* QueryCallback)
 	{
 		using namespace ChaosInterface;
 		FQueryFilterData QueryFilterData = MakeQueryFilterData(FilterData, QueryFlags, Params);
@@ -157,12 +153,12 @@ struct TSQTraits
 #if !(UE_BUILD_TEST || UE_BUILD_SHIPPING)
 		DebugParams.bDebugQuery = Params.bDebugQuery;
 #endif
-		LowLevelRaycast(Scene, StartTM.GetLocation(), Dir, DeltaMag, HitBuffer, OutputFlags, QueryFlags, FilterData, QueryFilterData, QueryCallback, DebugParams);	//todo(ocohen): namespace?
+		LowLevelRaycast(Container, StartTM.GetLocation(), Dir, DeltaMag, HitBuffer, OutputFlags, QueryFlags, FilterData, QueryFilterData, QueryCallback, DebugParams);	//todo(ocohen): namespace?
 	}
 
 	//SceneTrace - sweep
-	template <typename TGeomInputs, ESweepOrRay T = GeometryQuery>
-	static typename TEnableIf<T == ESweepOrRay::Sweep, void>::Type SceneTrace(FPhysScene& Scene, const TGeomInputs& GeomInputs, const FVector& Dir, float DeltaMag, const FTransform& StartTM, THitBuffer& HitBuffer, EHitFlags OutputFlags, EQueryFlags QueryFlags, const FCollisionFilterData& FilterData, const FCollisionQueryParams& Params, ICollisionQueryFilterCallbackBase* QueryCallback)
+	template <typename TAccelContainer, typename TGeomInputs, ESweepOrRay T = GeometryQuery>
+	static typename TEnableIf<T == ESweepOrRay::Sweep, void>::Type SceneTrace(const TAccelContainer& Container, const TGeomInputs& GeomInputs, const FVector& Dir, float DeltaMag, const FTransform& StartTM, THitBuffer& HitBuffer, EHitFlags OutputFlags, EQueryFlags QueryFlags, const FCollisionFilterData& FilterData, const FCollisionQueryParams& Params, ICollisionQueryFilterCallbackBase* QueryCallback)
 	{
 		using namespace ChaosInterface;
 		FQueryFilterData QueryFilterData = MakeQueryFilterData(FilterData, QueryFlags, Params);
@@ -170,7 +166,7 @@ struct TSQTraits
 #if !(UE_BUILD_TEST || UE_BUILD_SHIPPING)
 		DebugParams.bDebugQuery = Params.bDebugQuery;
 #endif
-		LowLevelSweep(Scene, *GeomInputs.GetGeometry(), StartTM, Dir, DeltaMag, HitBuffer, OutputFlags, QueryFlags, FilterData, QueryFilterData, QueryCallback, DebugParams);	//todo(ocohen): namespace?
+		LowLevelSweep(Container, *GeomInputs.GetGeometry(), StartTM, Dir, DeltaMag, HitBuffer, OutputFlags, QueryFlags, FilterData, QueryFilterData, QueryCallback, DebugParams);	//todo(ocohen): namespace?
 	}
 
 	static void ResetOutHits(TArray<FHitResult>& OutHits, const FVector& Start, const FVector& End)
@@ -338,8 +334,36 @@ struct FClusterUnionHit
 	bool bHit = false;
 };
 
-template <typename Traits, typename TGeomInputs>
-bool TSceneCastCommonImp(const UWorld* World, typename Traits::TOutHits& OutHits, const TGeomInputs& GeomInputs, const FVector Start, const FVector End, ECollisionChannel TraceChannel, const struct FCollisionQueryParams& Params, const struct FCollisionResponseParams& ResponseParams, const struct FCollisionObjectQueryParams& ObjectParams)
+struct FDefaultAccelContainer
+{
+	static constexpr bool HasAccelerationStructureOverride()
+	{
+		return false;
+	}
+};
+
+struct FOverrideAccelContainer
+{
+	explicit FOverrideAccelContainer(const Chaos::IDefaultChaosSpatialAcceleration& InSpatialAcceleration)
+		: SpatialAcceleration(InSpatialAcceleration)
+	{}
+
+	static constexpr bool HasAccelerationStructureOverride()
+	{
+		return true;
+	}
+
+	const Chaos::IDefaultChaosSpatialAcceleration& GetSpatialAcceleration() const
+	{
+		return SpatialAcceleration;
+	}
+
+private:
+	const Chaos::IDefaultChaosSpatialAcceleration& SpatialAcceleration;
+};
+
+template <typename Traits, typename TGeomInputs, typename TAccelContainer>
+bool TSceneCastCommonImp(const UWorld* World, typename Traits::TOutHits& OutHits, const TGeomInputs& GeomInputs, const FVector Start, const FVector End, ECollisionChannel TraceChannel, const struct FCollisionQueryParams& Params, const struct FCollisionResponseParams& ResponseParams, const struct FCollisionObjectQueryParams& ObjectParams, const TAccelContainer& AccelContainer)
 {
 	using namespace ChaosInterface;
 
@@ -387,7 +411,14 @@ bool TSceneCastCommonImp(const UWorld* World, typename Traits::TOutHits& OutHits
 			FScopedSQHitchRepeater<decltype(HitBufferSync)> HitchRepeater(HitBufferSync, QueryCallback, FHitchDetectionInfo(Start, End, TraceChannel, Params));
 			do
 			{
-				Traits::SceneTrace(PhysScene, GeomInputs, Dir, DeltaMag, StartTM, HitchRepeater.GetBuffer(), Traits::GetHitFlags(), Traits::GetQueryFlags(), Filter, Params, &QueryCallback);
+				if constexpr (TAccelContainer::HasAccelerationStructureOverride())
+				{
+					Traits::SceneTrace(AccelContainer.GetSpatialAcceleration(), GeomInputs, Dir, DeltaMag, StartTM, HitchRepeater.GetBuffer(), Traits::GetHitFlags(), Traits::GetQueryFlags(), Filter, Params, &QueryCallback);
+				}
+				else
+				{
+					Traits::SceneTrace(PhysScene, GeomInputs, Dir, DeltaMag, StartTM, HitchRepeater.GetBuffer(), Traits::GetHitFlags(), Traits::GetQueryFlags(), Filter, Params, &QueryCallback);
+				}
 			} while (HitchRepeater.RepeatOnHitch());
 		}
 
@@ -493,7 +524,7 @@ bool TSceneCastCommonImp(const UWorld* World, typename Traits::TOutHits& OutHits
 								// is the best hit! But there could be other things we could've hit instead if we ignored the fact that we hit a cluster union.
 								FCollisionQueryParams NewParams = Params;
 								NewParams.AddIgnoredActor(ClusterUnionActor);
-								return TSceneCastCommonImp<Traits, TGeomInputs>(World, OutHits, GeomInputs, Start, End, TraceChannel, NewParams, ResponseParams, ObjectParams);
+								return TSceneCastCommonImp<Traits, TGeomInputs, TAccelContainer>(World, OutHits, GeomInputs, Start, End, TraceChannel, NewParams, ResponseParams, ObjectParams, AccelContainer);
 							}
 							else
 							{
@@ -523,8 +554,8 @@ bool TSceneCastCommonImp(const UWorld* World, typename Traits::TOutHits& OutHits
 	return bHaveBlockingHit;
 }
 
-template <typename Traits, typename PTTraits, typename TGeomInputs>
-bool TSceneCastCommon(const UWorld* World, typename Traits::TOutHits& OutHits, const TGeomInputs& GeomInputs, const FVector Start, const FVector End, ECollisionChannel TraceChannel, const struct FCollisionQueryParams& Params, const struct FCollisionResponseParams& ResponseParams, const struct FCollisionObjectQueryParams& ObjectParams)
+template <typename Traits, typename PTTraits, typename TGeomInputs, typename TAccelContainer = FDefaultAccelContainer>
+bool TSceneCastCommon(const UWorld* World, typename Traits::TOutHits& OutHits, const TGeomInputs& GeomInputs, const FVector Start, const FVector End, ECollisionChannel TraceChannel, const struct FCollisionQueryParams& Params, const struct FCollisionResponseParams& ResponseParams, const struct FCollisionObjectQueryParams& ObjectParams, const TAccelContainer& AccelContainer = FDefaultAccelContainer{})
 {
 	if ((World == NULL) || (World->GetPhysicsScene() == NULL))
 	{
@@ -534,11 +565,11 @@ bool TSceneCastCommon(const UWorld* World, typename Traits::TOutHits& OutHits, c
 	const EThreadQueryContext ThreadContext = GetThreadQueryContext(*World->GetPhysicsScene()->GetSolver());
 	if(ThreadContext == EThreadQueryContext::GTData)
 	{
-		return TSceneCastCommonImp<Traits, TGeomInputs>(World, OutHits, GeomInputs, Start, End, TraceChannel, Params, ResponseParams, ObjectParams);
+		return TSceneCastCommonImp<Traits, TGeomInputs>(World, OutHits, GeomInputs, Start, End, TraceChannel, Params, ResponseParams, ObjectParams, AccelContainer);
 	}
 	else
 	{
-		return TSceneCastCommonImp<PTTraits, TGeomInputs>(World, OutHits, GeomInputs, Start, End, TraceChannel, Params, ResponseParams, ObjectParams);
+		return TSceneCastCommonImp<PTTraits, TGeomInputs>(World, OutHits, GeomInputs, Start, End, TraceChannel, Params, ResponseParams, ObjectParams, AccelContainer);
 	}
 }
 
@@ -560,6 +591,21 @@ bool FGenericPhysicsInterface::RaycastTest(const UWorld* World, const FVector St
 	return TSceneCastCommon<TCastTraits, TPTCastTraits>(World, DummyHit, FRaycastSQAdditionalInputs(), Start, End, TraceChannel, Params, ResponseParams, ObjectParams);
 }
 
+bool FGenericPhysicsInterface::RaycastTestUsingSpatialAcceleration(const Chaos::IDefaultChaosSpatialAcceleration& Accel, const UWorld* World, const FVector Start, const FVector End, ECollisionChannel TraceChannel, const FCollisionQueryParams& Params, const FCollisionResponseParams& ResponseParams, const FCollisionObjectQueryParams& ObjectParams)
+{
+	using namespace ChaosInterface;
+
+	SCOPE_CYCLE_COUNTER(STAT_Collision_SceneQueryTotal);
+	SCOPE_CYCLE_COUNTER(STAT_Collision_RaycastAny);
+	CSV_SCOPED_TIMING_STAT(SceneQuery, RaycastTest);
+
+	using TCastTraits = TSQTraits<FHitRaycast, ESweepOrRay::Raycast, ESingleMultiOrTest::Test>;
+	using TPTCastTraits = TSQTraits<FPTRaycastHit, ESweepOrRay::Raycast, ESingleMultiOrTest::Test>;
+	FHitResult DummyHit(NoInit);
+
+	return TSceneCastCommon<TCastTraits, TPTCastTraits>(World, DummyHit, FRaycastSQAdditionalInputs(), Start, End, TraceChannel, Params, ResponseParams, ObjectParams, FOverrideAccelContainer{ Accel });
+}
+
 bool FGenericPhysicsInterface::RaycastSingle(const UWorld* World, struct FHitResult& OutHit, const FVector Start, const FVector End, ECollisionChannel TraceChannel, const struct FCollisionQueryParams& Params, const struct FCollisionResponseParams& ResponseParams, const struct FCollisionObjectQueryParams& ObjectParams)
 {
 	using namespace ChaosInterface;
@@ -573,6 +619,18 @@ bool FGenericPhysicsInterface::RaycastSingle(const UWorld* World, struct FHitRes
 	return TSceneCastCommon<TCastTraits, TPTCastTraits>(World, OutHit, FRaycastSQAdditionalInputs(), Start, End, TraceChannel, Params, ResponseParams, ObjectParams);
 }
 
+bool FGenericPhysicsInterface::RaycastSingleUsingSpatialAcceleration(const Chaos::IDefaultChaosSpatialAcceleration& Accel, const UWorld* World, struct FHitResult& OutHit, const FVector Start, const FVector End, ECollisionChannel TraceChannel, const FCollisionQueryParams& Params, const FCollisionResponseParams& ResponseParams, const FCollisionObjectQueryParams& ObjectParams)
+{
+	using namespace ChaosInterface;
+
+	SCOPE_CYCLE_COUNTER(STAT_Collision_SceneQueryTotal);
+	SCOPE_CYCLE_COUNTER(STAT_Collision_RaycastSingle);
+	CSV_SCOPED_TIMING_STAT(SceneQuery, RaycastSingle);
+
+	using TCastTraits = TSQTraits<FRaycastHit, ESweepOrRay::Raycast, ESingleMultiOrTest::Single>;
+	using TPTCastTraits = TSQTraits<FPTRaycastHit, ESweepOrRay::Raycast, ESingleMultiOrTest::Single>;
+	return TSceneCastCommon<TCastTraits, TPTCastTraits>(World, OutHit, FRaycastSQAdditionalInputs(), Start, End, TraceChannel, Params, ResponseParams, ObjectParams, FOverrideAccelContainer{ Accel });
+}
 
 
 bool FGenericPhysicsInterface::RaycastMulti(const UWorld* World, TArray<struct FHitResult>& OutHits, const FVector& Start, const FVector& End, ECollisionChannel TraceChannel, const struct FCollisionQueryParams& Params, const struct FCollisionResponseParams& ResponseParams, const struct FCollisionObjectQueryParams& ObjectParams)
@@ -588,6 +646,19 @@ bool FGenericPhysicsInterface::RaycastMulti(const UWorld* World, TArray<struct F
 	return TSceneCastCommon<TCastTraits, TPTCastTraits> (World, OutHits, FRaycastSQAdditionalInputs(), Start, End, TraceChannel, Params, ResponseParams, ObjectParams);
 }
 
+bool FGenericPhysicsInterface::RaycastMultiUsingSpatialAcceleration(const Chaos::IDefaultChaosSpatialAcceleration& Accel, const UWorld* World, TArray<struct FHitResult>& OutHits, const FVector& Start, const FVector& End, ECollisionChannel TraceChannel, const struct FCollisionQueryParams& Params, const struct FCollisionResponseParams& ResponseParams, const struct FCollisionObjectQueryParams& ObjectParams)
+{
+	using namespace ChaosInterface;
+
+	SCOPE_CYCLE_COUNTER(STAT_Collision_SceneQueryTotal);
+	SCOPE_CYCLE_COUNTER(STAT_Collision_RaycastMultiple);
+	CSV_SCOPED_TIMING_STAT(SceneQuery, RaycastMultiple);
+
+	using TCastTraits = TSQTraits<FHitRaycast, ESweepOrRay::Raycast, ESingleMultiOrTest::Multi>;
+	using TPTCastTraits = TSQTraits<FPTRaycastHit, ESweepOrRay::Raycast, ESingleMultiOrTest::Multi>;
+	return TSceneCastCommon<TCastTraits, TPTCastTraits>(World, OutHits, FRaycastSQAdditionalInputs(), Start, End, TraceChannel, Params, ResponseParams, ObjectParams, FOverrideAccelContainer{ Accel });
+}
+
 //////////////////////////////////////////////////////////////////////////
 // GEOM SWEEP
 
@@ -606,6 +677,36 @@ bool FGenericPhysicsInterface::GeomSweepTest(const UWorld* World, const struct F
 	return TSceneCastCommon<TCastTraits, TPTCastTraits>(World, DummyHit, FGeomSQAdditionalInputs(CollisionShape, Rot), Start, End, TraceChannel, Params, ResponseParams, ObjectParams);
 }
 
+bool FGenericPhysicsInterface::GeomSweepTestUsingSpatialAcceleration(const Chaos::IDefaultChaosSpatialAcceleration& Accel, const UWorld* World, const struct FCollisionShape& CollisionShape, const FQuat& Rot, FVector Start, FVector End, ECollisionChannel TraceChannel, const struct FCollisionQueryParams& Params, const struct FCollisionResponseParams& ResponseParams, const struct FCollisionObjectQueryParams& ObjectParams)
+{
+	using namespace ChaosInterface;
+
+	SCOPE_CYCLE_COUNTER(STAT_Collision_SceneQueryTotal);
+	SCOPE_CYCLE_COUNTER(STAT_Collision_GeomSweepAny);
+	CSV_SCOPED_TIMING_STAT(SceneQuery, GeomSweepTest);
+
+	using TCastTraits = TSQTraits<FHitSweep, ESweepOrRay::Sweep, ESingleMultiOrTest::Test>;
+	using TPTCastTraits = TSQTraits<FPTSweepHit, ESweepOrRay::Sweep, ESingleMultiOrTest::Test>;
+	FHitResult DummyHit(NoInit);
+
+	return TSceneCastCommon<TCastTraits, TPTCastTraits>(World, DummyHit, FGeomSQAdditionalInputs(CollisionShape, Rot), Start, End, TraceChannel, Params, ResponseParams, ObjectParams, FOverrideAccelContainer{ Accel });
+}
+
+template<>
+bool FGenericPhysicsInterface::GeomSweepSingle(const UWorld* World, const FPhysicsGeometryCollection& InGeom, const FQuat& Rot, FHitResult& OutHit, FVector Start, FVector End, ECollisionChannel TraceChannel, const struct FCollisionQueryParams& Params, const struct FCollisionResponseParams& ResponseParams, const struct FCollisionObjectQueryParams& ObjectParams)
+{
+	using namespace ChaosInterface;
+
+	SCOPE_CYCLE_COUNTER(STAT_Collision_SceneQueryTotal);
+	SCOPE_CYCLE_COUNTER(STAT_Collision_GeomSweepSingle);
+	CSV_SCOPED_TIMING_STAT(SceneQuery, GeomSweepSingle);
+
+	using TCastTraits = TSQTraits<FSweepHit, ESweepOrRay::Sweep, ESingleMultiOrTest::Single>;
+	using TPTCastTraits = TSQTraits<FPTSweepHit, ESweepOrRay::Sweep, ESingleMultiOrTest::Single>;
+	return TSceneCastCommon<TCastTraits, TPTCastTraits>(World, OutHit, FGeomCollectionSQAdditionalInputs(InGeom, Rot), Start, End, TraceChannel, Params, ResponseParams, ObjectParams);
+}
+
+template<>
 bool FGenericPhysicsInterface::GeomSweepSingle(const UWorld* World, const struct FCollisionShape& CollisionShape, const FQuat& Rot, FHitResult& OutHit, FVector Start, FVector End, ECollisionChannel TraceChannel, const struct FCollisionQueryParams& Params, const struct FCollisionResponseParams& ResponseParams, const struct FCollisionObjectQueryParams& ObjectParams)
 {
 	using namespace ChaosInterface;
@@ -620,7 +721,13 @@ bool FGenericPhysicsInterface::GeomSweepSingle(const UWorld* World, const struct
 }
 
 template<>
-bool FGenericPhysicsInterface::GeomSweepMulti(const UWorld* World, const FPhysicsGeometryCollection& InGeom, const FQuat& InGeomRot, TArray<FHitResult>& OutHits, FVector Start, FVector End, ECollisionChannel TraceChannel, const FCollisionQueryParams& Params, const FCollisionResponseParams& ResponseParams, const FCollisionObjectQueryParams& ObjectParams /*= FCollisionObjectQueryParams::DefaultObjectQueryParam*/)
+bool FGenericPhysicsInterface::GeomSweepSingle(const UWorld* World, const FPhysicsGeometry& InGeom, const FQuat& Rot, FHitResult& OutHit, FVector Start, FVector End, ECollisionChannel TraceChannel, const struct FCollisionQueryParams& Params, const struct FCollisionResponseParams& ResponseParams, const struct FCollisionObjectQueryParams& ObjectParams)
+{
+	return GeomSweepSingle(World, FChaosEngineInterface::GetGeometryCollection(InGeom), Rot, OutHit, Start, End, TraceChannel, Params, ResponseParams, ObjectParams);
+}
+
+template<>
+bool FGenericPhysicsInterface::GeomSweepMulti(const UWorld* World, const FPhysicsGeometryCollection& InGeom, const FQuat& InGeomRot, TArray<FHitResult>& OutHits, FVector Start, FVector End, ECollisionChannel TraceChannel, const FCollisionQueryParams& Params, const FCollisionResponseParams& ResponseParams, const FCollisionObjectQueryParams& ObjectParams)
 {
 	using namespace ChaosInterface;
 
@@ -634,7 +741,7 @@ bool FGenericPhysicsInterface::GeomSweepMulti(const UWorld* World, const FPhysic
 }
 
 template<>
-bool FGenericPhysicsInterface::GeomSweepMulti(const UWorld* World, const FCollisionShape& InGeom, const FQuat& InGeomRot, TArray<FHitResult>& OutHits, FVector Start, FVector End, ECollisionChannel TraceChannel, const FCollisionQueryParams& Params, const FCollisionResponseParams& ResponseParams, const FCollisionObjectQueryParams& ObjectParams /*= FCollisionObjectQueryParams::DefaultObjectQueryParam*/)
+bool FGenericPhysicsInterface::GeomSweepMulti(const UWorld* World, const FCollisionShape& InGeom, const FQuat& InGeomRot, TArray<FHitResult>& OutHits, FVector Start, FVector End, ECollisionChannel TraceChannel, const FCollisionQueryParams& Params, const FCollisionResponseParams& ResponseParams, const FCollisionObjectQueryParams& ObjectParams)
 {
 	using namespace ChaosInterface;
 
@@ -647,6 +754,11 @@ bool FGenericPhysicsInterface::GeomSweepMulti(const UWorld* World, const FCollis
 	return TSceneCastCommon<TCastTraits, TPTCastTraits>(World, OutHits, FGeomSQAdditionalInputs(InGeom, InGeomRot), Start, End, TraceChannel, Params, ResponseParams, ObjectParams);
 }
 
+template<>
+bool FGenericPhysicsInterface::GeomSweepMulti(const UWorld* World, const FPhysicsGeometry& InGeom, const FQuat& InGeomRot, TArray<FHitResult>& OutHits, FVector Start, FVector End, ECollisionChannel TraceChannel, const FCollisionQueryParams& Params, const FCollisionResponseParams& ResponseParams, const FCollisionObjectQueryParams& ObjectParams)
+{
+	return GeomSweepMulti(World, FChaosEngineInterface::GetGeometryCollection(InGeom), InGeomRot, OutHits, Start, End, TraceChannel, Params, ResponseParams, ObjectParams);
+}
 
 //////////////////////////////////////////////////////////////////////////
 // GEOM OVERLAP
@@ -662,8 +774,8 @@ namespace EQueryInfo
 	};
 }
 
-template <typename TOverlapHit, EQueryInfo::Type InfoType, typename TCollisionAnalyzerType>
-bool GeomOverlapMultiImp(const UWorld* World, const FPhysicsGeometry& Geom, const TCollisionAnalyzerType& CollisionAnalyzerType, const FTransform& GeomPose, TArray<FOverlapResult>& OutOverlaps, ECollisionChannel TraceChannel, const struct FCollisionQueryParams& Params, const struct FCollisionResponseParams& ResponseParams, const struct FCollisionObjectQueryParams& ObjectParams)
+template <typename TOverlapHit, EQueryInfo::Type InfoType, typename TCollisionAnalyzerType, typename TAccelContainer>
+bool GeomOverlapMultiImp(const UWorld* World, const FPhysicsGeometry& Geom, const TCollisionAnalyzerType& CollisionAnalyzerType, const FTransform& GeomPose, TArray<FOverlapResult>& OutOverlaps, ECollisionChannel TraceChannel, const struct FCollisionQueryParams& Params, const struct FCollisionResponseParams& ResponseParams, const struct FCollisionObjectQueryParams& ObjectParams, const TAccelContainer& AccelContainer)
 {
 	using namespace ChaosInterface;
 
@@ -709,7 +821,14 @@ bool GeomOverlapMultiImp(const UWorld* World, const FPhysicsGeometry& Geom, cons
 				FScopedSQHitchRepeater<typename TRemoveReference<decltype(OverlapBuffer)>::Type> HitchRepeater(OverlapBuffer, QueryCallback, FHitchDetectionInfo(GeomPose, TraceChannel, Params));
 				do
 				{
-					LowLevelOverlap(PhysScene, Geom, GeomPose, HitchRepeater.GetBuffer(), QueryFlags, Filter, MakeQueryFilterData(Filter, QueryFlags, Params), &QueryCallback, DebugParams);
+					if constexpr (TAccelContainer::HasAccelerationStructureOverride())
+					{
+						LowLevelOverlap(AccelContainer.GetSpatialAcceleration(), Geom, GeomPose, HitchRepeater.GetBuffer(), QueryFlags, Filter, MakeQueryFilterData(Filter, QueryFlags, Params), &QueryCallback, DebugParams);
+					}
+					else
+					{
+						LowLevelOverlap(PhysScene, Geom, GeomPose, HitchRepeater.GetBuffer(), QueryFlags, Filter, MakeQueryFilterData(Filter, QueryFlags, Params), &QueryCallback, DebugParams);
+					}
 				} while(HitchRepeater.RepeatOnHitch());
 
 				if(GetHasBlock(OverlapBuffer) && InfoType != EQueryInfo::GatherAll)	//just want true or false so don't bother gathering info
@@ -794,8 +913,8 @@ bool GeomOverlapMultiImp(const UWorld* World, const FPhysicsGeometry& Geom, cons
 	return bHaveBlockingHit;
 }
 
-template <EQueryInfo::Type InfoType, typename TCollisionAnalyzerType>
-bool GeomOverlapMultiHelper(const UWorld* World, const FPhysicsGeometry& Geom, const TCollisionAnalyzerType& CollisionAnalyzerType, const FTransform& GeomPose, TArray<FOverlapResult>& OutOverlaps, ECollisionChannel TraceChannel, const struct FCollisionQueryParams& Params, const struct FCollisionResponseParams& ResponseParams, const struct FCollisionObjectQueryParams& ObjectParams)
+template <EQueryInfo::Type InfoType, typename TCollisionAnalyzerType, typename TAccelContainer = FDefaultAccelContainer>
+bool GeomOverlapMultiHelper(const UWorld* World, const FPhysicsGeometry& Geom, const TCollisionAnalyzerType& CollisionAnalyzerType, const FTransform& GeomPose, TArray<FOverlapResult>& OutOverlaps, ECollisionChannel TraceChannel, const struct FCollisionQueryParams& Params, const struct FCollisionResponseParams& ResponseParams, const struct FCollisionObjectQueryParams& ObjectParams, const TAccelContainer& AccelContainer = FDefaultAccelContainer{})
 {
 	using namespace ChaosInterface;
 
@@ -807,11 +926,11 @@ bool GeomOverlapMultiHelper(const UWorld* World, const FPhysicsGeometry& Geom, c
 	const EThreadQueryContext ThreadContext = GetThreadQueryContext(*World->GetPhysicsScene()->GetSolver());
 	if (ThreadContext == EThreadQueryContext::GTData)
 	{
-		return GeomOverlapMultiImp<FOverlapHit, InfoType, TCollisionAnalyzerType>(World, Geom, CollisionAnalyzerType, GeomPose, OutOverlaps, TraceChannel, Params, ResponseParams, ObjectParams);
+		return GeomOverlapMultiImp<FOverlapHit, InfoType, TCollisionAnalyzerType>(World, Geom, CollisionAnalyzerType, GeomPose, OutOverlaps, TraceChannel, Params, ResponseParams, ObjectParams, AccelContainer);
 	}
 	else
 	{
-		return GeomOverlapMultiImp<FPTOverlapHit, InfoType, TCollisionAnalyzerType>(World, Geom, CollisionAnalyzerType, GeomPose, OutOverlaps, TraceChannel, Params, ResponseParams, ObjectParams);
+		return GeomOverlapMultiImp<FPTOverlapHit, InfoType, TCollisionAnalyzerType>(World, Geom, CollisionAnalyzerType, GeomPose, OutOverlaps, TraceChannel, Params, ResponseParams, ObjectParams, AccelContainer);
 	}
 }
 
@@ -827,6 +946,18 @@ bool FGenericPhysicsInterface::GeomOverlapBlockingTest(const UWorld* World, cons
 	return GeomOverlapMultiHelper<EQueryInfo::IsBlocking>(World, Adaptor.GetGeometry(), CollisionShape, Adaptor.GetGeomPose(GeomTransform.GetTranslation()), Overlaps, TraceChannel, Params, ResponseParams, ObjectParams);
 }
 
+bool FGenericPhysicsInterface::GeomOverlapBlockingTestUsingSpatialAcceleration(const Chaos::IDefaultChaosSpatialAcceleration& Accel, const UWorld* World, const struct FCollisionShape& CollisionShape, const FVector& Pos, const FQuat& Rot, ECollisionChannel TraceChannel, const struct FCollisionQueryParams& Params, const struct FCollisionResponseParams& ResponseParams, const struct FCollisionObjectQueryParams& ObjectParams)
+{
+	SCOPE_CYCLE_COUNTER(STAT_Collision_SceneQueryTotal);
+	SCOPE_CYCLE_COUNTER(STAT_Collision_GeomOverlapBlocking);
+	CSV_SCOPED_TIMING_STAT(SceneQuery, GeomOverlapBlocking);
+
+	TArray<FOverlapResult> Overlaps;	//needed only for template shared code
+	FTransform GeomTransform(Rot, Pos);
+	FPhysicsShapeAdapter Adaptor(GeomTransform.GetRotation(), CollisionShape);
+	return GeomOverlapMultiHelper<EQueryInfo::IsBlocking>(World, Adaptor.GetGeometry(), CollisionShape, Adaptor.GetGeomPose(GeomTransform.GetTranslation()), Overlaps, TraceChannel, Params, ResponseParams, ObjectParams, FOverrideAccelContainer{ Accel });
+}
+
 bool FGenericPhysicsInterface::GeomOverlapAnyTest(const UWorld* World, const struct FCollisionShape& CollisionShape, const FVector& Pos, const FQuat& Rot, ECollisionChannel TraceChannel, const struct FCollisionQueryParams& Params, const struct FCollisionResponseParams& ResponseParams, const struct FCollisionObjectQueryParams& ObjectParams)
 {
 	SCOPE_CYCLE_COUNTER(STAT_Collision_SceneQueryTotal);
@@ -839,6 +970,18 @@ bool FGenericPhysicsInterface::GeomOverlapAnyTest(const UWorld* World, const str
 	return GeomOverlapMultiHelper<EQueryInfo::IsAnything>(World, Adaptor.GetGeometry(), CollisionShape, Adaptor.GetGeomPose(GeomTransform.GetTranslation()), Overlaps, TraceChannel, Params, ResponseParams, ObjectParams);
 }
 
+bool FGenericPhysicsInterface::GeomOverlapAnyTestUsingSpatialAcceleration(const Chaos::IDefaultChaosSpatialAcceleration& Accel, const UWorld* World, const struct FCollisionShape& CollisionShape, const FVector& Pos, const FQuat& Rot, ECollisionChannel TraceChannel, const struct FCollisionQueryParams& Params, const struct FCollisionResponseParams& ResponseParams, const struct FCollisionObjectQueryParams& ObjectParams)
+{
+	SCOPE_CYCLE_COUNTER(STAT_Collision_SceneQueryTotal);
+	SCOPE_CYCLE_COUNTER(STAT_Collision_GeomOverlapAny);
+	CSV_SCOPED_TIMING_STAT(SceneQuery, GeomOverlapAny);
+
+	TArray<FOverlapResult> Overlaps;	//needed only for template shared code
+	FTransform GeomTransform(Rot, Pos);
+	FPhysicsShapeAdapter Adaptor(GeomTransform.GetRotation(), CollisionShape);
+	return GeomOverlapMultiHelper<EQueryInfo::IsAnything>(World, Adaptor.GetGeometry(), CollisionShape, Adaptor.GetGeomPose(GeomTransform.GetTranslation()), Overlaps, TraceChannel, Params, ResponseParams, ObjectParams, FOverrideAccelContainer{ Accel });
+}
+
 template<>
 bool FGenericPhysicsInterface::GeomOverlapMulti(const UWorld* World, const FPhysicsGeometryCollection& InGeom, const FVector& InPosition, const FQuat& InRotation, TArray<FOverlapResult>& OutOverlaps, ECollisionChannel TraceChannel, const FCollisionQueryParams& Params, const FCollisionResponseParams& ResponseParams, const FCollisionObjectQueryParams& ObjectParams)
 {
@@ -861,3 +1004,106 @@ bool FGenericPhysicsInterface::GeomOverlapMulti(const UWorld* World, const FColl
 	FPhysicsShapeAdapter Adaptor(GeomTransform.GetRotation(), InGeom);
 	return GeomOverlapMultiHelper<EQueryInfo::GatherAll>(World, Adaptor.GetGeometry(), InGeom, Adaptor.GetGeomPose(GeomTransform.GetTranslation()), OutOverlaps, TraceChannel, Params, ResponseParams, ObjectParams);
 }
+
+template<>
+bool FGenericPhysicsInterface::GeomOverlapMulti(const UWorld* World, const FPhysicsGeometry& InGeom, const FVector& InPosition, const FQuat& InRotation, TArray<FOverlapResult>& OutOverlaps, ECollisionChannel TraceChannel, const FCollisionQueryParams& Params, const FCollisionResponseParams& ResponseParams, const FCollisionObjectQueryParams& ObjectParams)
+{
+	return GeomOverlapMulti(World, FChaosEngineInterface::GetGeometryCollection(InGeom), InPosition, InRotation, OutOverlaps, TraceChannel, Params, ResponseParams, ObjectParams);
+}
+
+template<>
+bool FGenericPhysicsInterface::GeomSweepSingleUsingSpatialAcceleration(const Chaos::IDefaultChaosSpatialAcceleration& Accel, const UWorld* World, const FPhysicsGeometryCollection& InGeom, const FQuat& Rot, FHitResult& OutHit, FVector Start, FVector End, ECollisionChannel TraceChannel, const FCollisionQueryParams& Params, const FCollisionResponseParams& ResponseParams, const FCollisionObjectQueryParams& ObjectParams)
+{
+	using namespace ChaosInterface;
+
+	SCOPE_CYCLE_COUNTER(STAT_Collision_SceneQueryTotal);
+	SCOPE_CYCLE_COUNTER(STAT_Collision_GeomSweepSingle);
+	CSV_SCOPED_TIMING_STAT(SceneQuery, GeomSweepSingle);
+
+	using TCastTraits = TSQTraits<FSweepHit, ESweepOrRay::Sweep, ESingleMultiOrTest::Single>;
+	using TPTCastTraits = TSQTraits<FPTSweepHit, ESweepOrRay::Sweep, ESingleMultiOrTest::Single>;
+	return TSceneCastCommon<TCastTraits, TPTCastTraits>(World, OutHit, FGeomCollectionSQAdditionalInputs(InGeom, Rot), Start, End, TraceChannel, Params, ResponseParams, ObjectParams, FOverrideAccelContainer{ Accel });
+}
+
+template<>
+bool FGenericPhysicsInterface::GeomSweepSingleUsingSpatialAcceleration(const Chaos::IDefaultChaosSpatialAcceleration& Accel, const UWorld* World, const FCollisionShape& InGeom, const FQuat& Rot, FHitResult& OutHit, FVector Start, FVector End, ECollisionChannel TraceChannel, const FCollisionQueryParams& Params, const FCollisionResponseParams& ResponseParams, const FCollisionObjectQueryParams& ObjectParams)
+{
+	using namespace ChaosInterface;
+
+	SCOPE_CYCLE_COUNTER(STAT_Collision_SceneQueryTotal);
+	SCOPE_CYCLE_COUNTER(STAT_Collision_GeomSweepSingle);
+	CSV_SCOPED_TIMING_STAT(SceneQuery, GeomSweepSingle);
+
+	using TCastTraits = TSQTraits<FSweepHit, ESweepOrRay::Sweep, ESingleMultiOrTest::Single>;
+	using TPTCastTraits = TSQTraits<FPTSweepHit, ESweepOrRay::Sweep, ESingleMultiOrTest::Single>;
+	return TSceneCastCommon<TCastTraits, TPTCastTraits>(World, OutHit, FGeomSQAdditionalInputs(InGeom, Rot), Start, End, TraceChannel, Params, ResponseParams, ObjectParams, FOverrideAccelContainer{ Accel });
+}
+
+template<>
+bool FGenericPhysicsInterface::GeomSweepSingleUsingSpatialAcceleration(const Chaos::IDefaultChaosSpatialAcceleration& Accel, const UWorld* World, const FPhysicsGeometry& InGeom, const FQuat& Rot, FHitResult& OutHit, FVector Start, FVector End, ECollisionChannel TraceChannel, const FCollisionQueryParams& Params, const FCollisionResponseParams& ResponseParams, const FCollisionObjectQueryParams& ObjectParams)
+{
+	return GeomSweepSingleUsingSpatialAcceleration(Accel, World, FChaosEngineInterface::GetGeometryCollection(InGeom), Rot, OutHit, Start, End, TraceChannel, Params, ResponseParams, ObjectParams);
+}
+
+template<>
+bool FGenericPhysicsInterface::GeomSweepMultiUsingSpatialAcceleration(const Chaos::IDefaultChaosSpatialAcceleration& Accel, const UWorld* World, const FPhysicsGeometryCollection& InGeom, const FQuat& InGeomRot, TArray<FHitResult>& OutHits, FVector Start, FVector End, ECollisionChannel TraceChannel, const FCollisionQueryParams& Params, const FCollisionResponseParams& ResponseParams, const FCollisionObjectQueryParams& ObjectParams)
+{
+	using namespace ChaosInterface;
+
+	SCOPE_CYCLE_COUNTER(STAT_Collision_SceneQueryTotal);
+	SCOPE_CYCLE_COUNTER(STAT_Collision_GeomSweepMultiple);
+	CSV_SCOPED_TIMING_STAT(SceneQuery, GeomSweepMultiple);
+
+	using TCastTraits = TSQTraits<FHitSweep, ESweepOrRay::Sweep, ESingleMultiOrTest::Multi>;
+	using TPTCastTraits = TSQTraits<FPTSweepHit, ESweepOrRay::Sweep, ESingleMultiOrTest::Multi>;
+	return TSceneCastCommon<TCastTraits, TPTCastTraits>(World, OutHits, FGeomCollectionSQAdditionalInputs(InGeom, InGeomRot), Start, End, TraceChannel, Params, ResponseParams, ObjectParams, FOverrideAccelContainer{ Accel });
+}
+
+template<>
+bool FGenericPhysicsInterface::GeomSweepMultiUsingSpatialAcceleration(const Chaos::IDefaultChaosSpatialAcceleration& Accel, const UWorld* World, const FCollisionShape& InGeom, const FQuat& InGeomRot, TArray<FHitResult>& OutHits, FVector Start, FVector End, ECollisionChannel TraceChannel, const FCollisionQueryParams& Params, const FCollisionResponseParams& ResponseParams, const FCollisionObjectQueryParams& ObjectParams)
+{
+	using namespace ChaosInterface;
+
+	SCOPE_CYCLE_COUNTER(STAT_Collision_SceneQueryTotal);
+	SCOPE_CYCLE_COUNTER(STAT_Collision_GeomSweepMultiple);
+	CSV_SCOPED_TIMING_STAT(SceneQuery, GeomSweepMultiple);
+
+	using TCastTraits = TSQTraits<FHitSweep, ESweepOrRay::Sweep, ESingleMultiOrTest::Multi>;
+	using TPTCastTraits = TSQTraits<FPTSweepHit, ESweepOrRay::Sweep, ESingleMultiOrTest::Multi>;
+	return TSceneCastCommon<TCastTraits, TPTCastTraits>(World, OutHits, FGeomSQAdditionalInputs(InGeom, InGeomRot), Start, End, TraceChannel, Params, ResponseParams, ObjectParams, FOverrideAccelContainer{ Accel });
+}
+
+template<>
+bool FGenericPhysicsInterface::GeomSweepMultiUsingSpatialAcceleration(const Chaos::IDefaultChaosSpatialAcceleration& Accel, const UWorld* World, const FPhysicsGeometry& InGeom, const FQuat& InGeomRot, TArray<FHitResult>& OutHits, FVector Start, FVector End, ECollisionChannel TraceChannel, const FCollisionQueryParams& Params, const FCollisionResponseParams& ResponseParams, const FCollisionObjectQueryParams& ObjectParams)
+{
+	return GeomSweepMultiUsingSpatialAcceleration(Accel, World, FChaosEngineInterface::GetGeometryCollection(InGeom), InGeomRot, OutHits, Start, End, TraceChannel, Params, ResponseParams, ObjectParams);
+}
+
+template<>
+bool FGenericPhysicsInterface::GeomOverlapMultiUsingSpatialAcceleration(const Chaos::IDefaultChaosSpatialAcceleration& Accel, const UWorld* World, const FPhysicsGeometryCollection& InGeom, const FVector& InPosition, const FQuat& InRotation, TArray<FOverlapResult>& OutOverlaps, ECollisionChannel TraceChannel, const FCollisionQueryParams& Params, const FCollisionResponseParams& ResponseParams, const FCollisionObjectQueryParams& ObjectParams)
+{
+	SCOPE_CYCLE_COUNTER(STAT_Collision_SceneQueryTotal);
+	SCOPE_CYCLE_COUNTER(STAT_Collision_GeomOverlapMultiple);
+	CSV_SCOPED_TIMING_STAT(SceneQuery, GeomOverlapMultiple);
+
+	FTransform GeomTransform(InRotation, InPosition);
+	return GeomOverlapMultiHelper<EQueryInfo::GatherAll>(World, InGeom.GetGeometry(), InGeom, GeomTransform, OutOverlaps, TraceChannel, Params, ResponseParams, ObjectParams, FOverrideAccelContainer{ Accel });
+}
+
+template<>
+bool FGenericPhysicsInterface::GeomOverlapMultiUsingSpatialAcceleration(const Chaos::IDefaultChaosSpatialAcceleration& Accel, const UWorld* World, const FCollisionShape& InGeom, const FVector& InPosition, const FQuat& InRotation, TArray<FOverlapResult>& OutOverlaps, ECollisionChannel TraceChannel, const FCollisionQueryParams& Params, const FCollisionResponseParams& ResponseParams, const FCollisionObjectQueryParams& ObjectParams)
+{
+	SCOPE_CYCLE_COUNTER(STAT_Collision_SceneQueryTotal);
+	SCOPE_CYCLE_COUNTER(STAT_Collision_GeomOverlapMultiple);
+	CSV_SCOPED_TIMING_STAT(SceneQuery, GeomOverlapMultiple);
+
+	FTransform GeomTransform(InRotation, InPosition);
+	FPhysicsShapeAdapter Adaptor(GeomTransform.GetRotation(), InGeom);
+	return GeomOverlapMultiHelper<EQueryInfo::GatherAll>(World, Adaptor.GetGeometry(), InGeom, Adaptor.GetGeomPose(GeomTransform.GetTranslation()), OutOverlaps, TraceChannel, Params, ResponseParams, ObjectParams, FOverrideAccelContainer{ Accel });
+}
+
+template<>
+bool FGenericPhysicsInterface::GeomOverlapMultiUsingSpatialAcceleration(const Chaos::IDefaultChaosSpatialAcceleration& Accel, const UWorld* World, const FPhysicsGeometry& InGeom, const FVector& InPosition, const FQuat& InRotation, TArray<FOverlapResult>& OutOverlaps, ECollisionChannel TraceChannel, const FCollisionQueryParams& Params, const FCollisionResponseParams& ResponseParams, const FCollisionObjectQueryParams& ObjectParams)
+{
+	return GeomOverlapMultiUsingSpatialAcceleration(Accel, World, FChaosEngineInterface::GetGeometryCollection(InGeom), InPosition, InRotation, OutOverlaps, TraceChannel, Params, ResponseParams, ObjectParams);
+}
\ No newline at end of file
diff --git a/Engine/Source/Runtime/Engine/Private/Collision/SceneQueryLowLevel.cpp b/Engine/Source/Runtime/Engine/Private/Collision/SceneQueryLowLevel.cpp
--- a/Engine/Source/Runtime/Engine/Private/Collision/SceneQueryLowLevel.cpp
+++ b/Engine/Source/Runtime/Engine/Private/Collision/SceneQueryLowLevel.cpp
@@ -66,7 +66,7 @@ namespace
 #endif
 	}
 
-	void SweepSQCaptureHelper(float QueryDurationSeconds, const FChaosSQAccelerator& SQAccelerator, FPhysScene& Scene, const FPhysicsGeometry& QueryGeom, const FTransform& StartTM, const FVector& Dir, float DeltaMag, const FPhysicsHitCallback<FHitSweep>& HitBuffer, EHitFlags OutputFlags, FQueryFlags QueryFlags, const FCollisionFilterData& Filter, const ChaosInterface::FQueryFilterData& QueryFilterData, ICollisionQueryFilterCallbackBase* QueryCallback, const ChaosInterface::FQueryDebugParams& DebugParams)
+	void SweepSQCaptureHelper(float QueryDurationSeconds, const FChaosSQAccelerator& SQAccelerator, const FPhysScene& Scene, const FPhysicsGeometry& QueryGeom, const FTransform& StartTM, const FVector& Dir, float DeltaMag, const FPhysicsHitCallback<FHitSweep>& HitBuffer, EHitFlags OutputFlags, FQueryFlags QueryFlags, const FCollisionFilterData& Filter, const ChaosInterface::FQueryFilterData& QueryFilterData, ICollisionQueryFilterCallbackBase* QueryCallback, const ChaosInterface::FQueryDebugParams& DebugParams)
 	{
 #if !UE_BUILD_SHIPPING
 		float QueryDurationMicro = QueryDurationSeconds * 1000.0 * 1000.0;
@@ -102,7 +102,7 @@ namespace
 #endif
 	}
 
-	void RaycastSQCaptureHelper(float QueryDurationSeconds, const FChaosSQAccelerator& SQAccelerator, FPhysScene& Scene, const FVector& Start, const FVector& Dir, float DeltaMag, const FPhysicsHitCallback<FHitRaycast>& HitBuffer, EHitFlags OutputFlags, FQueryFlags QueryFlags, const FCollisionFilterData& Filter, const ChaosInterface::FQueryFilterData& QueryFilterData, ICollisionQueryFilterCallbackBase* QueryCallback, const ChaosInterface::FQueryDebugParams& DebugParams)
+	void RaycastSQCaptureHelper(float QueryDurationSeconds, const FChaosSQAccelerator& SQAccelerator, const FPhysScene& Scene, const FVector& Start, const FVector& Dir, float DeltaMag, const FPhysicsHitCallback<FHitRaycast>& HitBuffer, EHitFlags OutputFlags, FQueryFlags QueryFlags, const FCollisionFilterData& Filter, const ChaosInterface::FQueryFilterData& QueryFilterData, ICollisionQueryFilterCallbackBase* QueryCallback, const ChaosInterface::FQueryDebugParams& DebugParams)
 	{
 #if !UE_BUILD_SHIPPING
 		float QueryDurationMicro = QueryDurationSeconds * 1000.0 * 1000.0;
@@ -138,7 +138,7 @@ namespace
 #endif
 	}
 
-	void OverlapSQCaptureHelper(float QueryDurationSeconds, const FChaosSQAccelerator& SQAccelerator, FPhysScene& Scene, const FPhysicsGeometry& QueryGeom, const FTransform& GeomPose, const FPhysicsHitCallback<FHitOverlap>& HitBuffer, FQueryFlags QueryFlags, const FCollisionFilterData& Filter, const ChaosInterface::FQueryFilterData& QueryFilterData, ICollisionQueryFilterCallbackBase* QueryCallback, const ChaosInterface::FQueryDebugParams& DebugParams)
+	void OverlapSQCaptureHelper(float QueryDurationSeconds, const FChaosSQAccelerator& SQAccelerator, const FPhysScene& Scene, const FPhysicsGeometry& QueryGeom, const FTransform& GeomPose, const FPhysicsHitCallback<FHitOverlap>& HitBuffer, FQueryFlags QueryFlags, const FCollisionFilterData& Filter, const ChaosInterface::FQueryFilterData& QueryFilterData, ICollisionQueryFilterCallbackBase* QueryCallback, const ChaosInterface::FQueryDebugParams& DebugParams)
 	{
 #if !UE_BUILD_SHIPPING
 		float QueryDurationMicro = QueryDurationSeconds * 1000.0 * 1000.0;
@@ -173,14 +173,43 @@ namespace
 		}
 #endif
 	}
+
+	template<typename TAccelContainer, bool bGTData>
+	struct FAccelerationContainerTraits
+	{
+		static const Chaos::IDefaultChaosSpatialAcceleration* GetSpatialAccelerationFromContainer(const TAccelContainer& Container)
+		{
+			if constexpr (IsPhysScene())
+			{
+				if constexpr (bGTData)
+				{
+					return Container.GetSpacialAcceleration();
+				}
+				else
+				{
+					return Container.GetSolver()->GetInternalAccelerationStructure_Internal();
+				}
+			}
+			else
+			{
+				return &Container;
+			}
+
+			return nullptr;
+		}
+
+		static constexpr bool IsPhysScene()
+		{
+			return std::is_same_v<TAccelContainer, FPhysScene>;
+		}
+	};
 }
 
-template <typename THitRaycast>
-void LowLevelRaycastImp(FPhysScene& Scene, const FVector& Start, const FVector& Dir, float DeltaMag, FPhysicsHitCallback<THitRaycast>& HitBuffer, EHitFlags OutputFlags, FQueryFlags QueryFlags, const FCollisionFilterData& Filter, const ChaosInterface::FQueryFilterData& QueryFilterData, ICollisionQueryFilterCallbackBase* QueryCallback, const ChaosInterface::FQueryDebugParams& DebugParams)
+template <typename TAccelContainer, typename THitRaycast>
+void LowLevelRaycastImp(const TAccelContainer& Container, const FVector& Start, const FVector& Dir, float DeltaMag, FPhysicsHitCallback<THitRaycast>& HitBuffer, EHitFlags OutputFlags, FQueryFlags QueryFlags, const FCollisionFilterData& Filter, const ChaosInterface::FQueryFilterData& QueryFilterData, ICollisionQueryFilterCallbackBase* QueryCallback, const ChaosInterface::FQueryDebugParams& DebugParams)
 {
-	constexpr bool bGTData = std::is_same<THitRaycast, FHitRaycast>::value;
-	const auto SolverAccelerationStructure = bGTData ? Scene.GetSpacialAcceleration() : Scene.GetSolver()->GetInternalAccelerationStructure_Internal();
-	if (SolverAccelerationStructure)
+	constexpr bool bGTData = std::is_same_v<THitRaycast, FHitRaycast>;
+	if (const Chaos::IDefaultChaosSpatialAcceleration* SolverAccelerationStructure = FAccelerationContainerTraits<TAccelContainer, bGTData>::GetSpatialAccelerationFromContainer(Container))
 	{
 		FChaosSQAccelerator SQAccelerator(*SolverAccelerationStructure);
 		double Time = 0.0;
@@ -189,69 +218,86 @@ void LowLevelRaycastImp(FPhysScene& Scene, const FVector& Start, const FVector&
 			SQAccelerator.Raycast(Start, Dir, DeltaMag, HitBuffer, OutputFlags, QueryFilterData, *QueryCallback, DebugParams);
 		}
 
-		if constexpr(bGTData)
+		if constexpr (bGTData && FAccelerationContainerTraits<TAccelContainer, bGTData>::IsPhysScene())
 		{
 			if (!!SerializeSQs && !!EnableRaycastSQCapture)
 			{
-				RaycastSQCaptureHelper(Time, SQAccelerator, Scene, Start, Dir, DeltaMag, HitBuffer, OutputFlags, QueryFlags, Filter, QueryFilterData, QueryCallback, DebugParams);
+				RaycastSQCaptureHelper(Time, SQAccelerator, Container, Start, Dir, DeltaMag, HitBuffer, OutputFlags, QueryFlags, Filter, QueryFilterData, QueryCallback, DebugParams);
 			}
 		}
 	}
 }
 
-void LowLevelRaycast(FPhysScene& Scene, const FVector& Start, const FVector& Dir, float DeltaMag, FPhysicsHitCallback<FHitRaycast>& HitBuffer, EHitFlags OutputFlags, FQueryFlags QueryFlags, const FCollisionFilterData& Filter, const ChaosInterface::FQueryFilterData& QueryFilterData, ICollisionQueryFilterCallbackBase* QueryCallback, const ChaosInterface::FQueryDebugParams& DebugParams)
+void LowLevelRaycast(const FPhysScene& Scene, const FVector& Start, const FVector& Dir, float DeltaMag, FPhysicsHitCallback<FHitRaycast>& HitBuffer, EHitFlags OutputFlags, FQueryFlags QueryFlags, const FCollisionFilterData& Filter, const ChaosInterface::FQueryFilterData& QueryFilterData, ICollisionQueryFilterCallbackBase* QueryCallback, const ChaosInterface::FQueryDebugParams& DebugParams)
 {
 	LowLevelRaycastImp(Scene, Start, Dir, DeltaMag, HitBuffer, OutputFlags, QueryFlags, Filter, QueryFilterData, QueryCallback, DebugParams);
 }
 
-void LowLevelRaycast(FPhysScene& Scene, const FVector& Start, const FVector& Dir, float DeltaMag, FPhysicsHitCallback<ChaosInterface::FPTRaycastHit>& HitBuffer, EHitFlags OutputFlags, FQueryFlags QueryFlags, const FCollisionFilterData& Filter, const ChaosInterface::FQueryFilterData& QueryFilterData, ICollisionQueryFilterCallbackBase* QueryCallback, const ChaosInterface::FQueryDebugParams& DebugParams)
+void LowLevelRaycast(const FPhysScene& Scene, const FVector& Start, const FVector& Dir, float DeltaMag, FPhysicsHitCallback<ChaosInterface::FPTRaycastHit>& HitBuffer, EHitFlags OutputFlags, FQueryFlags QueryFlags, const FCollisionFilterData& Filter, const ChaosInterface::FQueryFilterData& QueryFilterData, ICollisionQueryFilterCallbackBase* QueryCallback, const ChaosInterface::FQueryDebugParams& DebugParams)
 {
 	LowLevelRaycastImp(Scene, Start, Dir, DeltaMag, HitBuffer, OutputFlags, QueryFlags, Filter, QueryFilterData, QueryCallback, DebugParams);
 }
 
-template <typename THitSweep>
-void LowLevelSweepImp(FPhysScene& Scene, const FPhysicsGeometry& QueryGeom, const FTransform& StartTM, const FVector& Dir, float DeltaMag, FPhysicsHitCallback<THitSweep>& HitBuffer, EHitFlags OutputFlags, FQueryFlags QueryFlags, const FCollisionFilterData& Filter, const ChaosInterface::FQueryFilterData& QueryFilterData, ICollisionQueryFilterCallbackBase* QueryCallback, const ChaosInterface::FQueryDebugParams& DebugParams)
+void LowLevelRaycast(const Chaos::IDefaultChaosSpatialAcceleration& AccelerationStructure, const FVector& Start, const FVector& Dir, float DeltaMag, FPhysicsHitCallback<FHitRaycast>& HitBuffer, EHitFlags OutputFlags, FQueryFlags QueryFlags, const FCollisionFilterData& Filter, const ChaosInterface::FQueryFilterData& QueryFilterData, ICollisionQueryFilterCallbackBase* QueryCallback, const ChaosInterface::FQueryDebugParams& DebugParams)
+{
+	LowLevelRaycastImp(AccelerationStructure, Start, Dir, DeltaMag, HitBuffer, OutputFlags, QueryFlags, Filter, QueryFilterData, QueryCallback, DebugParams);
+}
+
+void LowLevelRaycast(const Chaos::IDefaultChaosSpatialAcceleration& AccelerationStructure, const FVector& Start, const FVector& Dir, float DeltaMag, FPhysicsHitCallback<ChaosInterface::FPTRaycastHit>& HitBuffer, EHitFlags OutputFlags, FQueryFlags QueryFlags, const FCollisionFilterData& Filter, const ChaosInterface::FQueryFilterData& QueryFilterData, ICollisionQueryFilterCallbackBase* QueryCallback, const ChaosInterface::FQueryDebugParams& DebugParams)
+{
+	LowLevelRaycastImp(AccelerationStructure, Start, Dir, DeltaMag, HitBuffer, OutputFlags, QueryFlags, Filter, QueryFilterData, QueryCallback, DebugParams);
+}
+
+template <typename TAccelContainer, typename THitSweep>
+void LowLevelSweepImp(const TAccelContainer& Container, const FPhysicsGeometry& QueryGeom, const FTransform& StartTM, const FVector& Dir, float DeltaMag, FPhysicsHitCallback<THitSweep>& HitBuffer, EHitFlags OutputFlags, FQueryFlags QueryFlags, const FCollisionFilterData& Filter, const ChaosInterface::FQueryFilterData& QueryFilterData, ICollisionQueryFilterCallbackBase* QueryCallback, const ChaosInterface::FQueryDebugParams& DebugParams)
 {
 	constexpr bool bGTData = std::is_same<THitSweep, FHitSweep>::value;
-	const auto SolverAccelerationStructure = bGTData ? Scene.GetSpacialAcceleration() : Scene.GetSolver()->GetInternalAccelerationStructure_Internal();
-	if(SolverAccelerationStructure)
+	if (const Chaos::IDefaultChaosSpatialAcceleration* SolverAccelerationStructure = FAccelerationContainerTraits<TAccelContainer, bGTData>::GetSpatialAccelerationFromContainer(Container))
 	{
 		FChaosSQAccelerator SQAccelerator(*SolverAccelerationStructure);
 		{
-			//ISQAccelerator* SQAccelerator = Scene.GetSQAccelerator();
 			double Time = 0.0;
 			{
 				FScopedDurationTimer Timer(Time);
 				SQAccelerator.Sweep(QueryGeom, StartTM, Dir, DeltaMag, HitBuffer, OutputFlags, QueryFilterData, *QueryCallback, DebugParams);
 			}
 
-			if constexpr(bGTData)
+			if constexpr (bGTData && FAccelerationContainerTraits<TAccelContainer, bGTData>::IsPhysScene())
 			{
 				if(!!SerializeSQs && !!EnableSweepSQCapture)
 				{
-					SweepSQCaptureHelper(Time, SQAccelerator, Scene, QueryGeom, StartTM, Dir, DeltaMag, HitBuffer, OutputFlags, QueryFlags, Filter, QueryFilterData, QueryCallback, DebugParams);
+					SweepSQCaptureHelper(Time, SQAccelerator, Container, QueryGeom, StartTM, Dir, DeltaMag, HitBuffer, OutputFlags, QueryFlags, Filter, QueryFilterData, QueryCallback, DebugParams);
 				}
 			}
 		}
 	}
 }
 
-void LowLevelSweep(FPhysScene& Scene, const FPhysicsGeometry& QueryGeom, const FTransform& StartTM, const FVector& Dir, float DeltaMag, FPhysicsHitCallback<FHitSweep>& HitBuffer, EHitFlags OutputFlags, FQueryFlags QueryFlags, const FCollisionFilterData& Filter, const ChaosInterface::FQueryFilterData& QueryFilterData, ICollisionQueryFilterCallbackBase* QueryCallback, const ChaosInterface::FQueryDebugParams& DebugParams)
+void LowLevelSweep(const FPhysScene& Scene, const FPhysicsGeometry& QueryGeom, const FTransform& StartTM, const FVector& Dir, float DeltaMag, FPhysicsHitCallback<FHitSweep>& HitBuffer, EHitFlags OutputFlags, FQueryFlags QueryFlags, const FCollisionFilterData& Filter, const ChaosInterface::FQueryFilterData& QueryFilterData, ICollisionQueryFilterCallbackBase* QueryCallback, const ChaosInterface::FQueryDebugParams& DebugParams)
+{
+	LowLevelSweepImp(Scene, QueryGeom, StartTM, Dir, DeltaMag, HitBuffer, OutputFlags, QueryFlags, Filter, QueryFilterData, QueryCallback, DebugParams);
+}
+
+void LowLevelSweep(const FPhysScene& Scene, const FPhysicsGeometry& QueryGeom, const FTransform& StartTM, const FVector& Dir, float DeltaMag, FPhysicsHitCallback<ChaosInterface::FPTSweepHit>& HitBuffer, EHitFlags OutputFlags, FQueryFlags QueryFlags, const FCollisionFilterData& Filter, const ChaosInterface::FQueryFilterData& QueryFilterData, ICollisionQueryFilterCallbackBase* QueryCallback, const ChaosInterface::FQueryDebugParams& DebugParams)
 {
-	return LowLevelSweepImp(Scene, QueryGeom, StartTM, Dir, DeltaMag, HitBuffer, OutputFlags, QueryFlags, Filter, QueryFilterData, QueryCallback, DebugParams);
+	LowLevelSweepImp(Scene, QueryGeom, StartTM, Dir, DeltaMag, HitBuffer, OutputFlags, QueryFlags, Filter, QueryFilterData, QueryCallback, DebugParams);
 }
 
-void LowLevelSweep(FPhysScene& Scene, const FPhysicsGeometry& QueryGeom, const FTransform& StartTM, const FVector& Dir, float DeltaMag, FPhysicsHitCallback<ChaosInterface::FPTSweepHit>& HitBuffer, EHitFlags OutputFlags, FQueryFlags QueryFlags, const FCollisionFilterData& Filter, const ChaosInterface::FQueryFilterData& QueryFilterData, ICollisionQueryFilterCallbackBase* QueryCallback, const ChaosInterface::FQueryDebugParams& DebugParams)
+void LowLevelSweep(const Chaos::IDefaultChaosSpatialAcceleration& AccelerationStructure, const FPhysicsGeometry& Geom, const FTransform& StartTM, const FVector& Dir, float DeltaMag, FPhysicsHitCallback<FHitSweep>& HitBuffer, EHitFlags OutputFlags, FQueryFlags QueryFlags, const FCollisionFilterData& Filter, const ChaosInterface::FQueryFilterData& QueryFilterData, ICollisionQueryFilterCallbackBase* QueryCallback, const ChaosInterface::FQueryDebugParams& DebugParams)
 {
-	return LowLevelSweepImp(Scene, QueryGeom, StartTM, Dir, DeltaMag, HitBuffer, OutputFlags, QueryFlags, Filter, QueryFilterData, QueryCallback, DebugParams);
+	LowLevelSweepImp(AccelerationStructure, Geom, StartTM, Dir, DeltaMag, HitBuffer, OutputFlags, QueryFlags, Filter, QueryFilterData, QueryCallback, DebugParams);
 }
 
-template <typename THitOverlap>
-void LowLevelOverlapImp(FPhysScene& Scene, const FPhysicsGeometry& QueryGeom, const FTransform& GeomPose, FPhysicsHitCallback<THitOverlap>& HitBuffer, FQueryFlags QueryFlags, const FCollisionFilterData& Filter, const ChaosInterface::FQueryFilterData& QueryFilterData, ICollisionQueryFilterCallbackBase* QueryCallback, const ChaosInterface::FQueryDebugParams& DebugParams)
+void LowLevelSweep(const Chaos::IDefaultChaosSpatialAcceleration& AccelerationStructure, const FPhysicsGeometry& Geom, const FTransform& StartTM, const FVector& Dir, float DeltaMag, FPhysicsHitCallback<ChaosInterface::FPTSweepHit>& HitBuffer, EHitFlags OutputFlags, FQueryFlags QueryFlags, const FCollisionFilterData& Filter, const ChaosInterface::FQueryFilterData& QueryFilterData, ICollisionQueryFilterCallbackBase* QueryCallback, const ChaosInterface::FQueryDebugParams& DebugParams)
+{
+	LowLevelSweepImp(AccelerationStructure, Geom, StartTM, Dir, DeltaMag, HitBuffer, OutputFlags, QueryFlags, Filter, QueryFilterData, QueryCallback, DebugParams);
+}
+
+template <typename TAccelContainer, typename THitOverlap>
+void LowLevelOverlapImp(const TAccelContainer& Container, const FPhysicsGeometry& QueryGeom, const FTransform& GeomPose, FPhysicsHitCallback<THitOverlap>& HitBuffer, FQueryFlags QueryFlags, const FCollisionFilterData& Filter, const ChaosInterface::FQueryFilterData& QueryFilterData, ICollisionQueryFilterCallbackBase* QueryCallback, const ChaosInterface::FQueryDebugParams& DebugParams)
 {
 	constexpr bool bGTData = std::is_same<THitOverlap, FHitOverlap>::value;
-	const auto SolverAccelerationStructure = bGTData ? Scene.GetSpacialAcceleration() : Scene.GetSolver()->GetInternalAccelerationStructure_Internal();
-	if(SolverAccelerationStructure)
+	if (const Chaos::IDefaultChaosSpatialAcceleration* SolverAccelerationStructure = FAccelerationContainerTraits<TAccelContainer, bGTData>::GetSpatialAccelerationFromContainer(Container))
 	{
 		FChaosSQAccelerator SQAccelerator(*SolverAccelerationStructure);
 		double Time = 0.0;
@@ -260,22 +306,32 @@ void LowLevelOverlapImp(FPhysScene& Scene, const FPhysicsGeometry& QueryGeom, co
 			SQAccelerator.Overlap(QueryGeom, GeomPose, HitBuffer, QueryFilterData, *QueryCallback, DebugParams);
 		}
 
-		if constexpr(bGTData)
+		if constexpr(bGTData && FAccelerationContainerTraits<TAccelContainer, bGTData>::IsPhysScene())
 		{
 			if(!!SerializeSQs && !!EnableOverlapSQCapture)
 			{
-				OverlapSQCaptureHelper(Time, SQAccelerator, Scene, QueryGeom, GeomPose, HitBuffer, QueryFlags, Filter, QueryFilterData, QueryCallback, DebugParams);
+				OverlapSQCaptureHelper(Time, SQAccelerator, Container, QueryGeom, GeomPose, HitBuffer, QueryFlags, Filter, QueryFilterData, QueryCallback, DebugParams);
 			}
 		}
 	}
 }
 
-void LowLevelOverlap(FPhysScene& Scene, const FPhysicsGeometry& QueryGeom, const FTransform& GeomPose, FPhysicsHitCallback<FHitOverlap>& HitBuffer, FQueryFlags QueryFlags, const FCollisionFilterData& Filter, const ChaosInterface::FQueryFilterData& QueryFilterData, ICollisionQueryFilterCallbackBase* QueryCallback, const ChaosInterface::FQueryDebugParams& DebugParams)
+void LowLevelOverlap(const FPhysScene& Scene, const FPhysicsGeometry& QueryGeom, const FTransform& GeomPose, FPhysicsHitCallback<FHitOverlap>& HitBuffer, FQueryFlags QueryFlags, const FCollisionFilterData& Filter, const ChaosInterface::FQueryFilterData& QueryFilterData, ICollisionQueryFilterCallbackBase* QueryCallback, const ChaosInterface::FQueryDebugParams& DebugParams)
+{
+	LowLevelOverlapImp(Scene, QueryGeom, GeomPose, HitBuffer, QueryFlags, Filter, QueryFilterData, QueryCallback, DebugParams);
+}
+
+void LowLevelOverlap(const FPhysScene& Scene, const FPhysicsGeometry& QueryGeom, const FTransform& GeomPose, FPhysicsHitCallback<ChaosInterface::FPTOverlapHit>& HitBuffer, FQueryFlags QueryFlags, const FCollisionFilterData& Filter, const ChaosInterface::FQueryFilterData& QueryFilterData, ICollisionQueryFilterCallbackBase* QueryCallback, const ChaosInterface::FQueryDebugParams& DebugParams)
 {
-	return LowLevelOverlapImp(Scene, QueryGeom, GeomPose, HitBuffer, QueryFlags, Filter, QueryFilterData, QueryCallback, DebugParams);
+	LowLevelOverlapImp(Scene, QueryGeom, GeomPose, HitBuffer, QueryFlags, Filter, QueryFilterData, QueryCallback, DebugParams);
 }
 
-void LowLevelOverlap(FPhysScene& Scene, const FPhysicsGeometry& QueryGeom, const FTransform& GeomPose, FPhysicsHitCallback<ChaosInterface::FPTOverlapHit>& HitBuffer, FQueryFlags QueryFlags, const FCollisionFilterData& Filter, const ChaosInterface::FQueryFilterData& QueryFilterData, ICollisionQueryFilterCallbackBase* QueryCallback, const ChaosInterface::FQueryDebugParams& DebugParams)
+void LowLevelOverlap(const Chaos::IDefaultChaosSpatialAcceleration& AccelerationStructure, const FPhysicsGeometry& Geom, const FTransform& GeomPose, FPhysicsHitCallback<FHitOverlap>& HitBuffer, FQueryFlags QueryFlags, const FCollisionFilterData& Filter, const ChaosInterface::FQueryFilterData& QueryFilterData, ICollisionQueryFilterCallbackBase* QueryCallback, const ChaosInterface::FQueryDebugParams& DebugParams)
 {
-	return LowLevelOverlapImp(Scene, QueryGeom, GeomPose, HitBuffer, QueryFlags, Filter, QueryFilterData, QueryCallback, DebugParams);
+	LowLevelOverlapImp(AccelerationStructure, Geom, GeomPose, HitBuffer, QueryFlags, Filter, QueryFilterData, QueryCallback, DebugParams);
 }
+
+void LowLevelOverlap(const Chaos::IDefaultChaosSpatialAcceleration& AccelerationStructure, const FPhysicsGeometry& Geom, const FTransform& GeomPose, FPhysicsHitCallback<ChaosInterface::FPTOverlapHit>& HitBuffer, FQueryFlags QueryFlags, const FCollisionFilterData& Filter, const ChaosInterface::FQueryFilterData& QueryFilterData, ICollisionQueryFilterCallbackBase* QueryCallback, const ChaosInterface::FQueryDebugParams& DebugParams)
+{
+	LowLevelOverlapImp(AccelerationStructure, Geom, GeomPose, HitBuffer, QueryFlags, Filter, QueryFilterData, QueryCallback, DebugParams);
+}
\ No newline at end of file
diff --git a/Engine/Source/Runtime/Engine/Private/PhysicsEngine/ClusterUnionComponent.cpp b/Engine/Source/Runtime/Engine/Private/PhysicsEngine/ClusterUnionComponent.cpp
--- a/Engine/Source/Runtime/Engine/Private/PhysicsEngine/ClusterUnionComponent.cpp
+++ b/Engine/Source/Runtime/Engine/Private/PhysicsEngine/ClusterUnionComponent.cpp
@@ -1,10 +1,14 @@
 // Copyright Epic Games, Inc. All Rights Reserved.
 #include "PhysicsEngine/ClusterUnionComponent.h"
 
+#include "Chaos/ImplicitObject.h"
+#include "Chaos/Serializable.h"
+#include "Chaos/ShapeInstance.h"
 #include "Engine/World.h"
 #include "Net/Core/PushModel/PushModel.h"
 #include "Net/UnrealNetwork.h"
 #include "Physics/Experimental/PhysScene_Chaos.h"
+#include "Physics/GenericPhysicsInterface.h"
 #include "PhysicsEngine/ClusterUnionReplicatedProxyComponent.h"
 #include "PhysicsEngine/PhysicsObjectExternalInterface.h"
 #include "PhysicsProxy/ClusterUnionPhysicsProxy.h"
@@ -14,10 +18,31 @@
 
 DEFINE_LOG_CATEGORY(LogClusterUnion);
 
-// todo : remove when proven to be the right fix 
-bool bClusterUnionClientTransformFix = true;
-FAutoConsoleVariableRef CVarClusterUnionClientTransformFix(TEXT("clusterunion.clienttransformfix"), bClusterUnionClientTransformFix, TEXT("fix transform not always being properly set for cluster union"));
+namespace
+{
+	bool bUseClusterUnionAccelerationStructure = false;
+	FAutoConsoleVariableRef CVarUseClusterUnionAccelerationStructure(TEXT("ClusterUnion.UseAccelerationStructure"), bUseClusterUnionAccelerationStructure, TEXT("Whether component level sweeps and overlaps against cluster unions should use an acceleration structure instead."));
+
+	// TODO: Should this be exposed in Chaos instead?
+	using FAccelerationStructure = Chaos::TAABBTree<Chaos::FAccelerationStructureHandle, Chaos::TAABBTreeLeafArray<Chaos::FAccelerationStructureHandle>>;
+
+	TUniquePtr<Chaos::ISpatialAcceleration<Chaos::FAccelerationStructureHandle, Chaos::FReal, 3>> CreateEmptyAccelerationStructure()
+	{
+		const Chaos::TConstParticleView<Chaos::FSpatialAccelerationCache> Empty;
+		FAccelerationStructure* Structure = new FAccelerationStructure(
+			Empty,
+			FAccelerationStructure::DefaultMaxChildrenInLeaf,
+			FAccelerationStructure::DefaultMaxTreeDepth,
+			FAccelerationStructure::DefaultMaxPayloadBounds,
+			FAccelerationStructure::DefaultMaxNumToProcess,
+			true,
+			false
+		);
 
+		check(Structure != nullptr);
+		return TUniquePtr<Chaos::ISpatialAcceleration<Chaos::FAccelerationStructureHandle, Chaos::FReal, 3>>(Structure);
+	}
+}
 
 UClusterUnionComponent::UClusterUnionComponent(const FObjectInitializer& ObjectInitializer)
 	: UPrimitiveComponent(ObjectInitializer)
@@ -114,12 +139,24 @@ void UClusterUnionComponent::AddComponentToCluster(UPrimitiveComponent* InCompon
 	}
 
 	PendingData.BoneIds.Empty(Objects.Num());
+
 	for (Chaos::FPhysicsObjectHandle Object : Objects)
 	{
 		PendingData.BoneIds.Add(Chaos::FPhysicsObjectInterface::GetId(Object));
+
+		if (AccelerationStructure && Object)
+		{
+			Chaos::FAccelerationStructureHandle Handle = Interface->CreateAccelerationStructureHandle(Object);
+			const FBox HandleBounds = Interface->GetWorldBounds({ &Object, 1 });
+			AccelerationStructure->UpdateElement(Handle, Chaos::TAABB<Chaos::FReal, 3>{HandleBounds.Min, HandleBounds.Max}, HandleBounds.IsValid != 0);
+		}
 	}
 
 	PendingComponentSync.Add(InComponent, PendingData);
+
+	// Need to listen to changes in the component's physics state. If it gets destroyed it should be removed from the cluster union as well.
+	InComponent->OnComponentPhysicsStateChanged.AddDynamic(this, &UClusterUnionComponent::HandleComponentPhysicsStateChangePostAddIntoClusterUnion);
+
 	PhysicsProxy->AddPhysicsObjects_External(Objects);
 	ForceRebuildGTParticleGeometry();
 }
@@ -139,6 +176,7 @@ void UClusterUnionComponent::RemoveComponentFromCluster(UPrimitiveComponent* InC
 	}
 
 	TSet<Chaos::FPhysicsObjectHandle> PhysicsObjectsToRemove;
+	FLockedReadPhysicsObjectExternalInterface Interface = FPhysicsObjectExternalInterface::LockRead(GetWorld()->GetPhysicsScene());
 
 	if (FClusteredComponentData* ComponentData = ComponentToPhysicsObjects.Find(InComponent))
 	{
@@ -164,6 +202,18 @@ void UClusterUnionComponent::RemoveComponentFromCluster(UPrimitiveComponent* InC
 		{
 			PhysicsObjectsToRemove = ComponentData->PhysicsObjects;
 		}
+
+		if (AccelerationStructure)
+		{
+			for (Chaos::FPhysicsObjectHandle PhysicsObject : PhysicsObjectsToRemove)
+			{
+				if (PhysicsObject)
+				{
+					Chaos::FAccelerationStructureHandle Handle = Interface->CreateAccelerationStructureHandle(PhysicsObject);
+					AccelerationStructure->RemoveElement(Handle);
+				}
+			}
+		}
 	}
 
 	PhysicsProxy->RemovePhysicsObjects_External(PhysicsObjectsToRemove);
@@ -290,6 +340,12 @@ void UClusterUnionComponent::OnCreatePhysicsState()
 	{
 		Scene->AddObject(this, PhysicsProxy);
 	}
+	
+	if (bUseClusterUnionAccelerationStructure)
+	{
+		AccelerationStructure = CreateEmptyAccelerationStructure();
+		check(AccelerationStructure != nullptr);
+	}
 
 	// It's just logically easier to be consistent on the client to go through the replication route.
 	if (IsAuthority())
@@ -306,13 +362,10 @@ void UClusterUnionComponent::OnCreatePhysicsState()
 	}
 	else
 	{
-		if (bClusterUnionClientTransformFix)
-		{
-			// we should not rely on the OnTransformUpdated callback to set the initial position of the external particle
-			// if it is not set and the callback does not get called then the cluster union component will be moved to the origin of the world
-			const FTransform Transform = GetComponentTransform();
-			PhysicsProxy->SetXR_External(Transform.GetLocation(), Transform.GetRotation());
-		}
+		// we should not rely on the OnTransformUpdated callback to set the initial position of the external particle
+		// if it is not set and the callback does not get called then the cluster union component will be moved to the origin of the world
+		const FTransform Transform = GetComponentTransform();
+		PhysicsProxy->SetXR_External(Transform.GetLocation(), Transform.GetRotation());
 	}
 }
 
@@ -347,6 +400,7 @@ void UClusterUnionComponent::OnDestroyPhysicsState()
 	}
 
 	PhysicsProxy = nullptr;
+	AccelerationStructure.Reset();
 }
 
 void UClusterUnionComponent::OnUpdateTransform(EUpdateTransformFlags UpdateTransformFlags, ETeleportType Teleport)
@@ -420,7 +474,6 @@ Chaos::FPhysicsObjectId UClusterUnionComponent::GetIdFromGTParticle(Chaos::FGeom
 
 void UClusterUnionComponent::HandleComponentPhysicsStateChange(UPrimitiveComponent* ChangedComponent, EComponentPhysicsStateChange StateChange)
 {
-	// TODO: Maybe we should handle the destroyed state change too?
 	if (!ChangedComponent || StateChange != EComponentPhysicsStateChange::Created)
 	{
 		return;
@@ -434,6 +487,17 @@ void UClusterUnionComponent::HandleComponentPhysicsStateChange(UPrimitiveCompone
 	}
 }
 
+void UClusterUnionComponent::HandleComponentPhysicsStateChangePostAddIntoClusterUnion(UPrimitiveComponent* ChangedComponent, EComponentPhysicsStateChange StateChange)
+{
+	if (!ChangedComponent || StateChange != EComponentPhysicsStateChange::Destroyed)
+	{
+		return;
+	}
+
+	ChangedComponent->OnComponentPhysicsStateChanged.RemoveDynamic(this, &UClusterUnionComponent::HandleComponentPhysicsStateChangePostAddIntoClusterUnion);
+	RemoveComponentFromCluster(ChangedComponent);
+}
+
 void UClusterUnionComponent::SyncVelocitiesFromPhysics(const FVector& LinearVelocity, const FVector& AngularVelocity)
 {
 	if (!IsAuthority())
@@ -563,6 +627,7 @@ void UClusterUnionComponent::HandleAddOrModifiedClusteredComponent(UPrimitiveCom
 
 	}
 
+	const TSet<int32> OldBoneIds{ ComponentData.BoneIds };
 	PerBoneChildToParent.GetKeys(ComponentData.BoneIds);
 	OnComponentAddedEvent.Broadcast(ChangedComponent, ComponentData.BoneIds, bIsNew);
 
@@ -582,16 +647,43 @@ void UClusterUnionComponent::HandleAddOrModifiedClusteredComponent(UPrimitiveCom
 		}
 	}
 
+	TArray<Chaos::FPhysicsObjectHandle> AllPhysicsObjects = ChangedComponent->GetAllPhysicsObjects();
+	FLockedReadPhysicsObjectExternalInterface Interface = FPhysicsObjectExternalInterface::LockRead(AllPhysicsObjects);
+
 	// One more loop to ensure that our sets of physics objects are valid and up to date.
 	// This needs to happen on both the client and the server.
 	for (const TPair<int32, FTransform>& Kvp : PerBoneChildToParent)
 	{
 		Chaos::FPhysicsObjectHandle PhysicsObject = ChangedComponent->GetPhysicsObjectById(Kvp.Key);
+
+		if (AccelerationStructure)
+		{
+			if (PhysicsObject)
+			{
+				Chaos::FAccelerationStructureHandle Handle = Interface->CreateAccelerationStructureHandle(PhysicsObject);
+				const FBox HandleBounds = Interface->GetWorldBounds({ &PhysicsObject, 1 });
+				AccelerationStructure->UpdateElement(Handle, Chaos::TAABB<Chaos::FReal, 3>{HandleBounds.Min, HandleBounds.Max}, HandleBounds.IsValid != 0);
+			}
+		}
+
 		ComponentData.PhysicsObjects.Add(PhysicsObject);
 	}
 
-	ComponentData.AllPhysicsObjects.Empty();
-	ComponentData.AllPhysicsObjects.Append(ChangedComponent->GetAllPhysicsObjects());
+	// In the case where we need to keep the acceleration structure up to date, we need to make sure old bone ids are properly
+	// removed from the acceleration structure.
+	if (AccelerationStructure)
+	{
+		for (int32 BoneId : OldBoneIds.Difference(ComponentData.BoneIds))
+		{
+			if (Chaos::FPhysicsObjectHandle PhysicsObject = ChangedComponent->GetPhysicsObjectById(BoneId))
+			{
+				Chaos::FAccelerationStructureHandle Handle = Interface->CreateAccelerationStructureHandle(PhysicsObject);
+				AccelerationStructure->RemoveElement(Handle);
+			}
+		}
+	}
+
+	ComponentData.AllPhysicsObjects = MoveTemp(AllPhysicsObjects);
 }
 
 void UClusterUnionComponent::HandleRemovedClusteredComponent(UPrimitiveComponent* ChangedComponent, bool bDestroyReplicatedProxy)
@@ -614,14 +706,18 @@ void UClusterUnionComponent::HandleRemovedClusteredComponent(UPrimitiveComponent
 		if (ChangedComponent->HasValidPhysicsState())
 		{
 			FLockedReadPhysicsObjectExternalInterface Interface = FPhysicsObjectExternalInterface::LockRead(ComponentData->AllPhysicsObjects);
-			for (Chaos::FGeometryParticle* Particle : Interface->GetAllParticles(ComponentData->AllPhysicsObjects))
+			for (Chaos::FPhysicsObjectHandle PhysicsObject : ComponentData->AllPhysicsObjects)
 			{
-				if (!Particle)
+				if (Chaos::FGeometryParticle* Particle = Interface->GetParticle(PhysicsObject))
 				{
-					continue;
+					UniqueIdxToComponent.Remove(Particle->UniqueIdx().Idx);
 				}
 
-				UniqueIdxToComponent.Remove(Particle->UniqueIdx().Idx);
+				if (AccelerationStructure && PhysicsObject)
+				{
+					Chaos::FAccelerationStructureHandle Handle = Interface->CreateAccelerationStructureHandle(PhysicsObject);
+					AccelerationStructure->RemoveElement(Handle);
+				}
 			}
 		}
 
@@ -709,8 +805,14 @@ void UClusterUnionComponent::SetSimulatePhysics(bool bSimulate)
 	PhysicsProxy->SetObjectState_External(bSimulate ? Chaos::EObjectStateType::Dynamic : Chaos::EObjectStateType::Kinematic);
 }
 
+DECLARE_CYCLE_STAT(TEXT("UClusterUnionComponent::LineTraceComponentMulti"), STAT_ClusterUnionComponentLineTraceComponentMulti, STATGROUP_Chaos);
 bool UClusterUnionComponent::LineTraceComponent(TArray<FHitResult>& OutHit, const FVector Start, const FVector End, ECollisionChannel TraceChannel, const struct FCollisionQueryParams& Params, const struct FCollisionResponseParams& ResponseParams, const struct FCollisionObjectQueryParams& ObjectParams)
 {
+	SCOPE_CYCLE_COUNTER(STAT_ClusterUnionComponentLineTraceComponentMulti);
+	if (AccelerationStructure)
+	{
+		return FGenericPhysicsInterface::RaycastMultiUsingSpatialAcceleration(*AccelerationStructure, GetWorld(), OutHit, Start, End, TraceChannel, Params, ResponseParams, ObjectParams);
+	}
 	OutHit.Reset();
 
 	VisitAllCurrentChildComponentsForCollision(TraceChannel, Params, ResponseParams, ObjectParams,
@@ -729,8 +831,15 @@ bool UClusterUnionComponent::LineTraceComponent(TArray<FHitResult>& OutHit, cons
 	return !OutHit.IsEmpty();
 }
 
+DECLARE_CYCLE_STAT(TEXT("UClusterUnionComponent::SweepComponentMulti"), STAT_ClusterUnionComponentSweepComponentMulti, STATGROUP_Chaos);
 bool UClusterUnionComponent::SweepComponent(TArray<FHitResult>& OutHit, const FVector Start, const FVector End, const FQuat& ShapeWorldRotation, const FPhysicsGeometry& Geometry, ECollisionChannel TraceChannel, const struct FCollisionQueryParams& Params, const struct FCollisionResponseParams& ResponseParams, const struct FCollisionObjectQueryParams& ObjectParams)
 {
+	SCOPE_CYCLE_COUNTER(STAT_ClusterUnionComponentSweepComponentMulti);
+	if (AccelerationStructure)
+	{
+		return FGenericPhysicsInterface::GeomSweepMultiUsingSpatialAcceleration(*AccelerationStructure, GetWorld(), Geometry, ShapeWorldRotation, OutHit, Start, End, TraceChannel, Params, ResponseParams, ObjectParams);
+	}
+
 	OutHit.Reset();
 
 	VisitAllCurrentChildComponentsForCollision(TraceChannel, Params, ResponseParams, ObjectParams,
@@ -749,8 +858,15 @@ bool UClusterUnionComponent::SweepComponent(TArray<FHitResult>& OutHit, const FV
 	return !OutHit.IsEmpty();
 }
 
+DECLARE_CYCLE_STAT(TEXT("UClusterUnionComponent::LineTraceComponentSingle"), STAT_ClusterUnionComponentLineTraceComponentSingle, STATGROUP_Chaos);
 bool UClusterUnionComponent::LineTraceComponent(FHitResult& OutHit, const FVector Start, const FVector End, ECollisionChannel TraceChannel, const struct FCollisionQueryParams& Params, const struct FCollisionResponseParams& ResponseParams, const struct FCollisionObjectQueryParams& ObjectParams)
 {
+	SCOPE_CYCLE_COUNTER(STAT_ClusterUnionComponentLineTraceComponentSingle);
+	if (AccelerationStructure)
+	{
+		return FGenericPhysicsInterface::RaycastSingleUsingSpatialAcceleration(*AccelerationStructure, GetWorld(), OutHit, Start, End, TraceChannel, Params, ResponseParams, ObjectParams);
+	}
+
 	bool bHasHit = false;
 	OutHit.Distance = TNumericLimits<float>::Max();
 
@@ -771,8 +887,15 @@ bool UClusterUnionComponent::LineTraceComponent(FHitResult& OutHit, const FVecto
 	return bHasHit;
 }
 
+DECLARE_CYCLE_STAT(TEXT("UClusterUnionComponent::SweepComponentSingle"), STAT_ClusterUnionComponentSweepComponentSingle, STATGROUP_Chaos);
 bool UClusterUnionComponent::SweepComponent(FHitResult& OutHit, const FVector Start, const FVector End, const FQuat& ShapeWorldRotation, const FPhysicsGeometry& Geometry, ECollisionChannel TraceChannel, const struct FCollisionQueryParams& Params, const struct FCollisionResponseParams& ResponseParams, const struct FCollisionObjectQueryParams& ObjectParams)
 {
+	SCOPE_CYCLE_COUNTER(STAT_ClusterUnionComponentSweepComponentSingle);
+	if (AccelerationStructure)
+	{
+		return FGenericPhysicsInterface::GeomSweepSingleUsingSpatialAcceleration(*AccelerationStructure, GetWorld(), Geometry, ShapeWorldRotation, OutHit, Start, End, TraceChannel, Params, ResponseParams, ObjectParams);
+	}
+
 	bool bHasHit = false;
 	OutHit.Distance = TNumericLimits<float>::Max();
 
@@ -793,8 +916,15 @@ bool UClusterUnionComponent::SweepComponent(FHitResult& OutHit, const FVector St
 	return bHasHit;
 }
 
+DECLARE_CYCLE_STAT(TEXT("UClusterUnionComponent::OverlapComponentWithResult"), STAT_ClusterUnionComponentOverlapComponentWithResult, STATGROUP_Chaos);
 bool UClusterUnionComponent::OverlapComponentWithResult(const FVector& Pos, const FQuat& Rot, const FPhysicsGeometry& Geometry, ECollisionChannel TraceChannel, const struct FCollisionQueryParams& Params, const struct FCollisionResponseParams& ResponseParams, const struct FCollisionObjectQueryParams& ObjectParams, TArray<FOverlapResult>& OutOverlap) const
 {
+	SCOPE_CYCLE_COUNTER(STAT_ClusterUnionComponentOverlapComponentWithResult);
+	if (AccelerationStructure)
+	{
+		return FGenericPhysicsInterface::GeomOverlapMultiUsingSpatialAcceleration(*AccelerationStructure, GetWorld(), Geometry, Pos, Rot, OutOverlap, TraceChannel, Params, ResponseParams, ObjectParams);
+	}
+
 	bool bHasOverlap = false;
 
 	FVector QueryHalfExtent = Geometry.BoundingBox().Extents();
@@ -821,14 +951,43 @@ bool UClusterUnionComponent::OverlapComponentWithResult(const FVector& Pos, cons
 	return bHasOverlap;
 }
 
+DECLARE_CYCLE_STAT(TEXT("UClusterUnionComponent::ComponentOverlapComponentWithResultImpl"), STAT_ClusterUnionComponentComponentOverlapComponentWithResultImpl, STATGROUP_Chaos);
 bool UClusterUnionComponent::ComponentOverlapComponentWithResultImpl(const class UPrimitiveComponent* const PrimComp, const FVector& Pos, const FQuat& Rot, const FCollisionQueryParams& Params, TArray<FOverlapResult>& OutOverlap) const
 {
+	SCOPE_CYCLE_COUNTER(STAT_ClusterUnionComponentComponentOverlapComponentWithResultImpl);
 	if(!PrimComp)
 	{
 		return false;
 	}
 
 	bool bHasOverlap = false;
+	if (AccelerationStructure)
+	{
+		TArray<Chaos::FPhysicsObjectHandle> InObjects = PrimComp->GetAllPhysicsObjects();
+		FLockedReadPhysicsObjectExternalInterface Interface = FPhysicsObjectExternalInterface::LockRead(InObjects);
+		InObjects = InObjects.FilterByPredicate(
+			[&Interface](Chaos::FPhysicsObjectHandle Handle)
+			{
+				return !Interface->AreAllDisabled({ &Handle, 1 });
+			}
+		);
+
+		Interface->VisitEveryShape(InObjects,
+			[this, &bHasOverlap, &Pos, &Rot, &OutOverlap, &Params](const Chaos::FConstPhysicsObjectHandle Handle, Chaos::FShapeInstanceProxy* Shape)
+			{
+				if (Shape)
+				{
+					if (Chaos::TSerializablePtr<Chaos::FImplicitObject> Geometry = Shape->GetGeometry())
+					{
+						bHasOverlap |= FGenericPhysicsInterface::GeomOverlapMultiUsingSpatialAcceleration(*AccelerationStructure, GetWorld(), *Geometry, Pos, Rot, OutOverlap, DefaultCollisionChannel, Params, FCollisionResponseParams::DefaultResponseParam, FCollisionObjectQueryParams::DefaultObjectQueryParam);
+					}
+				}
+
+				return true;
+			});
+
+		return bHasOverlap;
+	}
 
 	FBoxSphereBounds QueryBounds = PrimComp->Bounds.TransformBy(FTransform(Rot, Pos));
 
diff --git a/Engine/Source/Runtime/Engine/Public/Physics/Experimental/ChaosInterfaceWrapper.h b/Engine/Source/Runtime/Engine/Public/Physics/Experimental/ChaosInterfaceWrapper.h
--- a/Engine/Source/Runtime/Engine/Public/Physics/Experimental/ChaosInterfaceWrapper.h
+++ b/Engine/Source/Runtime/Engine/Public/Physics/Experimental/ChaosInterfaceWrapper.h
@@ -3,6 +3,7 @@
 
 #include "ChaosInterfaceWrapperCore.h"
 #include "Chaos/ParticleHandleFwd.h"
+#include "Chaos/SpatialAccelerationFwd.h"
 #include "PhysicsInterfaceUtilsCore.h"
 #include "CollisionQueryFilterCallbackCore.h"
 #include "Chaos/PBDRigidsEvolutionFwd.h"
@@ -41,11 +42,17 @@ UPhysicalMaterial* GetUserData(const Chaos::FChaosPhysicsMaterial& Material);
 
 }
 
-void LowLevelRaycast(FPhysScene& Scene, const FVector& Start, const FVector& Dir, float DeltaMag, FPhysicsHitCallback<FHitRaycast>& HitBuffer, EHitFlags OutputFlags, FQueryFlags QueryFlags, const FCollisionFilterData& Filter, const ChaosInterface::FQueryFilterData& QueryFilterData, ICollisionQueryFilterCallbackBase* QueryCallback, const ChaosInterface::FQueryDebugParams& DebugParams = {});
-void LowLevelRaycast(FPhysScene& Scene, const FVector& Start, const FVector& Dir, float DeltaMag, FPhysicsHitCallback<ChaosInterface::FPTRaycastHit>& HitBuffer, EHitFlags OutputFlags, FQueryFlags QueryFlags, const FCollisionFilterData& Filter, const ChaosInterface::FQueryFilterData& QueryFilterData, ICollisionQueryFilterCallbackBase* QueryCallback, const ChaosInterface::FQueryDebugParams& DebugParams = {});
+void LowLevelRaycast(const FPhysScene& Scene, const FVector& Start, const FVector& Dir, float DeltaMag, FPhysicsHitCallback<FHitRaycast>& HitBuffer, EHitFlags OutputFlags, FQueryFlags QueryFlags, const FCollisionFilterData& Filter, const ChaosInterface::FQueryFilterData& QueryFilterData, ICollisionQueryFilterCallbackBase* QueryCallback, const ChaosInterface::FQueryDebugParams& DebugParams = {});
+void LowLevelRaycast(const FPhysScene& Scene, const FVector& Start, const FVector& Dir, float DeltaMag, FPhysicsHitCallback<ChaosInterface::FPTRaycastHit>& HitBuffer, EHitFlags OutputFlags, FQueryFlags QueryFlags, const FCollisionFilterData& Filter, const ChaosInterface::FQueryFilterData& QueryFilterData, ICollisionQueryFilterCallbackBase* QueryCallback, const ChaosInterface::FQueryDebugParams& DebugParams = {});
+void LowLevelRaycast(const Chaos::IDefaultChaosSpatialAcceleration& AccelerationStructure, const FVector& Start, const FVector& Dir, float DeltaMag, FPhysicsHitCallback<FHitRaycast>& HitBuffer, EHitFlags OutputFlags, FQueryFlags QueryFlags, const FCollisionFilterData& Filter, const ChaosInterface::FQueryFilterData& QueryFilterData, ICollisionQueryFilterCallbackBase* QueryCallback, const ChaosInterface::FQueryDebugParams& DebugParams = {});
+void LowLevelRaycast(const Chaos::IDefaultChaosSpatialAcceleration& AccelerationStructure, const FVector& Start, const FVector& Dir, float DeltaMag, FPhysicsHitCallback<ChaosInterface::FPTRaycastHit>& HitBuffer, EHitFlags OutputFlags, FQueryFlags QueryFlags, const FCollisionFilterData& Filter, const ChaosInterface::FQueryFilterData& QueryFilterData, ICollisionQueryFilterCallbackBase* QueryCallback, const ChaosInterface::FQueryDebugParams& DebugParams = {});
 
-void LowLevelSweep(FPhysScene& Scene, const FPhysicsGeometry& Geom, const FTransform& StartTM, const FVector& Dir, float DeltaMag, FPhysicsHitCallback<FHitSweep>& HitBuffer, EHitFlags OutputFlags, FQueryFlags QueryFlags, const FCollisionFilterData& Filter, const ChaosInterface::FQueryFilterData& QueryFilterData, ICollisionQueryFilterCallbackBase* QueryCallback, const ChaosInterface::FQueryDebugParams& DebugParams = {});
-void LowLevelSweep(FPhysScene& Scene, const FPhysicsGeometry& Geom, const FTransform& StartTM, const FVector& Dir, float DeltaMag, FPhysicsHitCallback<ChaosInterface::FPTSweepHit>& HitBuffer, EHitFlags OutputFlags, FQueryFlags QueryFlags, const FCollisionFilterData& Filter, const ChaosInterface::FQueryFilterData& QueryFilterData, ICollisionQueryFilterCallbackBase* QueryCallback, const ChaosInterface::FQueryDebugParams& DebugParams = {});
+void LowLevelSweep(const FPhysScene& Scene, const FPhysicsGeometry& Geom, const FTransform& StartTM, const FVector& Dir, float DeltaMag, FPhysicsHitCallback<FHitSweep>& HitBuffer, EHitFlags OutputFlags, FQueryFlags QueryFlags, const FCollisionFilterData& Filter, const ChaosInterface::FQueryFilterData& QueryFilterData, ICollisionQueryFilterCallbackBase* QueryCallback, const ChaosInterface::FQueryDebugParams& DebugParams = {});
+void LowLevelSweep(const FPhysScene& Scene, const FPhysicsGeometry& Geom, const FTransform& StartTM, const FVector& Dir, float DeltaMag, FPhysicsHitCallback<ChaosInterface::FPTSweepHit>& HitBuffer, EHitFlags OutputFlags, FQueryFlags QueryFlags, const FCollisionFilterData& Filter, const ChaosInterface::FQueryFilterData& QueryFilterData, ICollisionQueryFilterCallbackBase* QueryCallback, const ChaosInterface::FQueryDebugParams& DebugParams = {});
+void LowLevelSweep(const Chaos::IDefaultChaosSpatialAcceleration& AccelerationStructure, const FPhysicsGeometry& Geom, const FTransform& StartTM, const FVector& Dir, float DeltaMag, FPhysicsHitCallback<FHitSweep>& HitBuffer, EHitFlags OutputFlags, FQueryFlags QueryFlags, const FCollisionFilterData& Filter, const ChaosInterface::FQueryFilterData& QueryFilterData, ICollisionQueryFilterCallbackBase* QueryCallback, const ChaosInterface::FQueryDebugParams& DebugParams = {});
+void LowLevelSweep(const Chaos::IDefaultChaosSpatialAcceleration& AccelerationStructure, const FPhysicsGeometry& Geom, const FTransform& StartTM, const FVector& Dir, float DeltaMag, FPhysicsHitCallback<ChaosInterface::FPTSweepHit>& HitBuffer, EHitFlags OutputFlags, FQueryFlags QueryFlags, const FCollisionFilterData& Filter, const ChaosInterface::FQueryFilterData& QueryFilterData, ICollisionQueryFilterCallbackBase* QueryCallback, const ChaosInterface::FQueryDebugParams& DebugParams = {});
 
-void LowLevelOverlap(FPhysScene& Scene, const FPhysicsGeometry& Geom, const FTransform& GeomPose, FPhysicsHitCallback<FHitOverlap>& HitBuffer, FQueryFlags QueryFlags, const FCollisionFilterData& Filter, const ChaosInterface::FQueryFilterData& QueryFilterData, ICollisionQueryFilterCallbackBase* QueryCallback, const ChaosInterface::FQueryDebugParams& DebugParams = {});
-void LowLevelOverlap(FPhysScene& Scene, const FPhysicsGeometry& Geom, const FTransform& GeomPose, FPhysicsHitCallback<ChaosInterface::FPTOverlapHit>& HitBuffer, FQueryFlags QueryFlags, const FCollisionFilterData& Filter, const ChaosInterface::FQueryFilterData& QueryFilterData, ICollisionQueryFilterCallbackBase* QueryCallback, const ChaosInterface::FQueryDebugParams& DebugParams = {});
+void LowLevelOverlap(const FPhysScene& Scene, const FPhysicsGeometry& Geom, const FTransform& GeomPose, FPhysicsHitCallback<FHitOverlap>& HitBuffer, FQueryFlags QueryFlags, const FCollisionFilterData& Filter, const ChaosInterface::FQueryFilterData& QueryFilterData, ICollisionQueryFilterCallbackBase* QueryCallback, const ChaosInterface::FQueryDebugParams& DebugParams = {});
+void LowLevelOverlap(const FPhysScene& Scene, const FPhysicsGeometry& Geom, const FTransform& GeomPose, FPhysicsHitCallback<ChaosInterface::FPTOverlapHit>& HitBuffer, FQueryFlags QueryFlags, const FCollisionFilterData& Filter, const ChaosInterface::FQueryFilterData& QueryFilterData, ICollisionQueryFilterCallbackBase* QueryCallback, const ChaosInterface::FQueryDebugParams& DebugParams = {});
+void LowLevelOverlap(const Chaos::IDefaultChaosSpatialAcceleration& AccelerationStructure, const FPhysicsGeometry& Geom, const FTransform& GeomPose, FPhysicsHitCallback<FHitOverlap>& HitBuffer, FQueryFlags QueryFlags, const FCollisionFilterData& Filter, const ChaosInterface::FQueryFilterData& QueryFilterData, ICollisionQueryFilterCallbackBase* QueryCallback, const ChaosInterface::FQueryDebugParams& DebugParams = {});
+void LowLevelOverlap(const Chaos::IDefaultChaosSpatialAcceleration& AccelerationStructure, const FPhysicsGeometry& Geom, const FTransform& GeomPose, FPhysicsHitCallback<ChaosInterface::FPTOverlapHit>& HitBuffer, FQueryFlags QueryFlags, const FCollisionFilterData& Filter, const ChaosInterface::FQueryFilterData& QueryFilterData, ICollisionQueryFilterCallbackBase* QueryCallback, const ChaosInterface::FQueryDebugParams& DebugParams = {});
\ No newline at end of file
diff --git a/Engine/Source/Runtime/Engine/Public/Physics/Experimental/PhysInterface_Chaos.h b/Engine/Source/Runtime/Engine/Public/Physics/Experimental/PhysInterface_Chaos.h
--- a/Engine/Source/Runtime/Engine/Public/Physics/Experimental/PhysInterface_Chaos.h
+++ b/Engine/Source/Runtime/Engine/Public/Physics/Experimental/PhysInterface_Chaos.h
@@ -177,11 +177,4 @@ inline void SetShape(FPhysTypeDummy& Hit, FPhysTypeDummy* Shape)
 
 }
 
-bool IsBlocking(const FPhysicsShape& PShape, const FCollisionFilterData& QueryFilter);
-
-template<>
-ENGINE_API bool FGenericPhysicsInterface::GeomSweepMulti(const UWorld* World, const FPhysicsGeometryCollection& InGeom, const FQuat& InGeomRot, TArray<FHitResult>& OutHits, FVector Start, FVector End, ECollisionChannel TraceChannel, const FCollisionQueryParams& Params, const FCollisionResponseParams& ResponseParams, const FCollisionObjectQueryParams& ObjectParams /*= FCollisionObjectQueryParams::DefaultObjectQueryParam*/);
-
-template<>
-ENGINE_API bool FGenericPhysicsInterface::GeomOverlapMulti(const UWorld* World, const FPhysicsGeometryCollection& InGeom, const FVector& InPosition, const FQuat& InRotation, TArray<FOverlapResult>& OutOverlaps, ECollisionChannel TraceChannel, const FCollisionQueryParams& Params, const FCollisionResponseParams& ResponseParams, const FCollisionObjectQueryParams& ObjectParams);
-
+bool IsBlocking(const FPhysicsShape& PShape, const FCollisionFilterData& QueryFilter);
\ No newline at end of file
diff --git a/Engine/Source/Runtime/Engine/Public/Physics/GenericPhysicsInterface.h b/Engine/Source/Runtime/Engine/Public/Physics/GenericPhysicsInterface.h
--- a/Engine/Source/Runtime/Engine/Public/Physics/GenericPhysicsInterface.h
+++ b/Engine/Source/Runtime/Engine/Public/Physics/GenericPhysicsInterface.h
@@ -1,6 +1,8 @@
 // Copyright Epic Games, Inc. All Rights Reserved.
 
 #pragma once
+
+#include "Chaos/SpatialAccelerationFwd.h"
 #include "Engine/EngineTypes.h"
 #include "CollisionQueryParams.h"
 #include "WorldCollision.h"
@@ -13,10 +15,11 @@ struct ENGINE_API FGenericPhysicsInterface
 {
 	/** Trace a ray against the world and return if a blocking hit is found */
 	static bool RaycastTest(const UWorld* World, const FVector Start, const FVector End, ECollisionChannel TraceChannel, const FCollisionQueryParams& Params, const FCollisionResponseParams& ResponseParams, const FCollisionObjectQueryParams& ObjectParams = FCollisionObjectQueryParams::DefaultObjectQueryParam);
+	static bool RaycastTestUsingSpatialAcceleration(const Chaos::IDefaultChaosSpatialAcceleration& Accel, const UWorld* World, const FVector Start, const FVector End, ECollisionChannel TraceChannel, const FCollisionQueryParams& Params, const FCollisionResponseParams& ResponseParams, const FCollisionObjectQueryParams& ObjectParams = FCollisionObjectQueryParams::DefaultObjectQueryParam);
 
 	/** Trace a ray against the world and return the first blocking hit */
 	static bool RaycastSingle(const UWorld* World, struct FHitResult& OutHit, const FVector Start, const FVector End, ECollisionChannel TraceChannel, const FCollisionQueryParams& Params, const FCollisionResponseParams& ResponseParams, const FCollisionObjectQueryParams& ObjectParams = FCollisionObjectQueryParams::DefaultObjectQueryParam);
-
+	static bool RaycastSingleUsingSpatialAcceleration(const Chaos::IDefaultChaosSpatialAcceleration& Accel, const UWorld* World, struct FHitResult& OutHit, const FVector Start, const FVector End, ECollisionChannel TraceChannel, const FCollisionQueryParams& Params, const FCollisionResponseParams& ResponseParams, const FCollisionObjectQueryParams& ObjectParams = FCollisionObjectQueryParams::DefaultObjectQueryParam);
 
 	/**
 	*  Trace a ray against the world and return touching hits and then first blocking hit
@@ -24,30 +27,94 @@ struct ENGINE_API FGenericPhysicsInterface
 	*  Only the single closest blocking result will be generated, no tests will be done after that
 	*/
 	static bool RaycastMulti(const UWorld* World, TArray<struct FHitResult>& OutHits, const FVector& Start, const FVector& End, ECollisionChannel TraceChannel, const FCollisionQueryParams& Params, const FCollisionResponseParams& ResponseParams, const FCollisionObjectQueryParams& ObjectParams = FCollisionObjectQueryParams::DefaultObjectQueryParam);
+	static bool RaycastMultiUsingSpatialAcceleration(const Chaos::IDefaultChaosSpatialAcceleration& Accel, const UWorld* World, TArray<struct FHitResult>& OutHits, const FVector& Start, const FVector& End, ECollisionChannel TraceChannel, const FCollisionQueryParams& Params, const FCollisionResponseParams& ResponseParams, const FCollisionObjectQueryParams& ObjectParams = FCollisionObjectQueryParams::DefaultObjectQueryParam);
 
 	/** Function used for sweeping a supplied shape against the world as a test */
 	static bool GeomSweepTest(const UWorld* World, const FCollisionShape& CollisionShape, const FQuat& Rot, FVector Start, FVector End, ECollisionChannel TraceChannel, const FCollisionQueryParams& Params, const FCollisionResponseParams& ResponseParams, const FCollisionObjectQueryParams& ObjectParams = FCollisionObjectQueryParams::DefaultObjectQueryParam);
+	static bool GeomSweepTestUsingSpatialAcceleration(const Chaos::IDefaultChaosSpatialAcceleration& Accel, const UWorld* World, const FCollisionShape& CollisionShape, const FQuat& Rot, FVector Start, FVector End, ECollisionChannel TraceChannel, const FCollisionQueryParams& Params, const FCollisionResponseParams& ResponseParams, const FCollisionObjectQueryParams& ObjectParams = FCollisionObjectQueryParams::DefaultObjectQueryParam);
 
 	/** Function for sweeping a supplied shape against the world */
-	static bool GeomSweepSingle(const UWorld* World, const FCollisionShape& CollisionShape, const FQuat& Rot, FHitResult& OutHit, FVector Start, FVector End, ECollisionChannel TraceChannel, const FCollisionQueryParams& Params, const FCollisionResponseParams& ResponseParams, const FCollisionObjectQueryParams& ObjectParams = FCollisionObjectQueryParams::DefaultObjectQueryParam);
+	template<typename GeomWrapper>
+	static bool GeomSweepSingle(const UWorld* World, const GeomWrapper& InGeom, const FQuat& Rot, FHitResult& OutHit, FVector Start, FVector End, ECollisionChannel TraceChannel, const FCollisionQueryParams& Params, const FCollisionResponseParams& ResponseParams, const FCollisionObjectQueryParams& ObjectParams = FCollisionObjectQueryParams::DefaultObjectQueryParam);
+
+	template<typename GeomWrapper>
+	static bool GeomSweepSingleUsingSpatialAcceleration(const Chaos::IDefaultChaosSpatialAcceleration& Accel, const UWorld* World, const GeomWrapper& InGeom, const FQuat& Rot, FHitResult& OutHit, FVector Start, FVector End, ECollisionChannel TraceChannel, const FCollisionQueryParams& Params, const FCollisionResponseParams& ResponseParams, const FCollisionObjectQueryParams& ObjectParams = FCollisionObjectQueryParams::DefaultObjectQueryParam);
 
 	/** Sweep a supplied shape against the world and do not stop until the first blocking hit */
 	template <typename GeomWrapper>
 	static bool GeomSweepMulti(const UWorld* World, const GeomWrapper& InGeom, const FQuat& InGeomRot, TArray<FHitResult>& OutHits, FVector Start, FVector End, ECollisionChannel TraceChannel, const FCollisionQueryParams& Params, const FCollisionResponseParams& ResponseParams, const FCollisionObjectQueryParams& ObjectParams = FCollisionObjectQueryParams::DefaultObjectQueryParam);
 
+	template<typename GeomWrapper>
+	static bool GeomSweepMultiUsingSpatialAcceleration(const Chaos::IDefaultChaosSpatialAcceleration& Accel, const UWorld* World, const GeomWrapper& InGeom, const FQuat& InGeomRot, TArray<FHitResult>& OutHits, FVector Start, FVector End, ECollisionChannel TraceChannel, const FCollisionQueryParams& Params, const FCollisionResponseParams& ResponseParams, const FCollisionObjectQueryParams& ObjectParams = FCollisionObjectQueryParams::DefaultObjectQueryParam);
+
 	/** Find overlapping shapes with a given shape */
 	template<typename GeomWrapper>
 	static bool GeomOverlapMulti(const UWorld* World, const GeomWrapper& InGeom, const FVector& InPosition, const FQuat& InRotation, TArray<FOverlapResult>& OutOverlaps, ECollisionChannel TraceChannel, const FCollisionQueryParams& Params, const FCollisionResponseParams& ResponseParams, const FCollisionObjectQueryParams& ObjectParams);
 
+	template<typename GeomWrapper>
+	static bool GeomOverlapMultiUsingSpatialAcceleration(const Chaos::IDefaultChaosSpatialAcceleration& Accel, const UWorld* World, const GeomWrapper& InGeom, const FVector& InPosition, const FQuat& InRotation, TArray<FOverlapResult>& OutOverlaps, ECollisionChannel TraceChannel, const FCollisionQueryParams& Params, const FCollisionResponseParams& ResponseParams, const FCollisionObjectQueryParams& ObjectParams);
+
 	/** Function for testing overlaps between a supplied PxGeometry and the world. Returns true if at least one overlapping shape is blocking*/
 	static bool GeomOverlapBlockingTest(const UWorld* World, const FCollisionShape& CollisionShape, const FVector& Pos, const FQuat& Rot, ECollisionChannel TraceChannel, const FCollisionQueryParams& Params, const FCollisionResponseParams& ResponseParams, const FCollisionObjectQueryParams& ObjectParams = FCollisionObjectQueryParams::DefaultObjectQueryParam);
+	static bool GeomOverlapBlockingTestUsingSpatialAcceleration(const Chaos::IDefaultChaosSpatialAcceleration& Accel, const UWorld* World, const FCollisionShape& CollisionShape, const FVector& Pos, const FQuat& Rot, ECollisionChannel TraceChannel, const FCollisionQueryParams& Params, const FCollisionResponseParams& ResponseParams, const FCollisionObjectQueryParams& ObjectParams = FCollisionObjectQueryParams::DefaultObjectQueryParam);
 
 	/** Function for testing overlaps between a supplied PxGeometry and the world. Returns true if anything is overlapping (blocking or touching)*/
 	static bool GeomOverlapAnyTest(const UWorld* World, const FCollisionShape& CollisionShape, const FVector& Pos, const FQuat& Rot, ECollisionChannel TraceChannel, const FCollisionQueryParams& Params, const FCollisionResponseParams& ResponseParams, const FCollisionObjectQueryParams& ObjectParams = FCollisionObjectQueryParams::DefaultObjectQueryParam);
+	static bool GeomOverlapAnyTestUsingSpatialAcceleration(const Chaos::IDefaultChaosSpatialAcceleration& Accel, const UWorld* World, const FCollisionShape& CollisionShape, const FVector& Pos, const FQuat& Rot, ECollisionChannel TraceChannel, const FCollisionQueryParams& Params, const FCollisionResponseParams& ResponseParams, const FCollisionObjectQueryParams& ObjectParams = FCollisionObjectQueryParams::DefaultObjectQueryParam);
 };
 
 template<>
-ENGINE_API bool FGenericPhysicsInterface::GeomSweepMulti(const UWorld* World, const FCollisionShape& InGeom, const FQuat& InGeomRot, TArray<FHitResult>& OutHits, FVector Start, FVector End, ECollisionChannel TraceChannel, const FCollisionQueryParams& Params, const FCollisionResponseParams& ResponseParams, const FCollisionObjectQueryParams& ObjectParams /*= FCollisionObjectQueryParams::DefaultObjectQueryParam*/);
+ENGINE_API bool FGenericPhysicsInterface::GeomSweepSingle(const UWorld* World, const FCollisionShape& InGeom, const FQuat& Rot, FHitResult& OutHit, FVector Start, FVector End, ECollisionChannel TraceChannel, const FCollisionQueryParams& Params, const FCollisionResponseParams& ResponseParams, const FCollisionObjectQueryParams& ObjectParams);
+
+template<>
+ENGINE_API bool FGenericPhysicsInterface::GeomSweepSingle(const UWorld* World, const FPhysicsGeometry& InGeom, const FQuat& Rot, FHitResult& OutHit, FVector Start, FVector End, ECollisionChannel TraceChannel, const FCollisionQueryParams& Params, const FCollisionResponseParams& ResponseParams, const FCollisionObjectQueryParams& ObjectParams);
+
+template<>
+ENGINE_API bool FGenericPhysicsInterface::GeomSweepMulti(const UWorld* World, const FCollisionShape& InGeom, const FQuat& InGeomRot, TArray<FHitResult>& OutHits, FVector Start, FVector End, ECollisionChannel TraceChannel, const FCollisionQueryParams& Params, const FCollisionResponseParams& ResponseParams, const FCollisionObjectQueryParams& ObjectParams);
+
+template<>
+ENGINE_API bool FGenericPhysicsInterface::GeomSweepMulti(const UWorld* World, const FPhysicsGeometry& InGeom, const FQuat& InGeomRot, TArray<FHitResult>& OutHits, FVector Start, FVector End, ECollisionChannel TraceChannel, const FCollisionQueryParams& Params, const FCollisionResponseParams& ResponseParams, const FCollisionObjectQueryParams& ObjectParams);
+
+template<>
+ENGINE_API bool FGenericPhysicsInterface::GeomOverlapMulti(const UWorld* World, const FCollisionShape& InGeom, const FVector& InPosition, const FQuat& InRotation, TArray<FOverlapResult>& OutOverlaps, ECollisionChannel TraceChannel, const FCollisionQueryParams& Params, const FCollisionResponseParams& ResponseParams, const FCollisionObjectQueryParams& ObjectParams);
+
+template<>
+ENGINE_API bool FGenericPhysicsInterface::GeomOverlapMulti(const UWorld* World, const FPhysicsGeometry& InGeom, const FVector& InPosition, const FQuat& InRotation, TArray<FOverlapResult>& OutOverlaps, ECollisionChannel TraceChannel, const FCollisionQueryParams& Params, const FCollisionResponseParams& ResponseParams, const FCollisionObjectQueryParams& ObjectParams);
+
+template<>
+ENGINE_API bool FGenericPhysicsInterface::GeomSweepSingle(const UWorld* World, const FPhysicsGeometryCollection& InGeom, const FQuat& Rot, FHitResult& OutHit, FVector Start, FVector End, ECollisionChannel TraceChannel, const FCollisionQueryParams& Params, const FCollisionResponseParams& ResponseParams, const FCollisionObjectQueryParams& ObjectParams);
+
+template<>
+ENGINE_API bool FGenericPhysicsInterface::GeomSweepMulti(const UWorld* World, const FPhysicsGeometryCollection& InGeom, const FQuat& InGeomRot, TArray<FHitResult>& OutHits, FVector Start, FVector End, ECollisionChannel TraceChannel, const FCollisionQueryParams& Params, const FCollisionResponseParams& ResponseParams, const FCollisionObjectQueryParams& ObjectParams);
+
+template<>
+ENGINE_API bool FGenericPhysicsInterface::GeomOverlapMulti(const UWorld* World, const FPhysicsGeometryCollection& InGeom, const FVector& InPosition, const FQuat& InRotation, TArray<FOverlapResult>& OutOverlaps, ECollisionChannel TraceChannel, const FCollisionQueryParams& Params, const FCollisionResponseParams& ResponseParams, const FCollisionObjectQueryParams& ObjectParams);
+
+////////////////////////////////////
+
+template<>
+ENGINE_API bool FGenericPhysicsInterface::GeomSweepSingleUsingSpatialAcceleration(const Chaos::IDefaultChaosSpatialAcceleration& Accel, const UWorld* World, const FCollisionShape& InGeom, const FQuat& Rot, FHitResult& OutHit, FVector Start, FVector End, ECollisionChannel TraceChannel, const FCollisionQueryParams& Params, const FCollisionResponseParams& ResponseParams, const FCollisionObjectQueryParams& ObjectParams);
+
+template<>
+ENGINE_API bool FGenericPhysicsInterface::GeomSweepSingleUsingSpatialAcceleration(const Chaos::IDefaultChaosSpatialAcceleration& Accel, const UWorld* World, const FPhysicsGeometry& InGeom, const FQuat& Rot, FHitResult& OutHit, FVector Start, FVector End, ECollisionChannel TraceChannel, const FCollisionQueryParams& Params, const FCollisionResponseParams& ResponseParams, const FCollisionObjectQueryParams& ObjectParams);
+
+template<>
+ENGINE_API bool FGenericPhysicsInterface::GeomSweepMultiUsingSpatialAcceleration(const Chaos::IDefaultChaosSpatialAcceleration& Accel, const UWorld* World, const FCollisionShape& InGeom, const FQuat& InGeomRot, TArray<FHitResult>& OutHits, FVector Start, FVector End, ECollisionChannel TraceChannel, const FCollisionQueryParams& Params, const FCollisionResponseParams& ResponseParams, const FCollisionObjectQueryParams& ObjectParams);
+
+template<>
+ENGINE_API bool FGenericPhysicsInterface::GeomSweepMultiUsingSpatialAcceleration(const Chaos::IDefaultChaosSpatialAcceleration& Accel, const UWorld* World, const FPhysicsGeometry& InGeom, const FQuat& InGeomRot, TArray<FHitResult>& OutHits, FVector Start, FVector End, ECollisionChannel TraceChannel, const FCollisionQueryParams& Params, const FCollisionResponseParams& ResponseParams, const FCollisionObjectQueryParams& ObjectParams);
+
+template<>
+ENGINE_API bool FGenericPhysicsInterface::GeomOverlapMultiUsingSpatialAcceleration(const Chaos::IDefaultChaosSpatialAcceleration& Accel, const UWorld* World, const FCollisionShape& InGeom, const FVector& InPosition, const FQuat& InRotation, TArray<FOverlapResult>& OutOverlaps, ECollisionChannel TraceChannel, const FCollisionQueryParams& Params, const FCollisionResponseParams& ResponseParams, const FCollisionObjectQueryParams& ObjectParams);
+
+template<>
+ENGINE_API bool FGenericPhysicsInterface::GeomOverlapMultiUsingSpatialAcceleration(const Chaos::IDefaultChaosSpatialAcceleration& Accel, const UWorld* World, const FPhysicsGeometry& InGeom, const FVector& InPosition, const FQuat& InRotation, TArray<FOverlapResult>& OutOverlaps, ECollisionChannel TraceChannel, const FCollisionQueryParams& Params, const FCollisionResponseParams& ResponseParams, const FCollisionObjectQueryParams& ObjectParams);
+
+template<>
+ENGINE_API bool FGenericPhysicsInterface::GeomSweepSingleUsingSpatialAcceleration(const Chaos::IDefaultChaosSpatialAcceleration& Accel, const UWorld* World, const FPhysicsGeometryCollection& InGeom, const FQuat& Rot, FHitResult& OutHit, FVector Start, FVector End, ECollisionChannel TraceChannel, const FCollisionQueryParams& Params, const FCollisionResponseParams& ResponseParams, const FCollisionObjectQueryParams& ObjectParams);
+
+template<>
+ENGINE_API bool FGenericPhysicsInterface::GeomSweepMultiUsingSpatialAcceleration(const Chaos::IDefaultChaosSpatialAcceleration& Accel, const UWorld* World, const FPhysicsGeometryCollection& InGeom, const FQuat& InGeomRot, TArray<FHitResult>& OutHits, FVector Start, FVector End, ECollisionChannel TraceChannel, const FCollisionQueryParams& Params, const FCollisionResponseParams& ResponseParams, const FCollisionObjectQueryParams& ObjectParams);
 
 template<>
-ENGINE_API bool FGenericPhysicsInterface::GeomOverlapMulti(const UWorld* World, const FCollisionShape& InGeom, const FVector& InPosition, const FQuat& InRotation, TArray<FOverlapResult>& OutOverlaps, ECollisionChannel TraceChannel, const FCollisionQueryParams& Params, const FCollisionResponseParams& ResponseParams, const FCollisionObjectQueryParams& ObjectParams);
\ No newline at end of file
+ENGINE_API bool FGenericPhysicsInterface::GeomOverlapMultiUsingSpatialAcceleration(const Chaos::IDefaultChaosSpatialAcceleration& Accel, const UWorld* World, const FPhysicsGeometryCollection& InGeom, const FVector& InPosition, const FQuat& InRotation, TArray<FOverlapResult>& OutOverlaps, ECollisionChannel TraceChannel, const FCollisionQueryParams& Params, const FCollisionResponseParams& ResponseParams, const FCollisionObjectQueryParams& ObjectParams);
\ No newline at end of file
diff --git a/Engine/Source/Runtime/Experimental/Chaos/Public/Chaos/SpatialAccelerationFwd.h b/Engine/Source/Runtime/Experimental/Chaos/Public/Chaos/SpatialAccelerationFwd.h
new file mode 100644
--- /dev/null
+++ b/Engine/Source/Runtime/Experimental/Chaos/Public/Chaos/SpatialAccelerationFwd.h
@@ -0,0 +1,13 @@
+// Copyright Epic Games, Inc. All Rights Reserved.
+#pragma once
+
+#include "Chaos/Real.h"
+
+namespace Chaos
+{
+	class FAccelerationStructureHandle;
+	template <typename TPayload, typename T, int d>
+	class ISpatialAcceleration;
+
+	using IDefaultChaosSpatialAcceleration = ISpatialAcceleration<FAccelerationStructureHandle, FReal, 3>;
+}
\ No newline at end of file
diff --git a/Engine/Source/Runtime/PhysicsCore/Private/ChaosEngineInterface.cpp b/Engine/Source/Runtime/PhysicsCore/Private/ChaosEngineInterface.cpp
--- a/Engine/Source/Runtime/PhysicsCore/Private/ChaosEngineInterface.cpp
+++ b/Engine/Source/Runtime/PhysicsCore/Private/ChaosEngineInterface.cpp
@@ -105,6 +105,11 @@ FPhysicsGeometryCollection_Chaos::FPhysicsGeometryCollection_Chaos(const FPhysic
 {
 }
 
+FPhysicsGeometryCollection_Chaos::FPhysicsGeometryCollection_Chaos(const FPhysicsGeometry& InGeom)
+	: Geom(InGeom)
+{
+}
+
 FPhysicsShapeAdapter_Chaos::FPhysicsShapeAdapter_Chaos(const FQuat& Rot,const FCollisionShape& CollisionShape)
 	: GeometryRotation(Rot)
 {
@@ -1841,6 +1846,11 @@ FPhysicsGeometryCollection_Chaos FChaosEngineInterface::GetGeometryCollection(co
 	return NewCollection;
 }
 
+FPhysicsGeometryCollection_Chaos FChaosEngineInterface::GetGeometryCollection(const FPhysicsGeometry& InShape)
+{
+	return FPhysicsGeometryCollection_Chaos{ InShape };
+}
+
 void FChaosEngineInterface::SetMaskFilter(const FPhysicsShapeHandle& InShape, FMaskFilter InFilter)
 {
 	FCollisionFilterData SimFilter = GetSimulationFilter(InShape);
diff --git a/Engine/Source/Runtime/PhysicsCore/Public/Chaos/ChaosEngineInterface.h b/Engine/Source/Runtime/PhysicsCore/Public/Chaos/ChaosEngineInterface.h
--- a/Engine/Source/Runtime/PhysicsCore/Public/Chaos/ChaosEngineInterface.h
+++ b/Engine/Source/Runtime/PhysicsCore/Public/Chaos/ChaosEngineInterface.h
@@ -316,6 +316,7 @@ struct PHYSICSCORE_API FPhysicsGeometryCollection_Chaos
 private:
 	friend class FChaosEngineInterface;
 	explicit FPhysicsGeometryCollection_Chaos(const FPhysicsShapeReference_Chaos& InShape);
+	explicit FPhysicsGeometryCollection_Chaos(const FPhysicsGeometry& InGeom);
 
 	const Chaos::FImplicitObject& Geom;
 };
@@ -538,7 +539,7 @@ public:
 	static void SetUserData(const FPhysicsShapeHandle& InShape, void* InUserData);
 	static FPhysicsShapeHandle CloneShape(const FPhysicsShapeHandle& InShape);
 	static FPhysicsGeometryCollection_Chaos GetGeometryCollection(const FPhysicsShapeHandle& InShape);
-
+	static FPhysicsGeometryCollection_Chaos GetGeometryCollection(const FPhysicsGeometry& InShape);
 
 	// @todo(mlentine): Which of these do we need to support?
 	// Set the mask filter of a shape, which is an extra level of filtering during collision detection / query for extra channels like "Blue Team" and "Red Team"
