Chaos : optimization : precompute proxy related data to avoid having to pay the cost of cache misses to get the proxy from the handles

diff --git a/Engine/Source/Runtime/Engine/Private/PhysicsEngine/Experimental/PhysScene_Chaos.cpp b/Engine/Source/Runtime/Engine/Private/PhysicsEngine/Experimental/PhysScene_Chaos.cpp
--- a/Engine/Source/Runtime/Engine/Private/PhysicsEngine/Experimental/PhysScene_Chaos.cpp
+++ b/Engine/Source/Runtime/Engine/Private/PhysicsEngine/Experimental/PhysScene_Chaos.cpp
@@ -1927,36 +1927,44 @@ void FPhysScene_Chaos::OnSyncBodies(Chaos::FPhysicsSolverBase* Solver)
 			// Don't pass in anything here so we don't end up locking anything because we can assume the scene is already locked.
 			FLockedWritePhysicsObjectExternalInterface Interface = FPhysicsObjectExternalInterface::LockWrite({});
 
-			auto ShouldGCParticleBeInAccelerationStructure = [&Interface](Chaos::FPhysicsObjectHandle Handle)
+			const bool bIsParentProxyNull = (Proxy->GetParentProxy() == nullptr);
+
+			if (bIsParentProxyNull)
 			{
-				IPhysicsProxyBase* HandleProxy = Chaos::FPhysicsObjectInterface::GetProxy({ &Handle, 1 });
-
-				// It's possible to be an enabled particle and not qualify for query collisions if the GC particle has been replication abandoned. 
-				// Furthermore, it's possible for a particle to be enabled on the game thread but not actually enabled on the physics thread. A particle
-				// with an internal cluster parent (e.g. a cluster union) could get into this state where the geometry collection physics proxy will not
-				// transfer the disabled state from the PT to the GT because the particle has an internal cluster for a parent. We can detect this case by making
-				// sure that the proxy of the GC does not have a parent proxy (only happens with a cluster union currently).
-				return !Interface->AreAllDisabled({ &Handle, 1 })
-					&& Interface->AreAllShapesQueryEnabled({ &Handle, 1 })
-					&& (HandleProxy->GetParentProxy() == nullptr);
-			};
-
-			TArray<Chaos::FPhysicsObjectHandle> ActiveHandles = Proxy->GetAllPhysicsObjects().FilterByPredicate(
-				[&ShouldGCParticleBeInAccelerationStructure](Chaos::FPhysicsObjectHandle Handle)
+				auto ShouldGCParticleBeInAccelerationStructure = [&Interface](Chaos::FPhysicsObjectHandle Handle)
 				{
-					return ShouldGCParticleBeInAccelerationStructure(Handle);
-				}
-			);
+					// It's possible to be an enabled particle and not qualify for query collisions if the GC particle has been replication abandoned. 
+					// Furthermore, it's possible for a particle to be enabled on the game thread but not actually enabled on the physics thread. A particle
+					// with an internal cluster parent (e.g. a cluster union) could get into this state where the geometry collection physics proxy will not
+					// transfer the disabled state from the PT to the GT because the particle has an internal cluster for a parent. We can detect this case by making
+					// sure that the proxy of the GC does not have a parent proxy (only happens with a cluster union currently).
+					return !Interface->AreAllDisabled({ &Handle, 1 })
+						&& Interface->AreAllShapesQueryEnabled({ &Handle, 1 });
+				};
+
+				TArray<Chaos::FPhysicsObjectHandle> ActiveHandles = Proxy->GetAllPhysicsObjects().FilterByPredicate(
+					[&ShouldGCParticleBeInAccelerationStructure](Chaos::FPhysicsObjectHandle Handle)
+					{
+						return ShouldGCParticleBeInAccelerationStructure(Handle);
+					}
+				);
 
-			TArray<Chaos::FPhysicsObjectHandle> DisabledHandles = Proxy->GetAllPhysicsObjects().FilterByPredicate(
-				[&ShouldGCParticleBeInAccelerationStructure](Chaos::FPhysicsObjectHandle Handle)
-				{
-					return !ShouldGCParticleBeInAccelerationStructure(Handle);
-				}
-			);
+				TArray<Chaos::FPhysicsObjectHandle> DisabledHandles = Proxy->GetAllPhysicsObjects().FilterByPredicate(
+					[&ShouldGCParticleBeInAccelerationStructure](Chaos::FPhysicsObjectHandle Handle)
+					{
+						return !ShouldGCParticleBeInAccelerationStructure(Handle);
+					}
+				);
 
-			Interface->AddToSpatialAcceleration(ActiveHandles, Outer->GetSpacialAcceleration());
-			Interface->RemoveFromSpatialAcceleration(DisabledHandles, Outer->GetSpacialAcceleration());
+				Interface->AddToSpatialAcceleration(ActiveHandles, Outer->GetSpacialAcceleration());
+				Interface->RemoveFromSpatialAcceleration(DisabledHandles, Outer->GetSpacialAcceleration());
+			}
+			else
+			{
+				// if we have aparent proxy ( like attached to a cluster union) then none of the handles should be in the acceleration structure 
+				// todo : right now we are sending all the particles, but we should be able to get away with active + direct children that could reduce the overall cost 
+				Interface->RemoveFromSpatialAcceleration(Proxy->GetAllPhysicsObjects(), Outer->GetSpacialAcceleration());
+			}
 		}
 
 		void operator()(FClusterUnionPhysicsProxy* Proxy)
