Networked Physics: Implement settings for prediction and resimulation.

UI Options:
Project Settings > Physics > Replication > Physics Prediction:
- bool Enable Physics Prediction
- bool Enable Physics Resimulation
- float Resimulation Error Threshold (centimeters)
- float Max Supported Latency Prediction (milliseconds)

UI Options:
Actor > Replication
- EPhysicsReplicationMode PhysicsReplicationMode

The different settings replaces previously existing CVars
Get and Set functions are present where needed and are available via blueprint.

The per-actor "PhysicsReplicationMode" setting is new and holds three options in the Enum: Default, Predicive Interpolation, Resimulation.
Where Predictive Interpolation is a new work in progress physics replication, more updates to that coming shortly.
Resimulation enables the actor to trigger a resimulation of the island if the error threshold is large enough.

"PhysicsReplicationMode" per actor can be set in runtime and each client can have different settings for the same actor. The autonomous proxy might run Resimulation while the same actor runs as Predictive Interpolation on another client where it's a simulated proxy.

diff --git a/Engine/Plugins/Experimental/ChaosVehiclesPlugin/Source/ChaosVehicles/Private/ChaosVehicleMovementComponent.cpp b/Engine/Plugins/Experimental/ChaosVehiclesPlugin/Source/ChaosVehicles/Private/ChaosVehicleMovementComponent.cpp
--- a/Engine/Plugins/Experimental/ChaosVehiclesPlugin/Source/ChaosVehicles/Private/ChaosVehicleMovementComponent.cpp
+++ b/Engine/Plugins/Experimental/ChaosVehiclesPlugin/Source/ChaosVehicles/Private/ChaosVehicleMovementComponent.cpp
@@ -1090,13 +1090,12 @@ void UChaosVehicleMovementComponent::OnCreatePhysicsState()
 							AController* Controller = GetController();
 							const bool bIsLocalHistory = (Controller && Controller->IsLocalController());
 
-							static const IConsoleVariable* NumFramesCVar = IConsoleManager::Get().FindConsoleVariable(TEXT("p.RewindCaptureNumFrames"));
-							const int32 NumFrames = NumFramesCVar ? NumFramesCVar->GetInt() : 50;
+							const int32 PhysicsHistoryLength = UPhysicsSettings::Get()->GetPhysicsHistoryCount();
 
-							VehicleSimulationPT->InputsHistory = MakeShared<FVehicleInputsHistory>(NumFrames, bIsLocalHistory, VehicleSimulationPT.Get());
+							VehicleSimulationPT->InputsHistory = MakeShared<FVehicleInputsHistory>(PhysicsHistoryLength, bIsLocalHistory, VehicleSimulationPT.Get());
 							RewindData->AddInputsHistory(VehicleSimulationPT->InputsHistory);
 
-							VehicleSimulationPT->StatesHistory = MakeShared<FVehicleStatesHistory>(NumFrames, bIsLocalHistory, VehicleSimulationPT.Get());
+							VehicleSimulationPT->StatesHistory = MakeShared<FVehicleStatesHistory>(PhysicsHistoryLength, bIsLocalHistory, VehicleSimulationPT.Get());
 							RewindData->AddStatesHistory(VehicleSimulationPT->StatesHistory);
 						}
 					}
diff --git a/Engine/Source/Runtime/Engine/Classes/Engine/EngineTypes.h b/Engine/Source/Runtime/Engine/Classes/Engine/EngineTypes.h
--- a/Engine/Source/Runtime/Engine/Classes/Engine/EngineTypes.h
+++ b/Engine/Source/Runtime/Engine/Classes/Engine/EngineTypes.h
@@ -2965,6 +2965,27 @@ enum ENetDormancy : int
 	DORM_MAX UMETA(Hidden),
 };
 
+UENUM(BlueprintType)
+enum class EPhysicsReplicationMode : uint8
+{
+	/** Default physics replication.*/
+	Default UMETA(DisplayName = "Default"),
+
+	/** Physics replication performing velocity interpolation.
+	* Recommendation: Set on actors with a local role of ENetRole::ROLE_SimulatedProxy.
+	* Designed to handle local predictive interactions with other actors, especially actors of the local role ENetRole::ROLE_AutonomousProxy.
+	* Work in progress, available when Project Settings > Physics Prediction is enabled.
+	*/
+	PredictiveInterpolation UMETA(DisplayName = "Predictive Interpolation (WIP)"),
+
+	/** Forward predicted replication by simulating physics and correcting errors through resimulating physics from a correct state in the past.
+	* Recommendation: Set on actors with a local role of ENetRole::ROLE_AutonomousProxy.
+	* Can be used for both ROLE_AutonomousProxy and ROLE_SimulatedProxy, though not recommended for ROLE_SimulatedProxy due to CPU performance.
+	* Work in progress, available when Project Settings > Physics Prediction > Resimulation is enabled.
+	*/
+	Resimulation UMETA(DisplayName = "Resimulation (WIP)"),
+};
+
 /** Specifies which player index will pass input to this actor/component */
 UENUM()
 namespace EAutoReceiveInput
diff --git a/Engine/Source/Runtime/Engine/Classes/GameFramework/Actor.h b/Engine/Source/Runtime/Engine/Classes/GameFramework/Actor.h
--- a/Engine/Source/Runtime/Engine/Classes/GameFramework/Actor.h
+++ b/Engine/Source/Runtime/Engine/Classes/GameFramework/Actor.h
@@ -762,6 +762,28 @@ public:
 	UPROPERTY(Category=Replication, EditDefaultsOnly, BlueprintReadWrite)
 	float NetPriority;
 
+private:
+	/** Which mode to replicate physics through for this actor. Only relevant if the actor replicates movement and has a component that simulate physics.*/
+	UPROPERTY(EditDefaultsOnly, Category = Replication)
+	EPhysicsReplicationMode PhysicsReplicationMode;
+	
+public:
+	/** Set the physics replication mode of this body, via EPhysicsReplicationMode */
+	UFUNCTION(BlueprintCallable, Category = Physics)
+	void SetPhysicsReplicationMode(const EPhysicsReplicationMode ReplicationMode);
+
+	/** Get the physics replication mode of this body, via EPhysicsReplicationMode */
+	UFUNCTION(BlueprintCallable, Category = Physics)
+	EPhysicsReplicationMode GetPhysicsReplicationMode();
+
+	/** Can this body trigger a resimulation when Physics Prediction is enabled */
+	UFUNCTION(BlueprintCallable, Category = Physics)
+	bool CanTriggerResimulation() const;
+
+	/** Get the error threshold in centimeters before this object should enforce a resimulation to trigger. */
+	UFUNCTION(BlueprintCallable, Category = Physics)
+	float GetResimulationThreshold() const;
+
 private:
 	/**
 	 * The value of WorldSettings->TimeSeconds for the frame when one of this actor's components was last rendered.  This is written
diff --git a/Engine/Source/Runtime/Engine/Classes/PhysicsEngine/BodyInstance.h b/Engine/Source/Runtime/Engine/Classes/PhysicsEngine/BodyInstance.h
--- a/Engine/Source/Runtime/Engine/Classes/PhysicsEngine/BodyInstance.h
+++ b/Engine/Source/Runtime/Engine/Classes/PhysicsEngine/BodyInstance.h
@@ -863,6 +863,9 @@ public:
 
 	bool IsPhysicsDisabled() const;
 
+	/** Get the EPhysicsReplicationMode from the owning actor. It's recommended to get the parameter directly from the Actor if possible. */
+	EPhysicsReplicationMode GetPhysicsReplicationMode() const;
+
 private:
 
 	/** Struct of body instance delegates that are rarely bound so that we can only allocate memory if one is actually being used. */
diff --git a/Engine/Source/Runtime/Engine/Classes/PhysicsEngine/PhysicsSettings.h b/Engine/Source/Runtime/Engine/Classes/PhysicsEngine/PhysicsSettings.h
--- a/Engine/Source/Runtime/Engine/Classes/PhysicsEngine/PhysicsSettings.h
+++ b/Engine/Source/Runtime/Engine/Classes/PhysicsEngine/PhysicsSettings.h
@@ -85,6 +85,36 @@ namespace ESettingsDOF
 	};
 }
 
+/** Physics Prediction Settings */
+USTRUCT()
+struct FPhysicsPredictionSettings
+{
+	GENERATED_BODY();
+
+	/** Enable networked physics prediction */
+	UPROPERTY(EditAnywhere, Category = "Replication")
+	bool bEnablePhysicsPrediction;
+
+	/** Enable physics resimulation */
+	UPROPERTY(EditAnywhere, Category = "Replication", meta = (editcondition = "bEnablePhysicsPrediction"))
+	bool bEnablePhysicsResimulation;
+
+	/** Distance in centimeters before a state discrepancy triggers a resimulation */
+	UPROPERTY(EditAnywhere, Category = "Replication", meta = (editcondition = "bEnablePhysicsPrediction"))
+	float ResimulationErrorThreshold;
+
+	/** Amount of RTT (Round Trip Time) latency for the prediction to support in milliseconds. */
+	UPROPERTY(EditAnywhere, Category = "Replication", meta = (editcondition = "bEnablePhysicsPrediction"))
+	float MaxSupportedLatencyPrediction;
+
+	FPhysicsPredictionSettings()
+		: bEnablePhysicsPrediction(false)
+		, bEnablePhysicsResimulation(false)
+		, ResimulationErrorThreshold(10.0)
+		, MaxSupportedLatencyPrediction(1000)
+	{ }
+};
+
 UENUM()
 namespace ESettingsLockedAxis
 {
@@ -111,6 +141,10 @@ class ENGINE_API UPhysicsSettings : public UPhysicsSettingsCore
 {
 	GENERATED_UCLASS_BODY()
 
+	/** Settings for Networked Physics Prediction, experimental. */
+	UPROPERTY(Config, EditAnywhere, Category = Replication, meta = (DisplayName = "Physics Prediction (Experimental)"))
+	FPhysicsPredictionSettings PhysicsPrediction;
+
 	/** Error correction data for replicating simulated physics (rigid bodies) */
 	UPROPERTY(config, EditAnywhere, Category = Replication)
 	FRigidBodyErrorCorrection PhysicErrorCorrection;
@@ -221,6 +255,12 @@ public:
 
 	static UPhysicsSettings* Get() { return CastChecked<UPhysicsSettings>(UPhysicsSettings::StaticClass()->GetDefaultObject()); }
 
+	UFUNCTION(BlueprintCallable, Category = "Physics")
+	int32 GetPhysicsHistoryCount() const
+	{
+		return FMath::CeilToInt(1.f / AsyncFixedTimeStepSize) * (PhysicsPrediction.MaxSupportedLatencyPrediction / 1000.f);
+	}
+
 	virtual void PostInitProperties() override;
 
 #if WITH_EDITOR
diff --git a/Engine/Source/Runtime/Engine/Private/Actor.cpp b/Engine/Source/Runtime/Engine/Private/Actor.cpp
--- a/Engine/Source/Runtime/Engine/Private/Actor.cpp
+++ b/Engine/Source/Runtime/Engine/Private/Actor.cpp
@@ -156,6 +156,7 @@ void AActor::InitializeDefaults()
 	bCallPreReplication = true;
 	bCallPreReplicationForReplay = true;
 	bReplicateUsingRegisteredSubObjectList = GDefaultUseSubObjectReplicationList;
+	PhysicsReplicationMode = EPhysicsReplicationMode::Default;
 	NetPriority = 1.0f;
 	NetUpdateFrequency = 100.0f;
 	MinNetUpdateFrequency = 2.0f;
diff --git a/Engine/Source/Runtime/Engine/Private/ActorReplication.cpp b/Engine/Source/Runtime/Engine/Private/ActorReplication.cpp
--- a/Engine/Source/Runtime/Engine/Private/ActorReplication.cpp
+++ b/Engine/Source/Runtime/Engine/Private/ActorReplication.cpp
@@ -11,6 +11,7 @@
 #include "GameFramework/GameNetworkManager.h"
 #include "NetworkingDistanceConstants.h"
 #include "PhysicsReplication.h"
+#include "PhysicsEngine/PhysicsSettings.h"
 #include "DrawDebugHelpers.h"
 #include "Net/Core/PushModel/PushModel.h"
 #include "Interfaces/Interface_ActorSubobject.h"
@@ -259,8 +260,8 @@ void AActor::PostNetReceiveLocationAndRotation()
 
 void AActor::PostNetReceiveVelocity(const FVector& NewVelocity)
 {
-	static IConsoleVariable* EnableNetworkPhysicsCVar = IConsoleManager::Get().FindConsoleVariable(TEXT("np2.EnableNetworkPhysicsPrediction"));
-	if (EnableNetworkPhysicsCVar && (EnableNetworkPhysicsCVar->GetInt() == 1))
+	const FPhysicsPredictionSettings& PhysicsPredictionSettings = UPhysicsSettings::Get()->PhysicsPrediction;
+	if (PhysicsPredictionSettings.bEnablePhysicsPrediction)
 	{
 		// required to keep client deterministic with server simulation
 		if (RootComponent && RootComponent->IsRegistered() && (NewVelocity != GetVelocity()))
@@ -1073,3 +1074,22 @@ void AActor::UpdateReplicatePhysicsCondition()
 #endif // UE_WITH_IRIS
 }
 
+bool AActor::CanTriggerResimulation() const
+{
+	return PhysicsReplicationMode == EPhysicsReplicationMode::Resimulation;
+}
+
+void AActor::SetPhysicsReplicationMode(const EPhysicsReplicationMode ReplicationMode)
+{
+	PhysicsReplicationMode = ReplicationMode;
+}
+
+EPhysicsReplicationMode AActor::GetPhysicsReplicationMode()
+{
+	return PhysicsReplicationMode;
+}
+
+float AActor::GetResimulationThreshold() const
+{
+	return UPhysicsSettings::Get()->PhysicsPrediction.ResimulationErrorThreshold;
+}
diff --git a/Engine/Source/Runtime/Engine/Private/Pawn.cpp b/Engine/Source/Runtime/Engine/Private/Pawn.cpp
--- a/Engine/Source/Runtime/Engine/Private/Pawn.cpp
+++ b/Engine/Source/Runtime/Engine/Private/Pawn.cpp
@@ -26,6 +26,7 @@
 #include "Components/PawnNoiseEmitterComponent.h"
 #include "GameFramework/GameNetworkManager.h"
 #include "GameFramework/InputSettings.h"
+#include "PhysicsEngine/PhysicsSettings.h"
 #include "Engine/DemoNetDriver.h"
 #include "Misc/EngineNetworkCustomVersion.h"
 
@@ -1168,8 +1169,8 @@ void APawn::GetMoveGoalReachTest(const AActor* MovingActor, const FVector& MoveO
 
 void APawn::PostNetReceiveVelocity(const FVector& NewVelocity)
 {
-	static IConsoleVariable* EnableNetworkPhysicsCVar = IConsoleManager::Get().FindConsoleVariable(TEXT("np2.EnableNetworkPhysicsPrediction"));
-	if (EnableNetworkPhysicsCVar && EnableNetworkPhysicsCVar->GetInt() == 1)
+	const FPhysicsPredictionSettings& PhysicsPredictionSettings = UPhysicsSettings::Get()->PhysicsPrediction;
+	if (PhysicsPredictionSettings.bEnablePhysicsPrediction)
 	{
 		Super::PostNetReceiveVelocity(NewVelocity);
 	}
diff --git a/Engine/Source/Runtime/Engine/Private/PhysicsEngine/AsyncPhysicsInputComponent.cpp b/Engine/Source/Runtime/Engine/Private/PhysicsEngine/AsyncPhysicsInputComponent.cpp
--- a/Engine/Source/Runtime/Engine/Private/PhysicsEngine/AsyncPhysicsInputComponent.cpp
+++ b/Engine/Source/Runtime/Engine/Private/PhysicsEngine/AsyncPhysicsInputComponent.cpp
@@ -342,14 +342,9 @@ void UAsyncPhysicsInputComponent::SetDataClass(TSubclassOf<UAsyncPhysicsData> In
 				// Hack: enable Rewind capture here. This needs to be moved somewhere else
 				// -------------------------------------------------------------------------------------
 
-				int32 NumFrames = 64;
-				static const IConsoleVariable* NumFramesCVar = IConsoleManager::Get().FindConsoleVariable(TEXT("p.RewindCaptureNumFrames"));
-				if (ensure(NumFramesCVar))
-				{
-					// 1 frame is required to enable rewind capture. 
-					NumFrames = FMath::Max<int32>(1, NumFramesCVar->GetInt());
-				}
-				Solver->EnableRewindCapture(NumFrames, true, MakeUnique<FAsyncPhysicsInputRewindCallback>(World));
+				const int32 PhysicsHistoryLength = FChaosSolversModule::GetModule()->GetSettingsProvider().GetPhysicsHistoryCount();
+
+				Solver->EnableRewindCapture(PhysicsHistoryLength, true, MakeUnique<FAsyncPhysicsInputRewindCallback>(World));
 			}
 
 			FAsyncPhysicsInputRewindCallback* Callback = static_cast<FAsyncPhysicsInputRewindCallback*>(Solver->GetRewindCallback());
diff --git a/Engine/Source/Runtime/Engine/Private/PhysicsEngine/BodyInstance.cpp b/Engine/Source/Runtime/Engine/Private/PhysicsEngine/BodyInstance.cpp
--- a/Engine/Source/Runtime/Engine/Private/PhysicsEngine/BodyInstance.cpp
+++ b/Engine/Source/Runtime/Engine/Private/PhysicsEngine/BodyInstance.cpp
@@ -3529,6 +3529,18 @@ bool FBodyInstance::IsPhysicsDisabled() const
 	return bIsDisabled;
 }
 
+EPhysicsReplicationMode FBodyInstance::GetPhysicsReplicationMode() const
+{
+	UPrimitiveComponent* OwnerComponentInst = OwnerComponent.Get();
+	AActor* OwningActor = OwnerComponentInst ? OwnerComponentInst->GetOwner() : nullptr;
+	if (OwningActor)
+	{
+		return OwningActor->GetPhysicsReplicationMode();
+	}
+
+	return EPhysicsReplicationMode::Default;
+}
+
 void FBodyInstance::AddRadialImpulseToBody(const FVector& Origin, float Radius, float Strength, uint8 Falloff, bool bVelChange)
 {
 	FPhysicsCommand::ExecuteWrite(ActorHandle, [&](const FPhysicsActorHandle& Actor)
@@ -4304,7 +4316,6 @@ FRecalculatedMassProperties& FBodyInstance::OnRecalculatedMassProperties()
 	return BodyInstanceDelegates->OnRecalculatedMassProperties;
 }
 
-
 bool FBodyInstanceAsyncPhysicsTickHandle::IsValid() const
 {
 	return Proxy && Proxy->GetPhysicsThreadAPI() != nullptr;
@@ -4315,6 +4326,7 @@ Chaos::FRigidBodyHandle_Internal* FBodyInstanceAsyncPhysicsTickHandle::operator-
 	return Proxy->GetPhysicsThreadAPI();
 }
 
+
 ////////////////////////////////////////////////////////////////////////////
 // FBodyInstanceEditorHelpers
 
diff --git a/Engine/Source/Runtime/Engine/Private/PhysicsEngine/Experimental/PhysScene_Chaos.cpp b/Engine/Source/Runtime/Engine/Private/PhysicsEngine/Experimental/PhysScene_Chaos.cpp
--- a/Engine/Source/Runtime/Engine/Private/PhysicsEngine/Experimental/PhysScene_Chaos.cpp
+++ b/Engine/Source/Runtime/Engine/Private/PhysicsEngine/Experimental/PhysScene_Chaos.cpp
@@ -95,11 +95,10 @@ public:
 			}
 		}
 
-		static IConsoleVariable* ResimNetworkPhysicsCVar = IConsoleManager::Get().FindConsoleVariable(TEXT("np2.ResimNetworkPhysicsPrediction"));
-		const bool bEnableResim =  (ResimNetworkPhysicsCVar && ResimNetworkPhysicsCVar->GetInt() == 1);
 
-		static const IConsoleVariable* NumFramesCVar = IConsoleManager::Get().FindConsoleVariable(TEXT("p.RewindCaptureNumFrames"));
-		const int32 NumFrames = NumFramesCVar ? NumFramesCVar->GetInt() : 0;
+		const UPhysicsSettings* PhysicsSettings = UPhysicsSettings::Get();
+		const bool bEnableResim = PhysicsSettings->PhysicsPrediction.bEnablePhysicsResimulation;
+		const int32 NumFrames = PhysicsSettings->GetPhysicsHistoryCount();
 
 		for (int32 Idx = 0; Idx < PendingCommands.Num(); ++Idx)
 		{
diff --git a/Engine/Source/Runtime/Engine/Private/PhysicsEngine/PhysLevel.cpp b/Engine/Source/Runtime/Engine/Private/PhysicsEngine/PhysLevel.cpp
--- a/Engine/Source/Runtime/Engine/Private/PhysicsEngine/PhysLevel.cpp
+++ b/Engine/Source/Runtime/Engine/Private/PhysicsEngine/PhysLevel.cpp
@@ -42,6 +42,24 @@ public:
 	{
 		return GetSettings()->MinDeltaVelocityForHitEvents;
 	}
+
+	virtual bool GetPhysicsPredictionEnabled() const override
+	{
+		return GetSettings()->PhysicsPrediction.bEnablePhysicsPrediction;
+	}
+	virtual bool GetPhysicsResimulationEnabled() const override
+	{
+		return GetSettings()->PhysicsPrediction.bEnablePhysicsResimulation;
+	}
+	virtual float GetResimulationErrorThreshold() const override
+	{
+		return GetSettings()->PhysicsPrediction.ResimulationErrorThreshold;
+	}
+	virtual int32 GetPhysicsHistoryCount() const override
+	{
+		return GetSettings()->GetPhysicsHistoryCount();
+	}
+
 private:
 
 	const UPhysicsSettings* GetSettings() const
diff --git a/Engine/Source/Runtime/Engine/Private/PhysicsEngine/PhysicsReplication.cpp b/Engine/Source/Runtime/Engine/Private/PhysicsEngine/PhysicsReplication.cpp
--- a/Engine/Source/Runtime/Engine/Private/PhysicsEngine/PhysicsReplication.cpp
+++ b/Engine/Source/Runtime/Engine/Private/PhysicsEngine/PhysicsReplication.cpp
@@ -153,14 +153,15 @@ bool FPhysicsReplication::ApplyRigidBodyState(float DeltaSeconds, FBodyInstance*
 	Chaos::FPhysicsSolver* Solver = PhysScene->GetSolver();
 	Chaos::FRewindData* RewindData = Solver->GetRewindData();
 	const FPhysicsActorHandle& Proxy = BI->GetPhysicsActorHandle();
+	EPhysicsReplicationMode PhysicsReplicationMode = BI->GetPhysicsReplicationMode();
 
-	if (Proxy && Chaos::FPhysicsSolverBase::IsNetworkPhysicsPredictionEnabled())
+	if (PhysicsReplicationMode != EPhysicsReplicationMode::Default && Proxy && Chaos::FPhysicsSolverBase::IsNetworkPhysicsPredictionEnabled())
 	{
 		if (RewindData && (LocalFrame <= RewindData->CurrentFrame()) && (LocalFrame >= RewindData->GetEarliestFrame_Internal()))
 		{
 			static constexpr Chaos::FFrameAndPhase::EParticleHistoryPhase RewindPhase = Chaos::FFrameAndPhase::EParticleHistoryPhase::PostPushData;
 
-			if (Chaos::FPhysicsSolverBase::CanResimNetworkPhysicsPrediction())
+			if (PhysicsReplicationMode == EPhysicsReplicationMode::Resimulation && Chaos::FPhysicsSolverBase::IsPhysicsResimulationEnabled())
 			{
 				if (LocalFrame < RewindData->CurrentFrame() && !RewindData->IsResim())
 				{
@@ -180,22 +181,23 @@ bool FPhysicsReplication::ApplyRigidBodyState(float DeltaSeconds, FBodyInstance*
 								}
 								#endif
 
-								PlayerController->ExecuteAsyncPhysicsCommand(TimeStamp, BI->OwnerComponent.Get(), [this, Proxy, PhysicsTarget, RewindData, LocalFrame]
+								const float ResimErrorThreshold = Chaos::FPhysicsSolverBase::ResimulationErrorThreshold();
+
+								PlayerController->ExecuteAsyncPhysicsCommand(TimeStamp, BI->OwnerComponent.Get(), [this, Proxy, PhysicsTarget, RewindData, LocalFrame, ResimErrorThreshold]
 								{
 									if (RewindData->IsResim())
 										return;
 
 									auto PastState = RewindData->GetPastStateAtFrame(*Proxy->GetHandle_LowLevel(), LocalFrame, RewindPhase);
 
-									static const float ErrorThreshold = Chaos::FPhysicsSolverBase::NetworkPhysicsErrorThreshold();
-
 									const float ErrorPosition = (PastState.X() - PhysicsTarget.TargetState.Position).Size();
-									float ColorLerp = (ErrorPosition >= ErrorThreshold) ? 1.0f : 0.0f;
+									const bool ShouldTriggerResim = ErrorPosition >= ResimErrorThreshold;
+									float ColorLerp = ShouldTriggerResim ? 1.0f : 0.0f;
 
 									#if !(UE_BUILD_SHIPPING || UE_BUILD_TEST)
 									if (Chaos::FPhysicsSolverBase::CanDebugNetworkPhysicsPrediction())
 									{
-										UE_LOG(LogTemp, Log, TEXT("Particle Position Error = %f | Resim Trigger = %f | Server Frame = %d | Client Frame = %d"), ErrorPosition, ColorLerp, PhysicsTarget.ServerFrame, LocalFrame);
+										UE_LOG(LogTemp, Log, TEXT("Particle Position Error = %f | Should Trigger Resim = %s | Server Frame = %d | Client Frame = %d"), ErrorPosition, (ShouldTriggerResim ? TEXT("True") : TEXT("False")), PhysicsTarget.ServerFrame, LocalFrame);
 										UE_LOG(LogTemp, Log, TEXT("Particle Target Position = %s | Current Position = %s"), *PhysicsTarget.TargetState.Position.ToString(), *PastState.X().ToString());
 										UE_LOG(LogTemp, Log, TEXT("Particle Target Velocity = %s | Current Velocity = %s"), *PhysicsTarget.TargetState.LinVel.ToString(), *PastState.V().ToString());
 										UE_LOG(LogTemp, Log, TEXT("Particle Target Quaternion = %s | Current Quaternion = %s"), *PhysicsTarget.TargetState.Quaternion.ToString(), *PastState.R().ToString());
@@ -220,7 +222,7 @@ bool FPhysicsReplication::ApplyRigidBodyState(float DeltaSeconds, FBodyInstance*
 										PhysicsTarget.TargetState.Position, PhysicsTarget.TargetState.Quaternion, 
 										PhysicsTarget.TargetState.LinVel, PhysicsTarget.TargetState.AngVel, bShouldSleep);
 
-									if (ColorLerp == 1.0f)
+									if (ShouldTriggerResim)
 									{
 										if (Chaos::FPBDRigidsSolver* RigidSolver = Proxy->GetSolver<Chaos::FPBDRigidsSolver>())
 										{
@@ -234,9 +236,9 @@ bool FPhysicsReplication::ApplyRigidBodyState(float DeltaSeconds, FBodyInstance*
 					}
 				}
 			}
-			else
+			else // PhysicsReplicationMode == EPhysicsReplicationMode::PredictiveInterpolation
 			{
-				static const float LerpFactor = Chaos::FPhysicsSolverBase::NetworkPhysicsInterpolationLerp();
+				const float LerpFactor = Chaos::FPhysicsSolverBase::NetworkPhysicsInterpolationLerp();
 
 				FRigidBodyState CurrentState;
 				BI->GetRigidBodyState(CurrentState);
@@ -651,7 +653,6 @@ void FPhysicsReplication::OnTick(float DeltaSeconds, TMap<TWeakObjectPtr<UPrimit
 		bool bRemoveItr = false;
 		if (UPrimitiveComponent* PrimComp = Itr.Key().Get())
 		{
-
 			if (FBodyInstance* BI = PrimComp->GetBodyInstance(Itr.Value().BoneName))
 			{
 				FReplicatedPhysicsTarget& PhysicsTarget = Itr.Value();
@@ -677,6 +678,7 @@ void FPhysicsReplication::OnTick(float DeltaSeconds, TMap<TWeakObjectPtr<UPrimit
 						{
 							const int32 LocalFrame = PhysicsTarget.ServerFrame - LocalFrameOffset;
 							const bool bRestoredState = ApplyRigidBodyState(DeltaSeconds, BI, PhysicsTarget, PhysicErrorCorrection, PingSecondsOneWay, LocalFrame, NumPredictedFrames);
+							
 							// Need to update the component to match new position.
 							if (PhysicsReplicationCVars::SkipSkeletalRepOptimization == 0 || Cast<USkeletalMeshComponent>(PrimComp) == nullptr)	//simulated skeletal mesh does its own polling of physics results so we don't need to call this as it'll happen at the end of the physics sim
 							{
diff --git a/Engine/Source/Runtime/Engine/Private/PlayerController.cpp b/Engine/Source/Runtime/Engine/Private/PlayerController.cpp
--- a/Engine/Source/Runtime/Engine/Private/PlayerController.cpp
+++ b/Engine/Source/Runtime/Engine/Private/PlayerController.cpp
@@ -67,6 +67,7 @@
 #include "WorldPartition/WorldPartitionSubsystem.h"
 #include "Physics/AsyncPhysicsInputComponent.h"
 #include "PBDRigidsSolver.h"
+#include "PhysicsEngine/PhysicsSettings.h"
 
 #if UE_WITH_IRIS
 #include "Iris/ReplicationSystem/ReplicationSystem.h"
@@ -109,24 +110,6 @@ namespace NetworkPhysicsCvars
 	int32 EnableDebugRPC = 1;
 #endif
 	FAutoConsoleVariableRef CVarEnableDebugRPC(TEXT("np2.EnableDebugRPC"), EnableDebugRPC, TEXT("Sends extra debug information to clients about server side input buffering"));
-
-	int32 EnableNetworkPhysicsPrediction = 0;
-	FAutoConsoleVariableRef CVarEnableNetworkPhysicsPrediction(TEXT("np2.EnableNetworkPhysicsPrediction"), EnableNetworkPhysicsPrediction, TEXT("Enables network physics prediction"));
-
-	int32 DebugNetworkPhysicsPrediction = 0;
-	FAutoConsoleVariableRef CVarDebugNetworkPhysicsPrediction(TEXT("np2.DebugNetworkPhysicsPrediction"), DebugNetworkPhysicsPrediction, TEXT("Debugs network physics prediction"));
-
-	int32 ResimNetworkPhysicsPrediction = 1;
-	FAutoConsoleVariableRef CVarResimNetworkPhysicsPrediction(TEXT("np2.ResimNetworkPhysicsPrediction"), ResimNetworkPhysicsPrediction, TEXT("Resim network physics prediction"));
-
-	float NetworkPhysicsPredictionErrorThreshold = 10.0f;
-	FAutoConsoleVariableRef CVarNetworkPhysicsPredictionErrorThreshold(TEXT("np2.NetworkPhysicsPredictionErrorThreshold"), NetworkPhysicsPredictionErrorThreshold, TEXT("Position error threshold that will trigger the client correction based on the server prediction"));
-	
-	float NetworkPhysicsPredictionInterpLerp = 0.1f;
-	FAutoConsoleVariableRef CVarNetworkPhysicsPredictionInterpLerp(TEXT("np2.NetworkPhysicsPredictionInterpLerp"), NetworkPhysicsPredictionInterpLerp, TEXT("State lerp value in between the target state and the current one in case resim is disabled or if the pawn is not possessed (continuous correction)"));
-
-	float NetworkPhysicsPredictionResimLerp = 1.0f;
-	FAutoConsoleVariableRef CVarNetworkPhysicsPredictionResimLerp(TEXT("np2.NetworkPhysicsPredictionResimLerp"), NetworkPhysicsPredictionResimLerp, TEXT("State lerp value in between the target state and the matching one in the history when resim is enabled"));
 }
 
 const float RetryClientRestartThrottleTime = 0.5f;
@@ -210,7 +193,7 @@ APlayerController::APlayerController(const FObjectInitializer& ObjectInitializer
 		RootComponent->SetUsingAbsoluteRotation(true);
 	}
 
-	if(NetworkPhysicsCvars::EnableNetworkPhysicsPrediction == 1)
+	if (UPhysicsSettings::Get()->PhysicsPrediction.bEnablePhysicsPrediction)
 	{
 		bAsyncPhysicsTickEnabled = true;
 		AsyncPhysicsDataClass = UAsyncPhysicsData::StaticClass();
@@ -1677,7 +1660,7 @@ void APlayerController::ClientSetCameraFade_Implementation(bool bEnableFading, F
 
 void APlayerController::SendClientAdjustment()
 {
-	if(!NetworkPhysicsCvars::EnableNetworkPhysicsPrediction)
+	if(!UPhysicsSettings::Get()->PhysicsPrediction.bEnablePhysicsPrediction)
 	{
 		if (ServerFrameInfo.LastProcessedInputFrame != INDEX_NONE && ServerFrameInfo.LastProcessedInputFrame != ServerFrameInfo.LastSentLocalFrame)
 		{
@@ -5177,7 +5160,7 @@ void APlayerController::TickActor( float DeltaSeconds, ELevelTick TickType, FAct
 	// Clear old axis inputs since we are done with them. 
 	RotationInput = FRotator::ZeroRotator;
 
-	if(!!NetworkPhysicsCvars::EnableNetworkPhysicsPrediction && GetLocalRole() == ROLE_AutonomousProxy && bIsClient)
+	if(UPhysicsSettings::Get()->PhysicsPrediction.bEnablePhysicsPrediction && GetLocalRole() == ROLE_AutonomousProxy && bIsClient)
 	{
 		UpdateServerAsyncPhysicsTickOffset();
 	}
@@ -6158,7 +6141,7 @@ void APlayerController::AsyncPhysicsTickActor(float DeltaTime, float SimTime)
 {
 	Super::AsyncPhysicsTickActor(DeltaTime, SimTime);
 
-	if(NetworkPhysicsCvars::EnableNetworkPhysicsPrediction)
+	if(UPhysicsSettings::Get()->PhysicsPrediction.bEnablePhysicsPrediction)
 	{
 		//TODO: only kick this off if server and using this feature
 		if (IsLocalController()) { return; }
diff --git a/Engine/Source/Runtime/Experimental/Chaos/Private/Chaos/Collision/PBDCollisionConstraint.cpp b/Engine/Source/Runtime/Experimental/Chaos/Private/Chaos/Collision/PBDCollisionConstraint.cpp
--- a/Engine/Source/Runtime/Experimental/Chaos/Private/Chaos/Collision/PBDCollisionConstraint.cpp
+++ b/Engine/Source/Runtime/Experimental/Chaos/Private/Chaos/Collision/PBDCollisionConstraint.cpp
@@ -1051,7 +1051,7 @@ namespace Chaos
 	{
 		int32 MatchIndex = INDEX_NONE;
 
-		const bool bEnableNetworkPhysicsResim = FPhysicsSolverBase::IsNetworkPhysicsPredictionEnabled() && FPhysicsSolverBase::CanResimNetworkPhysicsPrediction();
+		const bool bEnableNetworkPhysicsResim = FPhysicsSolverBase::IsNetworkPhysicsPredictionEnabled() && FPhysicsSolverBase::IsPhysicsResimulationEnabled();
 		if (bChaos_Manifold_EnableFrictionRestore && !bEnableNetworkPhysicsResim)
 		{
 			const FManifoldPoint& ManifoldPoint = ManifoldPoints[ManifoldPointIndex];
diff --git a/Engine/Source/Runtime/Experimental/Chaos/Private/Chaos/CollisionOneShotManifolds.cpp b/Engine/Source/Runtime/Experimental/Chaos/Private/Chaos/CollisionOneShotManifolds.cpp
--- a/Engine/Source/Runtime/Experimental/Chaos/Private/Chaos/CollisionOneShotManifolds.cpp
+++ b/Engine/Source/Runtime/Experimental/Chaos/Private/Chaos/CollisionOneShotManifolds.cpp
@@ -775,7 +775,7 @@ namespace Chaos
 			const FReal Margin2 = Constraint.GetCollisionMargin1();
 			const FRigidTransform3 Convex2ToConvex1Transform = Convex2Transform.GetRelativeTransformNoScale(Convex1Transform);
 
-			const bool bEnableNetworkPhysicsResim = FPhysicsSolverBase::IsNetworkPhysicsPredictionEnabled() && FPhysicsSolverBase::CanResimNetworkPhysicsPrediction();
+			const bool bEnableNetworkPhysicsResim = FPhysicsSolverBase::IsNetworkPhysicsPredictionEnabled() && FPhysicsSolverBase::IsPhysicsResimulationEnabled();
 			
 			if (bEnableNetworkPhysicsResim)
 			{
diff --git a/Engine/Source/Runtime/Experimental/Chaos/Private/PBDRigidsSolver.cpp b/Engine/Source/Runtime/Experimental/Chaos/Private/PBDRigidsSolver.cpp
--- a/Engine/Source/Runtime/Experimental/Chaos/Private/PBDRigidsSolver.cpp
+++ b/Engine/Source/Runtime/Experimental/Chaos/Private/PBDRigidsSolver.cpp
@@ -950,12 +950,6 @@ namespace Chaos
 		SuspensionProxy->UpdateTargetOnPhysicsThread(this, TargetPos, Normal, Enabled);
 	}
 
-	CHAOS_API int32 RewindCaptureNumFrames = -1;
-	FAutoConsoleVariableRef CVarRewindCaptureNumFrames(TEXT("p.RewindCaptureNumFrames"),RewindCaptureNumFrames,TEXT("The number of frames to capture rewind for. Requires restart of solver"));
-
-	int32 UseResimCache = 0;
-	FAutoConsoleVariableRef CVarUseResimCache(TEXT("p.UseResimCache"),UseResimCache,TEXT("Whether resim uses cache to skip work, requires recreating world to take effect"));
-
 	void FPBDRigidsSolver::EnableRewindCapture(int32 NumFrames, bool InUseCollisionResimCache, TUniquePtr<IRewindCallback>&& RewindCallback)
 	{
 		//TODO: this function calls both internal and extrnal - sort of assumed during initialization. Should decide what thread it's called on and mark it as either external or internal
@@ -992,9 +986,12 @@ namespace Chaos
 		//todo: do we need this?
 		//MarshallingManager.Reset();
 
-		if(RewindCaptureNumFrames >= 0)
+		const bool PhysicsPredictionEnabled = FChaosSolversModule::GetModule()->GetSettingsProvider().GetPhysicsPredictionEnabled();
+		const int32 PhysicsHistoryLength = FChaosSolversModule::GetModule()->GetSettingsProvider().GetPhysicsHistoryCount();
+
+		if ((PhysicsPredictionEnabled || bUseCollisionResimCache) && PhysicsHistoryLength >= 0)
 		{
-			EnableRewindCapture(RewindCaptureNumFrames, bUseCollisionResimCache || !!UseResimCache);
+			EnableRewindCapture(PhysicsHistoryLength, bUseCollisionResimCache || PhysicsPredictionEnabled);
 		}
 
 		MEvolution->SetCaptureRewindDataFunction([this](const TParticleView<TPBDRigidParticles<FReal,3>>& ActiveParticles)
diff --git a/Engine/Source/Runtime/Experimental/Chaos/Private/RewindData.cpp b/Engine/Source/Runtime/Experimental/Chaos/Private/RewindData.cpp
--- a/Engine/Source/Runtime/Experimental/Chaos/Private/RewindData.cpp
+++ b/Engine/Source/Runtime/Experimental/Chaos/Private/RewindData.cpp
@@ -12,9 +12,6 @@
 namespace Chaos
 {
 
-int32 EnableResimCache = 1;
-FAutoConsoleVariableRef CVarEnableEnableResimCache(TEXT("p.EnableResimCache"), EnableResimCache, TEXT("If enabled, provides a resim cache to speed up certain computations"));
-
 FVec3 FGeometryParticleState::ZeroVector = FVec3(0);
 
 void FGeometryParticleStateBase::SyncSimWritablePropsFromSim(FDirtyPropData Manager,const TPBDRigidParticleHandle<FReal,3>& Rigid)
diff --git a/Engine/Source/Runtime/Experimental/Chaos/Public/Chaos/Framework/PhysicsSolverBase.h b/Engine/Source/Runtime/Experimental/Chaos/Public/Chaos/Framework/PhysicsSolverBase.h
--- a/Engine/Source/Runtime/Experimental/Chaos/Public/Chaos/Framework/PhysicsSolverBase.h
+++ b/Engine/Source/Runtime/Experimental/Chaos/Public/Chaos/Framework/PhysicsSolverBase.h
@@ -12,6 +12,7 @@
 #include "Containers/Queue.h"
 #include "Chaos/ChaosMarshallingManager.h"
 #include "Stats/Stats2.h"
+#include "ChaosSolversModule.h"
 
 #if WITH_CHAOS_VISUAL_DEBUGGER
 #include "ChaosVisualDebugger/ChaosVDContextProvider.h"
@@ -637,41 +638,46 @@ namespace Chaos
 		/**/
 		FReal GetLastDt() const { return MLastDt; }
 
-		/** Return the interpolation lerp in case the resim is off */
-		static float NetworkPhysicsInterpolationLerp()
+		/** Check if we can enable debugging informations for network physics */
+		static bool CanDebugNetworkPhysicsPrediction()
 		{
-			static IConsoleVariable* InterpLerpCVar = IConsoleManager::Get().FindConsoleVariable(TEXT("np2.NetworkPhysicsPredictionInterpLerp"));
-			return InterpLerpCVar ? InterpLerpCVar->GetFloat() : 0.1f;
-		}
+			static int32 DebugNetworkPhysicsPrediction = 0;
+			static FAutoConsoleVariableRef CVarDebugNetworkPhysicsPrediction(TEXT("np2.DebugNetworkPhysicsPrediction"), DebugNetworkPhysicsPrediction, TEXT("Debugs network physics prediction"));
 
-		/** Return the error position threshold to trigger a resim*/
-		static float NetworkPhysicsErrorThreshold()
-		{
-			static IConsoleVariable* ErrorThresholdCVar = IConsoleManager::Get().FindConsoleVariable(TEXT("np2.NetworkPhysicsPredictionErrorThreshold"));
-			return ErrorThresholdCVar ? ErrorThresholdCVar->GetFloat() : 10.0f;
+			return !!DebugNetworkPhysicsPrediction;
 		}
 
 		/** Check if network physics is enables or not */
 		static bool IsNetworkPhysicsPredictionEnabled()
 		{
-			static IConsoleVariable* EnableNetworkPhysicsCVar = IConsoleManager::Get().FindConsoleVariable(TEXT("np2.EnableNetworkPhysicsPrediction"));
-			return EnableNetworkPhysicsCVar && (EnableNetworkPhysicsCVar->GetInt() == 1);
+			const bool NetworkPhysicsEnabled = FChaosSolversModule::GetModule()->GetSettingsProvider().GetPhysicsPredictionEnabled();
+			return NetworkPhysicsEnabled;
 		}
 
-		/** Check if we can enable debugging informations for network physics */
-		static bool CanDebugNetworkPhysicsPrediction()
+		/** Check if resim is enabled for network physics */
+		static bool IsPhysicsResimulationEnabled()
 		{
-			static IConsoleVariable* DebugNetworkPhysicsCVar = IConsoleManager::Get().FindConsoleVariable(TEXT("np2.DebugNetworkPhysicsPrediction"));
-			return DebugNetworkPhysicsCVar && (DebugNetworkPhysicsCVar->GetInt() == 1);
+			const bool PhysicsResimulationEnabled = FChaosSolversModule::GetModule()->GetSettingsProvider().GetPhysicsResimulationEnabled();
+			return PhysicsResimulationEnabled;
 		}
 
-		/** Check if resim is enabled for network physics */
-		static bool CanResimNetworkPhysicsPrediction()
+		static float ResimulationErrorThreshold()
 		{
-			static IConsoleVariable* ResimNetworkPhysicsCVar = IConsoleManager::Get().FindConsoleVariable(TEXT("np2.ResimNetworkPhysicsPrediction"));
-			return ResimNetworkPhysicsCVar && (ResimNetworkPhysicsCVar->GetInt() == 1);
+			const float ResimulationErrorThreshold = FChaosSolversModule::GetModule()->GetSettingsProvider().GetResimulationErrorThreshold();
+			return ResimulationErrorThreshold;
+
 		}
 
+		/** Return the interpolation lerp in case the resim is off */
+		static float NetworkPhysicsInterpolationLerp()
+		{
+			static float NetworkPhysicsPredictionInterpLerp = 0.1f;
+			static FAutoConsoleVariableRef CVarNetworkPhysicsPredictionInterpLerp(TEXT("np2.NetworkPhysicsPredictionInterpLerp"), NetworkPhysicsPredictionInterpLerp, TEXT("State lerp value in between the target state and the current one in case resim is disabled or if the pawn is not possessed (continuous correction)"));
+
+			return NetworkPhysicsPredictionInterpLerp;
+		}
+
+
 	protected:
 		/** Mode that the results buffers should be set to (single, double, triple) */
 		EMultiBufferMode BufferMode;
diff --git a/Engine/Source/Runtime/Experimental/Chaos/Public/ChaosSolversModule.h b/Engine/Source/Runtime/Experimental/Chaos/Public/ChaosSolversModule.h
--- a/Engine/Source/Runtime/Experimental/Chaos/Public/ChaosSolversModule.h
+++ b/Engine/Source/Runtime/Experimental/Chaos/Public/ChaosSolversModule.h
@@ -35,6 +35,24 @@ public:
 	{
 		return 0.f;
 	}
+
+	virtual bool GetPhysicsPredictionEnabled() const
+	{
+		return false;
+	}
+	virtual bool GetPhysicsResimulationEnabled() const 
+	{
+		return false;
+	}
+	virtual float GetResimulationErrorThreshold() const
+	{
+		return 0;
+	}
+	virtual int32 GetPhysicsHistoryCount() const
+	{
+		return 0;
+	}
+	
 }; 
 
 namespace Chaos
diff --git a/Engine/Source/Runtime/Experimental/Chaos/Public/PBDRigidsSolver.h b/Engine/Source/Runtime/Experimental/Chaos/Public/PBDRigidsSolver.h
--- a/Engine/Source/Runtime/Experimental/Chaos/Public/PBDRigidsSolver.h
+++ b/Engine/Source/Runtime/Experimental/Chaos/Public/PBDRigidsSolver.h
@@ -59,8 +59,6 @@ namespace Chaos
 	class FGeometryParticleBuffer;
 	class FClusterUnionPhysicsProxy;
 
-	CHAOS_API extern int32 RewindCaptureNumFrames;
-
 	template <typename T,typename R,int d>
 	class ISpatialAccelerationCollection;
 
diff --git a/Engine/Source/Runtime/Experimental/Chaos/Public/RewindData.h b/Engine/Source/Runtime/Experimental/Chaos/Public/RewindData.h
--- a/Engine/Source/Runtime/Experimental/Chaos/Public/RewindData.h
+++ b/Engine/Source/Runtime/Experimental/Chaos/Public/RewindData.h
@@ -1009,8 +1009,6 @@ const T* ConstifyHelper(T* Ptr) { return Ptr; }
 template <typename T>
 T NoRefHelper(const T& Ref) { return Ref; }
 
-extern CHAOS_API int32 EnableResimCache;
-
 template <typename TVal>
 class TDirtyObjects
 {
@@ -1182,7 +1180,8 @@ public:
 
 	IResimCacheBase* GetCurrentStepResimCache() const
 	{
-		return !!EnableResimCache && bResimOptimization ? Managers[CurFrame].ExternalResimCache.Get() : nullptr;
+		const bool PhysicsPredictionEnabled = FPhysicsSolverBase::IsNetworkPhysicsPredictionEnabled();
+		return PhysicsPredictionEnabled && bResimOptimization ? Managers[CurFrame].ExternalResimCache.Get() : nullptr;
 	}
 
 	void CHAOS_API DumpHistory_Internal(const int32 FramePrintOffset, const FString& Filename = FString(TEXT("Dump")));
