TextureLayout3d - fix truncation warning

diff --git a/Engine/Source/Runtime/Engine/Public/TextureLayout3d.h b/Engine/Source/Runtime/Engine/Public/TextureLayout3d.h
--- a/Engine/Source/Runtime/Engine/Public/TextureLayout3d.h
+++ b/Engine/Source/Runtime/Engine/Public/TextureLayout3d.h
@@ -36,8 +36,8 @@ public:
 		bAllowShrink(bInAllowShrink)
 	{
 		check(MaxSizeX < USHRT_MAX && MaxSizeY < USHRT_MAX && MaxSizeZ < USHRT_MAX);
-		Nodes.Emplace(0, 0, 0, IntCastChecked<uint16>(MaxSizeX), IntCastChecked<uint16>(MaxSizeY), IntCastChecked<uint16>(MaxSizeZ), INDEX_NONE);
-		UnusedLeaves.Emplace(0, 0, 0, IntCastChecked<uint16>(MaxSizeX), IntCastChecked<uint16>(MaxSizeY), IntCastChecked<uint16>(MaxSizeZ), 0);
+		Nodes.Emplace(static_cast<uint16>(0), static_cast<uint16>(0), static_cast<uint16>(0), IntCastChecked<uint16>(MaxSizeX), IntCastChecked<uint16>(MaxSizeY), IntCastChecked<uint16>(MaxSizeZ), INDEX_NONE);
+		UnusedLeaves.Emplace(static_cast<uint16>(0), static_cast<uint16>(0), static_cast<uint16>(0), IntCastChecked<uint16>(MaxSizeX), IntCastChecked<uint16>(MaxSizeY), IntCastChecked<uint16>(MaxSizeZ), 0);
 	}
 
 	/**
