Add fallback default bone compression settings in case ACLPlugin is disabled

diff --git a/Engine/Config/BaseEngine.ini b/Engine/Config/BaseEngine.ini
--- a/Engine/Config/BaseEngine.ini
+++ b/Engine/Config/BaseEngine.ini
@@ -3203,8 +3203,13 @@ bStripAnimationDataOnDedicatedServer=False
 BoneTimecodeCustomAttributeNameSettings=(HourAttributeName="TCHour",MinuteAttributeName="TCMinute",SecondAttributeName="TCSecond",FrameAttributeName="TCFrame",SubframeAttributeName="TCSubframe",RateAttributeName="TCRate",TakenameAttributeName="Takename")
 
 [Animation.DefaultObjectSettings]
+; Fallback is used only if the default codec isn't found (e.g. if the plugin isn't enabled)
+BoneCompressionSettingsFallback="/Engine/Animation/DefaultAnimBoneCompressionSettings"
+; Default codec assigned when new sequences are imported/created
 BoneCompressionSettings="/ACLPlugin/ACLAnimBoneCompressionSettings"
+; Default codec for the animation recorder
 AnimationRecorderBoneCompressionSettings="/Engine/Animation/DefaultRecorderBoneCompression"
+; Default codec assigned when new sequences are imported/created
 CurveCompressionSettings="/Engine/Animation/DefaultAnimCurveCompressionSettings"
 
 [/Script/Engine.MeshSimplificationSettings]
diff --git a/Engine/Source/Runtime/Engine/Private/Animation/AnimationUtils.cpp b/Engine/Source/Runtime/Engine/Private/Animation/AnimationUtils.cpp
--- a/Engine/Source/Runtime/Engine/Private/Animation/AnimationUtils.cpp
+++ b/Engine/Source/Runtime/Engine/Private/Animation/AnimationUtils.cpp
@@ -881,14 +881,22 @@ static void EnsureDependenciesAreLoaded(UObject* Object)
 
 
 
-UObject* GetDefaultAnimationCompressionSettings(const TCHAR* IniValueName)
+UObject* GetDefaultAnimationCompressionSettings(const TCHAR* IniValueName, bool bIsFatal)
 {
 	FConfigSection* AnimDefaultObjectSettingsSection = GConfig->GetSectionPrivate(TEXT("Animation.DefaultObjectSettings"), false, true, GEngineIni);
 	const FConfigValue* Value = AnimDefaultObjectSettingsSection != nullptr ? AnimDefaultObjectSettingsSection->Find(IniValueName) : nullptr;
 
 	if (Value == nullptr)
 	{
-		UE_LOG(LogAnimationCompression, Fatal, TEXT("Couldn't find default compression setting for '%s' under '[Animation.DefaultObjectSettings]'"), IniValueName);
+		if (bIsFatal)
+		{
+			UE_LOG(LogAnimationCompression, Fatal, TEXT("Couldn't find default compression setting for '%s' under '[Animation.DefaultObjectSettings]'"), IniValueName);
+		}
+		else
+		{
+			UE_LOG(LogAnimationCompression, Warning, TEXT("Couldn't find default compression setting for '%s' under '[Animation.DefaultObjectSettings]'"), IniValueName);
+		}
+
 		return nullptr;
 	}
 
@@ -897,7 +905,16 @@ UObject* GetDefaultAnimationCompressionSettings(const TCHAR* IniValueName)
 
 	if (DefaultCompressionSettings == nullptr)
 	{
-		UE_LOG(LogAnimationCompression, Fatal, TEXT("Couldn't load default compression settings asset with path '%s'"), *CompressionSettingsName);
+		if (bIsFatal)
+		{
+			UE_LOG(LogAnimationCompression, Fatal, TEXT("Couldn't load default compression settings asset with path '%s'"), *CompressionSettingsName);
+		}
+		else
+		{
+			UE_LOG(LogAnimationCompression, Warning, TEXT("Couldn't load default compression settings asset with path '%s'"), *CompressionSettingsName);
+		}
+
+		return nullptr;
 	}
 
 	// Force load the default settings and all its dependencies just in case it hasn't happened yet
@@ -923,7 +940,12 @@ UAnimBoneCompressionSettings* FAnimationUtils::GetDefaultAnimationBoneCompressio
 {
 	if (DefaultBoneCompressionSettings == nullptr)
 	{
-		DefaultBoneCompressionSettings = Cast<UAnimBoneCompressionSettings>(GetDefaultAnimationCompressionSettings(TEXT("BoneCompressionSettings")));
+		DefaultBoneCompressionSettings = Cast<UAnimBoneCompressionSettings>(GetDefaultAnimationCompressionSettings(TEXT("BoneCompressionSettings"), false));
+
+		if (DefaultBoneCompressionSettings == nullptr)
+		{
+			DefaultBoneCompressionSettings = Cast<UAnimBoneCompressionSettings>(GetDefaultAnimationCompressionSettings(TEXT("BoneCompressionSettingsFallback"), true));
+		}
 	}
 
 	return DefaultBoneCompressionSettings;
@@ -933,7 +955,7 @@ UAnimBoneCompressionSettings* FAnimationUtils::GetDefaultAnimationRecorderBoneCo
 {
 	if (DefaultRecorderBoneCompressionSettings == nullptr)
 	{
-		DefaultRecorderBoneCompressionSettings = Cast<UAnimBoneCompressionSettings>(GetDefaultAnimationCompressionSettings(TEXT("AnimationRecorderBoneCompressionSettings")));
+		DefaultRecorderBoneCompressionSettings = Cast<UAnimBoneCompressionSettings>(GetDefaultAnimationCompressionSettings(TEXT("AnimationRecorderBoneCompressionSettings"), true));
 	}
 
 	return DefaultRecorderBoneCompressionSettings;
@@ -943,7 +965,7 @@ UAnimCurveCompressionSettings* FAnimationUtils::GetDefaultAnimationCurveCompress
 {
 	if (DefaultCurveCompressionSettings == nullptr)
 	{
-		DefaultCurveCompressionSettings = Cast<UAnimCurveCompressionSettings>(GetDefaultAnimationCompressionSettings(TEXT("CurveCompressionSettings")));
+		DefaultCurveCompressionSettings = Cast<UAnimCurveCompressionSettings>(GetDefaultAnimationCompressionSettings(TEXT("CurveCompressionSettings"), true));
 	}
 
 	return DefaultCurveCompressionSettings;
