World Partition Streaming Generation: reduce the number of containers per-instance data and centralize code for per-instance data inheritance.

diff --git a/Engine/Source/Runtime/Engine/Private/WorldPartition/WorldPartitionStreamingGeneration.cpp b/Engine/Source/Runtime/Engine/Private/WorldPartition/WorldPartitionStreamingGeneration.cpp
--- a/Engine/Source/Runtime/Engine/Private/WorldPartition/WorldPartitionStreamingGeneration.cpp
+++ b/Engine/Source/Runtime/Engine/Private/WorldPartition/WorldPartitionStreamingGeneration.cpp
@@ -7,10 +7,12 @@
 
 #include "Editor.h"
 #include "Algo/Transform.h"
+#include "Algo/Unique.h"
 #include "Engine/LevelScriptBlueprint.h"
 #include "ActorReferencesUtils.h"
 #include "Misc/PackageName.h"
 #include "ReferenceCluster.h"
+#include "Misc/HashBuilder.h"
 #include "Misc/Paths.h"
 #include "ProfilingDebugging/ScopedTimers.h"
 #include "WorldPartition/WorldPartitionRuntimeHash.h"
@@ -160,7 +162,7 @@ class FWorldPartitionStreamingGenerator
 					}
 
 					FActorSetInstance& ActorSetInstance = ActorSetInstances.Emplace_GetRef();
-					const FContainerInstanceDescriptor::FPerInstanceData& PerInstanceData = ContainerInstanceDescriptor.PerInstanceData.FindChecked(ReferenceActorDescView.GetGuid());
+					const FContainerInstanceDescriptor::FPerInstanceData& PerInstanceData = ContainerInstanceDescriptor.GetInstanceData(ReferenceActorDescView.GetGuid());
 				
 					ActorSetInstance.ContainerInstance = &ActorSetContainer;
 					ActorSetInstance.ActorSet = &ActorSet;
@@ -264,9 +266,6 @@ class FWorldPartitionStreamingGenerator
 		FTransform Transform;
 		const UActorDescContainer* Container;
 		EContainerClusterMode ClusterMode;
-		TArray<FName> RuntimeDataLayers;
-		FName RuntimeGrid;
-		bool bIsSpatiallyLoaded;
 		FString OwnerName;
 		FActorContainerID ID;
 		FActorContainerID ParentID;
@@ -278,9 +277,54 @@ class FWorldPartitionStreamingGenerator
 			bool bIsSpatiallyLoaded;
 			FName RuntimeGrid;
 			TArray<FName> DataLayers;
+
+			inline bool operator==(const FPerInstanceData& Other) const
+			{
+				// Assumes DataLayers are sorted
+				return 
+					(bIsSpatiallyLoaded == Other.bIsSpatiallyLoaded) && 
+					(RuntimeGrid == Other.RuntimeGrid) && 
+					(DataLayers == Other.DataLayers);
+			}
+
+			inline bool operator!=(const FPerInstanceData& Other) const
+			{
+				return !operator==(Other);
+			}
+
+			friend inline uint32 GetTypeHash(const FPerInstanceData& PerInstanceData)
+			{
+				FHashBuilder HashBuilder;
+				HashBuilder << PerInstanceData.bIsSpatiallyLoaded << PerInstanceData.RuntimeGrid << PerInstanceData.DataLayers;
+				return HashBuilder.GetHash();
+			}
 		};
 
-		TMap<FGuid, FPerInstanceData> PerInstanceData;
+		FPerInstanceData& GetInstanceData(const FGuid& ActorGuid)
+		{
+			if (const FSetElementId* PerInstanceDataId = PerInstanceData.Find(ActorGuid))
+			{
+				check(UniquePerInstanceData.IsValidId(*PerInstanceDataId));
+				return UniquePerInstanceData[*PerInstanceDataId];
+			}
+
+			return InstanceData;
+		}
+
+		const FPerInstanceData& GetInstanceData(const FGuid& ActorGuid) const
+		{
+			if (const FSetElementId* PerInstanceDataId = PerInstanceData.Find(ActorGuid))
+			{
+				check(UniquePerInstanceData.IsValidId(*PerInstanceDataId));
+				return UniquePerInstanceData[*PerInstanceDataId];
+			}
+
+			return InstanceData;
+		}
+
+		FPerInstanceData InstanceData;
+		TSet<FPerInstanceData> UniquePerInstanceData;
+		TMap<FGuid, FSetElementId> PerInstanceData;
 	};
 
 	void ResolveRuntimeSpatiallyLoaded(FWorldPartitionActorDescView& ActorDescView)
@@ -301,8 +345,6 @@ class FWorldPartitionStreamingGenerator
 
 	void ResolveRuntimeDataLayers(FWorldPartitionActorDescView& ActorDescView, const FActorDescViewMap& ActorDescViewMap)
 	{
-		const UDataLayerManager* DataLayerManager = WorldPartitionContext ? WorldPartitionContext->GetDataLayerManager() : nullptr;
-
 		// Resolve DataLayerInstanceNames of ActorDescView only when necessary (i.e. when container is a template)
 		if (!ActorDescView.GetActorDesc()->HasResolvedDataLayerInstanceNames())
 		{
@@ -432,25 +474,29 @@ class FWorldPartitionStreamingGenerator
 
 	void CreateActorDescriptorViewsRecursive(const FContainerInstanceDescriptor& InContainerInstanceDescriptor)
 	{
-		// Inherited parent properties logic
-		auto InheritParentContainerProperties = [](const FActorContainerID& InParentContainerID, bool bInParentIsSpatiallyLoaded, FName InParentRuntimeGrid, const TArray<FName>& InParentRuntimeDataLayers, bool& bInOutIsSpatiallyLoaded, FName& InOutRuntimeGrid, TArray<FName>& InOutRuntimeDataLayers)
+		// Inherited parent per-instance data logic
+		auto InheritParentContainerPerInstanceData = [&InContainerInstanceDescriptor](const FContainerInstanceDescriptor::FPerInstanceData& InParentPerInstanceData, const FWorldPartitionActorDescView& InActorDescView)
 		{
+			FContainerInstanceDescriptor::FPerInstanceData ResultPerInstanceData;
+
 			// Apply AND logic on spatially loaded flag
-			bInOutIsSpatiallyLoaded &= bInParentIsSpatiallyLoaded;
+			ResultPerInstanceData.bIsSpatiallyLoaded = InActorDescView.GetIsSpatiallyLoaded() && InParentPerInstanceData.bIsSpatiallyLoaded;
 
 			// Runtime grid is only inherited from the main world, since level instance doesn't support setting this value on actors
-			if (InParentContainerID.IsMainContainer())
-			{
-				InOutRuntimeGrid = InParentRuntimeGrid;
-			}
+			ResultPerInstanceData.RuntimeGrid = InContainerInstanceDescriptor.ID.IsMainContainer() ? InActorDescView.GetRuntimeGrid() : InParentPerInstanceData.RuntimeGrid;
 
 			// Data layers are accumulated down the hierarchy chain, since level instances supports data layers assignation on actors
-			if (InParentRuntimeDataLayers.Num())
+			ResultPerInstanceData.DataLayers = InActorDescView.GetRuntimeDataLayerInstanceNames();
+			ResultPerInstanceData.DataLayers.Append(InParentPerInstanceData.DataLayers);
+			ResultPerInstanceData.DataLayers.Sort(FNameFastLess());
+
+			if (InParentPerInstanceData.DataLayers.Num())
 			{
-				TSet<FName> RuntimeDataLayers(InOutRuntimeDataLayers);
-				RuntimeDataLayers.Append(InParentRuntimeDataLayers);
-				InOutRuntimeDataLayers = RuntimeDataLayers.Array();
+				// Remove potential duplicates from sorted data layers array
+				ResultPerInstanceData.DataLayers.SetNum(Algo::Unique(ResultPerInstanceData.DataLayers));
 			}
+
+			return ResultPerInstanceData;
 		};
 
 		// ContainerDescriptor may be reallocated after this scope
@@ -501,63 +547,39 @@ class FWorldPartitionStreamingGenerator
 			for (const FWorldPartitionActorDescView& ContainerInstanceView : ContainerInstanceViews)
 			{
 				FWorldPartitionActorDesc::FContainerInstance SubContainerInstance;
-				if (!ContainerInstanceView.GetContainerInstance(SubContainerInstance))
+				if (!ContainerInstanceView.GetContainerInstance(SubContainerInstance) || !SubContainerInstance.Container)
 				{
 					ErrorHandler->OnLevelInstanceInvalidWorldAsset(ContainerInstanceView, ContainerInstanceView.GetContainerPackage(), IStreamingGenerationErrorHandler::ELevelInstanceInvalidReason::WorldAssetHasInvalidContainer);
 					continue;
 				}
 
-				if (ContainerInstancesStack.Contains(SubContainerInstance.Container->ContainerPackageName))
+				bool bContainerWasAlreadyInSet;
+				ContainerInstancesStack.Add(SubContainerInstance.Container->ContainerPackageName, &bContainerWasAlreadyInSet);
+
+				if (bContainerWasAlreadyInSet)
 				{
 					ErrorHandler->OnLevelInstanceInvalidWorldAsset(ContainerInstanceView, ContainerInstanceView.GetContainerPackage(), IStreamingGenerationErrorHandler::ELevelInstanceInvalidReason::CirculalReference);
 					continue;
 				}
 
-				const FGuid ActorGuid = ContainerInstanceView.GetGuid();
-				const FActorContainerID ContainerID(InContainerInstanceDescriptor.ID, ActorGuid);
-						
-				TMap<FActorContainerID, TSet<FGuid>> ContainerInstanceFilter;
-				if (InContainerInstanceDescriptor.ID.IsMainContainer() && ContainerInstanceView.IsContainerFilter())
-				{
-					if (const FWorldPartitionActorFilter* ContainerFilter = ContainerInstanceView.GetContainerFilter())
-					{
-						UWorld* OwningWorld = WorldPartitionContext ? WorldPartitionContext->GetWorld() : nullptr;
-						if (UWorldPartitionSubsystem* WorldPartitionSubsystem = UWorld::GetSubsystem<UWorldPartitionSubsystem>(OwningWorld))
-						{
-							ContainerInstanceFilter = WorldPartitionSubsystem->GetFilteredActorsPerContainer(ContainerID, ContainerInstanceView.GetContainerPackage().ToString(), *ContainerFilter);
-						}
-					}
-				}
-
-				ContainerInstancesStack.Add(SubContainerInstance.Container->ContainerPackageName);
-
-				check(SubContainerInstance.Container);
 				FContainerInstanceDescriptor SubContainerInstanceDescriptor;
-
-				// Inherit fields from FContainerInstance
-				SubContainerInstanceDescriptor.ID = ContainerID;
+				SubContainerInstanceDescriptor.ID = FActorContainerID(InContainerInstanceDescriptor.ID, ContainerInstanceView.GetGuid());
 				SubContainerInstanceDescriptor.Container = SubContainerInstance.Container;
 				SubContainerInstanceDescriptor.Transform = SubContainerInstance.Transform * InContainerInstanceDescriptor.Transform;
-				SubContainerInstanceDescriptor.FilteredActors = MoveTemp(ContainerInstanceFilter);
 				SubContainerInstanceDescriptor.ParentID = InContainerInstanceDescriptor.ID;
 				SubContainerInstanceDescriptor.OwnerName = *ContainerInstanceView.GetActorLabelOrName().ToString();
 				// Since Content Bundles streaming generation happens in its own context, all actor set instances must have the same content bundle GUID for now, so Level Instances
 				// placed inside a Content Bundle will propagate their Content Bundle GUID to child instances.
 				SubContainerInstanceDescriptor.ContentBundleID = InContainerInstanceDescriptor.ContentBundleID;
-				SubContainerInstanceDescriptor.bIsSpatiallyLoaded = InContainerInstanceDescriptor.bIsSpatiallyLoaded;
-				SubContainerInstanceDescriptor.RuntimeGrid = InContainerInstanceDescriptor.RuntimeGrid;
-				SubContainerInstanceDescriptor.RuntimeDataLayers = InContainerInstanceDescriptor.RuntimeDataLayers;
-
-				// Inherit parent container properties
-				InheritParentContainerProperties(
-					InContainerInstanceDescriptor.ID,
-					ContainerInstanceView.GetIsSpatiallyLoaded(),
-					ContainerInstanceView.GetRuntimeGrid(),
-					ContainerInstanceView.GetRuntimeDataLayerInstanceNames(),
-					SubContainerInstanceDescriptor.bIsSpatiallyLoaded,
-					SubContainerInstanceDescriptor.RuntimeGrid,
-					SubContainerInstanceDescriptor.RuntimeDataLayers
-				);
+				SubContainerInstanceDescriptor.InstanceData = InheritParentContainerPerInstanceData(InContainerInstanceDescriptor.InstanceData, ContainerInstanceView);
+
+				if (WorldPartitionSubsystem && InContainerInstanceDescriptor.ID.IsMainContainer() && ContainerInstanceView.IsContainerFilter())
+				{
+					if (const FWorldPartitionActorFilter* ContainerFilter = ContainerInstanceView.GetContainerFilter())
+					{						
+						SubContainerInstanceDescriptor.FilteredActors = WorldPartitionSubsystem->GetFilteredActorsPerContainer(SubContainerInstanceDescriptor.ID, ContainerInstanceView.GetContainerPackage().ToString(), *ContainerFilter);
+					}
+				}
 
 				CreateActorDescriptorViewsRecursive(SubContainerInstanceDescriptor);
 
@@ -577,24 +599,15 @@ class FWorldPartitionStreamingGenerator
 
 		// Apply per-instance data
 		ContainerInstanceDescriptor.PerInstanceData.Reserve(ContainerDescriptor.ActorDescViewMap.Num());
-		ContainerDescriptor.ActorDescViewMap.ForEachActorDescView([&ContainerInstanceDescriptor, &InheritParentContainerProperties](FWorldPartitionActorDescView& ActorDescView)
+		ContainerDescriptor.ActorDescViewMap.ForEachActorDescView([&ContainerInstanceDescriptor, &InheritParentContainerPerInstanceData](FWorldPartitionActorDescView& ActorDescView)
 		{
-			FContainerInstanceDescriptor::FPerInstanceData& PerInstanceData = ContainerInstanceDescriptor.PerInstanceData.Add(ActorDescView.GetGuid());
-
-			PerInstanceData.bIsSpatiallyLoaded = ContainerInstanceDescriptor.bIsSpatiallyLoaded;
-			PerInstanceData.RuntimeGrid = ContainerInstanceDescriptor.RuntimeGrid;
-			PerInstanceData.DataLayers = ContainerInstanceDescriptor.RuntimeDataLayers;
-
-			// Inherit parent container properties
-			InheritParentContainerProperties(
-				ContainerInstanceDescriptor.ID,
-				ActorDescView.GetIsSpatiallyLoaded(),
-				ActorDescView.GetRuntimeGrid(),
-				ActorDescView.GetRuntimeDataLayerInstanceNames(),
-				PerInstanceData.bIsSpatiallyLoaded,
-				PerInstanceData.RuntimeGrid,
-				PerInstanceData.DataLayers
-			);
+			const FContainerInstanceDescriptor::FPerInstanceData PerInstanceData = InheritParentContainerPerInstanceData(ContainerInstanceDescriptor.InstanceData, ActorDescView);
+
+			if (PerInstanceData != ContainerInstanceDescriptor.InstanceData)
+			{
+				const FSetElementId UniquePerInstanceDataId = ContainerInstanceDescriptor.UniquePerInstanceData.Add(PerInstanceData);
+				ContainerInstanceDescriptor.PerInstanceData.Emplace(ActorDescView.GetGuid(), UniquePerInstanceDataId);
+			}
 		});
 
 		// Validate container instance, fixing anything illegal, etc.
@@ -610,10 +623,11 @@ class FWorldPartitionStreamingGenerator
 
 		FContainerInstanceDescriptor MainContainerInstance;
 		MainContainerInstance.Container = InContainer;
-		MainContainerInstance.bIsSpatiallyLoaded = true; // Since we apply AND logic on spatially loaded flag recursively, startup value must be true
 		MainContainerInstance.ClusterMode = EContainerClusterMode::Partitioned;
 		MainContainerInstance.OwnerName = TEXT("MainContainer");
 		MainContainerInstance.ContentBundleID = InContainer->GetContentBundleGuid();
+		MainContainerInstance.InstanceData.bIsSpatiallyLoaded = true; // Since we apply AND logic on spatially loaded flag recursively, startup value must be true
+
 		CreateActorDescriptorViewsRecursive(MainContainerInstance);
 	}
 
@@ -662,7 +676,7 @@ class FWorldPartitionStreamingGenerator
 				}
 
 				// Validate data layers
-				if (const UDataLayerManager* DataLayerManager = WorldPartitionContext ? WorldPartitionContext->GetDataLayerManager() : nullptr)
+				if (DataLayerManager)
 				{
 					DataLayerManager->ForEachDataLayerInstance([this](const UDataLayerInstance* DataLayerInstance)
 					{
@@ -964,7 +978,7 @@ class FWorldPartitionStreamingGenerator
 
 			ContainerDescriptor.ActorDescViewMap.ForEachActorDescView([this, &ContainerInstanceDescriptor, &NbErrorsDetected, PassType](FWorldPartitionActorDescView& ActorDescView)
 			{
-				FContainerInstanceDescriptor::FPerInstanceData& PerInstanceData = ContainerInstanceDescriptor.PerInstanceData.FindChecked(ActorDescView.GetGuid());
+				FContainerInstanceDescriptor::FPerInstanceData& PerInstanceData = ContainerInstanceDescriptor.GetInstanceData(ActorDescView.GetGuid());
 
 PRAGMA_DISABLE_DEPRECATION_WARNINGS
 				if (ActorDescView.ShouldValidateRuntimeGrid())
@@ -1036,7 +1050,11 @@ public:
 		, IsValidGrid(Params.IsValidGrid)
 		, ErrorHandler(Params.ErrorHandler)
 		, ActorGuidsToContainerMap(Params.ActorGuidsToContainerMap)
-	{}
+	{
+		UWorld* OwningWorld = WorldPartitionContext ? WorldPartitionContext->GetWorld() : nullptr;
+		WorldPartitionSubsystem = OwningWorld ? UWorld::GetSubsystem<UWorldPartitionSubsystem>(OwningWorld) : nullptr;
+		DataLayerManager = OwningWorld ? WorldPartitionContext->GetDataLayerManager() : nullptr;
+	}
 
 	void PreparationPhase(const UActorDescContainer* Container)
 	{
@@ -1160,13 +1178,13 @@ public:
 
 	TArray<const UDataLayerInstance*> GetRuntimeDataLayerInstances(const TArray<FName>& RuntimeDataLayers) const
 	{
-		static TArray<const UDataLayerInstance*> EmptyArray;
-		const UDataLayerManager* DataLayerManager = WorldPartitionContext ? WorldPartitionContext->GetDataLayerManager() : nullptr;
-		return DataLayerManager ? DataLayerManager->GetRuntimeDataLayerInstances(RuntimeDataLayers) : EmptyArray;
+		return DataLayerManager ? DataLayerManager->GetRuntimeDataLayerInstances(RuntimeDataLayers) : TArray<const UDataLayerInstance*>();
 	}
 
 private:
 	const UWorldPartition* WorldPartitionContext;
+	UWorldPartitionSubsystem* WorldPartitionSubsystem;
+	const UDataLayerManager* DataLayerManager;
 	bool bEnableStreaming;
 	FActorDescList* ModifiedActorsDescList;
 	TArray<TSubclassOf<AActor>> FilteredClasses;
