Fix linux build warning.

diff --git a/Engine/Source/Runtime/Engine/Private/WorldPartition/WorldPartitionStreamingGeneration.cpp b/Engine/Source/Runtime/Engine/Private/WorldPartition/WorldPartitionStreamingGeneration.cpp
--- a/Engine/Source/Runtime/Engine/Private/WorldPartition/WorldPartitionStreamingGeneration.cpp
+++ b/Engine/Source/Runtime/Engine/Private/WorldPartition/WorldPartitionStreamingGeneration.cpp
@@ -292,10 +292,10 @@ class FWorldPartitionStreamingGenerator
 				return !operator==(Other);
 			}
 
-			friend inline uint32 GetTypeHash(const FPerInstanceData& PerInstanceData)
+			friend inline uint32 GetTypeHash(const FPerInstanceData& InPerInstanceData)
 			{
 				FHashBuilder HashBuilder;
-				HashBuilder << PerInstanceData.bIsSpatiallyLoaded << PerInstanceData.RuntimeGrid << PerInstanceData.DataLayers;
+				HashBuilder << InPerInstanceData.bIsSpatiallyLoaded << InPerInstanceData.RuntimeGrid << InPerInstanceData.DataLayers;
 				return HashBuilder.GetHash();
 			}
 		};
