Make PhysicsReplication work with FPhysicsObject instead of just BodyInstance. Make it backwards compatible with existing project overrides.

diff --git a/Engine/Source/Runtime/Engine/Private/PrimitiveComponentPhysics.cpp b/Engine/Source/Runtime/Engine/Private/PrimitiveComponentPhysics.cpp
--- a/Engine/Source/Runtime/Engine/Private/PrimitiveComponentPhysics.cpp
+++ b/Engine/Source/Runtime/Engine/Private/PrimitiveComponentPhysics.cpp
@@ -28,6 +28,16 @@ DECLARE_CYCLE_STAT(TEXT("PrimComp SetCollisionProfileName"), STAT_PrimComp_SetCo
 	#define WarnInvalidPhysicsOperations(Text, BodyInstance, BoneName)
 #endif
 
+namespace PrimitiveComponentCVars
+{
+	bool bReplicatePhysicsObject = 1;
+	static FAutoConsoleVariableRef CVarReplicatePhysicsObject(
+		TEXT("p.PrimitiveComponent.ReplicatePhysicsObject"),
+		bReplicatePhysicsObject,
+		TEXT("When a primitive component has no BodyInstance, allow replication based on PhysicsObject\n"),
+		ECVF_Default);
+}
+
 void UPrimitiveComponent::SetRigidBodyReplicatedTarget(FRigidBodyState& UpdatedState, FName BoneName, int32 ServerFrame, int32 ServerHandle)
 {
 	if (UWorld* World = GetWorld())
@@ -36,12 +46,18 @@ void UPrimitiveComponent::SetRigidBodyReplicatedTarget(FRigidBodyState& UpdatedS
 		{
 			if (FPhysicsReplication* PhysicsReplication = PhysScene->GetPhysicsReplication())
 			{
-				FBodyInstance* BI = GetBodyInstance(BoneName);
-				if (BI && BI->IsValidBodyInstance())
+				// If we are not allowed to replicate physics objects,
+				// don't set replicated target unless we have a BodyInstance.
+				if (PrimitiveComponentCVars::bReplicatePhysicsObject == false)
 				{
-					PhysicsReplication->SetReplicatedTarget(this, BoneName, UpdatedState, ServerFrame);
-					BI->GetPhysicsActorHandle();// ->GetGameThreadAPI().SetParticleID(Chaos::FParticleID{ ServerPhysicsHandle, INDEX_NONE });
+					FBodyInstance* BI = GetBodyInstance(BoneName);
+					if (BI == nullptr || !BI->IsValidBodyInstance())
+					{
+						return;
+					}
 				}
+				
+				PhysicsReplication->SetReplicatedTarget(this, BoneName, UpdatedState, ServerFrame);
 			}
 		}
 	}
@@ -49,12 +65,33 @@ void UPrimitiveComponent::SetRigidBodyReplicatedTarget(FRigidBodyState& UpdatedS
 
 bool UPrimitiveComponent::GetRigidBodyState(FRigidBodyState& OutState, FName BoneName)
 {
+	// If we have a BodyInstance, use it
+	//
+	// TODO: Remove this code path
 	FBodyInstance* BI = GetBodyInstance(BoneName);
 	if (BI)
 	{
 		return BI->GetRigidBodyState(OutState);
 	}
 
+	// If we don't, get data from the physics object.
+	//
+	// TODO: Add support for multiple physics objects
+	if (PrimitiveComponentCVars::bReplicatePhysicsObject)
+	{
+		if (Chaos::FPhysicsObject* PhysicsObject = GetPhysicsObjectByName(BoneName))
+		{
+			FLockedReadPhysicsObjectExternalInterface Interface = FPhysicsObjectExternalInterface::LockRead(PhysicsObject);
+			const FTransform Transform = Interface->GetTransform(PhysicsObject);
+			OutState.Position = Transform.GetLocation();
+			OutState.Quaternion = Transform.GetRotation();
+			OutState.LinVel = Interface->GetV(PhysicsObject);
+			OutState.AngVel = Interface->GetW(PhysicsObject);
+			OutState.Flags = (Interface->AreAllSleeping({ PhysicsObject }) ? ERigidBodyFlags::Sleeping : ERigidBodyFlags::None);
+			return true;
+		}
+	}
+
 	return false;
 }
 
