moved FindAndPrintStaleReferencesToObjects from Engine to ReferenceChainSearch. Makes the function usable by CoreUObject
https://p4-swarm.epicgames.net/reviews/25519109

diff --git a/Engine/Source/Runtime/CoreUObject/Private/UObject/ReferenceChainSearch.cpp b/Engine/Source/Runtime/CoreUObject/Private/UObject/ReferenceChainSearch.cpp
--- a/Engine/Source/Runtime/CoreUObject/Private/UObject/ReferenceChainSearch.cpp
+++ b/Engine/Source/Runtime/CoreUObject/Private/UObject/ReferenceChainSearch.cpp
@@ -1734,4 +1734,214 @@ void FReferenceChainSearch::Cleanup()
 		delete ObjectToInfoPair.Value;
 	}
 	ObjectToInfoMap.Empty();
+}
+
+static FORCEINLINE bool HasGarbageCollectionKeepFlags(FGCObjectInfo* ObjectInfo)
+{
+	return ObjectInfo && (ObjectInfo->HasAnyFlags(GARBAGE_COLLECTION_KEEPFLAGS) || ObjectInfo->HasAnyInternalFlags(EInternalObjectFlags::GarbageCollectionKeepFlags | EInternalObjectFlags::RootSet));
+}
+
+static bool PrintStaleReferenceChainsAndFindReferencingObjects(UObject* ObjectToFindReferencesTo, FReferenceChainSearch& RefChainSearch, FGCObjectInfo*& OutGarbageObject, FGCObjectInfo*& OutReferencingObject, ELogVerbosity::Type Verbosity)
+{
+#if !NO_LOGGING
+	if (Verbosity != ELogVerbosity::NoLogging)
+	{
+		FMsg::Logf(__FILE__, __LINE__, LogLoad.GetCategoryName(), Verbosity, TEXT("Printing reference chains leading to %s: "), *ObjectToFindReferencesTo->GetFullName());
+	}
+#endif
+	return RefChainSearch.PrintResults([&ObjectToFindReferencesTo, &OutGarbageObject, &OutReferencingObject](FReferenceChainSearch::FCallbackParams& Params)
+		{
+			check(Params.Object);
+			if (!Params.Object->IsValid())
+			{
+				// We may find many chains that lead to the leak but for brevity we only report the first one in the fatal error below
+				if (!OutGarbageObject)
+				{
+					OutGarbageObject = Params.Object;
+					OutReferencingObject = Params.Referencer;
+				}
+				return true;
+			}
+			else
+			{
+				UObject* ResolvedObject = Params.Object->TryResolveObject();
+				if (ResolvedObject != ObjectToFindReferencesTo)
+				{
+					// Keep track of the last referencing object, we may need it if there's no garbage referencers
+					OutReferencingObject = Params.Object;
+				}
+				return true;
+			}
+		}, false, ObjectToFindReferencesTo) != 0;
+}
+
+static FString GetPathToStaleObjectReferencer(UObject* ObjectToFindReferencesTo, FReferenceChainSearch& RefChainSearch)
+{
+	bool bFirstGarbageObjectFound = false;
+	FString PathToCulprit = RefChainSearch.GetRootPath([&bFirstGarbageObjectFound, ObjectToFindReferencesTo](FReferenceChainSearch::FCallbackParams& Params)
+		{
+			check(Params.Object);
+			// Mark the first garbage reference or the target object reference as the culprit
+			if ((!Params.Object->IsValid() || Params.Object->TryResolveObject() == ObjectToFindReferencesTo) && !bFirstGarbageObjectFound)
+			{
+				bFirstGarbageObjectFound = true;
+				check(Params.Out);
+				Params.Out->Logf(ELogVerbosity::Log, TEXT("%s    ^^^^^^^^^^^^^^^^^^^^^^^^^^^^^^^^^^^^^^^^^^^^^^^^^^^^^^^^^^^^^^"), FCString::Spc(Params.Indent));
+				Params.Out->Logf(ELogVerbosity::Log, TEXT("%s    ^ This reference is preventing the old %s from being GC'd ^"), FCString::Spc(Params.Indent), *ObjectToFindReferencesTo->GetClass()->GetName());
+			}
+			return true;
+		});
+	return PathToCulprit;
+}
+
+FString FReferenceChainSearch::FindAndPrintStaleReferencesToObject(UObject* ObjectToFindReferencesTo, EPrintStaleReferencesOptions Options)
+{
+	return FindAndPrintStaleReferencesToObjects(TConstArrayView<UObject*>(&ObjectToFindReferencesTo, 1), Options)[0];
+}
+
+TArray<FString> FReferenceChainSearch::FindAndPrintStaleReferencesToObjects(TConstArrayView<UObject*> ObjectsToFindReferencesTo, EPrintStaleReferencesOptions Options)
+{
+	TArray<FString> Paths;
+	checkf(ObjectsToFindReferencesTo.Num() > 0, TEXT("ObjectsToFindReferencesTo cannot be empty"));
+	bool bAnyGarbage = false;
+	for (UObject* Obj : ObjectsToFindReferencesTo)
+	{
+		checkf(Obj, TEXT("Object to find references to cannot be null"));
+		bAnyGarbage = bAnyGarbage || Obj->HasAnyInternalFlags(EInternalObjectFlags::Garbage);
+	}
+
+	ELogVerbosity::Type Verbosity = ELogVerbosity::NoLogging;
+	switch (Options & EPrintStaleReferencesOptions::VerbosityMask)
+	{
+	case EPrintStaleReferencesOptions::None:
+		Verbosity = ELogVerbosity::NoLogging;
+		break;
+	case EPrintStaleReferencesOptions::Log:
+		Verbosity = ELogVerbosity::Log;
+		break;
+	case EPrintStaleReferencesOptions::Display:
+		Verbosity = ELogVerbosity::Display;
+		break;
+	case EPrintStaleReferencesOptions::Warning:
+		Verbosity = ELogVerbosity::Warning;
+		break;
+	case EPrintStaleReferencesOptions::Error:
+	case EPrintStaleReferencesOptions::Fatal: // Use verbosity error for reference chains so we can log all errors before crashing 
+	default:
+		Verbosity = ELogVerbosity::Error;
+		break;
+	}
+
+	UE_LOG(LogLoad, Log, TEXT("Beginning reference chain search..."));
+	for (UObject* Obj : ObjectsToFindReferencesTo)
+	{
+		UE_LOG(LogLoad, Log, TEXT(" - %s"), *Obj->GetFullName());
+	}
+
+	EReferenceChainSearchMode SearchMode = EReferenceChainSearchMode::GCOnly;
+	if (!!(Options & EPrintStaleReferencesOptions::Minimal))
+	{
+		SearchMode |= EReferenceChainSearchMode::Minimal;
+	}
+	else
+	{
+		SearchMode |= bAnyGarbage ? EReferenceChainSearchMode::ShortestToGarbage : EReferenceChainSearchMode::Shortest;
+	}
+
+	FReferenceChainSearch RefChainSearch(ObjectsToFindReferencesTo, SearchMode);
+#if ENABLE_GC_HISTORY
+	// If we have the full history already, don't only search for direct referencers 
+	FReferenceChainSearch HistorySearch(EReferenceChainSearchMode::ShortestToGarbage); // we may not need it but it has to live in this scope so that FGCObjectInfos are being kept alive
+#endif
+
+	for (UObject* ObjectToFindReferencesTo : ObjectsToFindReferencesTo)
+	{
+		FGCObjectInfo* OutGarbageObject = nullptr;
+		FGCObjectInfo* OutReferencingObject = nullptr;
+
+		bool bReferenceChainFound = PrintStaleReferenceChainsAndFindReferencingObjects(ObjectToFindReferencesTo, RefChainSearch, OutGarbageObject, OutReferencingObject, Verbosity);
+
+		FString PathToCulprit;
+		FString GarbageErrorMessage;
+		if (bReferenceChainFound)
+		{
+			PathToCulprit = GetPathToStaleObjectReferencer(ObjectToFindReferencesTo, RefChainSearch);
+			FString GarbageObjectName = OutGarbageObject ? OutGarbageObject->GetFullName() : ObjectToFindReferencesTo->GetFullName();
+			checkf(OutReferencingObject || HasGarbageCollectionKeepFlags(OutGarbageObject ? OutGarbageObject : OutReferencingObject), TEXT("No object referencing %s found even though we have a valid reference chain"), *ObjectToFindReferencesTo->GetPathName());
+			GarbageErrorMessage = FString::Printf(TEXT("Object %s is being referenced by %s"), *GarbageObjectName,
+				OutReferencingObject ? *OutReferencingObject->GetFullName() : (HasGarbageCollectionKeepFlags(OutGarbageObject) ? TEXT("GarbageCollectionKeepFlags") : TEXT("NULL")));
+		}
+#if ENABLE_GC_HISTORY
+		else
+		{
+			FGCSnapshot* LastGCSnapshot = FGCHistory::Get().GetLastSnapshot();
+			if (LastGCSnapshot)
+			{
+				UE_LOG(LogLoad, Log, TEXT("Looking for references to %s in the last GC run history..."), *ObjectToFindReferencesTo->GetFullName());
+
+				HistorySearch.PerformSearchFromGCSnapshot(ObjectsToFindReferencesTo, *LastGCSnapshot);
+
+				OutGarbageObject = nullptr;
+				OutReferencingObject = nullptr;
+				bReferenceChainFound = PrintStaleReferenceChainsAndFindReferencingObjects(ObjectToFindReferencesTo, HistorySearch, OutGarbageObject, OutReferencingObject, Verbosity);
+				if (bReferenceChainFound)
+				{
+					bReferenceChainFound = true;
+					PathToCulprit = GetPathToStaleObjectReferencer(ObjectToFindReferencesTo, HistorySearch);
+					FString GarbageObjectName = OutGarbageObject ? OutGarbageObject->GetFullName() : ObjectToFindReferencesTo->GetFullName();
+					checkf(OutReferencingObject || HasGarbageCollectionKeepFlags(OutGarbageObject ? OutGarbageObject : OutReferencingObject), TEXT("No object referencing %s found even though we have a valid reference chain"), *ObjectToFindReferencesTo->GetPathName());
+					GarbageErrorMessage = FString::Printf(TEXT("Garbage object %s was previously being referenced by %s"), *GarbageObjectName, OutReferencingObject ? *OutReferencingObject->GetFullName() : TEXT("NULL"));
+				}
+			}
+		}
+#endif // ENABLE_GC_HISTORY
+
+		if (!bReferenceChainFound)
+		{
+			if (OutGarbageObject)
+			{
+				GarbageErrorMessage = FString::Printf(TEXT("Garbage object %s is not referenced so it may have a flag set that's preventing it from being destroyed (see log for details)"), *OutGarbageObject->GetFullName());
+				PathToCulprit = OutGarbageObject->GetFullName();
+			}
+			else
+			{
+				GarbageErrorMessage = TEXT("However it's not referenced by any object. It may have a flag set that's preventing it from being destroyed (see log for details)");
+			}
+		}
+
+		Paths.Add(PathToCulprit);
+
+#if DO_ENSURE
+		if ((Options & EPrintStaleReferencesOptions::Ensure) == EPrintStaleReferencesOptions::Ensure)
+		{
+			ensureAlwaysMsgf(false, TEXT("Old %s not cleaned up by GC! %s:") LINE_TERMINATOR TEXT("%s"),
+				*ObjectToFindReferencesTo->GetFullName(),
+				*GarbageErrorMessage,
+				*PathToCulprit);
+		}
+		else
+#endif
+			if (Verbosity == ELogVerbosity::Fatal)
+			{
+				UE_LOG(LogLoad, Fatal, TEXT("Old %s not cleaned up by GC! %s:") LINE_TERMINATOR TEXT("%s"),
+					*ObjectToFindReferencesTo->GetFullName(),
+					*GarbageErrorMessage,
+					*PathToCulprit);
+			}
+			else
+			{
+				GLog->CategorizedLogf(TEXT("LogLoad"), Verbosity, TEXT("Old %s not cleaned up by GC! %s:") LINE_TERMINATOR TEXT("%s"),
+					*ObjectToFindReferencesTo->GetFullName(),
+					*GarbageErrorMessage,
+					*PathToCulprit);
+			}
+	}
+
+	if ((Options & EPrintStaleReferencesOptions::Fatal) == EPrintStaleReferencesOptions::Fatal)
+	{
+		UE_LOG(LogLoad, Fatal, TEXT("Fatal world leaks detected. Logging first error, check logs for additional information") LINE_TERMINATOR TEXT("%s"),
+			Paths.Num() ? *Paths[0] : TEXT("No paths to leaked objects found!"));
+	}
+
+	return Paths;
 }
\ No newline at end of file
diff --git a/Engine/Source/Runtime/CoreUObject/Public/UObject/ReferenceChainSearch.h b/Engine/Source/Runtime/CoreUObject/Public/UObject/ReferenceChainSearch.h
--- a/Engine/Source/Runtime/CoreUObject/Public/UObject/ReferenceChainSearch.h
+++ b/Engine/Source/Runtime/CoreUObject/Public/UObject/ReferenceChainSearch.h
@@ -59,6 +59,23 @@ enum class EReferenceChainSearchMode
 
 ENUM_CLASS_FLAGS(EReferenceChainSearchMode);
 
+/* Options for UEngine::FindAndPrintStaleReferencesToObject function */
+enum class EPrintStaleReferencesOptions
+{
+	None = 0,
+	Log = 1,
+	Display = 2,
+	Warning = 3,
+	Error = 4,
+	Fatal = 5,
+	Ensure = 0x00000100,
+
+	// Only search for direct references to the object or one of its inners, not a full reference chain 
+	Minimal = 0x00000200,
+
+	VerbosityMask = 0x000000ff
+};
+ENUM_CLASS_FLAGS(EPrintStaleReferencesOptions);
 
 class FReferenceChainSearch
 {
@@ -333,6 +350,13 @@ public:
 
 	int64 GetAllocatedSize() const;
 
+	/**
+	* Attempts to find a reference chain leading to a world that should have been garbage collected
+	* @param ObjectToFindReferencesTo World or its package (or any object from the world package that should've been destroyed)
+	* @param Options Determines how the stale references messages should be logged
+	*/
+	COREUOBJECT_API	static FString FindAndPrintStaleReferencesToObject(UObject* ObjectToFindReferencesTo, EPrintStaleReferencesOptions Options);
+	COREUOBJECT_API static TArray<FString> FindAndPrintStaleReferencesToObjects(TConstArrayView<UObject*> ObjectsToFindReferencesTo, EPrintStaleReferencesOptions Options);
 private:
 
 	/** The objects we're going to look for references to */
diff --git a/Engine/Source/Runtime/Engine/Classes/Engine/Engine.h b/Engine/Source/Runtime/Engine/Classes/Engine/Engine.h
--- a/Engine/Source/Runtime/Engine/Classes/Engine/Engine.h
+++ b/Engine/Source/Runtime/Engine/Classes/Engine/Engine.h
@@ -4,6 +4,7 @@
 
 #include "CoreMinimal.h"
 #include "Containers/IndirectArray.h"
+#include "UObject/ReferenceChainSearch.h"
 #include "UObject/ObjectMacros.h"
 #include "UObject/Object.h"
 #include "Misc/Guid.h"
@@ -674,24 +675,6 @@ struct FPluginRedirect
 	FString NewPluginName;
 };
 
-/* Options for UEngine::FindAndPrintStaleReferencesToObject function */
-enum class EPrintStaleReferencesOptions
-{
-	None = 0,
-	Log = 1,
-	Display = 2,
-	Warning = 3,
-	Error = 4,
-	Fatal = 5,
-	Ensure = 0x00000100,
-
-	// Only search for direct references to the object or one of its inners, not a full reference chain 
-	Minimal = 0x00000200, 
-
-	VerbosityMask = 0x000000ff
-};
-ENUM_CLASS_FLAGS(EPrintStaleReferencesOptions);
-
 /** Game thread events for dynamic resolution state. */
 enum class EDynamicResolutionStateEvent : uint8;
 
diff --git a/Engine/Source/Runtime/Engine/Private/UnrealEngine.cpp b/Engine/Source/Runtime/Engine/Private/UnrealEngine.cpp
--- a/Engine/Source/Runtime/Engine/Private/UnrealEngine.cpp
+++ b/Engine/Source/Runtime/Engine/Private/UnrealEngine.cpp
@@ -16000,64 +16000,6 @@ void UEngine::CheckAndHandleStaleWorldObjectReferences(FWorldContext* WorldConte
 	}
 }
 
-static FORCEINLINE bool HasGarbageCollectionKeepFlags(FGCObjectInfo* ObjectInfo)
-{
-	return ObjectInfo && (ObjectInfo->HasAnyFlags(GARBAGE_COLLECTION_KEEPFLAGS) || ObjectInfo->HasAnyInternalFlags(EInternalObjectFlags::GarbageCollectionKeepFlags | EInternalObjectFlags::RootSet));
-}
-
-static bool PrintStaleReferenceChainsAndFindReferencingObjects(UObject* ObjectToFindReferencesTo, FReferenceChainSearch& RefChainSearch, FGCObjectInfo*& OutGarbageObject, FGCObjectInfo*& OutReferencingObject, ELogVerbosity::Type Verbosity)
-{
-#if !NO_LOGGING
-	if (Verbosity != ELogVerbosity::NoLogging)
-	{
-		FMsg::Logf(__FILE__, __LINE__, LogLoad.GetCategoryName(), Verbosity, TEXT("Printing reference chains leading to %s: "), *ObjectToFindReferencesTo->GetFullName());
-	}
-#endif
-	return RefChainSearch.PrintResults([&ObjectToFindReferencesTo, &OutGarbageObject, &OutReferencingObject](FReferenceChainSearch::FCallbackParams& Params)
-		{
-			check(Params.Object);
-			if (!Params.Object->IsValid())
-			{
-				// We may find many chains that lead to the leak but for brevity we only report the first one in the fatal error below
-				if (!OutGarbageObject)
-				{
-					OutGarbageObject = Params.Object;
-					OutReferencingObject = Params.Referencer;
-				}
-				return true;
-			}
-			else
-			{
-				UObject* ResolvedObject = Params.Object->TryResolveObject();
-				if (ResolvedObject != ObjectToFindReferencesTo)
-				{
-					// Keep track of the last referencing object, we may need it if there's no garbage referencers
-					OutReferencingObject = Params.Object;
-				}
-				return true;
-			}
-		}, false, ObjectToFindReferencesTo) != 0;
-}
-
-static FString GetPathToStaleObjectReferencer(UObject* ObjectToFindReferencesTo, FReferenceChainSearch& RefChainSearch)
-{
-	bool bFirstGarbageObjectFound = false;
-	FString PathToCulprit = RefChainSearch.GetRootPath([&bFirstGarbageObjectFound, ObjectToFindReferencesTo](FReferenceChainSearch::FCallbackParams& Params)
-		{
-			check(Params.Object);
-			// Mark the first garbage reference or the target object reference as the culprit
-			if ((!Params.Object->IsValid() || Params.Object->TryResolveObject() == ObjectToFindReferencesTo) && !bFirstGarbageObjectFound)
-			{
-				bFirstGarbageObjectFound = true;
-				check(Params.Out);
-				Params.Out->Logf(ELogVerbosity::Log, TEXT("%s    ^^^^^^^^^^^^^^^^^^^^^^^^^^^^^^^^^^^^^^^^^^^^^^^^^^^^^^^^^^^^^^"), FCString::Spc(Params.Indent));
-				Params.Out->Logf(ELogVerbosity::Log, TEXT("%s    ^ This reference is preventing the old %s from being GC'd ^"), FCString::Spc(Params.Indent), *ObjectToFindReferencesTo->GetClass()->GetName());
-			}
-			return true;
-		});
-	return PathToCulprit;
-}
-
 void UEngine::FindAndPrintStaleReferencesToObject(UObject* ObjectToFindReferencesTo, ELogVerbosity::Type Verbosity)
 {
 	EPrintStaleReferencesOptions Options = EPrintStaleReferencesOptions::None;
@@ -16090,154 +16032,12 @@ void UEngine::FindAndPrintStaleReferencesToObject(UObject* ObjectToFindReference
 
 FString UEngine::FindAndPrintStaleReferencesToObject(UObject* ObjectToFindReferencesTo, EPrintStaleReferencesOptions Options)
 {
-	return FindAndPrintStaleReferencesToObjects(TConstArrayView<UObject*>(&ObjectToFindReferencesTo, 1), Options)[0];
+	return FReferenceChainSearch::FindAndPrintStaleReferencesToObject(ObjectToFindReferencesTo, Options);
 }
 
 TArray<FString> UEngine::FindAndPrintStaleReferencesToObjects(TConstArrayView<UObject*> ObjectsToFindReferencesTo, EPrintStaleReferencesOptions Options)
 {
-	TArray<FString> Paths; 
-	checkf(ObjectsToFindReferencesTo.Num() > 0, TEXT("ObjectsToFindReferencesTo cannot be empty"));
-	bool bAnyGarbage = false;
-	for (UObject* Obj : ObjectsToFindReferencesTo)
-	{
-		checkf(Obj, TEXT("Object to find references to cannot be null"));
-		bAnyGarbage = bAnyGarbage || Obj->HasAnyInternalFlags(EInternalObjectFlags::Garbage);
-	}
-
-	ELogVerbosity::Type Verbosity = ELogVerbosity::NoLogging;
-	switch (Options & EPrintStaleReferencesOptions::VerbosityMask)
-	{
-	case EPrintStaleReferencesOptions::None:
-		Verbosity = ELogVerbosity::NoLogging; 
-		break;
-	case EPrintStaleReferencesOptions::Log:
-		Verbosity = ELogVerbosity::Log; 
-		break;
-	case EPrintStaleReferencesOptions::Display:
-		Verbosity = ELogVerbosity::Display; 
-		break;
-	case EPrintStaleReferencesOptions::Warning:
-		Verbosity = ELogVerbosity::Warning; 
-		break;
-	case EPrintStaleReferencesOptions::Error:
-	case EPrintStaleReferencesOptions::Fatal: // Use verbosity error for reference chains so we can log all errors before crashing 
-	default:
-		Verbosity = ELogVerbosity::Error; 
-		break;
-	}
-
-	UE_LOG(LogLoad, Log, TEXT("Beginning reference chain search..."));
-	for (UObject* Obj : ObjectsToFindReferencesTo)
-	{
-		UE_LOG(LogLoad, Log, TEXT(" - %s"), *Obj->GetFullName());
-	}
-
-	EReferenceChainSearchMode SearchMode = EReferenceChainSearchMode::GCOnly;
-	if(!!(Options & EPrintStaleReferencesOptions::Minimal))
-	{
-		SearchMode |= EReferenceChainSearchMode::Minimal;
-	}
-	else
-	{
-		SearchMode |= bAnyGarbage ? EReferenceChainSearchMode::ShortestToGarbage : EReferenceChainSearchMode::Shortest;
-	}
-
-	FReferenceChainSearch RefChainSearch(ObjectsToFindReferencesTo, SearchMode);
-#if ENABLE_GC_HISTORY
-	// If we have the full history already, don't only search for direct referencers 
-	FReferenceChainSearch HistorySearch(EReferenceChainSearchMode::ShortestToGarbage); // we may not need it but it has to live in this scope so that FGCObjectInfos are being kept alive
-#endif
-
-	for (UObject* ObjectToFindReferencesTo : ObjectsToFindReferencesTo)
-	{
-		FGCObjectInfo* OutGarbageObject = nullptr;
-		FGCObjectInfo* OutReferencingObject = nullptr;
-
-		bool bReferenceChainFound = PrintStaleReferenceChainsAndFindReferencingObjects(ObjectToFindReferencesTo, RefChainSearch, OutGarbageObject, OutReferencingObject, Verbosity);
-
-		FString PathToCulprit;
-		FString GarbageErrorMessage;
-		if (bReferenceChainFound)
-		{
-			PathToCulprit = GetPathToStaleObjectReferencer(ObjectToFindReferencesTo, RefChainSearch);
-			FString GarbageObjectName = OutGarbageObject ? OutGarbageObject->GetFullName() : ObjectToFindReferencesTo->GetFullName();
-			checkf(OutReferencingObject || HasGarbageCollectionKeepFlags(OutGarbageObject ? OutGarbageObject : OutReferencingObject), TEXT("No object referencing %s found even though we have a valid reference chain"), *ObjectToFindReferencesTo->GetPathName());
-			GarbageErrorMessage = FString::Printf(TEXT("Object %s is being referenced by %s"), *GarbageObjectName, 
-				OutReferencingObject ? *OutReferencingObject->GetFullName() : (HasGarbageCollectionKeepFlags(OutGarbageObject) ? TEXT("GarbageCollectionKeepFlags") : TEXT("NULL")));
-		}
-#if ENABLE_GC_HISTORY
-		else
-		{
-			FGCSnapshot* LastGCSnapshot = FGCHistory::Get().GetLastSnapshot();
-			if (LastGCSnapshot)
-			{
-				UE_LOG(LogLoad, Log, TEXT("Looking for references to %s in the last GC run history..."), *ObjectToFindReferencesTo->GetFullName());
-
-				HistorySearch.PerformSearchFromGCSnapshot(ObjectsToFindReferencesTo, *LastGCSnapshot);
-
-				OutGarbageObject = nullptr;
-				OutReferencingObject = nullptr;
-				bReferenceChainFound = PrintStaleReferenceChainsAndFindReferencingObjects(ObjectToFindReferencesTo, HistorySearch, OutGarbageObject, OutReferencingObject, Verbosity);
-				if (bReferenceChainFound)
-				{
-					bReferenceChainFound = true;
-					PathToCulprit = GetPathToStaleObjectReferencer(ObjectToFindReferencesTo, HistorySearch);
-					FString GarbageObjectName = OutGarbageObject ? OutGarbageObject->GetFullName() : ObjectToFindReferencesTo->GetFullName();
-					checkf(OutReferencingObject || HasGarbageCollectionKeepFlags(OutGarbageObject ? OutGarbageObject : OutReferencingObject), TEXT("No object referencing %s found even though we have a valid reference chain"), *ObjectToFindReferencesTo->GetPathName());
-					GarbageErrorMessage = FString::Printf(TEXT("Garbage object %s was previously being referenced by %s"), *GarbageObjectName, OutReferencingObject ? *OutReferencingObject->GetFullName() : TEXT("NULL"));
-				}
-			}
-		}
-#endif // ENABLE_GC_HISTORY
-
-		if (!bReferenceChainFound)
-		{
-			if (OutGarbageObject)
-			{
-				GarbageErrorMessage = FString::Printf(TEXT("Garbage object %s is not referenced so it may have a flag set that's preventing it from being destroyed (see log for details)"), *OutGarbageObject->GetFullName());
-				PathToCulprit = OutGarbageObject->GetFullName();
-			}
-			else
-			{
-				GarbageErrorMessage = TEXT("However it's not referenced by any object. It may have a flag set that's preventing it from being destroyed (see log for details)");
-			}
-		}
-
-		Paths.Add(PathToCulprit);
-
-#if DO_ENSURE
-		if ((Options & EPrintStaleReferencesOptions::Ensure) == EPrintStaleReferencesOptions::Ensure)
-		{
-			ensureAlwaysMsgf(false, TEXT("Old %s not cleaned up by GC! %s:") LINE_TERMINATOR TEXT("%s"),
-				*ObjectToFindReferencesTo->GetFullName(),
-				*GarbageErrorMessage,
-				*PathToCulprit);
-		}
-		else
-#endif
-		if (Verbosity == ELogVerbosity::Fatal)
-		{
-			UE_LOG(LogLoad, Fatal, TEXT("Old %s not cleaned up by GC! %s:") LINE_TERMINATOR TEXT("%s"),
-				*ObjectToFindReferencesTo->GetFullName(), 
-				*GarbageErrorMessage,
-				*PathToCulprit);
-		}
-		else
-		{
-			GLog->CategorizedLogf(TEXT("LogLoad"), Verbosity, TEXT("Old %s not cleaned up by GC! %s:") LINE_TERMINATOR TEXT("%s"),
-				*ObjectToFindReferencesTo->GetFullName(), 
-				*GarbageErrorMessage,
-				*PathToCulprit);
-		}
-	}
-
-	if ((Options & EPrintStaleReferencesOptions::Fatal) == EPrintStaleReferencesOptions::Fatal)
-	{
-		UE_LOG(LogLoad, Fatal, TEXT("Fatal world leaks detected. Logging first error, check logs for additional information") LINE_TERMINATOR TEXT("%s"), 
-			Paths.Num() ? *Paths[0] : TEXT("No paths to leaked objects found!"));
-	}
-
-	return Paths;
+	return FReferenceChainSearch::FindAndPrintStaleReferencesToObjects(ObjectsToFindReferencesTo, Options);
 }
 
 /*-----------------------------------------------------------------------------
