Fix CL# 25534819 in light of CL# 24569999
#rnx

diff --git a/Engine/Source/Runtime/Engine/Private/StaticMeshRender.cpp b/Engine/Source/Runtime/Engine/Private/StaticMeshRender.cpp
--- a/Engine/Source/Runtime/Engine/Private/StaticMeshRender.cpp
+++ b/Engine/Source/Runtime/Engine/Private/StaticMeshRender.cpp
@@ -1418,7 +1418,7 @@ void FStaticMeshSceneProxy::OnEvaluateWorldPositionOffsetChanged_RenderThread()
 	if (ShouldOptimizedWPOAffectNonNaniteShaderSelection())
 	{
 		// Re-cache draw commands
-		GetRendererModule().BeginDeferredUpdateOfPrimitiveSceneInfo(GetPrimitiveSceneInfo());
+		GetRendererModule().RequestStaticMeshUpdate(GetPrimitiveSceneInfo());
 	}
 }
 
