Added UWorld::OverlapBlockingTestByObjectType().

diff --git a/Engine/Source/Runtime/Engine/Classes/Engine/World.h b/Engine/Source/Runtime/Engine/Classes/Engine/World.h
--- a/Engine/Source/Runtime/Engine/Classes/Engine/World.h
+++ b/Engine/Source/Runtime/Engine/Classes/Engine/World.h
@@ -2140,6 +2140,16 @@ public:
 	*/
 	bool OverlapAnyTestByChannel(const FVector& Pos, const FQuat& Rot, ECollisionChannel TraceChannel, const FCollisionShape& CollisionShape, const FCollisionQueryParams& Params = FCollisionQueryParams::DefaultQueryParam, const FCollisionResponseParams& ResponseParam = FCollisionResponseParams::DefaultResponseParam) const;
 
+	/**
+	*  Test the collision of a shape at the supplied location using object types, and return if any blocking overlap is found
+	*  @param  Pos             Location of center of box to test against the world
+	*  @param  ObjectQueryParams	List of object types it's looking for
+	*  @param  CollisionShape	CollisionShape - supports Box, Sphere, Capsule, Convex
+	*  @param  Params          Additional parameters used for the trace
+	*  @return TRUE if any blocking results are found
+	*/
+	bool OverlapBlockingTestByObjectType(const FVector& Pos, const FQuat& Rot, const FCollisionObjectQueryParams& ObjectQueryParams, const FCollisionShape& CollisionShape, const FCollisionQueryParams& Params = FCollisionQueryParams::DefaultQueryParam) const;
+
 	/**
 	*  Test the collision of a shape at the supplied location using object types, and return if any overlap is found
 	*  @param  Pos             Location of center of box to test against the world
diff --git a/Engine/Source/Runtime/Engine/Private/Collision/WorldCollision.cpp b/Engine/Source/Runtime/Engine/Private/Collision/WorldCollision.cpp
--- a/Engine/Source/Runtime/Engine/Private/Collision/WorldCollision.cpp
+++ b/Engine/Source/Runtime/Engine/Private/Collision/WorldCollision.cpp
@@ -210,6 +210,13 @@ bool UWorld::SweepMultiByObjectType(TArray<struct FHitResult>& OutHits, const FV
 	}
 }
 
+bool UWorld::OverlapBlockingTestByObjectType(const FVector& Pos, const FQuat& Rot, const FCollisionObjectQueryParams& ObjectQueryParams, const FCollisionShape& CollisionShape, const FCollisionQueryParams& Params /* = FCollisionQueryParams::DefaultQueryParam */) const
+{
+	bool bBlocking = false;
+	bBlocking = FPhysicsInterface::GeomOverlapBlockingTest(this, CollisionShape, Pos, Rot, DefaultCollisionChannel, Params, FCollisionResponseParams::DefaultResponseParam, ObjectQueryParams);
+	return bBlocking;
+}
+
 bool UWorld::OverlapAnyTestByObjectType(const FVector& Pos, const FQuat& Rot, const FCollisionObjectQueryParams& ObjectQueryParams, const FCollisionShape& CollisionShape, const FCollisionQueryParams& Params /* = FCollisionQueryParams::DefaultQueryParam */) const
 {
 	bool bBlocking = false;
