Removed deprecated mirroring functionality from SkeletalMesh

diff --git a/Engine/Plugins/Experimental/AnimNext/Source/AnimNext/Private/GenerationTools.cpp b/Engine/Plugins/Experimental/AnimNext/Source/AnimNext/Private/GenerationTools.cpp
--- a/Engine/Plugins/Experimental/AnimNext/Source/AnimNext/Private/GenerationTools.cpp
+++ b/Engine/Plugins/Experimental/AnimNext/Source/AnimNext/Private/GenerationTools.cpp
@@ -172,14 +172,6 @@ struct FCompareBoneIndexType
 			USkeletalMeshComponent::ExcludeHiddenBones(SkeletalMeshComponent, SkeletalMesh, OutRequiredBones);
 		}
 
-		// Add in any bones that may be required when mirroring.
-		PRAGMA_DISABLE_DEPRECATION_WARNINGS
-		if (SkeletalMesh->SkelMirrorTable.Num() > 0 &&
-			(SkeletalMeshComponent == nullptr || SkeletalMesh->SkelMirrorTable.Num() == SkeletalMeshComponent->BoneSpaceTransforms.Num()))
-		{
-			USkeletalMeshComponent::GetMirroringRequiredBones(SkeletalMesh, OutRequiredBones);
-		}
-		PRAGMA_ENABLE_DEPRECATION_WARNINGS
 
 			// Get socket bones set to animate and bones required to fill the component space base transforms
 		TArray<FBoneIndexType> NeededBonesForFillComponentSpaceTransforms;
diff --git a/Engine/Source/Editor/UnrealEd/Public/ImportUtils/SkelImport.h b/Engine/Source/Editor/UnrealEd/Public/ImportUtils/SkelImport.h
--- a/Engine/Source/Editor/UnrealEd/Public/ImportUtils/SkelImport.h
+++ b/Engine/Source/Editor/UnrealEd/Public/ImportUtils/SkelImport.h
@@ -96,11 +96,6 @@ struct FExistingSkelMeshData
 	bool									bExistingUseFullPrecisionUVs;
 	bool									bExistingUseHighPrecisionTangentBasis;
 
-	UE_DEPRECATED(5.0, "Please use UMirrorDataTable for mirroring ")
-	PRAGMA_DISABLE_DEPRECATION_WARNINGS
-	TArray<FBoneMirrorExport>				ExistingMirrorTable;
-	PRAGMA_ENABLE_DEPRECATION_WARNINGS
-
 	TWeakObjectPtr<UAssetImportData>		ExistingAssetImportData;
 	TWeakObjectPtr<UThumbnailInfo>			ExistingThumbnailInfo;
 
diff --git a/Engine/Source/Runtime/Engine/Classes/Components/SkeletalMeshComponent.h b/Engine/Source/Runtime/Engine/Classes/Components/SkeletalMeshComponent.h
--- a/Engine/Source/Runtime/Engine/Classes/Components/SkeletalMeshComponent.h
+++ b/Engine/Source/Runtime/Engine/Classes/Components/SkeletalMeshComponent.h
@@ -1662,6 +1662,7 @@ public:
 	static void ExcludeHiddenBones(const USkeletalMeshComponent* SkeletalMeshComponent, const USkeletalMesh* SkeletalMesh, TArray<FBoneIndexType>& OutRequiredBones);
 
 	// Get the bones used for mirroring in the skeletal mesh asset (if any)
+	UE_DEPRECATED(5.3, "This method has been deprecated. Please use UMirrorDataTable for mirroring support.")
 	static void GetMirroringRequiredBones(const USkeletalMesh* SkeletalMesh, TArray<FBoneIndexType>& OutRequiredBones);
 
 	// Get the bones that are used by sockets
diff --git a/Engine/Source/Runtime/Engine/Classes/Engine/SkeletalMesh.h b/Engine/Source/Runtime/Engine/Classes/Engine/SkeletalMesh.h
--- a/Engine/Source/Runtime/Engine/Classes/Engine/SkeletalMesh.h
+++ b/Engine/Source/Runtime/Engine/Classes/Engine/SkeletalMesh.h
@@ -790,45 +790,12 @@ public:
 	UFUNCTION(BlueprintSetter)
 	void SetMaterials(const TArray<FSkeletalMaterial>& InMaterials);
 
+#if WITH_EDITORONLY_DATA
 	/** List of bones that should be mirrored. */
 	UE_DEPRECATED(4.27, "Please do not access this member directly; Use UMirrorDataTable for mirroring support")
 	UPROPERTY(EditAnywhere, editfixedsize, Category=Mirroring)
 	TArray<struct FBoneMirrorInfo> SkelMirrorTable;
-
-    UE_DEPRECATED(5.0, "Please use UMirrorDataTable for mirroring support")
-	static FName GetSkelMirrorTableMemberName()
-	{
-		PRAGMA_DISABLE_DEPRECATION_WARNINGS
-		return GET_MEMBER_NAME_CHECKED(USkeletalMesh, SkelMirrorTable);
-		PRAGMA_ENABLE_DEPRECATION_WARNINGS
-	}
-
-    UE_DEPRECATED(5.0, "Please use UMirrorDataTable for mirroring support")
-	TArray<struct FBoneMirrorInfo>& GetSkelMirrorTable()
-	{
-		WaitUntilAsyncPropertyReleased(ESkeletalMeshAsyncProperties::SkelMirrorTable);
-		PRAGMA_DISABLE_DEPRECATION_WARNINGS
-		return SkelMirrorTable;
-		PRAGMA_ENABLE_DEPRECATION_WARNINGS
-	}
-
-    UE_DEPRECATED(5.0, "Please use UMirrorDataTable for mirroring support")
-	const TArray<struct FBoneMirrorInfo>& GetSkelMirrorTable() const
-	{
-		WaitUntilAsyncPropertyReleased(ESkeletalMeshAsyncProperties::SkelMirrorTable, ESkinnedAssetAsyncPropertyLockType::ReadOnly);
-		PRAGMA_DISABLE_DEPRECATION_WARNINGS
-		return SkelMirrorTable;
-		PRAGMA_ENABLE_DEPRECATION_WARNINGS
-	}
-
-    UE_DEPRECATED(5.0, "Please use UMirrorDataTable for mirroring support")
-	void SetSkelMirrorTable(const TArray<struct FBoneMirrorInfo>& InSkelMirrorTable)
-	{
-		WaitUntilAsyncPropertyReleased(ESkeletalMeshAsyncProperties::SkelMirrorTable);
-		PRAGMA_DISABLE_DEPRECATION_WARNINGS
-		SkelMirrorTable = InSkelMirrorTable;
-		PRAGMA_ENABLE_DEPRECATION_WARNINGS
-	}
+#endif
 
 private:
 	/** Struct containing information for each LOD level, such as materials to use, and when use the LOD. */
@@ -1157,62 +1124,17 @@ public:
 
 	UFUNCTION(BlueprintGetter)
 	TSoftObjectPtr<UObject> GetDefaultAnimatingRig() const;
-	
+
+#if WITH_EDITORONLY_DATA
 	UE_DEPRECATED(4.27, "Please do not access this member directly; use USkeletalMesh::GetSkelMirrorAxis() or USkeletalMesh::SetSkelMirrorAxis().")
 	UPROPERTY(EditAnywhere, Category=Mirroring)
 	TEnumAsByte<EAxis::Type> SkelMirrorAxis;
-	static FName GetSkelMirrorAxisMemberName()
-	{
-		PRAGMA_DISABLE_DEPRECATION_WARNINGS
-		return GET_MEMBER_NAME_CHECKED(USkeletalMesh, SkelMirrorAxis);
-		PRAGMA_ENABLE_DEPRECATION_WARNINGS
-	}
 
-	UE_DEPRECATED(5.0, "Please use UMirrorDataTable for mirroring support")
-	TEnumAsByte<EAxis::Type> GetSkelMirrorAxis() const
-	{
-		WaitUntilAsyncPropertyReleased(ESkeletalMeshAsyncProperties::SkelMirrorAxis, ESkinnedAssetAsyncPropertyLockType::ReadOnly);
-		PRAGMA_DISABLE_DEPRECATION_WARNINGS
-		return SkelMirrorAxis;
-		PRAGMA_ENABLE_DEPRECATION_WARNINGS
-	}
-
-	UE_DEPRECATED(5.0, "Please use UMirrorDataTable for mirroring support")
-	void SetSkelMirrorAxis(TEnumAsByte<EAxis::Type> InSkelMirrorAxis)
-	{
-		WaitUntilAsyncPropertyReleased(ESkeletalMeshAsyncProperties::SkelMirrorAxis);
-		PRAGMA_DISABLE_DEPRECATION_WARNINGS
-		SkelMirrorAxis = InSkelMirrorAxis;
-		PRAGMA_ENABLE_DEPRECATION_WARNINGS
-	}
 
 	UE_DEPRECATED(4.27, "Please do not access this member directly; Use UMirrorDataTable for mirroring support")
 	UPROPERTY(EditAnywhere, Category=Mirroring)
 	TEnumAsByte<EAxis::Type> SkelMirrorFlipAxis;
-	static FName GetSkelMirrorFlipAxisMemberName()
-	{
-		PRAGMA_DISABLE_DEPRECATION_WARNINGS
-		return GET_MEMBER_NAME_CHECKED(USkeletalMesh, SkelMirrorFlipAxis);
-		PRAGMA_ENABLE_DEPRECATION_WARNINGS
-	}
-
-    UE_DEPRECATED(5.0, "Please use UMirrorDataTable for mirroring support")
-	TEnumAsByte<EAxis::Type> GetSkelMirrorFlipAxis() const
-	{
-		WaitUntilAsyncPropertyReleased(ESkeletalMeshAsyncProperties::SkelMirrorFlipAxis, ESkinnedAssetAsyncPropertyLockType::ReadOnly);
-		PRAGMA_DISABLE_DEPRECATION_WARNINGS
-		return SkelMirrorFlipAxis;
-		PRAGMA_ENABLE_DEPRECATION_WARNINGS
-	}
-
-    UE_DEPRECATED(5.0, "Please use UMirrorDataTable for mirroring support")
-	void SetSkelMirrorFlipAxis(TEnumAsByte<EAxis::Type> InSkelMirrorFlipAxis)
-	{
-		WaitUntilAsyncPropertyReleased(ESkeletalMeshAsyncProperties::SkelMirrorFlipAxis);
-		PRAGMA_DISABLE_DEPRECATION_WARNINGS
-		SkelMirrorFlipAxis = InSkelMirrorFlipAxis;
-		PRAGMA_ENABLE_DEPRECATION_WARNINGS
-	}
+#endif 
 
 	/** If true, use 32 bit UVs. If false, use 16 bit UVs to save memory */
 	UPROPERTY()
@@ -2548,24 +2470,6 @@ public:
 	virtual FMatrix GetComposedRefPoseMatrix( FName InBoneName ) const override;
 	virtual FMatrix GetComposedRefPoseMatrix( int32 InBoneIndex ) const override;
 
-	UE_DEPRECATED(5.0, "Please use UMirrorDataTable for mirroring support.")
-	void InitBoneMirrorInfo();
-
-	UE_DEPRECATED(5.0, "Please use UMirrorDataTable for mirroring support.")
-	void CopyMirrorTableFrom(USkeletalMesh* SrcMesh);
-	
-	UE_DEPRECATED(5.0, "Please use UMirrorDataTable for mirroring support.")
-	PRAGMA_DISABLE_DEPRECATION_WARNINGS
-	void ExportMirrorTable(TArray<FBoneMirrorExport> &MirrorExportInfo) const;
-	PRAGMA_ENABLE_DEPRECATION_WARNINGS
-
-	UE_DEPRECATED(5.0, "Please use UMirrorDataTable for mirroring support.")
-	PRAGMA_DISABLE_DEPRECATION_WARNINGS
-	void ImportMirrorTable(const TArray<FBoneMirrorExport> &MirrorExportInfo);
-	PRAGMA_ENABLE_DEPRECATION_WARNINGS
-
-	UE_DEPRECATED(5.0, "Please use UMirrorDataTable for mirroring support.")
-	bool MirrorTableIsGood(FString& ProblemBones) const;
 
 	/**
 	 * Returns the mesh only socket list - this ignores any sockets in the skeleton
diff --git a/Engine/Source/Runtime/Engine/Private/Components/SkeletalMeshComponent.cpp b/Engine/Source/Runtime/Engine/Private/Components/SkeletalMeshComponent.cpp
--- a/Engine/Source/Runtime/Engine/Private/Components/SkeletalMeshComponent.cpp
+++ b/Engine/Source/Runtime/Engine/Private/Components/SkeletalMeshComponent.cpp
@@ -1819,15 +1819,6 @@ void USkeletalMeshComponent::ComputeRequiredBones(TArray<FBoneIndexType>& OutReq
 	// mirror table/phys body ones has to be calculated
 	ExcludeHiddenBones(this, SkelMesh, OutRequiredBones);
 
-	// Add in any bones that may be required when mirroring.
-	PRAGMA_DISABLE_DEPRECATION_WARNINGS
-	if (SkelMesh->SkelMirrorTable.Num() > 0 &&
-		SkelMesh->SkelMirrorTable.Num() == BoneSpaceTransforms.Num())
-	{
-		GetMirroringRequiredBones(SkelMesh, OutRequiredBones);
-	}
-	PRAGMA_ENABLE_DEPRECATION_WARNINGS
-
 	// Get socket bones set to animate and bones required to fill the component space base transforms
 	TArray<FBoneIndexType> NeededBonesForFillComponentSpaceTransforms;
 	GetSocketRequiredBones(SkelMesh, OutRequiredBones, NeededBonesForFillComponentSpaceTransforms);
@@ -1922,27 +1913,6 @@ void USkeletalMeshComponent::ComputeRequiredBones(TArray<FBoneIndexType>& OutReq
 
 /*static*/ void USkeletalMeshComponent::GetMirroringRequiredBones(const USkeletalMesh* SkeletalMesh, TArray<FBoneIndexType>& OutRequiredBones)
 {
-	check(SkeletalMesh != nullptr);
-
-	// Add in any bones that may be required when mirroring.
-	// JTODO: This is only required if there are mirroring nodes in the tree, but hard to know...
-	PRAGMA_DISABLE_DEPRECATION_WARNINGS
-	TArray<FBoneIndexType> MirroredDesiredBones;
-	MirroredDesiredBones.AddUninitialized(OutRequiredBones.Num());
-
-	const TArray<struct FBoneMirrorInfo>& SkelMirrorTable = SkeletalMesh->GetSkelMirrorTable();
-	// Look up each bone in the mirroring table.
-	for (int32 i = 0; i<OutRequiredBones.Num(); i++)
-	{
-		MirroredDesiredBones[i] = SkelMirrorTable[OutRequiredBones[i]].SourceIndex;
-	}
-
-	// Sort to ensure strictly increasing order.
-	MirroredDesiredBones.Sort();
-
-	// Make sure all of these are in OutRequiredBones, and 
-	MergeInBoneIndexArrays(OutRequiredBones, MirroredDesiredBones);
-	PRAGMA_ENABLE_DEPRECATION_WARNINGS
 }
 
 /*static*/ void USkeletalMeshComponent::GetSocketRequiredBones(const USkeletalMesh* SkeletalMesh, TArray<FBoneIndexType>& OutRequiredBones, TArray<FBoneIndexType>& NeededBonesForFillComponentSpaceTransforms)
diff --git a/Engine/Source/Runtime/Engine/Private/SkeletalMesh.cpp b/Engine/Source/Runtime/Engine/Private/SkeletalMesh.cpp
--- a/Engine/Source/Runtime/Engine/Private/SkeletalMesh.cpp
+++ b/Engine/Source/Runtime/Engine/Private/SkeletalMesh.cpp
@@ -395,11 +395,6 @@ namespace SkeletalMeshImpl
 USkeletalMesh::USkeletalMesh(const FObjectInitializer& ObjectInitializer)
 	: Super(ObjectInitializer)
 {
-PRAGMA_DISABLE_DEPRECATION_WARNINGS
-	SetSkelMirrorAxis(EAxis::X);
-	SetSkelMirrorFlipAxis(EAxis::Z);
-PRAGMA_ENABLE_DEPRECATION_WARNINGS
-
 #if WITH_EDITORONLY_DATA
 	SetImportedModel(MakeShareable(new FSkeletalMeshModel()));
 	SetVertexColorGuid(FGuid());
@@ -4588,200 +4583,6 @@ void USkeletalMesh::ReleaseCPUResources()
 	}
 }
 
-/** Allocate and initialise bone mirroring table for this skeletal mesh. Default is source = destination for each bone. */
-void USkeletalMesh::InitBoneMirrorInfo()
-{
-PRAGMA_DISABLE_DEPRECATION_WARNINGS
-	TArray<FBoneMirrorInfo>& LocalSkelMirrorTable = GetSkelMirrorTable();
-	LocalSkelMirrorTable.Empty(GetRefSkeleton().GetNum());
-	LocalSkelMirrorTable.AddZeroed(GetRefSkeleton().GetNum());
-
-	// By default, no bone mirroring, and source is ourself.
-	for(int32 i=0; i< LocalSkelMirrorTable.Num(); i++)
-	{
-		LocalSkelMirrorTable[i].SourceIndex = i;
-	}
-PRAGMA_ENABLE_DEPRECATION_WARNINGS
-}
-
-/** Utility for copying and converting a mirroring table from another SkeletalMesh. */
-void USkeletalMesh::CopyMirrorTableFrom(USkeletalMesh* SrcMesh)
-{
-PRAGMA_DISABLE_DEPRECATION_WARNINGS
-	TArray<FBoneMirrorInfo>& SrcSkelMirrorTable = SrcMesh->GetSkelMirrorTable();
-	TArray<FBoneMirrorInfo>& LocalSkelMirrorTable = GetSkelMirrorTable();
-	// Do nothing if no mirror table in source mesh
-	if(SrcSkelMirrorTable.Num() == 0)
-	{
-		return;
-	}
-
-	// First, allocate and default mirroring table.
-	InitBoneMirrorInfo();
-
-	// Keep track of which entries in the source we have already copied
-	TArray<bool> EntryCopied;
-	EntryCopied.AddZeroed(SrcSkelMirrorTable.Num() );
-
-	// Mirror table must always be size of ref skeleton.
-	check(SrcSkelMirrorTable.Num() == SrcMesh->GetRefSkeleton().GetNum());
-
-	// Iterate over each entry in the source mesh mirror table.
-	// We assume that the src table is correct, and don't check for errors here (ie two bones using the same one as source).
-	for(int32 i=0; i< SrcSkelMirrorTable.Num(); i++)
-	{
-		if(!EntryCopied[i])
-		{
-			// Get name of source and dest bone for this entry in the source table.
-			FName DestBoneName = SrcMesh->GetRefSkeleton().GetBoneName(i);
-			int32 SrcBoneIndex = SrcSkelMirrorTable[i].SourceIndex;
-			FName SrcBoneName = SrcMesh->GetRefSkeleton().GetBoneName(SrcBoneIndex);
-			EAxis::Type FlipAxis = SrcSkelMirrorTable[i].BoneFlipAxis;
-
-			// Look up bone names in target mesh (this one)
-			int32 DestBoneIndexTarget = GetRefSkeleton().FindBoneIndex(DestBoneName);
-			int32 SrcBoneIndexTarget = GetRefSkeleton().FindBoneIndex(SrcBoneName);
-
-			// If both bones found, copy data to this mesh's mirror table.
-			if( DestBoneIndexTarget != INDEX_NONE && SrcBoneIndexTarget != INDEX_NONE )
-			{
-				LocalSkelMirrorTable[DestBoneIndexTarget].SourceIndex = SrcBoneIndexTarget;
-				LocalSkelMirrorTable[DestBoneIndexTarget].BoneFlipAxis = FlipAxis;
-
-
-				LocalSkelMirrorTable[SrcBoneIndexTarget].SourceIndex = DestBoneIndexTarget;
-				LocalSkelMirrorTable[SrcBoneIndexTarget].BoneFlipAxis = FlipAxis;
-
-				// Flag entries as copied, so we don't try and do it again.
-				EntryCopied[i] = true;
-				EntryCopied[SrcBoneIndex] = true;
-			}
-		}
-	}
-PRAGMA_ENABLE_DEPRECATION_WARNINGS
-}
-
-PRAGMA_DISABLE_DEPRECATION_WARNINGS
-/** Utility for copying and converting a mirroring table from another SkeletalMesh. */
-void USkeletalMesh::ExportMirrorTable(TArray<FBoneMirrorExport> &MirrorExportInfo) const
-{
-	const TArray<FBoneMirrorInfo>& LocalSkelMirrorTable = GetSkelMirrorTable();
-	// Do nothing if no mirror table in source mesh
-	if(LocalSkelMirrorTable.Num() == 0 )
-	{
-		return;
-	}
-	
-	// Mirror table must always be size of ref skeleton.
-	check(LocalSkelMirrorTable.Num() == GetRefSkeleton().GetNum());
-
-	MirrorExportInfo.Empty(LocalSkelMirrorTable.Num());
-	MirrorExportInfo.AddZeroed(LocalSkelMirrorTable.Num());
-
-	// Iterate over each entry in the source mesh mirror table.
-	// We assume that the src table is correct, and don't check for errors here (ie two bones using the same one as source).
-	for(int32 i=0; i< LocalSkelMirrorTable.Num(); i++)
-	{
-		MirrorExportInfo[i].BoneName		= GetRefSkeleton().GetBoneName(i);
-		MirrorExportInfo[i].SourceBoneName	= GetRefSkeleton().GetBoneName(LocalSkelMirrorTable[i].SourceIndex);
-		MirrorExportInfo[i].BoneFlipAxis	= LocalSkelMirrorTable[i].BoneFlipAxis;
-	}
-}
-PRAGMA_ENABLE_DEPRECATION_WARNINGS
-
-PRAGMA_DISABLE_DEPRECATION_WARNINGS
-/** Utility for copying and converting a mirroring table from another SkeletalMesh. */
-void USkeletalMesh::ImportMirrorTable(const TArray<FBoneMirrorExport> &MirrorExportInfo)
-{
-
-	// Do nothing if no mirror table in source mesh
-	if( MirrorExportInfo.Num() == 0 )
-	{
-		return;
-	}
-
-	// First, allocate and default mirroring table.
-	InitBoneMirrorInfo();
-
-	// Keep track of which entries in the source we have already copied
-	TArray<bool> EntryCopied;
-	EntryCopied.AddZeroed(GetRefSkeleton().GetNum() );
-
-	TArray<FBoneMirrorInfo>& LocalSkelMirrorTable = GetSkelMirrorTable();
-	// Mirror table must always be size of ref skeleton.
-	check(LocalSkelMirrorTable.Num() == GetRefSkeleton().GetNum());
-
-	// Iterate over each entry in the source mesh mirror table.
-	// We assume that the src table is correct, and don't check for errors here (ie two bones using the same one as source).
-	for(int32 i=0; i<MirrorExportInfo.Num(); i++)
-	{
-		FName DestBoneName	= MirrorExportInfo[i].BoneName;
-		int32 DestBoneIndex	= GetRefSkeleton().FindBoneIndex(DestBoneName);
-
-		if( DestBoneIndex != INDEX_NONE && !EntryCopied[DestBoneIndex] )
-		{
-			FName SrcBoneName	= MirrorExportInfo[i].SourceBoneName;
-			int32 SrcBoneIndex	= GetRefSkeleton().FindBoneIndex(SrcBoneName);
-			EAxis::Type FlipAxis		= MirrorExportInfo[i].BoneFlipAxis;
-
-			// If both bones found, copy data to this mesh's mirror table.
-			if( SrcBoneIndex != INDEX_NONE )
-			{
-				LocalSkelMirrorTable[DestBoneIndex].SourceIndex = SrcBoneIndex;
-				LocalSkelMirrorTable[DestBoneIndex].BoneFlipAxis = FlipAxis;
-
-				LocalSkelMirrorTable[SrcBoneIndex].SourceIndex = DestBoneIndex;
-				LocalSkelMirrorTable[SrcBoneIndex].BoneFlipAxis = FlipAxis;
-
-				// Flag entries as copied, so we don't try and do it again.
-				EntryCopied[DestBoneIndex]	= true;
-				EntryCopied[SrcBoneIndex]	= true;
-			}
-		}
-	}
-}
-PRAGMA_ENABLE_DEPRECATION_WARNINGS
-
-/** 
- *	Utility for checking that the bone mirroring table of this mesh is good.
- *	Return true if mirror table is OK, false if there are problems.
- *	@param	ProblemBones	Output string containing information on bones that are currently bad.
- */
-bool USkeletalMesh::MirrorTableIsGood(FString& ProblemBones) const
-{
-PRAGMA_DISABLE_DEPRECATION_WARNINGS
-	TArray<int32>	BadBoneMirror;
-
-	const TArray<FBoneMirrorInfo>& LocalSkelMirrorTable = GetSkelMirrorTable();
-
-	for(int32 i=0; i< LocalSkelMirrorTable.Num(); i++)
-	{
-		int32 SrcIndex = LocalSkelMirrorTable[i].SourceIndex;
-		if(LocalSkelMirrorTable[SrcIndex].SourceIndex != i)
-		{
-			BadBoneMirror.Add(i);
-		}
-	}
-
-	if(BadBoneMirror.Num() > 0)
-	{
-		for(int32 i=0; i<BadBoneMirror.Num(); i++)
-		{
-			int32 BoneIndex = BadBoneMirror[i];
-			FName BoneName = GetRefSkeleton().GetBoneName(BoneIndex);
-
-			ProblemBones += FString::Printf( TEXT("%s (%d)\n"), *BoneName.ToString(), BoneIndex );
-		}
-
-		return false;
-	}
-	else
-	{
-		return true;
-	}
-PRAGMA_ENABLE_DEPRECATION_WARNINGS
-}
-
 void USkeletalMesh::CreateBodySetup()
 {
 	const USkeletalMesh* ConstThis = this;
diff --git a/Engine/Source/Runtime/Engine/Private/SkeletalMeshMerge.cpp b/Engine/Source/Runtime/Engine/Private/SkeletalMeshMerge.cpp
--- a/Engine/Source/Runtime/Engine/Private/SkeletalMeshMerge.cpp
+++ b/Engine/Source/Runtime/Engine/Private/SkeletalMeshMerge.cpp
@@ -803,9 +803,6 @@ bool FSkeletalMeshMerge::ProcessMergeMesh()
 	
 	// copy settings and bone info from src meshes
 	bool bNeedsInit=true;
-	
-PRAGMA_DISABLE_DEPRECATION_WARNINGS
-	MergeMesh->GetSkelMirrorTable().Empty();
 
 	for( int32 MeshIdx=0; MeshIdx < SrcMeshList.Num(); MeshIdx++ )
 	{
@@ -816,10 +813,6 @@ PRAGMA_DISABLE_DEPRECATION_WARNINGS
 			{
 				// initialize the merged mesh with the first src mesh entry used
 				MergeMesh->SetImportedBounds(SrcMesh->GetImportedBounds());
-
-				MergeMesh->SetSkelMirrorAxis(SrcMesh->GetSkelMirrorAxis());
-				MergeMesh->SetSkelMirrorFlipAxis(SrcMesh->GetSkelMirrorFlipAxis());
-
 				// only initialize once
 				bNeedsInit = false;
 			}
@@ -830,7 +823,6 @@ PRAGMA_DISABLE_DEPRECATION_WARNINGS
 			}
 		}
 	}
-PRAGMA_ENABLE_DEPRECATION_WARNINGS
 	// Rebuild inverse ref pose matrices.
 	MergeMesh->GetRefBasesInvMatrix().Empty();
 	MergeMesh->CalculateInvRefMatrices();
