Substrata - HUD - display of material tyep and details.

diff --git a/Engine/Source/Editor/MaterialEditor/Private/SMaterialEditorStrataWidget.cpp b/Engine/Source/Editor/MaterialEditor/Private/SMaterialEditorStrataWidget.cpp
--- a/Engine/Source/Editor/MaterialEditor/Private/SMaterialEditorStrataWidget.cpp
+++ b/Engine/Source/Editor/MaterialEditor/Private/SMaterialEditorStrataWidget.cpp
@@ -228,7 +228,67 @@ void SMaterialEditorStrataWidget::Tick(const FGeometry& AllottedGeometry, const
 					MaterialDescription += FString::Printf(TEXT("Material complexity          = UNKOWN => ERROR!\r\n"));
 				}
 
-				MaterialDescription += FString::Printf(TEXT("Is Thin                      = %i\r\n"), CompilationOutput.bIsThin);
+				MaterialDescription += FString::Printf(TEXT("Root Node Is Thin            = %i\r\n"), CompilationOutput.bIsThin);
+
+				//if (CompilationOutput.StrataBSDFCount == 1)
+				//{
+				//	auto GetSingleOperatorBSDF = [&](const FStrataOperator& Op)
+				//	{
+				//		switch (Op.OperatorType)
+				//		{
+				//		case STRATA_OPERATOR_WEIGHT:
+				//		{
+				//			return ProcessOperator(CompilationOutput.Operators[Op.LeftIndex]);	// Continue walking the tree
+				//			break;
+				//		}
+				//		case STRATA_OPERATOR_VERTICAL:
+				//		case STRATA_OPERATOR_HORIZONTAL:
+				//		case STRATA_OPERATOR_ADD:
+				//		{
+				//			return nullptr;	// A operator with multiple entries must use slabs and so should be a Strata material made of slabs
+				//			break;
+				//		}
+				//		case STRATA_OPERATOR_BSDF_LEGACY:
+				//		case STRATA_OPERATOR_BSDF:
+				//		{
+				//			return Op;
+				//			break;
+				//		}
+				//		}
+				//		return nullptr;
+				//	};
+				//	const FStrataOperator& RootOperator = CompilationOutput.Operators[CompilationOutput.RootOperatorIndex];
+				//	const FStrataOperator* OpBSDF = GetSingleOperatorBSDF(RootOperator);
+				//
+				//	if (OpBSDF)
+				//	{
+				//		// STRATA_TODO gather information about SSSPRofile, subsurface, two sided, etc?
+				//	}
+				//}
+
+				FString MaterialType;
+				switch (CompilationOutput.MaterialType)
+				{
+				case SUBSTRATE_MATERIAL_TYPE_SINGLESLAB:			MaterialType = TEXT("Single Slab"); break;
+				case SUBSTRATE_MATERIAL_TYPE_MULTIPLESLABS:			MaterialType = TEXT("Multiple Slabs"); break;
+				case SUBSTRATE_MATERIAL_TYPE_VOLUMETRICFOGCLOUD:	MaterialType = TEXT("Volumetric (Fog or Cloud)"); break;
+				case SUBSTRATE_MATERIAL_TYPE_UNLIT:					MaterialType = TEXT("Unlit"); break;
+				case SUBSTRATE_MATERIAL_TYPE_HAIR:					MaterialType = TEXT("Hair"); break;
+				case SUBSTRATE_MATERIAL_TYPE_SINGLELAYERWATER:		MaterialType = TEXT("SingleLayerWater"); break;
+				case SUBSTRATE_MATERIAL_TYPE_EYE:					MaterialType = TEXT("Eye"); break;
+				case SUBSTRATE_MATERIAL_TYPE_LIGHTFUNCTION:			MaterialType = TEXT("LightFunction"); break;
+				case SUBSTRATE_MATERIAL_TYPE_POSTPROCESS:			MaterialType = TEXT("PostProcess"); break;
+				case SUBSTRATE_MATERIAL_TYPE_UI:					MaterialType = TEXT("UI"); break;
+				case SUBSTRATE_MATERIAL_TYPE_DECAL:					MaterialType = TEXT("Decal"); break;
+				}
+
+				FString MaterialTypeDetail;
+				if ((CompilationOutput.MaterialType == SUBSTRATE_MATERIAL_TYPE_SINGLESLAB || CompilationOutput.MaterialType == SUBSTRATE_MATERIAL_TYPE_MULTIPLESLABS) && CompilationOutput.bIsThin)
+				{
+					MaterialTypeDetail = TEXT("(Thin Surface / Two-Sided Lighting)");
+				}
+
+				MaterialDescription += FString::Printf(TEXT("Material Type                = %s %s\r\n"), *MaterialType, *MaterialTypeDetail);
 
 				MaterialDescription += FString::Printf(TEXT(" \r\n"));
 				MaterialDescription += FString::Printf(TEXT(" \r\n"));
diff --git a/Engine/Source/Runtime/Engine/Private/Materials/HLSLMaterialTranslator.cpp b/Engine/Source/Runtime/Engine/Private/Materials/HLSLMaterialTranslator.cpp
--- a/Engine/Source/Runtime/Engine/Private/Materials/HLSLMaterialTranslator.cpp
+++ b/Engine/Source/Runtime/Engine/Private/Materials/HLSLMaterialTranslator.cpp
@@ -11297,14 +11297,13 @@ bool FHLSLMaterialTranslator::FStrataCompilationContext::StrataGenerateDerivedMa
 		// Allocate BSDFIndex at the same time.
 		// Check BSDF that should be unique.
 		//
+		bool bHasUnlit = false;
+		bool bHasVFogCloud = false;
+		bool bHasHair = false;
+		bool bHasEye = false;
+		bool bHasSLW = false;
+		bool bHasSlab = false;
 		{
-			bool bHasUnlit = false;
-			bool bHasVFogCloud = false;
-			bool bHasHair = false;
-			bool bHasEye = false;
-			bool bHasSLW = false;
-			bool bHasSlab = false;
-
 			int VOpTopBranchCountTaken = 0;
 			int VOpBottomBranchCountTaken = 0;
 			bool bStrataUsesVerticalLayering = false;
@@ -11790,13 +11789,60 @@ bool FHLSLMaterialTranslator::FStrataCompilationContext::StrataGenerateDerivedMa
 				Compiler->MaterialCompilationOutput.StrataMaterialCompilationOutput.StrataMaterialType = bStrataMaterialIsSimple ? 0 : bStrataMaterialIsSingle ? 1 : 2;
 				Compiler->MaterialCompilationOutput.StrataMaterialCompilationOutput.StrataBSDFCount = StrataMaterialBSDFCount;
 				Compiler->MaterialCompilationOutput.StrataMaterialCompilationOutput.StrataUintPerPixel = uint8(FMath::Clamp(RequestedSizeInUint, 0u, 0xFF));
-				Compiler->MaterialCompilationOutput.StrataMaterialCompilationOutput.bIsThin = CompilerMaterial->IsThinSurface() ? 1 : 0;
 
 #if WITH_EDITOR
 				Compiler->MaterialCompilationOutput.StrataMaterialCompilationOutput.SharedLocalBasesCount = 0; // FinalUsedSharedLocalBasesCount is not valid yet
 				Compiler->MaterialCompilationOutput.StrataMaterialCompilationOutput.RequestedBytePixePixel = StrataMaterialRequestedSizeByte;
 				Compiler->MaterialCompilationOutput.StrataMaterialCompilationOutput.PlatformBytePixePixel = StrataBytePerPixel;
 
+				Compiler->MaterialCompilationOutput.StrataMaterialCompilationOutput.bIsThin = CompilerMaterial->IsThinSurface() ? 1 : 0;
+
+				// The order of ifs here is important.
+				if (CompilerMaterial->IsLightFunction())
+				{
+					Compiler->MaterialCompilationOutput.StrataMaterialCompilationOutput.MaterialType = SUBSTRATE_MATERIAL_TYPE_LIGHTFUNCTION;
+				}
+				else if (CompilerMaterial->IsPostProcessMaterial())
+				{
+					Compiler->MaterialCompilationOutput.StrataMaterialCompilationOutput.MaterialType = SUBSTRATE_MATERIAL_TYPE_POSTPROCESS;
+				}
+				else if (CompilerMaterial->IsUIMaterial())
+				{
+					Compiler->MaterialCompilationOutput.StrataMaterialCompilationOutput.MaterialType = SUBSTRATE_MATERIAL_TYPE_UI;
+				}
+				else if (CompilerMaterial->IsDeferredDecal())
+				{
+					Compiler->MaterialCompilationOutput.StrataMaterialCompilationOutput.MaterialType = SUBSTRATE_MATERIAL_TYPE_DECAL;
+				}
+				else if (StrataMaterialBSDFCount > 1)
+				{
+					Compiler->MaterialCompilationOutput.StrataMaterialCompilationOutput.MaterialType = SUBSTRATE_MATERIAL_TYPE_MULTIPLESLABS;
+				}
+				else if (bHasUnlit)
+				{
+					Compiler->MaterialCompilationOutput.StrataMaterialCompilationOutput.MaterialType = SUBSTRATE_MATERIAL_TYPE_UNLIT;
+				}
+				else if (bHasVFogCloud)
+				{
+					Compiler->MaterialCompilationOutput.StrataMaterialCompilationOutput.MaterialType = SUBSTRATE_MATERIAL_TYPE_VOLUMETRICFOGCLOUD;
+				}
+				else if (bHasHair)
+				{
+					Compiler->MaterialCompilationOutput.StrataMaterialCompilationOutput.MaterialType = SUBSTRATE_MATERIAL_TYPE_HAIR;
+				}
+				else if (bHasEye)
+				{
+					Compiler->MaterialCompilationOutput.StrataMaterialCompilationOutput.MaterialType = SUBSTRATE_MATERIAL_TYPE_EYE;
+				}
+				else if (bHasSLW)
+				{
+					Compiler->MaterialCompilationOutput.StrataMaterialCompilationOutput.MaterialType = SUBSTRATE_MATERIAL_TYPE_SINGLELAYERWATER;
+				}
+				else
+				{
+					Compiler->MaterialCompilationOutput.StrataMaterialCompilationOutput.MaterialType = SUBSTRATE_MATERIAL_TYPE_SINGLESLAB;
+				}
+
 				Compiler->MaterialCompilationOutput.StrataMaterialCompilationOutput.bMaterialOutOfBudgetHasBeenSimplified |= !StrataSimplificationStatus.bMaterialFitsInMemoryBudget;
 
 				if (OperatorCount <= STRATA_COMPILATION_OUTPUT_MAX_OPERATOR)
diff --git a/Engine/Source/Runtime/Engine/Private/Materials/MaterialShared.cpp b/Engine/Source/Runtime/Engine/Private/Materials/MaterialShared.cpp
--- a/Engine/Source/Runtime/Engine/Private/Materials/MaterialShared.cpp
+++ b/Engine/Source/Runtime/Engine/Private/Materials/MaterialShared.cpp
@@ -1545,6 +1545,7 @@ bool FMaterialResource::ShouldDoSSR() const { return Material->bScreenSpaceRefle
 bool FMaterialResource::ShouldDoContactShadows() const { return Material->bContactShadows; }
 bool FMaterialResource::IsWireframe() const { return Material->Wireframe; }
 bool FMaterialResource::IsUIMaterial() const { return Material->MaterialDomain == MD_UI; }
+bool FMaterialResource::IsPostProcessMaterial() const { return Material->MaterialDomain == MD_PostProcess; }
 bool FMaterialResource::IsLightFunction() const { return Material->MaterialDomain == MD_LightFunction; }
 bool FMaterialResource::IsUsedWithEditorCompositing() const { return Material->bUsedWithEditorCompositing; }
 bool FMaterialResource::IsDeferredDecal() const { return Material->MaterialDomain == MD_DeferredDecal; }
diff --git a/Engine/Source/Runtime/Engine/Public/MaterialShared.h b/Engine/Source/Runtime/Engine/Public/MaterialShared.h
--- a/Engine/Source/Runtime/Engine/Public/MaterialShared.h
+++ b/Engine/Source/Runtime/Engine/Public/MaterialShared.h
@@ -1844,6 +1844,7 @@ public:
 	virtual bool IsVolumetricPrimitive() const = 0;
 	virtual bool IsWireframe() const = 0;
 	virtual bool IsUIMaterial() const { return false; }
+	virtual bool IsPostProcessMaterial() const { return false; }
 	virtual bool IsSpecialEngineMaterial() const = 0;
 	virtual bool IsUsedWithSkeletalMesh() const { return false; }
 	virtual bool IsUsedWithLandscape() const { return false; }
@@ -2575,6 +2576,7 @@ public:
 	ENGINE_API virtual bool IsWireframe() const override;
 	ENGINE_API virtual EMaterialShadingRate  GetShadingRate() const override;
 	ENGINE_API virtual bool IsUIMaterial() const override;
+	ENGINE_API virtual bool IsPostProcessMaterial() const override;
 	ENGINE_API virtual bool IsSpecialEngineMaterial() const override;
 	ENGINE_API virtual bool IsUsedWithSkeletalMesh() const override;
 	ENGINE_API virtual bool IsUsedWithLandscape() const override;
diff --git a/Engine/Source/Runtime/Engine/Public/Rendering/StrataMaterialShared.h b/Engine/Source/Runtime/Engine/Public/Rendering/StrataMaterialShared.h
--- a/Engine/Source/Runtime/Engine/Public/Rendering/StrataMaterialShared.h
+++ b/Engine/Source/Runtime/Engine/Public/Rendering/StrataMaterialShared.h
@@ -74,6 +74,19 @@ public:
 
 #define STRATA_COMPILATION_OUTPUT_MAX_OPERATOR 24
 
+
+#define SUBSTRATE_MATERIAL_TYPE_SINGLESLAB			0
+#define SUBSTRATE_MATERIAL_TYPE_MULTIPLESLABS		1
+#define SUBSTRATE_MATERIAL_TYPE_VOLUMETRICFOGCLOUD	2
+#define SUBSTRATE_MATERIAL_TYPE_UNLIT				3
+#define SUBSTRATE_MATERIAL_TYPE_HAIR				4
+#define SUBSTRATE_MATERIAL_TYPE_SINGLELAYERWATER	5
+#define SUBSTRATE_MATERIAL_TYPE_EYE					6
+#define SUBSTRATE_MATERIAL_TYPE_LIGHTFUNCTION		7
+#define SUBSTRATE_MATERIAL_TYPE_POSTPROCESS			8
+#define SUBSTRATE_MATERIAL_TYPE_UI					9
+#define SUBSTRATE_MATERIAL_TYPE_DECAL				10
+
 struct FStrataMaterialCompilationOutput
 {
 	DECLARE_TYPE_LAYOUT(FStrataMaterialCompilationOutput, NonVirtual);
@@ -94,14 +107,9 @@ public:
 	/** Strata uint per pixel, at compile time (0-255) */
 	LAYOUT_FIELD(uint8, StrataUintPerPixel);
 
-	/** Indicate that the material is considered a thin surface instead of a volume filled up with matter */
-	LAYOUT_FIELD(uint8, bIsThin);
-
 	////
 	//// The following data is only needed when compiling with the editor.
 	////
-	
-	// STRATA_TODO corresponding uint8 ShadingModel? IsThin?
 
 	/** The Strata verbose description */
 	LAYOUT_FIELD_EDITORONLY(FMemoryImageString, StrataMaterialDescription);
@@ -113,6 +121,11 @@ public:
 	/** The byte count per pixel supported by the platform the material has been compiled against */
 	LAYOUT_FIELD_EDITORONLY(uint8, PlatformBytePixePixel);
 
+	/** Indicate that the material is considered a thin surface instead of a volume filled up with matter */
+	LAYOUT_FIELD_EDITORONLY(uint8, bIsThin);
+	/** Indicate the final material type */
+	LAYOUT_FIELD_EDITORONLY(uint8, MaterialType);
+
 	/** The byte per pixel count supported by the platform the material has been compiled against */
 	LAYOUT_BITFIELD_EDITORONLY(uint8, bMaterialOutOfBudgetHasBeenSimplified, 1);
 
