World Partition : fixed scaling issue in 2D debug display when world partition has a rotation.

diff --git a/Engine/Source/Runtime/Engine/Private/WorldPartition/WorldPartitionRuntimeSpatialHash.cpp b/Engine/Source/Runtime/Engine/Private/WorldPartition/WorldPartitionRuntimeSpatialHash.cpp
--- a/Engine/Source/Runtime/Engine/Private/WorldPartition/WorldPartitionRuntimeSpatialHash.cpp
+++ b/Engine/Source/Runtime/Engine/Private/WorldPartition/WorldPartitionRuntimeSpatialHash.cpp
@@ -1801,7 +1801,7 @@ bool UWorldPartitionRuntimeSpatialHash::Draw2D(FWorldPartitionDraw2DContext& Dra
 		}
 
 		const FVector2D GridReferenceWorldPos = FVector2D(WorldRegion.GetCenter());
-		const FVector2D WorldRegionExtent = WorldRegion.GetExtent();
+		const FVector2D WorldRegionExtent = FVector2D(WorldRegion.GetExtent().GetMax());
 		const FVector2D GridScreenOffset = GridScreenInitialOffset + ((float)GridIndex * FVector2D(GridMaxScreenWidth, 0.f)) + GridScreenHalfExtent + FVector2D(GridScreenWidthShrinkSize * 0.5f);
 		const FVector2D WorldToScreenScale = GridScreenHalfExtent / WorldRegionExtent;
 		const FBox2D GridScreenBounds(GridScreenOffset - GridScreenHalfExtent, GridScreenOffset + GridScreenHalfExtent);
@@ -1847,10 +1847,11 @@ bool UWorldPartitionRuntimeSpatialHash::Draw2D(FWorldPartitionDraw2DContext& Dra
 	{
 		// Convert to 2D
 		FBox2D GridsShapeBounds2D = FBox2D(FVector2D(GridsShapeBounds.Min.X, GridsShapeBounds.Min.Y), FVector2D(GridsShapeBounds.Max.X, GridsShapeBounds.Max.Y));
+		FVector2D CenterGlobalPos = LocalWPToGlobal.TransformPoint(GridsShapeBounds2D.GetCenter());
 		// Use max extent of X/Y
-		GridsShapeBounds2D = FBox2D(GridsShapeBounds2D.GetCenter() - FVector2D(GridsShapeBounds2D.GetExtent().GetMax()), GridsShapeBounds2D.GetCenter() + FVector2D(GridsShapeBounds2D.GetExtent().GetMax()));
+		FVector2D Extent = FVector2D(GridsShapeBounds2D.GetExtent().GetMax());
 		// Transform to global space
-		GridsShapeBounds2D = FBox2D(LocalWPToGlobal.TransformPoint(GridsShapeBounds2D.Min), LocalWPToGlobal.TransformPoint(GridsShapeBounds2D.Max));
+		GridsShapeBounds2D = FBox2D(CenterGlobalPos - Extent, CenterGlobalPos + Extent);
 		// Expand by 10% computed bounds
 		DesiredWorldBounds = GridsShapeBounds2D.ExpandBy(GridsShapeBounds2D.GetExtent() * 0.1f);
 	}
