Uninitializing a partitioned world should always trigger unloading of all cells regardless if server streaming and/or server streaming out are enabled/disabled.

diff --git a/Engine/Source/Runtime/Engine/Private/WorldPartition/WorldPartitionStreamingPolicy.cpp b/Engine/Source/Runtime/Engine/Private/WorldPartition/WorldPartitionStreamingPolicy.cpp
--- a/Engine/Source/Runtime/Engine/Private/WorldPartition/WorldPartitionStreamingPolicy.cpp
+++ b/Engine/Source/Runtime/Engine/Private/WorldPartition/WorldPartitionStreamingPolicy.cpp
@@ -275,9 +275,11 @@ void UWorldPartitionStreamingPolicy::UpdateStreamingState()
 		CriticalPerformanceBlockTillLevelStreamingCompletedEpoch = World->GetBlockTillLevelStreamingCompletedEpoch();
 	}
 
+	const bool bCanStream = WorldPartition->CanStream();
 	const bool bForceFrameUpdate = (UWorldPartitionStreamingPolicy::ForceUpdateFrameCount > 0) ? ((UpdateStreamingStateCalls % UWorldPartitionStreamingPolicy::ForceUpdateFrameCount) == 0) : false;
 	const bool bCanOptimizeUpdate =
 		WorldPartition->RuntimeHash &&
+		bCanStream &&
 		!bForceFrameUpdate &&										// We garantee to update every N frame to force some internal updates like UpdateStreamingPerformance
 		IsUpdateStreamingOptimEnabled() &&							// Check CVars to see if optimization is enabled
 		bLastUpdateCompletedLoadingAndActivation &&					// Don't optimize if last frame didn't process all cells to load/activate
@@ -311,7 +313,6 @@ void UWorldPartitionStreamingPolicy::UpdateStreamingState()
 		FrameLoadCells.Reset();
 	};
 
-	const bool bCanStream = WorldPartition->CanStream();
 	const bool bIsServer = WorldPartition->IsServer();
 	const bool bIsServerStreamingEnabled = WorldPartition->IsServerStreamingEnabled();
 	const int32 NewServerStreamingEnabledEpoch = ComputeServerStreamingEnabledEpoch();
