Fix FActorInstanceHandle::FetchActor<> crashing if called on a managed handle where the LWI manager fails to spawn an actor ( for eg. on the client ).

diff --git a/Engine/Source/Runtime/Engine/Classes/Engine/ActorInstanceHandle.h b/Engine/Source/Runtime/Engine/Classes/Engine/ActorInstanceHandle.h
--- a/Engine/Source/Runtime/Engine/Classes/Engine/ActorInstanceHandle.h
+++ b/Engine/Source/Runtime/Engine/Classes/Engine/ActorInstanceHandle.h
@@ -118,7 +118,7 @@ T* FActorInstanceHandle::FetchActor() const
 {
 	if (DoesRepresent<T>())
 	{
-		return CastChecked<T>(FetchActor());
+		return CastChecked<T>(FetchActor(), ECastCheckedType::NullAllowed);
 	}
 	return nullptr;
 }
