Extract out some functions in the LightWeightInstanceStaticMeshManager to allow adding support for subclasses to absorb spawned actors back into the manager as instances.

diff --git a/Engine/Source/Runtime/Engine/Classes/GameFramework/LightWeightInstanceStaticMeshManager.h b/Engine/Source/Runtime/Engine/Classes/GameFramework/LightWeightInstanceStaticMeshManager.h
--- a/Engine/Source/Runtime/Engine/Classes/GameFramework/LightWeightInstanceStaticMeshManager.h
+++ b/Engine/Source/Runtime/Engine/Classes/GameFramework/LightWeightInstanceStaticMeshManager.h
@@ -30,9 +30,10 @@ class ENGINE_API ALightWeightInstanceStaticMeshManager : public ALightWeightInst
 
 protected:
 	virtual void AddNewInstanceAt(FLWIData* InitData, int32 Index) override;
-
 	virtual void RemoveInstance(int32 Index) override;
+	virtual void UpdateDataAtIndex(FLWIData* InData, int32 Index) override;
 
+	void AddInstanceToRendering(int32 DataIndex);
 	void RemoveInstanceFromRendering(int32 DataIndex);
 
 	void PostRemoveInstanceFromRendering();
diff --git a/Engine/Source/Runtime/Engine/Private/LightWeightInstanceStaticMeshManager.cpp b/Engine/Source/Runtime/Engine/Private/LightWeightInstanceStaticMeshManager.cpp
--- a/Engine/Source/Runtime/Engine/Private/LightWeightInstanceStaticMeshManager.cpp
+++ b/Engine/Source/Runtime/Engine/Private/LightWeightInstanceStaticMeshManager.cpp
@@ -77,36 +77,54 @@ void ALightWeightInstanceStaticMeshManager::AddNewInstanceAt(FLWIData* InitData,
 {
 	Super::AddNewInstanceAt(InitData, Index);
 
+	AddInstanceToRendering(Index);
+}
+
+void ALightWeightInstanceStaticMeshManager::RemoveInstance(const int32 Index)
+{
+#if WITH_EDITOR
+	Modify();
+#endif
+
+	RemoveInstanceFromRendering(Index);
+	Super::RemoveInstance(Index);
+}
+
+void ALightWeightInstanceStaticMeshManager::UpdateDataAtIndex(FLWIData* InData, int32 Index)
+{
+	Super::UpdateDataAtIndex(InData, Index);
+
+	if (!RenderingIndicesToDataIndices.Contains(Index))
+	{
+		AddInstanceToRendering(Index);
+	}
+}
+
+void ALightWeightInstanceStaticMeshManager::AddInstanceToRendering(int32 DataIndex)
+{
+	//cancel any pending deletes
+	DataIndicesToBeDeleted.RemoveSingle(DataIndex);
+
 	// The rendering indices are tightly packed so we know it's going on the end of the array
-	const int32 RenderingIdx = RenderingIndicesToDataIndices.Add(Index);
+	const int32 RenderingIdx = RenderingIndicesToDataIndices.Add(DataIndex);
 
 	// Now that we know the rendering index we can fill in the other side of the map
-	if (Index >= DataIndicesToRenderingIndices.Num())
+	if (DataIndex >= DataIndicesToRenderingIndices.Num())
 	{
-		ensure(Index == DataIndicesToRenderingIndices.Add(RenderingIdx));
+		ensure(DataIndex == DataIndicesToRenderingIndices.Add(RenderingIdx));
 	}
 	else
 	{
-		DataIndicesToRenderingIndices[Index] = RenderingIdx;
+		DataIndicesToRenderingIndices[DataIndex] = RenderingIdx;
 	}
 
 	// Update the HISMC
 	if (InstancedStaticMeshComponent)
 	{
-		ensure(RenderingIdx == InstancedStaticMeshComponent->AddInstance(InstanceTransforms[Index], /*bWorldSpace*/false));
+		ensure(RenderingIdx == InstancedStaticMeshComponent->AddInstance(InstanceTransforms[DataIndex], /*bWorldSpace*/false));
 	}
 }
 
-void ALightWeightInstanceStaticMeshManager::RemoveInstance(const int32 Index)
-{
-#if WITH_EDITOR
-	Modify();
-#endif
-
-	RemoveInstanceFromRendering(Index);
-	Super::RemoveInstance(Index);
-}
-
 void ALightWeightInstanceStaticMeshManager::RemoveInstanceFromRendering(int32 DataIndex)
 {
 	if (IsIndexValid(DataIndex))
