Fixed for bad shader map content non deterministic key.

diff --git a/Engine/Shaders/Public/ShaderVersion.ush b/Engine/Shaders/Public/ShaderVersion.ush
--- a/Engine/Shaders/Public/ShaderVersion.ush
+++ b/Engine/Shaders/Public/ShaderVersion.ush
@@ -2,4 +2,5 @@
 // Each time the console command is executed it generates a new GUID. As this file is included
 // in Platform.ush (which should be included in any shader) it allows to invalidate the shader DDC.
 // 
-// GUID = F7E16E56-DB0E-4CEA-99EA-3278948536F2
\ No newline at end of file
+// GUID = 760DE428-F661-11ED-B67E-0242AC120002
+
diff --git a/Engine/Source/Runtime/Engine/Public/Rendering/StrataMaterialShared.cpp b/Engine/Source/Runtime/Engine/Public/Rendering/StrataMaterialShared.cpp
--- a/Engine/Source/Runtime/Engine/Public/Rendering/StrataMaterialShared.cpp
+++ b/Engine/Source/Runtime/Engine/Public/Rendering/StrataMaterialShared.cpp
@@ -14,6 +14,8 @@ FStrataMaterialCompilationOutput::FStrataMaterialCompilationOutput()
 	, SharedLocalBasesCount(0)
 	, RequestedBytePixePixel(0)
 	, PlatformBytePixePixel(0)
+	, bIsThin(0)
+	, MaterialType(0)
 	, bMaterialOutOfBudgetHasBeenSimplified(0)
 	, RootOperatorIndex(0)
 #endif
