LevelInstance: Fix cook regression, include actor name in FLevelInstanceID (as it was before adding Short Package Name which should be IsGameWorld() only)

#rnx

diff --git a/Engine/Source/Runtime/Engine/Private/LevelInstance/LevelInstanceSubsystem.cpp b/Engine/Source/Runtime/Engine/Private/LevelInstance/LevelInstanceSubsystem.cpp
--- a/Engine/Source/Runtime/Engine/Private/LevelInstance/LevelInstanceSubsystem.cpp
+++ b/Engine/Source/Runtime/Engine/Private/LevelInstance/LevelInstanceSubsystem.cpp
@@ -126,13 +126,15 @@ FLevelInstanceID::FLevelInstanceID(ULevelInstanceSubsystem* LevelInstanceSubsyst
 	
 	uint64 TmpHash = 0;
 		
-	if (LevelInstanceActor->IsNameStableForNetworking() && LevelInstanceActor->GetWorld()->IsGameWorld())
+	if (LevelInstanceActor->IsNameStableForNetworking())
 	{
 		ActorName = LevelInstanceActor->GetFName();
 		FString NameStr = ActorName.ToString();
 		TmpHash = CityHash64((const char*)*NameStr, NameStr.Len() * sizeof(TCHAR));
 
-		if (UWorld* OuterWorld = LevelInstanceActor->GetTypedOuter<UWorld>())
+		// Only include OuterWorld shortpackage name in GameWorlds. In Editor this would cause the FLevelInstanceID to change for an actor within another LevelInstance
+		// based off if the parent LevelInstance was in Edit (non instanced) or not (instanced)
+		if (UWorld* OuterWorld = LevelInstanceActor->GetTypedOuter<UWorld>(); OuterWorld && LevelInstanceActor->GetWorld()->IsGameWorld())
 		{
 			PackageShortName = UWorld::RemovePIEPrefix(FPackageName::GetShortName(OuterWorld->GetPackage()));
 			TmpHash = CityHash64WithSeed((const char*)*PackageShortName, PackageShortName.Len() * sizeof(TCHAR), TmpHash);
