Avoid ensure being thrown in FGenericPlatformMath::Fmod during preshader material evaluation when the divisor is a 0 constant, corrects the value with GetSafeDivisor() call.

diff --git a/Engine/Source/Runtime/Engine/Private/Shader/ShaderValue.cpp b/Engine/Source/Runtime/Engine/Private/Shader/ShaderValue.cpp
--- a/Engine/Source/Runtime/Engine/Private/Shader/ShaderValue.cpp
+++ b/Engine/Source/Runtime/Engine/Private/Shader/ShaderValue.cpp
@@ -995,7 +995,7 @@ struct FOpMul : public FOpBase { template<typename T> T operator()(T Lhs, T Rhs)
 struct FOpDiv : public FOpBase { template<typename T> T operator()(T Lhs, T Rhs) const { return Lhs / GetSafeDivisor(Rhs); } };
 struct FOpMin : public FOpBase { template<typename T> T operator()(T Lhs, T Rhs) const { return FMath::Min(Lhs, Rhs); } };
 struct FOpMax : public FOpBase { template<typename T> T operator()(T Lhs, T Rhs) const { return FMath::Max(Lhs, Rhs); } };
-struct FOpFmod : public FOpBaseNoInt { template<typename T> T operator()(T Lhs, T Rhs) const { return FMath::Fmod(Lhs, Rhs); } };
+struct FOpFmod : public FOpBaseNoInt { template<typename T> T operator()(T Lhs, T Rhs) const { return FMath::Fmod(Lhs, GetSafeDivisor(Rhs)); } };
 struct FOpAtan2 : public FOpBaseNoInt { template<typename T> T operator()(T Lhs, T Rhs) const { return FMath::Atan2(Lhs, Rhs); } };
 struct FOpLess : public FOpBase { template<typename T> bool operator()(T Lhs, T Rhs) const { return Lhs < Rhs; } };
 struct FOpGreater : public FOpBase { template<typename T> bool operator()(T Lhs, T Rhs) const { return Lhs > Rhs; } };
