Packed Level Actor: Filter support

#rnx

diff --git a/Engine/Source/Editor/LevelInstanceEditor/Private/LevelInstanceEditorModule.cpp b/Engine/Source/Editor/LevelInstanceEditor/Private/LevelInstanceEditorModule.cpp
--- a/Engine/Source/Editor/LevelInstanceEditor/Private/LevelInstanceEditorModule.cpp
+++ b/Engine/Source/Editor/LevelInstanceEditor/Private/LevelInstanceEditorModule.cpp
@@ -263,6 +263,26 @@ namespace LevelInstanceMenuUtils
 		}
 	}
 
+	UClass* GetDefaultLevelInstanceClass(ELevelInstanceCreationType CreationType)
+	{
+		if (CreationType == ELevelInstanceCreationType::PackedLevelActor)
+		{
+			return APackedLevelActor::StaticClass();
+		}
+
+		ULevelInstanceEditorSettings* LevelInstanceEditorSettings = GetMutableDefault<ULevelInstanceEditorSettings>();
+		if (!LevelInstanceEditorSettings->LevelInstanceClassName.IsEmpty())
+		{
+			UClass* LevelInstanceClass = LoadClass<AActor>(nullptr, *LevelInstanceEditorSettings->LevelInstanceClassName, nullptr, LOAD_NoWarn);
+			if (LevelInstanceClass && LevelInstanceClass->ImplementsInterface(ULevelInstanceInterface::StaticClass()))
+			{
+				return LevelInstanceClass;
+			}
+		}
+
+		return ALevelInstance::StaticClass();
+	}
+
 	void CreateLevelInstanceFromSelection(ULevelInstanceSubsystem* LevelInstanceSubsystem, ELevelInstanceCreationType CreationType)
 	{
 		TArray<AActor*> ActorsToMove;
@@ -319,15 +339,7 @@ namespace LevelInstanceMenuUtils
 				UPackage* TemplatePackage = !TemplateMapPackage.IsEmpty() ? LoadPackage(nullptr, *TemplateMapPackage, LOAD_None) : nullptr;
 				
 				CreationParams.TemplateWorld = TemplatePackage ? UWorld::FindWorldInPackage(TemplatePackage) : nullptr;
-				
-				if (!LevelInstanceEditorSettings->LevelInstanceClassName.IsEmpty())
-				{
-					UClass* LevelInstanceClass = LoadClass<AActor>(nullptr, *LevelInstanceEditorSettings->LevelInstanceClassName, nullptr, LOAD_NoWarn);
-					if (LevelInstanceClass && LevelInstanceClass->ImplementsInterface(ULevelInstanceInterface::StaticClass()))
-					{
-						CreationParams.LevelInstanceClass = LevelInstanceClass;
-					}
-				}
+				CreationParams.LevelInstanceClass = GetDefaultLevelInstanceClass(CreationType);
 
 				if (!LevelInstanceSubsystem->CreateLevelInstanceFrom(ActorsToMove, CreationParams))
 				{
@@ -508,12 +520,12 @@ namespace LevelInstanceMenuUtils
 		
 		virtual bool IsClassAllowed(const FClassViewerInitializationOptions& InInitOptions, const UClass* InClass, TSharedRef< FClassViewerFilterFuncs > InFilterFuncs) override
 		{
-			return InClass && InClass->ImplementsInterface(ULevelInstanceInterface::StaticClass()) && !InClass->IsChildOf(APackedLevelActor::StaticClass()) && !InClass->HasAnyClassFlags(CLASS_Deprecated);
+			return InClass && InClass->ImplementsInterface(ULevelInstanceInterface::StaticClass()) && InClass->IsNative() && !InClass->HasAnyClassFlags(CLASS_Deprecated);
 		}
 
 		virtual bool IsUnloadedClassAllowed(const FClassViewerInitializationOptions& InInitOptions, const TSharedRef< const IUnloadedBlueprintData > InUnloadedClassData, TSharedRef< FClassViewerFilterFuncs > InFilterFuncs) override
 		{
-			return InUnloadedClassData->ImplementsInterface(ULevelInstanceInterface::StaticClass()) && !InUnloadedClassData->IsChildOf(APackedLevelActor::StaticClass())  && !InUnloadedClassData->HasAnyClassFlags(CLASS_Deprecated);
+			return false;
 		}
 	};
 
@@ -526,7 +538,7 @@ namespace LevelInstanceMenuUtils
 		LongPackageName.FindLastChar('/', LastSlashIndex);
 		
 		FString PackagePath = LongPackageName.Mid(0, LastSlashIndex == INDEX_NONE ? MAX_int32 : LastSlashIndex);
-		FString AssetName = LevelInstancePtr.GetAssetName() + "_LevelInstance";
+		FString AssetName = "BP_" + LevelInstancePtr.GetAssetName();
 		IAssetTools& AssetTools = FAssetToolsModule::GetModule().Get();
 
 		UBlueprintFactory* BlueprintFactory = NewObject<UBlueprintFactory>();
@@ -552,6 +564,11 @@ namespace LevelInstanceMenuUtils
 			LevelInstanceCDO->SetWorldAsset(LevelInstancePtr);
 			FBlueprintEditorUtils::MarkBlueprintAsModified(NewBlueprint);
 			
+			if (NewBlueprint->GeneratedClass->IsChildOf<APackedLevelActor>())
+			{
+				APackedLevelActor::UpdateBlueprint(NewBlueprint);
+			}
+
 			FContentBrowserModule& ContentBrowserModule = FModuleManager::LoadModuleChecked<FContentBrowserModule>("ContentBrowser");
 			TArray<UObject*> Assets;
 			Assets.Add(NewBlueprint);
diff --git a/Engine/Source/Editor/LevelInstanceEditor/Private/LevelInstanceFilterPropertyTypeCustomization.cpp b/Engine/Source/Editor/LevelInstanceEditor/Private/LevelInstanceFilterPropertyTypeCustomization.cpp
--- a/Engine/Source/Editor/LevelInstanceEditor/Private/LevelInstanceFilterPropertyTypeCustomization.cpp
+++ b/Engine/Source/Editor/LevelInstanceEditor/Private/LevelInstanceFilterPropertyTypeCustomization.cpp
@@ -24,6 +24,8 @@ TSharedPtr<FWorldPartitionActorFilterMode::FFilter> FLevelInstanceFilterProperty
 	// Find Selected Level Instances with matching WorldAssetPackage
 	FString WorldAssetPackage;
 	UWorld* World = nullptr;
+	TOptional<EWorldPartitionActorFilterType> Types;
+
 	for (int32 OuterObjectIndex = 0; OuterObjectIndex < OuterObjects.Num(); ++OuterObjectIndex)
 	{
 		UObject* OuterObject = OuterObjects[OuterObjectIndex];
@@ -33,7 +35,20 @@ TSharedPtr<FWorldPartitionActorFilterMode::FFilter> FLevelInstanceFilterProperty
 			{
 				LevelInstances.Add(LevelInstanceInterface, OuterObjectIndex);
 				ActorLabel = OuterActor->GetActorLabel();
-
+				
+				if (!Types.IsSet())
+				{
+					Types = LevelInstanceInterface->GetDetailsFilterTypes();
+					if (Types == EWorldPartitionActorFilterType::None)
+					{
+						return nullptr;
+					}
+				}
+				else if (Types != LevelInstanceInterface->GetDetailsFilterTypes())
+				{
+					return nullptr;
+				}
+				
 				if (WorldAssetPackage.IsEmpty())
 				{
 					WorldAssetPackage = LevelInstanceInterface->GetWorldAssetPackage();
@@ -72,7 +87,7 @@ TSharedPtr<FWorldPartitionActorFilterMode::FFilter> FLevelInstanceFilterProperty
 	UWorldPartitionSubsystem* WorldPartitionSubsystem = World->GetSubsystem<UWorldPartitionSubsystem>();
 
 	// Get the Default Filter for selected WorldAssetPackage
-	TSharedPtr<FWorldPartitionActorFilter> Filter = MakeShared<FWorldPartitionActorFilter>(WorldPartitionSubsystem->GetWorldPartitionActorFilter(WorldAssetPackage));
+	TSharedPtr<FWorldPartitionActorFilter> Filter = MakeShared<FWorldPartitionActorFilter>(WorldPartitionSubsystem->GetWorldPartitionActorFilter(WorldAssetPackage, Types.GetValue()));
 
 	// Set its name based on single/multi selection (root node name in outliner)
 	Filter->DisplayName = LevelInstances.Num() == 1 ? ActorLabel : TEXT("(Multiple Actors)");
diff --git a/Engine/Source/Runtime/Core/Public/UObject/FortniteMainBranchObjectVersion.h b/Engine/Source/Runtime/Core/Public/UObject/FortniteMainBranchObjectVersion.h
--- a/Engine/Source/Runtime/Core/Public/UObject/FortniteMainBranchObjectVersion.h
+++ b/Engine/Source/Runtime/Core/Public/UObject/FortniteMainBranchObjectVersion.h
@@ -322,6 +322,9 @@ struct CORE_API FFortniteMainBranchObjectVersion
 		// FWorldPartitionActorFilter go back to FString serialize of AssetPaths to avoid FArchiveReplaceOrClearExternalReferences clearing CDO references on BP Compile
 		WorldPartitionActorFilterStringAssetPath,
 
+		// Add FPackedLevelActorDesc for APackedLevelActor and support for APackedLevelActor Filters
+		PackedLevelActorDesc,
+
 		// -----<new versions can be added above this line>-------------------------------------------------
 		VersionPlusOne,
 		LatestVersion = VersionPlusOne - 1
diff --git a/Engine/Source/Runtime/Engine/Private/LevelInstance/LevelInstanceComponent.cpp b/Engine/Source/Runtime/Engine/Private/LevelInstance/LevelInstanceComponent.cpp
--- a/Engine/Source/Runtime/Engine/Private/LevelInstance/LevelInstanceComponent.cpp
+++ b/Engine/Source/Runtime/Engine/Private/LevelInstance/LevelInstanceComponent.cpp
@@ -73,10 +73,19 @@ void ULevelInstanceComponent::SetFilter(const FWorldPartitionActorFilter& InFilt
 	{
 		Modify();
 		Filter = InFilter;
+		OnFilterChanged();
 		FWorldPartitionActorFilter::GetOnWorldPartitionActorFilterChanged().Broadcast();
 	}
 }
 
+void ULevelInstanceComponent::OnFilterChanged()
+{
+	if (ILevelInstanceInterface* LevelInstance = Cast<ILevelInstanceInterface>(GetOwner()))
+	{
+		LevelInstance->OnFilterChanged();
+	}
+}
+
 bool ULevelInstanceComponent::ShouldShowSpriteComponent() const
 {
 	return GetOwner() && GetOwner()->GetLevel() && (GetOwner()->GetLevel()->IsPersistentLevel() || !GetOwner()->GetLevel()->IsInstancedLevel());
@@ -96,6 +105,7 @@ void ULevelInstanceComponent::PostEditUndo()
 
 	if (Filter != UndoRedoCachedFilter)
 	{
+		OnFilterChanged();
 		FWorldPartitionActorFilter::RequestFilterRefresh(false);
 		FWorldPartitionActorFilter::GetOnWorldPartitionActorFilterChanged().Broadcast();
 	}
@@ -107,6 +117,11 @@ void ULevelInstanceComponent::PostEditChangeProperty(FPropertyChangedEvent& Prop
 	Super::PostEditChangeProperty(PropertyChangedEvent);
 
 	UpdateEditorInstanceActor();
+
+	if (PropertyChangedEvent.Property && PropertyChangedEvent.Property->GetFName() == GET_MEMBER_NAME_CHECKED(ULevelInstanceComponent, Filter))
+	{
+		OnFilterChanged();
+	}
 }
 
 void ULevelInstanceComponent::OnUpdateTransform(EUpdateTransformFlags UpdateTransformFlags, ETeleportType Teleport)
@@ -182,7 +197,7 @@ const TMap<FActorContainerID, TSet<FGuid>>& ULevelInstanceComponent::GetFiltered
 		UWorldPartitionSubsystem* WorldPartitionSubsystem = GetWorld()->GetSubsystem<UWorldPartitionSubsystem>();
 		check(WorldPartitionSubsystem);
 				
-		FilteredActors = WorldPartitionSubsystem->GetFilteredActorsPerContainer(LevelInstance->GetLevelInstanceID().GetContainerID(), LevelInstance->GetWorldAssetPackage(), Filter);
+		FilteredActors = WorldPartitionSubsystem->GetFilteredActorsPerContainer(LevelInstance->GetLevelInstanceID().GetContainerID(), LevelInstance->GetWorldAssetPackage(), Filter, LevelInstance->GetLoadingFilterTypes());
 	}
 	
 	CachedFilteredActorsPerContainer = MoveTemp(FilteredActors);
diff --git a/Engine/Source/Runtime/Engine/Private/LevelInstance/LevelInstanceInterface.cpp b/Engine/Source/Runtime/Engine/Private/LevelInstance/LevelInstanceInterface.cpp
--- a/Engine/Source/Runtime/Engine/Private/LevelInstance/LevelInstanceInterface.cpp
+++ b/Engine/Source/Runtime/Engine/Private/LevelInstance/LevelInstanceInterface.cpp
@@ -319,6 +319,37 @@ const TMap<FActorContainerID, TSet<FGuid>>& ILevelInstanceInterface::GetFiltered
 	return NoFilteredActors;
 }
 
+void ILevelInstanceInterface::SetFilter(const FWorldPartitionActorFilter& InFilter)
+{
+	if (ULevelInstanceComponent* LevelInstanceComponent = GetLevelInstanceComponent())
+	{
+		LevelInstanceComponent->SetFilter(InFilter);
+	}
+}
+
+EWorldPartitionActorFilterType ILevelInstanceInterface::GetLoadingFilterTypes() const
+{
+	const AActor* Actor = CastChecked<AActor>(this);
+	if (Actor->GetWorld() && Actor->GetWorld()->IsPartitionedWorld())
+	{
+		return EWorldPartitionActorFilterType::All;
+	}
+
+	return EWorldPartitionActorFilterType::None;
+}
+
+
+EWorldPartitionActorFilterType ILevelInstanceInterface::GetDetailsFilterTypes() const
+{
+	const AActor* Actor = CastChecked<AActor>(this);
+	if (Actor->IsTemplate() || (Actor->GetWorld() && Actor->GetWorld()->IsPartitionedWorld()))
+	{
+		return EWorldPartitionActorFilterType::Loading;
+	}
+
+	return EWorldPartitionActorFilterType::None;
+}
+
 #endif
 
 
diff --git a/Engine/Source/Runtime/Engine/Private/LevelInstance/LevelInstanceSubsystem.cpp b/Engine/Source/Runtime/Engine/Private/LevelInstance/LevelInstanceSubsystem.cpp
--- a/Engine/Source/Runtime/Engine/Private/LevelInstance/LevelInstanceSubsystem.cpp
+++ b/Engine/Source/Runtime/Engine/Private/LevelInstance/LevelInstanceSubsystem.cpp
@@ -656,7 +656,7 @@ void ULevelInstanceSubsystem::PackAllLoadedActors()
 	TFunction<void(APackedLevelActor*, TArray<UBlueprint*>&, TArray<APackedLevelActor*>&)> GatherDepencenciesRecursive = [&GatherDepencenciesRecursive](APackedLevelActor* PackedLevelActor, TArray<UBlueprint*>& BPsToPack, TArray<APackedLevelActor*>& ToPack)
 	{
 		// Early out on already processed BPs or non BP Packed LIs.
-		UBlueprint* Blueprint = Cast<UBlueprint>(PackedLevelActor->GetClass()->ClassGeneratedBy);
+		UBlueprint* Blueprint = PackedLevelActor->GetRootBlueprint();
 		if ((Blueprint && BPsToPack.Contains(Blueprint)) || ToPack.Contains(PackedLevelActor))
 		{
 			return;
@@ -1058,7 +1058,7 @@ ILevelInstanceInterface* ULevelInstanceSubsystem::CreateLevelInstanceFrom(const
 
 	ULevelStreamingLevelInstanceEditor* LevelStreaming = nullptr;
 	{
-		const bool bIsPartitioned = (CreationParams.Type != ELevelInstanceCreationType::PackedLevelActor) && GetWorld()->IsPartitionedWorld();
+		const bool bIsPartitioned = GetWorld()->IsPartitionedWorld();
 
 		// We want to properly setup the world partition prior to its initialization
 		FDelegateHandle PreWorldInit = FWorldDelegates::OnPreWorldInitialization.AddLambda([&bIsPartitioned](UWorld* World, const UWorld::InitializationValues IVS)
@@ -1156,22 +1156,22 @@ ILevelInstanceInterface* ULevelInstanceSubsystem::CreateLevelInstanceFrom(const
 			
 	// Make sure newly created level asset gets scanned
 	ULevel::ScanLevelAssets(LoadedLevel->GetPackage()->GetName());
-
-	if (CreationParams.Type == ELevelInstanceCreationType::LevelInstance)
+	
+	// Use CreationParams class if provided
+	UClass* ActorClass = CreationParams.LevelInstanceClass;
+	if (!ActorClass)
 	{
-		TSubclassOf<AActor> ActorClass = ALevelInstance::StaticClass();
-		if (CreationParams.LevelInstanceClass)
-		{
-			ActorClass = CreationParams.LevelInstanceClass;
-		}
+		ActorClass = CreationParams.Type == ELevelInstanceCreationType::LevelInstance ? ALevelInstance::StaticClass() : APackedLevelActor::StaticClass();
+	}
 
-		check(ActorClass->ImplementsInterface(ULevelInstanceInterface::StaticClass()));
+	check(ActorClass->ImplementsInterface(ULevelInstanceInterface::StaticClass()));
 
+	if (!ActorClass->IsChildOf<APackedLevelActor>())
+	{
 		NewLevelInstanceActor = GetWorld()->SpawnActor<AActor>(ActorClass, SpawnParams);
 	}
 	else
 	{
-		check(CreationParams.Type == ELevelInstanceCreationType::PackedLevelActor);
 		FString PackageDir = FPaths::GetPath(WorldPtr.GetLongPackageName());
 		FString AssetName = FPackedLevelActorBuilder::GetPackedBPPrefix() + WorldPtr.GetAssetName();
 		FString BPAssetPath = FString::Format(TEXT("{0}/{1}.{1}"), { PackageDir , AssetName });
diff --git a/Engine/Source/Runtime/Engine/Private/PackedLevelActor/IPackedLevelActorBuilder.h b/Engine/Source/Runtime/Engine/Private/PackedLevelActor/IPackedLevelActorBuilder.h
--- a/Engine/Source/Runtime/Engine/Private/PackedLevelActor/IPackedLevelActorBuilder.h
+++ b/Engine/Source/Runtime/Engine/Private/PackedLevelActor/IPackedLevelActorBuilder.h
@@ -20,7 +20,7 @@ public:
 	virtual ~IPackedLevelActorBuilder() {}
 	virtual FPackedLevelActorBuilderID GetID() const = 0;
 	virtual void GetPackClusters(FPackedLevelActorBuilderContext& InContext, AActor* InActor) const = 0;
-	virtual void PackActors(FPackedLevelActorBuilderContext& InBuilder, APackedLevelActor* InPackingActor, const FPackedLevelActorBuilderClusterID& InClusterID, const TArray<UActorComponent*>& InComponents) const = 0;
+	virtual void PackActors(FPackedLevelActorBuilderContext& InBuilder, const FPackedLevelActorBuilderClusterID& InClusterID, const TArray<UActorComponent*>& InComponents) const = 0;
 };
 
 #endif
\ No newline at end of file
diff --git a/Engine/Source/Runtime/Engine/Private/PackedLevelActor/PackedLevelActor.cpp b/Engine/Source/Runtime/Engine/Private/PackedLevelActor/PackedLevelActor.cpp
--- a/Engine/Source/Runtime/Engine/Private/PackedLevelActor/PackedLevelActor.cpp
+++ b/Engine/Source/Runtime/Engine/Private/PackedLevelActor/PackedLevelActor.cpp
@@ -4,6 +4,7 @@
 #include "Engine/Blueprint.h"
 #include "Engine/World.h"
 #include "PackedLevelActor/PackedLevelActorBuilder.h"
+#include "WorldPartition/PackedLevelActor/PackedLevelActorDesc.h"
 #include "LevelInstance/LevelInstanceSubsystem.h"
 #include "LevelInstance/LevelInstancePrivate.h"
 #include "UObject/Package.h"
@@ -13,8 +14,9 @@
 
 APackedLevelActor::APackedLevelActor()
 	: Super()
-#if WITH_EDITORONLY_DATA
+#if WITH_EDITOR
 	, bChildChanged(false)
+	, bLoadForPacking(false)
 #endif
 {
 #if WITH_EDITORONLY_DATA
@@ -40,7 +42,7 @@ void APackedLevelActor::Serialize(FArchive& Ar)
 bool APackedLevelActor::IsLoadingEnabled() const
 {
 #if WITH_EDITOR
-	return HasChildEdit() || IsLoaded();
+	return HasChildEdit() || IsLoaded() || ShouldLoadForPacking();
 #else
 	return false;
 #endif
@@ -87,9 +89,79 @@ void APackedLevelActor::RerunConstructionScripts()
 	}
 }
 
+EWorldPartitionActorFilterType APackedLevelActor::GetDetailsFilterTypes() const
+{
+	return IsRootBlueprintTemplate() ? EWorldPartitionActorFilterType::All : EWorldPartitionActorFilterType::None;
+}
+
+EWorldPartitionActorFilterType APackedLevelActor::GetLoadingFilterTypes() const
+{
+	return ShouldLoadForPacking() ? EWorldPartitionActorFilterType::All : EWorldPartitionActorFilterType::None;
+}
+
+bool APackedLevelActor::IsRootBlueprintTemplate() const
+{
+	return IsTemplate() && IsRootBlueprint(GetClass());
+}
+
+bool APackedLevelActor::IsRootBlueprint(UClass* InClass)
+{
+	return InClass && InClass->ClassGeneratedBy != nullptr && InClass->GetSuperClass()->IsNative();
+}
+
+UBlueprint* APackedLevelActor::GetRootBlueprint() const
+{
+	UClass* Class = GetClass();
+	while (Class->GetSuperClass() && !Class->GetSuperClass()->IsNative())
+	{
+		Class = Class->GetSuperClass();
+	}
+
+	return Cast<UBlueprint>(Class->ClassGeneratedBy);
+}
+
+void APackedLevelActor::OnFilterChanged()
+{
+	Super::OnFilterChanged();
+
+	if (IsRootBlueprintTemplate())
+	{
+		// Reflect child changes
+		TSharedPtr<FPackedLevelActorBuilder> Builder = FPackedLevelActorBuilder::CreateDefaultBuilder();
+		Builder->UpdateBlueprint(GetRootBlueprint(), false);
+	}
+}
+
+bool APackedLevelActor::ShouldLoadForPacking() const
+{
+	if (bLoadForPacking)
+	{
+		return true;
+	}
+
+	bool bAncestorLoadForPacking = false;
+	if (ULevelInstanceSubsystem* LevelInstanceSubsystem = GetLevelInstanceSubsystem())
+	{
+		LevelInstanceSubsystem->ForEachLevelInstanceAncestors(this, [&bAncestorLoadForPacking](const ILevelInstanceInterface* LevelInstanceInterface)
+		{
+			if(const APackedLevelActor* PackedAncestor = Cast<APackedLevelActor>(LevelInstanceInterface))
+			{
+				if (PackedAncestor->bLoadForPacking)
+				{
+					bAncestorLoadForPacking = true;
+					return false;
+				}
+			}
+			return true;
+		});
+	}
+
+	return bAncestorLoadForPacking;
+}
+
 TUniquePtr<class FWorldPartitionActorDesc> APackedLevelActor::CreateClassActorDesc() const
 {
-	return AActor::CreateClassActorDesc();
+	return TUniquePtr<FWorldPartitionActorDesc>(new FPackedLevelActorDesc());
 }
 
 bool APackedLevelActor::CreateOrUpdateBlueprint(ALevelInstance* InLevelInstance, TSoftObjectPtr<UBlueprint> InBlueprintAsset, bool bCheckoutAndSave, bool bPromptForSave)
@@ -102,6 +174,11 @@ bool APackedLevelActor::CreateOrUpdateBlueprint(TSoftObjectPtr<UWorld> InWorldAs
 	return FPackedLevelActorBuilder::CreateDefaultBuilder()->CreateOrUpdateBlueprint(InWorldAsset, InBlueprintAsset, bCheckoutAndSave, bPromptForSave);
 }
 
+void APackedLevelActor::UpdateBlueprint(UBlueprint* InBlueprint, bool bCheckoutAndSave)
+{
+	FPackedLevelActorBuilder::CreateDefaultBuilder()->UpdateBlueprint(InBlueprint, bCheckoutAndSave);
+}
+
 FName APackedLevelActor::GetPackedComponentTag()
 {
 	static FName PackedComponentTag("PackedComponent");
@@ -112,7 +189,7 @@ void APackedLevelActor::UpdateLevelInstanceFromWorldAsset()
 {
 	Super::UpdateLevelInstanceFromWorldAsset();
 
-	if (!Cast<UBlueprint>(GetClass()->ClassGeneratedBy))
+	if (!GetRootBlueprint())
 	{
 		if (IsWorldAssetValid())
 		{
@@ -152,9 +229,9 @@ void APackedLevelActor::OnCommitChild(bool bChanged)
 			// Reflect child changes
 			TSharedPtr<FPackedLevelActorBuilder> Builder = FPackedLevelActorBuilder::CreateDefaultBuilder();
 
-			if (UBlueprint* GeneratedBy = Cast<UBlueprint>(GetClass()->ClassGeneratedBy))
+			if (UBlueprint* BlueprintToPack = GetRootBlueprint())
 			{
-				Builder->UpdateBlueprint(GeneratedBy);
+				Builder->UpdateBlueprint(BlueprintToPack);
 				return; // return here because Actor might have been reinstanced
 			}
 			else
@@ -179,10 +256,10 @@ void APackedLevelActor::OnCommit(bool bChanged)
 
 	if (bChanged)
 	{
-		if (UBlueprint* GeneratedBy = Cast<UBlueprint>(GetClass()->ClassGeneratedBy))
+		if (UBlueprint* BlueprintToPack = GetRootBlueprint())
 		{
 			TSharedPtr<FPackedLevelActorBuilder> Builder = FPackedLevelActorBuilder::CreateDefaultBuilder();
-			Builder->UpdateBlueprint(GeneratedBy);
+			Builder->UpdateBlueprint(BlueprintToPack);
 			return; // return here because Actor might have been reinstanced
 		}
 	}
@@ -193,7 +270,7 @@ void APackedLevelActor::OnCommit(bool bChanged)
 
 bool APackedLevelActor::IsHiddenEd() const
 {
-	return Super::IsHiddenEd() || IsEditing() || HasChildEdit();
+	return Super::IsHiddenEd() || IsEditing() || HasChildEdit() || ShouldLoadForPacking();
 }
 
 bool APackedLevelActor::IsHLODRelevant() const
diff --git a/Engine/Source/Runtime/Engine/Private/PackedLevelActor/PackedLevelActorBuilder.cpp b/Engine/Source/Runtime/Engine/Private/PackedLevelActor/PackedLevelActorBuilder.cpp
--- a/Engine/Source/Runtime/Engine/Private/PackedLevelActor/PackedLevelActorBuilder.cpp
+++ b/Engine/Source/Runtime/Engine/Private/PackedLevelActor/PackedLevelActorBuilder.cpp
@@ -23,12 +23,15 @@
 #include "Engine/LevelStreaming.h"
 #include "FileHelpers.h"
 #include "Editor.h"
+#include "Editor/Transactor.h"
 #include "AssetToolsModule.h"
 #include "Factories/BlueprintFactory.h"
 
 #include "IContentBrowserSingleton.h"
 #include "ContentBrowserModule.h"
 #include "EditorDirectories.h"
+#include "WorldPartition/Filter/WorldPartitionActorFilter.h"
+
 
 #define LOCTEXT_NAMESPACE "FPackedLevelActorBuilder"
 
@@ -159,54 +162,45 @@ bool FPackedLevelActorBuilder::PackActor(APackedLevelActor* InPackedLevelActor)
 
 bool FPackedLevelActorBuilder::PackActor(APackedLevelActor* InPackedLevelActor, ALevelInstance* InLevelInstanceToPack)
 {
-	FMessageLog Log("PackedLevelActor");
-	Log.Info(FText::Format(LOCTEXT("PackingStarted", "Packing of '{0}' started..."), FText::FromString(InPackedLevelActor->GetWorldAssetPackage())));
-	
-	FPackedLevelActorBuilderContext Context(*this, InPackedLevelActor);
+	FPackedLevelActorBuilderContext Context(*this, InPackedLevelActor, InLevelInstanceToPack);
+	return PackActor(Context);
+}
 
-	InPackedLevelActor->DestroyPackedComponents();
+bool FPackedLevelActorBuilder::PackActor(FPackedLevelActorBuilderContext& InContext)
+{
+	FMessageLog Log("PackedLevelActor");
+	Log.Info(FText::Format(LOCTEXT("PackingStarted", "Packing of '{0}' started..."), FText::FromString(InContext.GetPackedLevelActor()->GetWorldAssetPackage())));
+		
+	InContext.GetPackedLevelActor()->DestroyPackedComponents();
 
-	ULevelInstanceSubsystem* LevelInstanceSubystem = InPackedLevelActor->GetLevelInstanceSubsystem();
+	ULevelInstanceSubsystem* LevelInstanceSubystem = InContext.GetPackedLevelActor()->GetLevelInstanceSubsystem();
 	check(LevelInstanceSubystem);
 	
-	ULevel* SourceLevel = LevelInstanceSubystem->GetLevelInstanceLevel(InLevelInstanceToPack);
+	ULevel* SourceLevel = LevelInstanceSubystem->GetLevelInstanceLevel(InContext.GetLevelInstanceToPack());
 	if (!SourceLevel)
 	{
-		Log.Error(FText::Format(LOCTEXT("FailedPackingNoLevel", "Packing of '{0}' failed"), FText::FromString(InPackedLevelActor->GetWorldAssetPackage())));
+		Log.Error(FText::Format(LOCTEXT("FailedPackingNoLevel", "Packing of '{0}' failed"), FText::FromString(InContext.GetPackedLevelActor()->GetWorldAssetPackage())));
 		return false;
 	}
 
 	ULevelStreaming* SourceLevelStreaming = ULevelStreaming::FindStreamingLevel(SourceLevel);
 	AWorldSettings* WorldSettings = SourceLevel->GetWorldSettings();
-	Context.DiscardActor(WorldSettings);
-
+	
 	// Build relative transform without rotation because pivots don't support rotation
-	FTransform CurrentPivotTransform(SourceLevelStreaming->LevelTransform.GetRelativeTransform(InPackedLevelActor->GetActorTransform()).GetTranslation());
+	FTransform CurrentPivotTransform(SourceLevelStreaming->LevelTransform.GetRelativeTransform(InContext.GetPackedLevelActor()->GetActorTransform()).GetTranslation());
 	FTransform NewPivotTransform(WorldSettings->LevelInstancePivotOffset);
 	FTransform RelativePivotTransform(NewPivotTransform.GetRelativeTransform(CurrentPivotTransform));
 		
-	Context.SetRelativePivotTransform(RelativePivotTransform);
-	
-	if (AActor* DefaultBrush = SourceLevel->GetDefaultBrush())
-	{
-		Context.DiscardActor(DefaultBrush);
-	}
-		
-	for (AActor* LevelActor : SourceLevel->Actors)
-	{
-		if (LevelActor)
-		{
-			Context.ClusterLevelActor(LevelActor);
-		}
-	}
+	InContext.SetRelativePivotTransform(RelativePivotTransform);	
+	InContext.ClusterLevelActor(InContext.GetLevelInstanceToPack());
 
-	for (const auto& Pair : Context.GetClusters())
+	for (const auto& Pair : InContext.GetClusters())
 	{
 		TUniquePtr<IPackedLevelActorBuilder>& Builder = Builders.FindChecked(Pair.Key.GetBuilderID());
-		Builder->PackActors(Context, InPackedLevelActor, Pair.Key, Pair.Value);
+		Builder->PackActors(InContext, Pair.Key, Pair.Value);
 	}
 			
-	Context.Report(Log);
+	InContext.Report(Log);
 	return true;
 }
 
@@ -289,7 +283,7 @@ void FPackedLevelActorBuilderContext::Report(FMessageLog& Log) const
 	Log.Info(FText::Format(LOCTEXT("PackCompleted", "Packing '{0}' completed with {1} warning(s)"), FText::FromString(PackedLevelActor->GetWorldAssetPackage()), FText::AsNumber(TotalWarningCount)));
 }
 
-ALevelInstance* FPackedLevelActorBuilder::CreateTransientLevelInstanceForPacking(TSoftObjectPtr<UWorld> InWorldAsset, const FVector& InLocation, const FRotator& InRotator)
+ALevelInstance* FPackedLevelActorBuilder::CreateTransientLevelInstanceForPacking(TSoftObjectPtr<UWorld> InWorldAsset, const FVector& InLocation, const FRotator& InRotator, const FWorldPartitionActorFilter& InFilter)
 {
 	// Create Temp Actor for Packing
 	FActorSpawnParameters SpawnParams;
@@ -301,7 +295,9 @@ ALevelInstance* FPackedLevelActorBuilder::CreateTransientLevelInstanceForPacking
 	UWorld* World = GEditor->GetEditorWorldContext().World();
 	check(World);
 	SpawnParams.OverrideLevel = World->PersistentLevel;
-	ALevelInstance* LevelInstance = World->SpawnActor<ALevelInstance>(InLocation, InRotator, SpawnParams);
+	APackedLevelActor* LevelInstance = World->SpawnActor<APackedLevelActor>(InLocation, InRotator, SpawnParams);
+	LevelInstance->SetFilter(InFilter);
+	LevelInstance->SetShouldLoadForPacking(true);
 	LevelInstance->SetWorldAsset(InWorldAsset);
 
 	// Wait for load
@@ -312,7 +308,7 @@ ALevelInstance* FPackedLevelActorBuilder::CreateTransientLevelInstanceForPacking
 
 bool FPackedLevelActorBuilder::PackActor(APackedLevelActor* InActor, TSoftObjectPtr<UWorld> InWorldAsset)
 {
-	ALevelInstance* TransientLevelInstance = CreateTransientLevelInstanceForPacking(InWorldAsset, InActor->GetActorLocation(), InActor->GetActorRotation());
+	ALevelInstance* TransientLevelInstance = CreateTransientLevelInstanceForPacking(InWorldAsset, InActor->GetActorLocation(), InActor->GetActorRotation(), InActor->GetFilter());
 	ON_SCOPE_EXIT
 	{
 		TransientLevelInstance->GetWorld()->DestroyActor(TransientLevelInstance);
@@ -333,7 +329,17 @@ bool FPackedLevelActorBuilder::CreateOrUpdateBlueprint(TSoftObjectPtr<UWorld> In
 {
 	bool bResult = true;
 	
-	ALevelInstance* TransientLevelInstance = CreateTransientLevelInstanceForPacking(InWorldAsset, FVector::ZeroVector, FRotator::ZeroRotator);
+	// Transfer filter from existing BP CDO to Transient Level Instance
+	FWorldPartitionActorFilter Filter;
+	if (!InBlueprintAsset.IsNull())
+	{
+		if (UBlueprint* BP = InBlueprintAsset.LoadSynchronous())
+		{
+			Filter = CastChecked<APackedLevelActor>(BP->GeneratedClass->GetDefaultObject())->GetFilter();
+		}
+	}
+
+	ALevelInstance* TransientLevelInstance = CreateTransientLevelInstanceForPacking(InWorldAsset, FVector::ZeroVector, FRotator::ZeroRotator, Filter);
 	
 	bResult = CreateOrUpdateBlueprintFromUnpacked(TransientLevelInstance, InBlueprintAsset, bCheckoutAndSave, bPromptForSave);
 
@@ -362,12 +368,14 @@ bool FPackedLevelActorBuilder::CreateOrUpdateBlueprintFromUnpacked(ALevelInstanc
 	SpawnParams.bHideFromSceneOutliner = true;
 	SpawnParams.bNoFail = true;
 	SpawnParams.ObjectFlags |= RF_Transient;
-
+	SpawnParams.ObjectFlags &= ~RF_Transactional;
+	
 	UWorld* World = InActor->GetWorld();
 	SpawnParams.OverrideLevel = World->PersistentLevel;
 
 	APackedLevelActor* PackedLevelActor = World->SpawnActor<APackedLevelActor>(InActor->GetActorLocation(), InActor->GetActorRotation(), SpawnParams);	
 	PackedLevelActor->SetWorldAsset(InActor->GetWorldAsset());
+	PackedLevelActor->SetFilter(InActor->GetFilter());
 	ON_SCOPE_EXIT
 	{
 		InActor->GetWorld()->DestroyActor(PackedLevelActor);
@@ -385,6 +393,11 @@ bool FPackedLevelActorBuilder::CreateOrUpdateBlueprintFromUnpacked(ALevelInstanc
 
 bool FPackedLevelActorBuilder::CreateOrUpdateBlueprintFromPacked(APackedLevelActor* InActor, TSoftObjectPtr<UBlueprint> InBlueprintAsset, bool bCheckoutAndSave, bool bPromptToSave)
 {
+	if (GEditor->Trans)
+	{
+		GEditor->Trans->Reset(LOCTEXT("BlueprintPacking", "Blueprint Packing"));
+	}
+
 	UBlueprint* BP = nullptr;
 	if (!InBlueprintAsset.IsNull())
 	{
@@ -405,37 +418,25 @@ bool FPackedLevelActorBuilder::CreateOrUpdateBlueprintFromPacked(APackedLevelAct
 		BP = CreatePackedLevelActorBlueprintWithDialog(TSoftObjectPtr<UBlueprint>(AssetPath), InActor->GetWorldAsset(), bCompile);
 	}
 
-	if (BP && BP->SimpleConstructionScript)
-	{
-		TArray<USCS_Node*> AllNodes = BP->SimpleConstructionScript->GetAllNodes();
-		for (USCS_Node* Node : AllNodes)
-		{
-			BP->SimpleConstructionScript->RemoveNodeAndPromoteChildren(Node);
-		}
-		FKismetEditorUtilities::GenerateBlueprintSkeleton(BP,true);
-	}
-	else
+	if (!BP || !BP->SimpleConstructionScript)
 	{
 		return false;
 	}
-
+	
 	BP->Modify();
+
+	TArray<USCS_Node*> AllNodes = BP->SimpleConstructionScript->GetAllNodes();
+	for (USCS_Node* Node : AllNodes)
+	{
+		BP->SimpleConstructionScript->RemoveNodeAndPromoteChildren(Node);
+	}
+	FKismetEditorUtilities::GenerateBlueprintSkeleton(BP,true);
+		
 	// Avoid running construction script while dragging an instance of that BP for performance reasons
 	BP->bRunConstructionScriptOnDrag = false;
 	FGuid NewVersion = FGuid::NewGuid();
 	APackedLevelActor* CDO = CastChecked<APackedLevelActor>(BP->GeneratedClass->GetDefaultObject());
-	auto PropagatePropertiesToActor = [InActor, NewVersion](APackedLevelActor* TargetActor)
-	{
-		TargetActor->Modify(false);
-		TargetActor->SetWorldAsset(InActor->GetWorldAsset());
-		TargetActor->PackedBPDependencies = InActor->PackedBPDependencies;
-		TargetActor->SetPackedVersion(NewVersion);
-		// match root component mobility to source actor
-		USceneComponent* Root = TargetActor->GetRootComponent();
-		Root->SetMobility(InActor->GetRootComponent()->Mobility);
-	};
-	PropagatePropertiesToActor(CDO);
-		
+			
 	// Prep AddComponentsToBlueprintParam
 	FKismetEditorUtilities::FAddComponentsToBlueprintParams AddCompToBPParams;
 	AddCompToBPParams.HarvestMode = FKismetEditorUtilities::EAddComponentToBPHarvestMode::None;
@@ -459,16 +460,29 @@ bool FPackedLevelActorBuilder::CreateOrUpdateBlueprintFromPacked(APackedLevelAct
 	}
 
 	// Propagate properties before BP Compilation so that they are considered default (no delta)
-	if (BP->GeneratedClass)
+	TArray<UObject*> ObjectsOfClass;
+	GetObjectsOfClass(BP->GeneratedClass, ObjectsOfClass, true, RF_NoFlags, EInternalObjectFlags::Garbage);
+	for (UObject* ObjectOfClass : ObjectsOfClass)
 	{
-		TArray<UObject*> ObjectsOfClass;
-		GetObjectsOfClass(BP->GeneratedClass, ObjectsOfClass, true, RF_ClassDefaultObject | RF_ArchetypeObject, EInternalObjectFlags::Garbage);
-		for (UObject* ObjectOfClass : ObjectsOfClass)
+		const bool bWasDirty = ObjectOfClass->GetPackage()->IsDirty();
+
+		APackedLevelActor* PackedLevelActor = CastChecked<APackedLevelActor>(ObjectOfClass);
+
+		PackedLevelActor->Modify(false);
+		PackedLevelActor->SetWorldAsset(InActor->GetWorldAsset());
+		PackedLevelActor->SetFilter(InActor->GetFilter());
+		PackedLevelActor->PackedBPDependencies = InActor->PackedBPDependencies;
+		PackedLevelActor->SetPackedVersion(NewVersion);
+		// match root component mobility to source actor
+		USceneComponent* Root = PackedLevelActor->GetRootComponent();
+		Root->SetMobility(InActor->GetRootComponent()->Mobility);
+
+		if (!bWasDirty && ObjectOfClass->GetPackage()->IsDirty())
 		{
-			APackedLevelActor* PackedLevelActor = CastChecked<APackedLevelActor>(ObjectOfClass);
-			PropagatePropertiesToActor(PackedLevelActor);
+			ObjectOfClass->GetPackage()->SetDirtyFlag(false);
 		}
 	}
+	
 			
 	// Synchronous compile
 	FKismetEditorUtilities::CompileBlueprint(BP, EBlueprintCompileOptions::SkipGarbageCollection);
diff --git a/Engine/Source/Runtime/Engine/Private/PackedLevelActor/PackedLevelActorBuilder.h b/Engine/Source/Runtime/Engine/Private/PackedLevelActor/PackedLevelActorBuilder.h
--- a/Engine/Source/Runtime/Engine/Private/PackedLevelActor/PackedLevelActorBuilder.h
+++ b/Engine/Source/Runtime/Engine/Private/PackedLevelActor/PackedLevelActorBuilder.h
@@ -17,6 +17,7 @@ class AActor;
 class UActorComponent;
 class UBlueprint;
 class FMessageLog;
+struct FWorldPartitionActorFilter;
 
 /**
  * FPackedLevelActorBuilder handles packing of ALevelInstance actors into APackedLevelActor actors and Blueprints.
@@ -49,12 +50,14 @@ public:
 	static UBlueprint* CreatePackedLevelActorBlueprint(TSoftObjectPtr<UBlueprint> InBlueprintAsset, TSoftObjectPtr<UWorld> InWorldAsset, bool bInCompile);
 	
 private:
+	bool PackActor(FPackedLevelActorBuilderContext& InContext);
+
 	/* Create/Updates a APackedLevelActor Blueprint from InPackedActor (will overwrite existing asset or show a dialog if InBlueprintAsset is null) */
 	bool CreateOrUpdateBlueprintFromPacked(APackedLevelActor* InPackedActor, TSoftObjectPtr<UBlueprint> InBlueprintAsset, bool bCheckoutAndSave, bool bPromptForSave);
 	/* Creates/Updates a APackedLevelActor Blueprint from InLevelInstance (will overwrite existing asset or show a dialog if InBlueprintAsset is null) */
 	bool CreateOrUpdateBlueprintFromUnpacked(ALevelInstance* InPackedActor, TSoftObjectPtr<UBlueprint> InBlueprintAsset, bool bCheckoutAndSave, bool bPromptForSave);
 	/* Creates and loads a ALevelInstance so it can be used for packing */
-	ALevelInstance* CreateTransientLevelInstanceForPacking(TSoftObjectPtr<UWorld> InWorldAsset, const FVector& InLocation, const FRotator& InRotator);
+	ALevelInstance* CreateTransientLevelInstanceForPacking(TSoftObjectPtr<UWorld> InWorldAsset, const FVector& InLocation, const FRotator& InRotator, const FWorldPartitionActorFilter& InFilter);
 
 	FPackedLevelActorBuilder(const FPackedLevelActorBuilder&) = delete;
 	FPackedLevelActorBuilder& operator=(const FPackedLevelActorBuilder&) = delete;
@@ -72,8 +75,8 @@ private:
 class FPackedLevelActorBuilderContext
 {
 public:
-	FPackedLevelActorBuilderContext(const FPackedLevelActorBuilder& InBuilder, APackedLevelActor* InPackedLevelActor) 
-		: Builders(InBuilder.Builders), ClassDiscards(InBuilder.ClassDiscards), PackedLevelActor(InPackedLevelActor), RelativePivotTransform(FTransform::Identity) {}
+	FPackedLevelActorBuilderContext(const FPackedLevelActorBuilder& InBuilder, APackedLevelActor* InPackedLevelActor, ALevelInstance* InLevelInstanceToPack) 
+		: Builders(InBuilder.Builders), ClassDiscards(InBuilder.ClassDiscards), PackedLevelActor(InPackedLevelActor), LevelInstanceToPack(InLevelInstanceToPack), RelativePivotTransform(FTransform::Identity) {}
 
 	/* Interface for IPackedLevelActorBuilder's to use */
 	void ClusterLevelActor(AActor* InLevelActor);
@@ -87,11 +90,15 @@ public:
 	const FTransform GetRelativePivotTransform() const { return RelativePivotTransform; }
 
 	bool ShouldPackComponent(UActorComponent* InActorComponent) const;
+
+	APackedLevelActor* GetPackedLevelActor() const { return PackedLevelActor; }
+	ALevelInstance* GetLevelInstanceToPack() const { return LevelInstanceToPack; }
 private:
 	const TMap<FPackedLevelActorBuilderID, TUniquePtr<IPackedLevelActorBuilder>>& Builders;
 	const TSet<UClass*>& ClassDiscards;
 
 	APackedLevelActor* PackedLevelActor;
+	ALevelInstance* LevelInstanceToPack;
 	
 	TMap<FPackedLevelActorBuilderClusterID, TArray<UActorComponent*>> Clusters;
 
diff --git a/Engine/Source/Runtime/Engine/Private/PackedLevelActor/PackedLevelActorISMBuilder.cpp b/Engine/Source/Runtime/Engine/Private/PackedLevelActor/PackedLevelActorISMBuilder.cpp
--- a/Engine/Source/Runtime/Engine/Private/PackedLevelActor/PackedLevelActorISMBuilder.cpp
+++ b/Engine/Source/Runtime/Engine/Private/PackedLevelActor/PackedLevelActorISMBuilder.cpp
@@ -23,6 +23,15 @@ FPackedLevelActorBuilderID FPackedLevelActorISMBuilder::GetID() const
 
 void FPackedLevelActorISMBuilder::GetPackClusters(FPackedLevelActorBuilderContext& InContext, AActor* InActor) const
 {
+	// Skip PackedLevelActors that are loaded for Packing as their Level Actors will get packed
+	if (APackedLevelActor* PackedLevelActor = Cast<APackedLevelActor>(InActor))
+	{
+		if (PackedLevelActor->ShouldLoadForPacking())
+		{
+			return;
+		}
+	}
+
 	TArray<UStaticMeshComponent*> StaticMeshComponents;
 	InActor->GetComponents(StaticMeshComponents);
 
@@ -37,10 +46,11 @@ void FPackedLevelActorISMBuilder::GetPackClusters(FPackedLevelActorBuilderContex
 	}
 }
 
-void FPackedLevelActorISMBuilder::PackActors(FPackedLevelActorBuilderContext& InContext, APackedLevelActor* InPackingActor, const FPackedLevelActorBuilderClusterID& InClusterID, const TArray<UActorComponent*>& InComponents) const
+void FPackedLevelActorISMBuilder::PackActors(FPackedLevelActorBuilderContext& InContext, const FPackedLevelActorBuilderClusterID& InClusterID, const TArray<UActorComponent*>& InComponents) const
 {
 	check(InClusterID.GetBuilderID() == GetID());
-	FTransform ActorTransform = InPackingActor->GetActorTransform();
+	APackedLevelActor* PackingActor = InContext.GetPackedLevelActor();
+	FTransform ActorTransform = PackingActor->GetActorTransform();
 
 	FPackedLevelActorISMBuilderCluster* ISMCluster = (FPackedLevelActorISMBuilderCluster*)InClusterID.GetData();
 	check(ISMCluster);
@@ -55,8 +65,8 @@ void FPackedLevelActorISMBuilder::PackActors(FPackedLevelActorBuilderContext& In
 		ComponentClass = UHierarchicalInstancedStaticMeshComponent::StaticClass();
 	}
 
-	UInstancedStaticMeshComponent* PackComponent = InPackingActor->AddPackedComponent<UInstancedStaticMeshComponent>(ComponentClass);
-	PackComponent->AttachToComponent(InPackingActor->GetRootComponent(), FAttachmentTransformRules::SnapToTargetIncludingScale);
+	UInstancedStaticMeshComponent* PackComponent = PackingActor->AddPackedComponent<UInstancedStaticMeshComponent>(ComponentClass);
+	PackComponent->AttachToComponent(PackingActor->GetRootComponent(), FAttachmentTransformRules::SnapToTargetIncludingScale);
 
 	// Initialize the ISM properties using the ISM descriptor
 	ISMCluster->ISMDescriptor.InitComponent(PackComponent);
diff --git a/Engine/Source/Runtime/Engine/Private/PackedLevelActor/PackedLevelActorISMBuilder.h b/Engine/Source/Runtime/Engine/Private/PackedLevelActor/PackedLevelActorISMBuilder.h
--- a/Engine/Source/Runtime/Engine/Private/PackedLevelActor/PackedLevelActorISMBuilder.h
+++ b/Engine/Source/Runtime/Engine/Private/PackedLevelActor/PackedLevelActorISMBuilder.h
@@ -17,7 +17,7 @@ public:
 
 	virtual FPackedLevelActorBuilderID GetID() const override;
 	virtual void GetPackClusters(FPackedLevelActorBuilderContext& InContext, AActor* InActor) const override;
-	virtual void PackActors(FPackedLevelActorBuilderContext& InContext, APackedLevelActor* InPackingActor, const FPackedLevelActorBuilderClusterID& InClusterID, const TArray<UActorComponent*>& InComponents) const override;
+	virtual void PackActors(FPackedLevelActorBuilderContext& InContext, const FPackedLevelActorBuilderClusterID& InClusterID, const TArray<UActorComponent*>& InComponents) const override;
 };
 
 class FPackedLevelActorISMBuilderCluster : public FPackedLevelActorBuilderCluster
diff --git a/Engine/Source/Runtime/Engine/Private/PackedLevelActor/PackedLevelActorRecursiveBuilder.cpp b/Engine/Source/Runtime/Engine/Private/PackedLevelActor/PackedLevelActorRecursiveBuilder.cpp
--- a/Engine/Source/Runtime/Engine/Private/PackedLevelActor/PackedLevelActorRecursiveBuilder.cpp
+++ b/Engine/Source/Runtime/Engine/Private/PackedLevelActor/PackedLevelActorRecursiveBuilder.cpp
@@ -4,6 +4,7 @@
 
 #if WITH_EDITOR
 
+#include "LevelInstance/LevelInstanceInterface.h"
 #include "PackedLevelActor/PackedLevelActorBuilder.h"
 #include "PackedLevelActor/PackedLevelActor.h"
 
@@ -11,7 +12,8 @@
 
 #include "Engine/Blueprint.h"
 #include "Engine/Level.h"
-
+#include "Engine/Brush.h"
+#include "GameFramework/WorldSettings.h"
 
 FPackedLevelActorBuilderID FPackedLevelActorRecursiveBuilder::BuilderID = 'RECP';
 
@@ -22,13 +24,13 @@ FPackedLevelActorBuilderID FPackedLevelActorRecursiveBuilder::GetID() const
 
 void FPackedLevelActorRecursiveBuilder::GetPackClusters(FPackedLevelActorBuilderContext& InContext, AActor* InActor) const
 {
-	if (ALevelInstance* LevelInstance = Cast<ALevelInstance>(InActor))
+	if (ILevelInstanceInterface* LevelInstance = Cast<ILevelInstanceInterface>(InActor))
 	{
 		FPackedLevelActorBuilderClusterID ClusterID(MakeUnique<FPackedLevelActorRecursiveBuilderCluster>(GetID(), LevelInstance));
 		InContext.FindOrAddCluster(MoveTemp(ClusterID));
 
 		// This Actor can be safely discarded without warning because it is a container
-		InContext.DiscardActor(LevelInstance);
+		InContext.DiscardActor(InActor);
 
 		ULevelInstanceSubsystem* LevelInstanceSubsystem = LevelInstance->GetLevelInstanceSubsystem();
 		check(LevelInstanceSubsystem);
@@ -38,29 +40,40 @@ void FPackedLevelActorRecursiveBuilder::GetPackClusters(FPackedLevelActorBuilder
 			{
 				if (LevelActor)
 				{
-					InContext.ClusterLevelActor(LevelActor);
+					if (LevelActor == Level->GetDefaultBrush())
+					{
+						InContext.DiscardActor(LevelActor);
+					}
+					else if (LevelActor == Level->GetWorldSettings())
+					{
+						InContext.DiscardActor(LevelActor);
+					}
+					else
+					{
+						InContext.ClusterLevelActor(LevelActor);
+					}
 				}
 			}
 		}
 	}
 }
 
-void FPackedLevelActorRecursiveBuilder::PackActors(FPackedLevelActorBuilderContext& InContext, APackedLevelActor* InPackingActor, const FPackedLevelActorBuilderClusterID& InClusterID, const TArray<UActorComponent*>& InComponents) const
+void FPackedLevelActorRecursiveBuilder::PackActors(FPackedLevelActorBuilderContext& InContext, const FPackedLevelActorBuilderClusterID& InClusterID, const TArray<UActorComponent*>& InComponents) const
 {
 	check(InClusterID.GetBuilderID() == GetID());
 	FPackedLevelActorRecursiveBuilderCluster* LevelInstanceCluster = (FPackedLevelActorRecursiveBuilderCluster*)InClusterID.GetData();
 	check(LevelInstanceCluster);
 
-	if (LevelInstanceCluster->LevelInstance->IsA<APackedLevelActor>())
+	if (APackedLevelActor* PackedLevelActor = Cast<APackedLevelActor>(LevelInstanceCluster->LevelInstance))
 	{
-		if (UBlueprint* GeneratedBy = Cast<UBlueprint>(LevelInstanceCluster->LevelInstance->GetClass()->ClassGeneratedBy))
+		if (UBlueprint* GeneratedBy = PackedLevelActor->GetRootBlueprint())
 		{
-			InPackingActor->PackedBPDependencies.AddUnique(GeneratedBy);
+			InContext.GetPackedLevelActor()->PackedBPDependencies.AddUnique(GeneratedBy);
 		}
 	}
 }
 
-FPackedLevelActorRecursiveBuilderCluster::FPackedLevelActorRecursiveBuilderCluster(FPackedLevelActorBuilderID InBuilderID, ALevelInstance* InLevelInstance)
+FPackedLevelActorRecursiveBuilderCluster::FPackedLevelActorRecursiveBuilderCluster(FPackedLevelActorBuilderID InBuilderID, ILevelInstanceInterface* InLevelInstance)
 	: FPackedLevelActorBuilderCluster(InBuilderID), LevelInstance(InLevelInstance)
 {
 }
diff --git a/Engine/Source/Runtime/Engine/Private/PackedLevelActor/PackedLevelActorRecursiveBuilder.h b/Engine/Source/Runtime/Engine/Private/PackedLevelActor/PackedLevelActorRecursiveBuilder.h
--- a/Engine/Source/Runtime/Engine/Private/PackedLevelActor/PackedLevelActorRecursiveBuilder.h
+++ b/Engine/Source/Runtime/Engine/Private/PackedLevelActor/PackedLevelActorRecursiveBuilder.h
@@ -7,7 +7,7 @@
 #include "PackedLevelActor/PackedLevelActorTypes.h"
 #include "PackedLevelActor/IPackedLevelActorBuilder.h"
 
-class ALevelInstance;
+class ILevelInstanceInterface;
 
 class FPackedLevelActorRecursiveBuilder : public IPackedLevelActorBuilder
 {
@@ -16,18 +16,18 @@ public:
 
 	virtual FPackedLevelActorBuilderID GetID() const override;
 	virtual void GetPackClusters(FPackedLevelActorBuilderContext& InContext, AActor* InActor) const override;
-	virtual void PackActors(FPackedLevelActorBuilderContext& InContext, APackedLevelActor* InPackingActor, const FPackedLevelActorBuilderClusterID& InClusterID, const TArray<UActorComponent*>& InComponents) const override;
+	virtual void PackActors(FPackedLevelActorBuilderContext& InContext, const FPackedLevelActorBuilderClusterID& InClusterID, const TArray<UActorComponent*>& InComponents) const override;
 };
 
 class FPackedLevelActorRecursiveBuilderCluster : public FPackedLevelActorBuilderCluster
 {
 public:
-	FPackedLevelActorRecursiveBuilderCluster(FPackedLevelActorBuilderID InBuilderID, ALevelInstance* InLevelInstance);
+	FPackedLevelActorRecursiveBuilderCluster(FPackedLevelActorBuilderID InBuilderID, ILevelInstanceInterface* InLevelInstance);
 
 	virtual bool Equals(const FPackedLevelActorBuilderCluster& InOther) const override;
 	virtual uint32 ComputeHash() const override;
 
-	ALevelInstance* LevelInstance = nullptr;
+	ILevelInstanceInterface* LevelInstance = nullptr;
 };
 
 #endif
\ No newline at end of file
diff --git a/Engine/Source/Runtime/Engine/Private/WorldPartition/PackedLevelActor/PackedLevelActorDesc.cpp b/Engine/Source/Runtime/Engine/Private/WorldPartition/PackedLevelActor/PackedLevelActorDesc.cpp
new file mode 100644
--- /dev/null
+++ b/Engine/Source/Runtime/Engine/Private/WorldPartition/PackedLevelActor/PackedLevelActorDesc.cpp
@@ -0,0 +1,50 @@
+// Copyright Epic Games, Inc. All Rights Reserved.
+
+#include "WorldPartition/PackedLevelActor/PackedLevelActorDesc.h"
+
+#if WITH_EDITOR
+#include "PackedLevelActor/PackedLevelActor.h"
+#include "WorldPartition/WorldPartitionActorDescArchive.h"
+#include "UObject/FortniteMainBranchObjectVersion.h"
+
+void FPackedLevelActorDesc::Init(const AActor* InActor)
+{
+	FWorldPartitionActorDesc::Init(InActor);
+
+	const APackedLevelActor* PackedLevelActor = CastChecked<APackedLevelActor>(InActor);
+	WorldAsset = PackedLevelActor->GetWorldAsset().ToSoftObjectPath();
+	ContainerFilter = PackedLevelActor->GetFilter();
+}
+
+void FPackedLevelActorDesc::Init(const FWorldPartitionActorDescInitData& DescData)
+{
+	FWorldPartitionActorDesc::Init(DescData);
+}
+
+bool FPackedLevelActorDesc::Equals(const FWorldPartitionActorDesc* Other) const
+{
+	if (FWorldPartitionActorDesc::Equals(Other))
+	{
+		const FPackedLevelActorDesc* PackedLevelActorDesc = (FPackedLevelActorDesc*)Other;
+
+		return WorldAsset == PackedLevelActorDesc->WorldAsset &&
+			ContainerFilter == PackedLevelActorDesc->ContainerFilter;
+	}
+
+	return false;
+}
+
+void FPackedLevelActorDesc::Serialize(FArchive& Ar)
+{
+	FWorldPartitionActorDesc::Serialize(Ar);
+
+	Ar.UsingCustomVersion(FFortniteMainBranchObjectVersion::GUID);
+
+	if (Ar.CustomVer(FFortniteMainBranchObjectVersion::GUID) >= FFortniteMainBranchObjectVersion::PackedLevelActorDesc)
+	{
+		Ar << TDeltaSerialize<FSoftObjectPath>(WorldAsset);
+		Ar << TDeltaSerialize<FWorldPartitionActorFilter>(ContainerFilter);
+	}
+}
+
+#endif
diff --git a/Engine/Source/Runtime/Engine/Private/WorldPartition/WorldPartitionActorDescView.cpp b/Engine/Source/Runtime/Engine/Private/WorldPartition/WorldPartitionActorDescView.cpp
--- a/Engine/Source/Runtime/Engine/Private/WorldPartition/WorldPartitionActorDescView.cpp
+++ b/Engine/Source/Runtime/Engine/Private/WorldPartition/WorldPartitionActorDescView.cpp
@@ -184,9 +184,9 @@ bool FWorldPartitionActorDescView::IsContainerInstance() const
 	return ActorDesc->IsContainerInstance();
 }
 
-bool FWorldPartitionActorDescView::IsContainerFilter() const
+EWorldPartitionActorFilterType FWorldPartitionActorDescView::GetContainerFilterType() const
 {
-	return ActorDesc->IsContainerFilter();
+	return ActorDesc->GetContainerFilterType();
 }
 
 FName FWorldPartitionActorDescView::GetContainerPackage() const
@@ -201,7 +201,7 @@ bool FWorldPartitionActorDescView::GetContainerInstance(FWorldPartitionActorDesc
 
 const FWorldPartitionActorFilter* FWorldPartitionActorDescView::GetContainerFilter() const
 {
-	check(IsContainerFilter());
+	check(GetContainerFilterType() != EWorldPartitionActorFilterType::None);
 	return ActorDesc->GetContainerFilter();
 }
 
diff --git a/Engine/Source/Runtime/Engine/Private/WorldPartition/WorldPartitionStreamingGeneration.cpp b/Engine/Source/Runtime/Engine/Private/WorldPartition/WorldPartitionStreamingGeneration.cpp
--- a/Engine/Source/Runtime/Engine/Private/WorldPartition/WorldPartitionStreamingGeneration.cpp
+++ b/Engine/Source/Runtime/Engine/Private/WorldPartition/WorldPartitionStreamingGeneration.cpp
@@ -573,7 +573,7 @@ class FWorldPartitionStreamingGenerator
 				SubContainerInstanceDescriptor.ContentBundleID = InContainerInstanceDescriptor.ContentBundleID;
 				SubContainerInstanceDescriptor.InstanceData = InheritParentContainerPerInstanceData(InContainerInstanceDescriptor.InstanceData, ContainerInstanceView);
 
-				if (WorldPartitionSubsystem && InContainerInstanceDescriptor.ID.IsMainContainer() && ContainerInstanceView.IsContainerFilter())
+				if (WorldPartitionSubsystem && InContainerInstanceDescriptor.ID.IsMainContainer() && ContainerInstanceView.GetContainerFilterType() == EWorldPartitionActorFilterType::Loading)
 				{
 					if (const FWorldPartitionActorFilter* ContainerFilter = ContainerInstanceView.GetContainerFilter())
 					{						
diff --git a/Engine/Source/Runtime/Engine/Private/WorldPartition/WorldPartitionSubsystem.cpp b/Engine/Source/Runtime/Engine/Private/WorldPartition/WorldPartitionSubsystem.cpp
--- a/Engine/Source/Runtime/Engine/Private/WorldPartition/WorldPartitionSubsystem.cpp
+++ b/Engine/Source/Runtime/Engine/Private/WorldPartition/WorldPartitionSubsystem.cpp
@@ -155,14 +155,14 @@ void UWorldPartitionSubsystem::AddReferencedObjects(UObject* InThis, FReferenceC
 	This->ActorDescContainerInstanceManager.AddReferencedObjects(Collector);
 }
 
-FWorldPartitionActorFilter UWorldPartitionSubsystem::GetWorldPartitionActorFilter(const FString& InWorldPackage) const
+FWorldPartitionActorFilter UWorldPartitionSubsystem::GetWorldPartitionActorFilter(const FString& InWorldPackage, EWorldPartitionActorFilterType InFilterTypes) const
 {
 	TSet<FString> VisitedPackages;
-	return GetWorldPartitionActorFilterInternal(InWorldPackage, VisitedPackages);
+	return GetWorldPartitionActorFilterInternal(InWorldPackage, InFilterTypes, VisitedPackages);
 }
 
 
-FWorldPartitionActorFilter UWorldPartitionSubsystem::GetWorldPartitionActorFilterInternal(const FString& InWorldPackage, TSet<FString>& InOutVisitedPackages) const
+FWorldPartitionActorFilter UWorldPartitionSubsystem::GetWorldPartitionActorFilterInternal(const FString& InWorldPackage, EWorldPartitionActorFilterType InFilterTypes, TSet<FString>& InOutVisitedPackages) const
 {
 	if (InOutVisitedPackages.Contains(InWorldPackage))
 	{
@@ -188,7 +188,7 @@ FWorldPartitionActorFilter UWorldPartitionSubsystem::GetWorldPartitionActorFilte
 			check(!WorldDataLayersActorDesc);
 			WorldDataLayersActorDesc = static_cast<const FWorldDataLayersActorDesc*>(*ActorDescIt);
 		}
-		else if (ActorDescIt->IsContainerFilter())
+		else if ((ActorDescIt->GetContainerFilterType() & InFilterTypes) != EWorldPartitionActorFilterType::None)
 		{
 			ContainerActorDescs.Add(*ActorDescIt);
 		}
@@ -213,7 +213,7 @@ FWorldPartitionActorFilter UWorldPartitionSubsystem::GetWorldPartitionActorFilte
 		TSet<FString> VisitedPackagesCopy(InOutVisitedPackages);
 
 		// Get World Default Filter
-		FWorldPartitionActorFilter* ChildFilter = new FWorldPartitionActorFilter(GetWorldPartitionActorFilterInternal(ContainerActorDesc->GetContainerPackage().ToString(), VisitedPackagesCopy));
+		FWorldPartitionActorFilter* ChildFilter = new FWorldPartitionActorFilter(GetWorldPartitionActorFilterInternal(ContainerActorDesc->GetContainerPackage().ToString(), InFilterTypes, VisitedPackagesCopy));
 		ChildFilter->DisplayName = ContainerActorDesc->GetActorLabelOrName().ToString();
 
 		// Apply Filter to Default
@@ -228,11 +228,11 @@ FWorldPartitionActorFilter UWorldPartitionSubsystem::GetWorldPartitionActorFilte
 	return Filter;
 }
 
-TMap<FActorContainerID, TSet<FGuid>> UWorldPartitionSubsystem::GetFilteredActorsPerContainer(const FActorContainerID& InContainerID, const FString& InWorldPackage, const FWorldPartitionActorFilter& InActorFilter)
+TMap<FActorContainerID, TSet<FGuid>> UWorldPartitionSubsystem::GetFilteredActorsPerContainer(const FActorContainerID& InContainerID, const FString& InWorldPackage, const FWorldPartitionActorFilter& InActorFilter, EWorldPartitionActorFilterType InFilterTypes)
 {
 	TMap<FActorContainerID, TSet<FGuid>> FilteredActors;
 
-	FWorldPartitionActorFilter ContainerFilter = GetWorldPartitionActorFilter(InWorldPackage);
+	FWorldPartitionActorFilter ContainerFilter = GetWorldPartitionActorFilter(InWorldPackage, InFilterTypes);
 	ContainerFilter.Override(InActorFilter);
 
 	// Flatten Filter to FActorContainerID map
@@ -270,7 +270,7 @@ TMap<FActorContainerID, TSet<FGuid>> UWorldPartitionSubsystem::GetFilteredActors
 		return RegisteredContainer;
 	};
 
-	TFunction<void(const FActorContainerID&, const UActorDescContainer*)> ProcessContainers = [&FindOrRegisterContainer, &DataLayerFiltersPerContainer, &FilteredActors, &ProcessContainers](const FActorContainerID& InContainerID, const UActorDescContainer* InContainer)
+	TFunction<void(const FActorContainerID&, const UActorDescContainer*)> ProcessContainers = [&FindOrRegisterContainer, &DataLayerFiltersPerContainer, &FilteredActors, &ProcessContainers, &InFilterTypes](const FActorContainerID& InContainerID, const UActorDescContainer* InContainer)
 	{
 		const TMap<FSoftObjectPath, FWorldPartitionActorFilter::FDataLayerFilter>& DataLayerFilters = DataLayerFiltersPerContainer.FindChecked(InContainerID);
 		for (FActorDescList::TConstIterator<> ActorDescIt(InContainer); ActorDescIt; ++ActorDescIt)
@@ -301,11 +301,14 @@ TMap<FActorContainerID, TSet<FGuid>> UWorldPartitionSubsystem::GetFilteredActors
 				}
 			}
 
-			if (ActorDescIt->IsContainerFilter())
+			if ((ActorDescIt->GetContainerFilterType() & InFilterTypes) != EWorldPartitionActorFilterType::None)
 			{
-				UActorDescContainer* ChildContainer = FindOrRegisterContainer(ActorDescIt->GetContainerPackage());
-				check(ChildContainer);
-				ProcessContainers(FActorContainerID(InContainerID, ActorDescIt->GetGuid()), ChildContainer);
+				if (const FWorldPartitionActorFilter* ChildFilter = ActorDescIt->GetContainerFilter())
+				{
+					UActorDescContainer* ChildContainer = FindOrRegisterContainer(ActorDescIt->GetContainerPackage());
+					check(ChildContainer);
+					ProcessContainers(FActorContainerID(InContainerID, ActorDescIt->GetGuid()), ChildContainer);
+				}
 			}
 		}
 	};
diff --git a/Engine/Source/Runtime/Engine/Public/LevelInstance/LevelInstanceComponent.h b/Engine/Source/Runtime/Engine/Public/LevelInstance/LevelInstanceComponent.h
--- a/Engine/Source/Runtime/Engine/Public/LevelInstance/LevelInstanceComponent.h
+++ b/Engine/Source/Runtime/Engine/Public/LevelInstance/LevelInstanceComponent.h
@@ -45,6 +45,7 @@ public:
 	const TMap<FActorContainerID, TSet<FGuid>>& GetFilteredActorsPerContainer() const;
 private:
 	bool ShouldShowSpriteComponent() const;
+	void OnFilterChanged();
 
 	TWeakObjectPtr<AActor> CachedEditorInstanceActorPtr;
 #endif
diff --git a/Engine/Source/Runtime/Engine/Public/LevelInstance/LevelInstanceInterface.h b/Engine/Source/Runtime/Engine/Public/LevelInstance/LevelInstanceInterface.h
--- a/Engine/Source/Runtime/Engine/Public/LevelInstance/LevelInstanceInterface.h
+++ b/Engine/Source/Runtime/Engine/Public/LevelInstance/LevelInstanceInterface.h
@@ -7,11 +7,11 @@
 #include "UObject/SoftObjectPtr.h"
 #include "LevelInstance/LevelInstanceTypes.h"
 #include "WorldPartition/WorldPartitionActorContainerID.h"
+#include "WorldPartition/Filter/WorldPartitionActorFilter.h"
 #include "LevelInstanceInterface.generated.h"
 
 class ULevelInstanceComponent;
 class ULevelStreamingLevelInstance;
-struct FWorldPartitionActorFilter;
 
 UINTERFACE(BlueprintType, meta = (CannotImplementInterfaceInBlueprint))
 class ENGINE_API ULevelInstanceInterface : public UInterface
@@ -124,5 +124,15 @@ class ENGINE_API ILevelInstanceInterface
 	virtual const FWorldPartitionActorFilter& GetFilter() const;
 
 	virtual const TMap<FActorContainerID, TSet<FGuid>>& GetFilteredActorsPerContainer() const;
+
+	virtual void SetFilter(const FWorldPartitionActorFilter& InFilter);
+
+	virtual void OnFilterChanged() {}
+
+	// Return supported filter types when using filter for loading actors
+	virtual EWorldPartitionActorFilterType GetLoadingFilterTypes() const;
+
+	// Return supported filter types when setting filter through details panel
+	virtual EWorldPartitionActorFilterType GetDetailsFilterTypes() const;
 #endif
 };
\ No newline at end of file
diff --git a/Engine/Source/Runtime/Engine/Public/PackedLevelActor/PackedLevelActor.h b/Engine/Source/Runtime/Engine/Public/PackedLevelActor/PackedLevelActor.h
--- a/Engine/Source/Runtime/Engine/Public/PackedLevelActor/PackedLevelActor.h
+++ b/Engine/Source/Runtime/Engine/Public/PackedLevelActor/PackedLevelActor.h
@@ -33,10 +33,19 @@ public:
 
 	virtual void Serialize(FArchive& Ar) override;
 #if WITH_EDITOR
+	virtual EWorldPartitionActorFilterType GetDetailsFilterTypes() const override;
+	virtual EWorldPartitionActorFilterType GetLoadingFilterTypes() const override;
+	virtual void OnFilterChanged() override;
+	void SetShouldLoadForPacking(bool bInLoadForPacking) { bLoadForPacking = bInLoadForPacking; }
+	bool ShouldLoadForPacking() const;
+	// When Loading a APackedLevelActor it needs to be fully loaded for packing.
+	virtual bool SupportsPartialEditorLoading() const override { return false; }
+
 	virtual TUniquePtr<class FWorldPartitionActorDesc> CreateClassActorDesc() const override;
 	
 	static bool CreateOrUpdateBlueprint(ALevelInstance* InLevelInstance, TSoftObjectPtr<UBlueprint> InBlueprintAsset, bool bCheckoutAndSave = true, bool bPromptForSave = true);
 	static bool CreateOrUpdateBlueprint(TSoftObjectPtr<UWorld> InWorldAsset, TSoftObjectPtr<UBlueprint> InBlueprintAsset, bool bCheckoutAndSave = true, bool bPromptForSave = true);
+	static void UpdateBlueprint(UBlueprint* InBlueprint, bool bCheckoutAndSave = true);
 
 	static FName GetPackedComponentTag();
 
@@ -61,6 +70,10 @@ public:
 
 	virtual void RerunConstructionScripts() override;
 
+	static bool IsRootBlueprint(UClass* InClass);
+	bool IsRootBlueprintTemplate() const;
+	UBlueprint* GetRootBlueprint() const;
+
 	template<class T>
 	T* AddPackedComponent(TSubclassOf<T> ComponentClass)
 	{
@@ -78,12 +91,15 @@ public:
 	TArray<TSoftObjectPtr<UBlueprint>> PackedBPDependencies;
 
 private:
-	bool bChildChanged;
-
 	UPROPERTY()
 	FGuid PackedVersion;
 #endif
+
+#if WITH_EDITOR
+	bool bChildChanged;
+	bool bLoadForPacking;
+#endif
 };
 
 
-DEFINE_ACTORDESC_TYPE(APackedLevelActor, FWorldPartitionActorDesc);
\ No newline at end of file
+DEFINE_ACTORDESC_TYPE(APackedLevelActor, FPackedLevelActorDesc);
\ No newline at end of file
diff --git a/Engine/Source/Runtime/Engine/Public/WorldPartition/Filter/WorldPartitionActorFilter.h b/Engine/Source/Runtime/Engine/Public/WorldPartition/Filter/WorldPartitionActorFilter.h
--- a/Engine/Source/Runtime/Engine/Public/WorldPartition/Filter/WorldPartitionActorFilter.h
+++ b/Engine/Source/Runtime/Engine/Public/WorldPartition/Filter/WorldPartitionActorFilter.h
@@ -9,6 +9,15 @@
 #include "UObject/SoftObjectPath.h"
 #include "WorldPartitionActorFilter.generated.h"
 
+enum class EWorldPartitionActorFilterType : uint8
+{
+	None = 0,
+	Loading = 1 << 1,
+	Packing = 1 << 0,
+	All = Packing | Loading,
+};
+ENUM_CLASS_FLAGS(EWorldPartitionActorFilterType);
+
 USTRUCT()
 struct ENGINE_API FWorldPartitionActorFilter
 {
diff --git a/Engine/Source/Runtime/Engine/Public/WorldPartition/LevelInstance/LevelInstanceActorDesc.h b/Engine/Source/Runtime/Engine/Public/WorldPartition/LevelInstance/LevelInstanceActorDesc.h
--- a/Engine/Source/Runtime/Engine/Public/WorldPartition/LevelInstance/LevelInstanceActorDesc.h
+++ b/Engine/Source/Runtime/Engine/Public/WorldPartition/LevelInstance/LevelInstanceActorDesc.h
@@ -26,7 +26,7 @@ public:
 	virtual ~FLevelInstanceActorDesc() override;
 
 	virtual bool IsContainerInstance() const override;
-	virtual bool IsContainerFilter() const override { return IsContainerInstance(); }
+	virtual EWorldPartitionActorFilterType GetContainerFilterType() const { return IsContainerInstance() ? EWorldPartitionActorFilterType::Loading : EWorldPartitionActorFilterType::None; }
 	virtual FName GetContainerPackage() const override { return WorldAsset.GetLongPackageFName(); }
 	virtual bool GetContainerInstance(FContainerInstance& OutContainerInstance) const override;
 	virtual const FWorldPartitionActorFilter* GetContainerFilter() const override { return &Filter; }
diff --git a/Engine/Source/Runtime/Engine/Public/WorldPartition/PackedLevelActor/PackedLevelActorDesc.h b/Engine/Source/Runtime/Engine/Public/WorldPartition/PackedLevelActor/PackedLevelActorDesc.h
new file mode 100644
--- /dev/null
+++ b/Engine/Source/Runtime/Engine/Public/WorldPartition/PackedLevelActor/PackedLevelActorDesc.h
@@ -0,0 +1,31 @@
+// Copyright Epic Games, Inc. All Rights Reserved.
+
+#pragma once
+
+#if WITH_EDITOR
+#include "CoreMinimal.h"
+#include "UObject/ObjectMacros.h"
+#include "WorldPartition/WorldPartitionActorDesc.h"
+#include "WorldPartition/Filter/WorldPartitionActorFilter.h"
+
+class ENGINE_API FPackedLevelActorDesc : public FWorldPartitionActorDesc
+{
+public:
+	FPackedLevelActorDesc() {}
+	virtual ~FPackedLevelActorDesc() override {}
+
+	virtual EWorldPartitionActorFilterType GetContainerFilterType() const { return EWorldPartitionActorFilterType::Packing; }
+	virtual FName GetContainerPackage() const override { return WorldAsset.GetLongPackageFName(); }
+	virtual const FWorldPartitionActorFilter* GetContainerFilter() const override { return &ContainerFilter; }
+
+protected:
+	virtual void Init(const AActor* InActor) override;
+	virtual void Init(const FWorldPartitionActorDescInitData& DescData) override;
+	virtual void Serialize(FArchive& Ar) override;
+	virtual bool Equals(const FWorldPartitionActorDesc* Other) const override;
+	virtual uint32 GetSizeOf() const override { return sizeof(FPackedLevelActorDesc); }
+
+	FSoftObjectPath WorldAsset;
+	FWorldPartitionActorFilter ContainerFilter;
+};
+#endif
diff --git a/Engine/Source/Runtime/Engine/Public/WorldPartition/WorldPartitionActorDesc.h b/Engine/Source/Runtime/Engine/Public/WorldPartition/WorldPartitionActorDesc.h
--- a/Engine/Source/Runtime/Engine/Public/WorldPartition/WorldPartitionActorDesc.h
+++ b/Engine/Source/Runtime/Engine/Public/WorldPartition/WorldPartitionActorDesc.h
@@ -13,6 +13,7 @@
 #include "Templates/SubclassOf.h"
 #include "Misc/Guid.h"
 #include "WorldPartition/WorldPartitionActorContainerID.h"
+#include "WorldPartition/Filter/WorldPartitionActorFilter.h"
 
 // Struct used to create actor descriptor
 struct FWorldPartitionActorDescInitData
@@ -30,7 +31,6 @@ struct FWorldPartitionActorDescInitData
 class AActor;
 class UActorDescContainer;
 class IStreamingGenerationErrorHandler;
-struct FWorldPartitionActorFilter;
 
 enum class EContainerClusterMode : uint8
 {
@@ -144,7 +144,7 @@ PRAGMA_ENABLE_DEPRECATION_WARNINGS
 	inline bool IsDefaultActorDesc() const { return bIsDefaultActorDesc; }
 
 	virtual bool IsContainerInstance() const { return false; }
-	virtual bool IsContainerFilter() const { return false; }
+	virtual EWorldPartitionActorFilterType GetContainerFilterType() const { return EWorldPartitionActorFilterType::None; }
 	virtual FName GetContainerPackage() const { return NAME_None; }
 	
 	virtual const FWorldPartitionActorFilter* GetContainerFilter() const { return nullptr; }
diff --git a/Engine/Source/Runtime/Engine/Public/WorldPartition/WorldPartitionActorDescView.h b/Engine/Source/Runtime/Engine/Public/WorldPartition/WorldPartitionActorDescView.h
--- a/Engine/Source/Runtime/Engine/Public/WorldPartition/WorldPartitionActorDescView.h
+++ b/Engine/Source/Runtime/Engine/Public/WorldPartition/WorldPartitionActorDescView.h
@@ -13,6 +13,7 @@ class IStreamingGenerationErrorHandler;
 class UActorDescContainer;
 class UWorldPartition;
 struct FWorldPartitionActorFilter;
+enum class EWorldPartitionActorFilterType : uint8;
 
 /**
  * A view on top of an actor desc, used to cache information that can be (potentially) different than the actor desc
@@ -63,7 +64,7 @@ public:
 	FName GetContainerPackage() const;
 	bool IsContainerInstance() const;
 	bool GetContainerInstance(FWorldPartitionActorDesc::FContainerInstance& OutContainerInstance) const;
-	bool IsContainerFilter() const;
+	EWorldPartitionActorFilterType GetContainerFilterType() const;
 	const FWorldPartitionActorFilter* GetContainerFilter() const;
 	
 	UE_DEPRECATED(5.3, "GetLevelPackage is deprecated use GetContainerPackage instead.")
diff --git a/Engine/Source/Runtime/Engine/Public/WorldPartition/WorldPartitionSubsystem.h b/Engine/Source/Runtime/Engine/Public/WorldPartition/WorldPartitionSubsystem.h
--- a/Engine/Source/Runtime/Engine/Public/WorldPartition/WorldPartitionSubsystem.h
+++ b/Engine/Source/Runtime/Engine/Public/WorldPartition/WorldPartitionSubsystem.h
@@ -103,8 +103,8 @@ public:
 	void ForEachWorldPartition(TFunctionRef<bool(UWorldPartition*)> Func);
 
 #if WITH_EDITOR
-	FWorldPartitionActorFilter GetWorldPartitionActorFilter(const FString& InWorldPackage) const;
-	TMap<FActorContainerID, TSet<FGuid>> GetFilteredActorsPerContainer(const FActorContainerID& InContainerID, const FString& InWorldPackage, const FWorldPartitionActorFilter& InActorFilter);
+	FWorldPartitionActorFilter GetWorldPartitionActorFilter(const FString& InWorldPackage, EWorldPartitionActorFilterType InFilterTypes = EWorldPartitionActorFilterType::Loading) const;
+	TMap<FActorContainerID, TSet<FGuid>> GetFilteredActorsPerContainer(const FActorContainerID& InContainerID, const FString& InWorldPackage, const FWorldPartitionActorFilter& InActorFilter, EWorldPartitionActorFilterType InFilterTypes = EWorldPartitionActorFilterType::Loading);
 
 	static bool IsRunningConvertWorldPartitionCommandlet();
 
@@ -148,7 +148,7 @@ public:
 		TMap<FName, FActorDescContainerInstance> ActorDescContainers;
 	};
 private:
-	FWorldPartitionActorFilter GetWorldPartitionActorFilterInternal(const FString& InWorldPackage, TSet<FString>& InOutVisitedPackages) const;
+	FWorldPartitionActorFilter GetWorldPartitionActorFilterInternal(const FString& InWorldPackage, EWorldPartitionActorFilterType InFilterTypes, TSet<FString>& InOutVisitedPackages) const;
 #endif
 
 protected:
