Fix character movement replication bug where time dilation was being applied inversely, causing client moves to be delayed longer than expected (if dilation < 1) or to break intended sending rates (if dilation > 1).   This changes net limits to use real time, rather than game time.  This takes dilation and pausing out of the equation.

diff --git a/Engine/Source/Runtime/Engine/Classes/GameFramework/CharacterMovementComponent.h b/Engine/Source/Runtime/Engine/Classes/GameFramework/CharacterMovementComponent.h
--- a/Engine/Source/Runtime/Engine/Classes/GameFramework/CharacterMovementComponent.h
+++ b/Engine/Source/Runtime/Engine/Classes/GameFramework/CharacterMovementComponent.h
@@ -3032,10 +3032,10 @@ public:
 	/** Allows references to be considered during GC */
 	void AddStructReferencedObjects(FReferenceCollector& Collector) const override;
 
-	/** Client timestamp of last time it sent a servermove() to the server. This is an increasing timestamp from the owning UWorld. Used for holding off on sending movement updates to save bandwidth. */
-	float ClientUpdateTime;
+	/** Timestamp of last time it sent a servermove() to the server. This is an increasing timestamp from the owning UWorld (undilated, real time seconds). Used for holding off on sending movement updates to save bandwidth. */
+	float ClientUpdateRealTime;
 
-	/** Current TimeStamp for sending new Moves to the Server. This time resets to zero at a frequency of MinTimeBetweenTimeStampResets. */
+	/** Current TimeStamp for sending new Moves to the Server. This time resets to zero at a frequency of MinTimeBetweenTimeStampResets and increases with this component's tick delta time. */
 	float CurrentTimeStamp;
 
 	/** Last World timestamp (undilated, real time) at which we received a server ack for a move. This could be either a good move or a correction from the server. */
diff --git a/Engine/Source/Runtime/Engine/Private/Components/CharacterMovementComponent.cpp b/Engine/Source/Runtime/Engine/Private/Components/CharacterMovementComponent.cpp
--- a/Engine/Source/Runtime/Engine/Private/Components/CharacterMovementComponent.cpp
+++ b/Engine/Source/Runtime/Engine/Private/Components/CharacterMovementComponent.cpp
@@ -8483,9 +8483,10 @@ void UCharacterMovementComponent::ReplicateMoveToServer(float DeltaTime, const F
 		if (bCanDelayMove && ClientData->PendingMove.IsValid() == false)
 		{
 			// Decide whether to hold off on move
-			const float NetMoveDelta = FMath::Clamp(GetClientNetSendDeltaTime(PC, ClientData, NewMovePtr), 1.f/120.f, 1.f/5.f);
+			const float NetMoveDeltaSeconds = FMath::Clamp(GetClientNetSendDeltaTime(PC, ClientData, NewMovePtr), 1.f/120.f, 1.f/5.f);
+			const float SecondsSinceLastMoveSent = MyWorld->GetRealTimeSeconds() - ClientData->ClientUpdateRealTime;
 
-			if ((MyWorld->TimeSeconds - ClientData->ClientUpdateTime) * MyWorld->GetWorldSettings()->GetEffectiveTimeDilation() < NetMoveDelta)
+			if (SecondsSinceLastMoveSent < NetMoveDeltaSeconds)
 			{
 				// Delay sending this move.
 				ClientData->PendingMove = NewMovePtr;
@@ -8493,7 +8494,7 @@ void UCharacterMovementComponent::ReplicateMoveToServer(float DeltaTime, const F
 			}
 		}
 
-		ClientData->ClientUpdateTime = MyWorld->TimeSeconds;
+		ClientData->ClientUpdateRealTime = MyWorld->GetRealTimeSeconds();
 
 		UE_CLOG(CharacterOwner && UpdatedComponent, LogNetPlayerMovement, VeryVerbose, TEXT("ClientMove Time %f Acceleration %s Velocity %s Position %s Rotation %s DeltaTime %f Mode %s MovementBase %s.%s (Dynamic:%d) DualMove? %d"),
 			NewMove->TimeStamp, *NewMove->Acceleration.ToString(), *Velocity.ToString(), *UpdatedComponent->GetComponentLocation().ToString(), *UpdatedComponent->GetComponentRotation().ToCompactString(), NewMove->DeltaTime, *GetMovementName(),
@@ -11587,7 +11588,7 @@ void UCharacterMovementComponent::ConvertRootMotionServerIDsToLocalIDs(const FRo
 PRAGMA_DISABLE_DEPRECATION_WARNINGS // For deprecated members of FNetworkPredictionData_Client_Character
 
 FNetworkPredictionData_Client_Character::FNetworkPredictionData_Client_Character(const UCharacterMovementComponent& ClientMovement)
-	: ClientUpdateTime(0.f)
+	: ClientUpdateRealTime(0.f)
 	, CurrentTimeStamp(0.f)
 	, LastReceivedAckRealTime(0.f)
 	, PendingMove(NULL)
@@ -12601,7 +12602,7 @@ void UCharacterMovementComponent::FlushServerMoves()
 		if (ClientData->PendingMove.IsValid())
 		{
 			const UWorld* MyWorld = GetWorld();
-			ClientData->ClientUpdateTime = MyWorld->TimeSeconds;
+			ClientData->ClientUpdateRealTime = MyWorld->GetRealTimeSeconds();
 
 			FSavedMovePtr NewMove = ClientData->PendingMove;
 			ClientData->PendingMove = nullptr;
