Fix the Input Device Subsystem not finding a valid Game World in a packaged windows project.The subsystem was searching for a PIE only play world, not a cooked game world as well.

diff --git a/Engine/Source/Runtime/Engine/Private/GameFramework/InputDeviceSubsystem.cpp b/Engine/Source/Runtime/Engine/Private/GameFramework/InputDeviceSubsystem.cpp
--- a/Engine/Source/Runtime/Engine/Private/GameFramework/InputDeviceSubsystem.cpp
+++ b/Engine/Source/Runtime/Engine/Private/GameFramework/InputDeviceSubsystem.cpp
@@ -245,7 +245,7 @@ UWorld* UInputDeviceSubsystem::GetTickableGameObjectWorld() const
 				continue;
 			}
 			// Prefer new PIE window worlds
-			else if (Context.WorldType == EWorldType::PIE)
+			else if (Context.WorldType == EWorldType::PIE || Context.WorldType == EWorldType::Game)
 			{
 				World = ThisWorld;
 				break;
