StaticMeshComponent: Fix regression in partial marquee selection

diff --git a/Engine/Source/Runtime/Engine/Private/Components/StaticMeshComponent.cpp b/Engine/Source/Runtime/Engine/Private/Components/StaticMeshComponent.cpp
--- a/Engine/Source/Runtime/Engine/Private/Components/StaticMeshComponent.cpp
+++ b/Engine/Source/Runtime/Engine/Private/Components/StaticMeshComponent.cpp
@@ -3192,10 +3192,13 @@ bool UStaticMeshComponent::ComponentIsTouchingSelectionBox(const FBox& InSelBBox
 				switch (ComponentIsTouchingSelectionHelpers::CheckSectionBounds(Section, Indices, Vertices, bMustEncompassEntireComponent))
 				{
 				case ComponentIsTouchingSelectionHelpers::ECheckSectionBoundsResult::Valid:
-					/* Proceed with triangle tests. */ ;
+					/* Proceed with triangle tests. */
+					break;
 				case ComponentIsTouchingSelectionHelpers::ECheckSectionBoundsResult::InvalidContinue:
+					/* Skip this mesh section, and continue with next one. */
 					continue;
 				case ComponentIsTouchingSelectionHelpers::ECheckSectionBoundsResult::InvalidFail:
+					/* Invalid data; fail test. */
 					return false;
 				}
 
@@ -3265,10 +3268,13 @@ bool UStaticMeshComponent::ComponentIsTouchingSelectionFrustum(const FConvexVolu
 				switch (ComponentIsTouchingSelectionHelpers::CheckSectionBounds(Section, Indices, Vertices, bMustEncompassEntireComponent))
 				{
 				case ComponentIsTouchingSelectionHelpers::ECheckSectionBoundsResult::Valid:
-					/* Proceed with triangle tests. */ ;
+					/* Proceed with triangle tests. */
+					break;
 				case ComponentIsTouchingSelectionHelpers::ECheckSectionBoundsResult::InvalidContinue:
+					/* Skip this mesh section, and continue with next one. */
 					continue;
 				case ComponentIsTouchingSelectionHelpers::ECheckSectionBoundsResult::InvalidFail:
+					/* Invalid data; fail test. */
 					return false;
 				}
 
