Fix for potential Divide By Zero in Sound Node Enveloper

diff --git a/Engine/Source/Runtime/Engine/Private/SoundNodeEnveloper.cpp b/Engine/Source/Runtime/Engine/Private/SoundNodeEnveloper.cpp
--- a/Engine/Source/Runtime/Engine/Private/SoundNodeEnveloper.cpp
+++ b/Engine/Source/Runtime/Engine/Private/SoundNodeEnveloper.cpp
@@ -87,8 +87,8 @@ void USoundNodeEnveloper::ParseNodes( FAudioDevice* AudioDevice, const UPTRINT N
 		}
 
 		float LoopDuration = LoopEnd - LoopStart;
-		int32 CurrentLoopCount = (int32)(PlayTime - LoopStart)/LoopDuration;
-		PlayTime -= CurrentLoopCount*LoopDuration;
+		int32 CurrentLoopCount = (int32)(PlayTime - LoopStart) / FMath::Min(UE_SMALL_NUMBER, LoopDuration);
+		PlayTime -= CurrentLoopCount * LoopDuration;
 
 		if( CurrentLoopCount == LoopCount && !bLoopIndefinitely && LoopCount != 0 )
 		{
