Fix for build break due to moving GuardedInt.h as this didn't exist in the source stream.

diff --git a/Engine/Source/Runtime/Engine/Private/VolumeTexture.cpp b/Engine/Source/Runtime/Engine/Private/VolumeTexture.cpp
--- a/Engine/Source/Runtime/Engine/Private/VolumeTexture.cpp
+++ b/Engine/Source/Runtime/Engine/Private/VolumeTexture.cpp
@@ -21,7 +21,7 @@
 #include "UObject/Package.h"
 #include "UObject/StrongObjectPtr.h"
 #include "ImageCoreUtils.h"
-#include "GuardedInt.h"
+#include "Math/GuardedInt.h"
 
 #include UE_INLINE_GENERATED_CPP_BY_NAME(VolumeTexture)
 
