Fixed AddressSanitizer delete mismatch: make sure the pointer is deleted through the actual type.

#tested full game session with address sanitizer

diff --git a/Engine/Source/Runtime/Engine/Public/GPUSkinVertexFactory.h b/Engine/Source/Runtime/Engine/Public/GPUSkinVertexFactory.h
--- a/Engine/Source/Runtime/Engine/Public/GPUSkinVertexFactory.h
+++ b/Engine/Source/Runtime/Engine/Public/GPUSkinVertexFactory.h
@@ -675,6 +675,17 @@ public:
 		, TGPUSkinVertexFactory<BoneInfluenceType>(InFeatureLevel, InNumVertices)
 	{}
 
+	/**
+	 * Destructor takes care of the Data pointer. Since FGPUBaseSkinVertexFactory does not know the real type of the Data,
+	 * delete the data here instead.
+	 */
+	virtual ~TGPUSkinAPEXClothVertexFactory() override
+	{
+		checkf(!ClothDataPtr->ClothBuffer.IsValid(), TEXT("ClothBuffer RHI resource should have been released in ReleaseDynamicRHI"));
+		delete ClothDataPtr;
+		this->Data.Release();
+	}
+
 	static void ModifyCompilationEnvironment(const FVertexFactoryShaderPermutationParameters& Parameters, FShaderCompilerEnvironment& OutEnvironment);
 	static bool ShouldCompilePermutation(const FVertexFactoryShaderPermutationParameters& Parameters);
 
