Explicitely discard the return value of Release().

diff --git a/Engine/Source/Runtime/Engine/Public/GPUSkinVertexFactory.h b/Engine/Source/Runtime/Engine/Public/GPUSkinVertexFactory.h
--- a/Engine/Source/Runtime/Engine/Public/GPUSkinVertexFactory.h
+++ b/Engine/Source/Runtime/Engine/Public/GPUSkinVertexFactory.h
@@ -683,7 +683,7 @@ public:
 	{
 		checkf(!ClothDataPtr->ClothBuffer.IsValid(), TEXT("ClothBuffer RHI resource should have been released in ReleaseDynamicRHI"));
 		delete ClothDataPtr;
-		this->Data.Release();
+		(void)this->Data.Release();
 	}
 
 	static void ModifyCompilationEnvironment(const FVertexFactoryShaderPermutationParameters& Parameters, FShaderCompilerEnvironment& OutEnvironment);
