#rnx

diff --git a/Engine/Source/Runtime/Engine/Private/Rendering/SkinWeightVertexBuffer.cpp b/Engine/Source/Runtime/Engine/Private/Rendering/SkinWeightVertexBuffer.cpp
--- a/Engine/Source/Runtime/Engine/Private/Rendering/SkinWeightVertexBuffer.cpp
+++ b/Engine/Source/Runtime/Engine/Private/Rendering/SkinWeightVertexBuffer.cpp
@@ -659,6 +659,20 @@ void FSkinWeightVertexBuffer::CleanUp()
 	LookupVertexBuffer.CleanUp();
 }
 
+void FSkinWeightVertexBuffer::RebuildLookupVertexBuffer()
+{
+	uint32 MaxBoneInfluences = DataVertexBuffer.GetMaxBoneInfluences();
+	uint32 NumVertices = DataVertexBuffer.GetNumBoneWeights() / DataVertexBuffer.GetMaxBoneInfluences();
+	LookupVertexBuffer.Init(NumVertices);
+
+	uint32 WeightOffset = 0;
+	for (uint32 VertIdx = 0; VertIdx < NumVertices; VertIdx++)
+	{
+		LookupVertexBuffer.SetWeightOffsetAndInfluenceCount(VertIdx, WeightOffset * GetBoneIndexAndWeightByteSize(), MaxBoneInfluences);
+		WeightOffset += MaxBoneInfluences;
+	}
+}
+
 FArchive& operator<<(FArchive& Ar, FSkinWeightVertexBuffer& VertexBuffer)
 {
 	Ar.UsingCustomVersion(FAnimObjectVersion::GUID);	
@@ -667,16 +681,7 @@ FArchive& operator<<(FArchive& Ar, FSkinWeightVertexBuffer& VertexBuffer)
 	if (Ar.IsLoading() && Ar.CustomVer(FAnimObjectVersion::GUID) < FAnimObjectVersion::UnlimitedBoneInfluences)
 	{
 		// LookupVertexBuffer doesn't exist before this version, so construct its content from scratch
-		uint32 MaxBoneInfluences = VertexBuffer.DataVertexBuffer.GetMaxBoneInfluences();
-		uint32 NumVertices = VertexBuffer.DataVertexBuffer.GetNumBoneWeights() / VertexBuffer.DataVertexBuffer.GetMaxBoneInfluences();
-		VertexBuffer.LookupVertexBuffer.Init(NumVertices);
-
-		uint32 WeightOffset = 0;
-		for (uint32 VertIdx = 0; VertIdx < NumVertices; VertIdx++)
-		{
-			VertexBuffer.LookupVertexBuffer.SetWeightOffsetAndInfluenceCount(VertIdx, WeightOffset * VertexBuffer.GetBoneIndexAndWeightByteSize(), MaxBoneInfluences);
-			WeightOffset += MaxBoneInfluences;
-		}
+		VertexBuffer.RebuildLookupVertexBuffer();
 	}
 	else
 	{
diff --git a/Engine/Source/Runtime/Engine/Public/Rendering/SkinWeightVertexBuffer.h b/Engine/Source/Runtime/Engine/Public/Rendering/SkinWeightVertexBuffer.h
--- a/Engine/Source/Runtime/Engine/Public/Rendering/SkinWeightVertexBuffer.h
+++ b/Engine/Source/Runtime/Engine/Public/Rendering/SkinWeightVertexBuffer.h
@@ -474,6 +474,10 @@ public:
 	ENGINE_API uint16 GetBoneWeight(uint32 VertexIndex, uint32 InfluenceIndex) const;
 	ENGINE_API void SetBoneWeight(uint32 VertexIndex, uint32 InfluenceIndex, uint16 BoneWeight);
 	ENGINE_API void ResetVertexBoneWeights(uint32 VertexIndex);
+	
+	/** Should only be called for legacy meshes or meshes that don't really have a LookupVertexBuffer, for example created in realtime.
+		All variable influence meshes imported after the FAnimObjectVersion::UnlimitedBoneInfluences version already have a LookupVertexBuffer */
+	ENGINE_API void RebuildLookupVertexBuffer();
 
 	/** SetOwnerName should be called before BeginInitResources for the owner name to be successfully tracked. */
 	ENGINE_API void SetOwnerName(const FName& OwnerName);
