Add Specular Profile asset.

This allows artist to provide a custom LUT describing specular behavior for view and light angles. The LUT can be procedurally generated or provided from a texture.

diff --git a/Engine/Config/BaseEditor.ini b/Engine/Config/BaseEditor.ini
--- a/Engine/Config/BaseEditor.ini
+++ b/Engine/Config/BaseEditor.ini
@@ -55,6 +55,7 @@ ColorFontAtlasSize=512
 +RenderableThumbnailTypes=(ClassNeedingThumbnailName="/Script/Engine.SoundWave",RendererClassName="/Script/UnrealEd.SoundWaveThumbnailRenderer")
 +RenderableThumbnailTypes=(ClassNeedingThumbnailName="/Script/Engine.SlateBrushAsset",RendererClassName="/Script/UnrealEd.SlateBrushThumbnailRenderer")
 +RenderableThumbnailTypes=(ClassNeedingThumbnailName="/Script/Engine.SubsurfaceProfile",RendererClassName="/Script/UnrealEd.SubsurfaceProfileRenderer")
++RenderableThumbnailTypes=(ClassNeedingThumbnailName="/Script/Engine.SpecularProfile",RendererClassName="/Script/UnrealEd.SpecularProfileRenderer")
 +RenderableThumbnailTypes=(ClassNeedingThumbnailName="/Script/Engine.PhysicsAsset",RendererClassName="/Script/UnrealEd.PhysicsAssetThumbnailRenderer")
 +RenderableThumbnailTypes=(ClassNeedingThumbnailName="/Script/Engine.CurveLinearColor",RendererClassName="/Script/UnrealEd.CurveLinearColorThumbnailRenderer")
 +RenderableThumbnailTypes=(ClassNeedingThumbnailName="/Script/Engine.CurveFloat",RendererClassName="/Script/UnrealEd.CurveFloatThumbnailRenderer")
diff --git a/Engine/Plugins/Editor/EngineAssetDefinitions/Source/Private/AssetDefinition_SpecularProfile.h b/Engine/Plugins/Editor/EngineAssetDefinitions/Source/Private/AssetDefinition_SpecularProfile.h
new file mode 100644
--- /dev/null
+++ b/Engine/Plugins/Editor/EngineAssetDefinitions/Source/Private/AssetDefinition_SpecularProfile.h
@@ -0,0 +1,26 @@
+// Copyright Epic Games, Inc. All Rights Reserved.
+
+#pragma once
+
+#include "Engine/SpecularProfile.h"
+#include "AssetDefinitionDefault.h"
+
+#include "AssetDefinition_SpecularProfile.generated.h"
+
+UCLASS()
+class UAssetDefinition_SpecularProfile : public UAssetDefinitionDefault
+{
+	GENERATED_BODY()
+
+public:
+	// UAssetDefinition Begin
+	virtual FText GetAssetDisplayName() const override { return NSLOCTEXT("AssetTypeActions", "AssetTypeActions_SpecularProfile", "Specular Profile"); }
+	virtual FLinearColor GetAssetColor() const override { return FLinearColor(FColor(128, 128, 255)); }
+	virtual TSoftClassPtr<UObject> GetAssetClass() const override { return USpecularProfile::StaticClass(); }
+	virtual TConstArrayView<FAssetCategoryPath> GetAssetCategories() const override
+    {
+    	static const auto Categories = { EAssetCategoryPaths::Material };
+    	return Categories;
+    }
+	// UAssetDefinition End
+};
diff --git a/Engine/Shaders/Private/DebugProbes.usf b/Engine/Shaders/Private/DebugProbes.usf
--- a/Engine/Shaders/Private/DebugProbes.usf
+++ b/Engine/Shaders/Private/DebugProbes.usf
@@ -139,6 +139,7 @@ void WriteOutMaterial(
 		/*float FuzzRoughness					*/0.0f,
 		/*float GlintValue						*/0.0f,
 		/*float GlintUV							*/0.0f.xx,
+		/*float SpecularProfile					*/0.0f,
 		/*float Thickness						*/0.0f,
 		/*bool bIsThinSurface					*/false,
 		/*uint SharedLocalBasisIndex			*/SharedLocalBasisIndex,
diff --git a/Engine/Shaders/Private/SpecularProfile.usf b/Engine/Shaders/Private/SpecularProfile.usf
new file mode 100644
--- /dev/null
+++ b/Engine/Shaders/Private/SpecularProfile.usf
@@ -0,0 +1,23 @@
+// Copyright Epic Games, Inc. All Rights Reserved.
+
+#include "Common.ush"
+
+int2 SourceResolution;
+uint SourceMipCount;
+int2 TargetResolution;
+uint TargetIndex;
+SamplerState SourceSampler;
+Texture2D<float4> SourceTexture;
+RWTexture2DArray<float4> TargetTexture;
+
+[numthreads(8, 8, 1)]
+void MainCS(uint3 DispatchThreadId : SV_DispatchThreadID)
+{
+	const uint2 PixelPos = DispatchThreadId.xy;
+	if (all(PixelPos < uint2(TargetResolution)))
+	{
+		const float SourceMip = clamp(log2(SourceResolution) - log2(TargetResolution), 0.f, SourceMipCount-1);
+		const float4 Value = SourceTexture.SampleLevel(SourceSampler, float2(PixelPos+0.5f) / float2(TargetResolution), SourceMip);
+		TargetTexture[uint3(PixelPos,TargetIndex)] = Value;
+	}
+}
\ No newline at end of file
diff --git a/Engine/Shaders/Private/Strata/Strata.ush b/Engine/Shaders/Private/Strata/Strata.ush
--- a/Engine/Shaders/Private/Strata/Strata.ush
+++ b/Engine/Shaders/Private/Strata/Strata.ush
@@ -369,7 +369,8 @@ struct FStrataBSDF
 #define BSDF_BIT_COUNT_SHAREDLOCALBASESID		2
 #define BSDF_BIT_COUNT_HASTRANSABOVE			1
 #define BSDF_BIT_COUNT_HASGLINT					1
-#define BSDF_BIT_COUNT___UNUSED___				4
+#define BSDF_BIT_COUNT_HASSPECPROFILE			1
+#define BSDF_BIT_COUNT___UNUSED___				3
 #define BSDF_BIT_COUNT_WEIGHT10F				10
 
 // Bit offset of each component store in BSDF state 
@@ -385,7 +386,8 @@ struct FStrataBSDF
 #define BSDF_BIT_OFFSET_BSDF_TYPE				(BSDF_BIT_OFFSET_ISTHIN				+ BSDF_BIT_COUNT_ISTHIN)
 #define BSDF_BIT_OFFSET_SHAREDLOCALBASESID		(BSDF_BIT_OFFSET_BSDF_TYPE			+ BSDF_BIT_COUNT_BSDF_TYPE)
 #define BSDF_BIT_OFFSET_HASTRANSABOVE			(BSDF_BIT_OFFSET_SHAREDLOCALBASESID	+ BSDF_BIT_COUNT_SHAREDLOCALBASESID)
-#define BSDF_BIT_OFFSET_HASGLINT				(BSDF_BIT_OFFSET_HASTRANSABOVE	+ BSDF_BIT_COUNT_HASTRANSABOVE)
+#define BSDF_BIT_OFFSET_HASGLINT				(BSDF_BIT_OFFSET_HASTRANSABOVE		+ BSDF_BIT_COUNT_HASTRANSABOVE)
+#define BSDF_BIT_OFFSET_HASSPECPROFILE			(BSDF_BIT_OFFSET_HASGLINT			+ BSDF_BIT_COUNT_HASGLINT)
 #define BSDF_BIT_OFFSET_WEIGHT10F				22
 
 // Special optmised modes for the single BSDF path.
@@ -479,10 +481,14 @@ struct FStrataBSDF
 #define BSDF_GETHASTRANSABOVE(X)		READ_BITS(X.State,  BSDF_BIT_COUNT_HASTRANSABOVE, BSDF_BIT_OFFSET_HASTRANSABOVE)
 #define BSDF_SETHASTRANSABOVE(X, Value)	WRITE_BITS(X.State, BSDF_BIT_COUNT_HASTRANSABOVE, BSDF_BIT_OFFSET_HASTRANSABOVE, Value)
 
-// Indicates if the BSDF is part of the top layer
+// Indicates if the BSDF has glint data
 #define BSDF_GETHASGLINT(X)				READ_BITS(X.State,  BSDF_BIT_COUNT_HASGLINT, BSDF_BIT_OFFSET_HASGLINT) 
 #define BSDF_SETHASGLINT(X, Value)		WRITE_BITS(X.State, BSDF_BIT_COUNT_HASGLINT, BSDF_BIT_OFFSET_HASGLINT, Value)
 
+// Indicates if the BSDF has a specular profile
+#define BSDF_GETHASSPECPROFILE(X)		READ_BITS(X.State,  BSDF_BIT_COUNT_HASSPECPROFILE, BSDF_BIT_OFFSET_HASSPECPROFILE) 
+#define BSDF_SETHASSPECPROFILE(X, Value)WRITE_BITS(X.State, BSDF_BIT_COUNT_HASSPECPROFILE, BSDF_BIT_OFFSET_HASSPECPROFILE, Value)
+
 #if STRATA_INLINE_SHADING
 #define BSDF_SETEMISSIVE(X, V)			X.Emissive = V
 #define BSDF_GETEMISSIVE(X)				X.Emissive
@@ -531,6 +537,8 @@ struct FStrataBSDF
 #define SLAB_GLINT_VALUE(X)				X.VGPRs[5].z
 #define SLAB_GLINT_UVDDX(X)				X.VGPRs[6].xy
 #define SLAB_GLINT_UVDDY(X)				X.VGPRs[6].zw
+// Specular profile
+#define SLAB_SPECPROFILEID(X)			X.VGPRs[5].w
 #endif
 
 ///////////////////////////////////////////////////////////////////////////////
@@ -1473,6 +1481,7 @@ FStrataData GetStrataSlabBSDF(
 	float SecondRoughness, float SecondRoughnessWeight, float SecondRoughnessAsSimpleClearCoat,
 	float FuzzAmount, float3 FuzzColor, float FuzzRoughness,
 	float GlintValue, float2 GlintUV,
+	float SpecularProfileId,
 	float ThicknessCm,
 	bool bIsThinSurface,
 	uint SharedLocalBasisIndex, inout uint SharedLocalBasisTypes)
@@ -1573,6 +1582,12 @@ FStrataData GetStrataSlabBSDF(
 		SLAB_GLINT_UVDDX(StrataData.InlinedBSDF)		= ddx(GlintUV);
 		SLAB_GLINT_UVDDY(StrataData.InlinedBSDF)		= ddy(GlintUV);
 	}
+	
+	if (SpecularProfileId > 0.0f)
+	{
+		BSDF_SETHASSPECPROFILE	(StrataData.InlinedBSDF, 1);
+		SLAB_SPECPROFILEID		(StrataData.InlinedBSDF) = SpecularProfileId * 255; // Stored as 0-255 slice index
+	}
 #endif
 
 #if !STRATA_FASTPATH
@@ -2381,6 +2396,9 @@ FStrataData StrataHorizontalMixingParameterBlending(FStrataData Background, FStr
 	SLAB_GLINT_VALUE(ResultBSDF) =		lerp(SLAB_GLINT_VALUE(BackgroundBSDF),			SLAB_GLINT_VALUE(ForegroundBSDF),			ForegroundOtherMixFactor);
 	SLAB_GLINT_UVDDX(ResultBSDF) =		lerp(SLAB_GLINT_UVDDX(BackgroundBSDF),			SLAB_GLINT_UVDDX(ForegroundBSDF),			ForegroundOtherMixFactor);
 	SLAB_GLINT_UVDDY(ResultBSDF) =		lerp(SLAB_GLINT_UVDDY(BackgroundBSDF),			SLAB_GLINT_UVDDY(ForegroundBSDF),			ForegroundOtherMixFactor);
+
+	BSDF_SETHASSPECPROFILE(ResultBSDF,	BSDF_GETHASSPECPROFILE(ForegroundBSDF)			| BSDF_GETHASSPECPROFILE(BackgroundBSDF));
+	SLAB_SPECPROFILEID(ResultBSDF) = BSDF_GETHASSPECPROFILE(ForegroundBSDF) ? SLAB_SPECPROFILEID(ForegroundBSDF) : SLAB_SPECPROFILEID(BackgroundBSDF);
 #endif
 	
 
@@ -2496,6 +2514,9 @@ FStrataData StrataAddParameterBlending(FStrataData A, FStrataData B, float AMixF
 	SLAB_GLINT_VALUE(ResultBSDF) =		lerp(SLAB_GLINT_VALUE(BBSDF),					SLAB_GLINT_VALUE(ABSDF), AMixFactor);
 	SLAB_GLINT_UVDDX(ResultBSDF) =		lerp(SLAB_GLINT_UVDDX(BBSDF),					SLAB_GLINT_UVDDX(ABSDF), AMixFactor);
 	SLAB_GLINT_UVDDY(ResultBSDF) =		lerp(SLAB_GLINT_UVDDY(BBSDF),					SLAB_GLINT_UVDDY(ABSDF), AMixFactor);
+
+	BSDF_SETHASSPECPROFILE(ResultBSDF,	BSDF_GETHASSPECPROFILE(BBSDF) | BSDF_GETHASSPECPROFILE(ABSDF));
+	SLAB_SPECPROFILEID(ResultBSDF) = BSDF_GETHASSPECPROFILE(BBSDF) ? SLAB_SPECPROFILEID(BBSDF) : SLAB_SPECPROFILEID(ABSDF);
 #endif
 
 	if (BSDF_GETSSSTYPE(ResultBSDF) == SSS_TYPE_DIFFUSION_PROFILE)
@@ -2578,6 +2599,7 @@ FStrataBSDF StrataVerticalLayeringParameterBlending(in FStrataBSDF TopBSDF, in F
 
 #if STRATA_GLINTS_ENABLED
 	BSDF_SETHASGLINT(ResultBSDF, 0); // Disable for now
+	BSDF_SETHASSPECPROFILE(ResultBSDF,	0); // Disable for now
 #endif
 
 	// Compute the top slab transmittance that will be used to blend the top and base material parameters.
@@ -2622,6 +2644,9 @@ FStrataBSDF StrataVerticalLayeringParameterBlending(in FStrataBSDF TopBSDF, in F
 	SLAB_GLINT_VALUE(ResultBSDF) =		lerp(SLAB_GLINT_VALUE(TopBSDF),					SLAB_GLINT_VALUE(BaseBSDF), TopToBottomLerp);
 	SLAB_GLINT_UVDDX(ResultBSDF) =		lerp(SLAB_GLINT_UVDDX(TopBSDF),					SLAB_GLINT_UVDDX(BaseBSDF), TopToBottomLerp);
 	SLAB_GLINT_UVDDY(ResultBSDF) =		lerp(SLAB_GLINT_UVDDY(TopBSDF),					SLAB_GLINT_UVDDY(BaseBSDF), TopToBottomLerp);
+
+	BSDF_SETHASSPECPROFILE(ResultBSDF,	BSDF_GETHASSPECPROFILE(TopBSDF) | BSDF_GETHASSPECPROFILE(BaseBSDF));
+	SLAB_SPECPROFILEID(ResultBSDF) = 	BSDF_GETHASSPECPROFILE(TopBSDF) ? SLAB_SPECPROFILEID(TopBSDF) : SLAB_SPECPROFILEID(BaseBSDF);
 #endif
 
 	// Anisotropy is disabled with vertical blending. 
@@ -3657,6 +3682,7 @@ void FStrataPixelHeader::StrataConvertFromDBuffer(
 		In.Roughness,		// FuzzRoughness
 		0.f,				// GlintValue, 
 		0.f.xx,				// GlintUV,
+		0.f, 				// SpecularProfileId
 		DefaultThicknessCm,
 		false,				// bIsThinSurface,
 		SharedLocalBasisIndex,
@@ -4300,6 +4326,11 @@ FStrataBSDF UnpackStrataBSDFIn(FStrataMaterialContainer StrataBuffer, inout FStr
 			SLAB_GLINT_UVDDX(OutBSDF)	= UnpackFloat2FromUInt(StrataLoadUint1(StrataBuffer, StrataAddressing));
 			SLAB_GLINT_UVDDY(OutBSDF)	= UnpackFloat2FromUInt(StrataLoadUint1(StrataBuffer, StrataAddressing));
 		}
+
+		if (BSDF_GETHASSPECPROFILE(OutBSDF))
+		{
+			SLAB_SPECPROFILEID(OutBSDF)	= StrataLoadUint1(StrataBuffer, StrataAddressing);
+		}
 	#endif
 	}
 	break;
diff --git a/Engine/Shaders/Private/Strata/StrataEvaluation.ush b/Engine/Shaders/Private/Strata/StrataEvaluation.ush
--- a/Engine/Shaders/Private/Strata/StrataEvaluation.ush
+++ b/Engine/Shaders/Private/Strata/StrataEvaluation.ush
@@ -681,6 +681,18 @@ FStrataEvaluateResult StrataEvaluateBSDFCommon(FStrataBSDFContext BSDFContext, F
 			}
 			#endif
 
+			// Lighting LUT
+			#if STRATA_FASTPATH==0 && STRATA_GLINTS_ENABLED
+			BRANCH
+			if (BSDF_GETHASSPECPROFILE(BSDFContext.BSDF))
+			{
+				const uint ProfileID = SLAB_SPECPROFILEID(BSDFContext.BSDF);
+				const float NoL = saturate(dot(BSDFContext.L, BSDFContext.N));
+				const float NoV = saturate(dot(BSDFContext.V, BSDFContext.N));
+				Sample.IntegratedSpecularValue *= View.SpecularProfileTexture.SampleLevel(View.SpecularProfileSampler, float3(NoV, NoL, ProfileID), 0);
+			}
+			#endif
+
 			{
 				float3 CommonTerm = 0.0f;
 				if (bIsRectLight)
@@ -1492,6 +1504,7 @@ FStrataEnvLightResult StrataEvaluateForEnvLight(FStrataBSDFContext BSDFContext,
 
 				// STRATA_TODO environment light importance sampling
 			#if STRATA_GLINTS_ENABLED
+				// Glints
 				BRANCH
 				if (BSDF_GETHASGLINT(BSDFContext.BSDF))
 				{
@@ -1520,6 +1533,16 @@ FStrataEnvLightResult StrataEvaluateForEnvLight(FStrataBSDFContext BSDFContext,
 					StrataEnvLightResult.SpecularWeight *= saturate(GlintD);
 					StrataEnvLightResult.SpecularSafeRoughness = lerp(StrataEnvLightResult.SpecularSafeRoughness, 0.001f, saturate(GlintD));
 				}
+
+				// Specular profile
+				BRANCH
+				if (BSDF_GETHASSPECPROFILE(BSDFContext.BSDF))
+				{
+					const uint ProfileID = SLAB_SPECPROFILEID(BSDFContext.BSDF);
+					const float NoR = saturate(dot(BSDFContext.R, BSDFContext.N));
+					const float NoV = saturate(dot(BSDFContext.V, BSDFContext.N));
+					StrataEnvLightResult.SpecularWeight *= View.SpecularProfileTexture.SampleLevel(View.SpecularProfileSampler, float3(NoV, NoR, ProfileID), 0);
+				}
 			#endif // STRATA_COMPLEXPATH
 
 				StrataEnvLightResult.SSRSpecularWeight = StrataEnvLightResult.SpecularWeight;
diff --git a/Engine/Shaders/Private/Strata/StrataExport.ush b/Engine/Shaders/Private/Strata/StrataExport.ush
--- a/Engine/Shaders/Private/Strata/StrataExport.ush
+++ b/Engine/Shaders/Private/Strata/StrataExport.ush
@@ -191,7 +191,7 @@ void PackStrataOut(
 						const bool LuminanceWeightEqualOne = all(BSDF.LuminanceWeightV == 1.0f);
 
 						#if STRATA_GLINTS
-						const bool bForceComplexMaterial = BSDF_GETHASANISOTROPY(BSDF) || BSDF_GETHASGLINT(BSDF);
+						const bool bForceComplexMaterial = BSDF_GETHASANISOTROPY(BSDF) || BSDF_GETHASGLINT(BSDF) || BSDF_GETHASSPECPROFILE(BSDF);
 						#else
 						const bool bForceComplexMaterial = BSDF_GETHASANISOTROPY(BSDF);
 						#endif
@@ -805,6 +805,11 @@ void PackStrataOut(
 								STRATA_STORE_UINT1(PackFloat2ToUInt(SLAB_GLINT_UVDDX(BSDF)));
 								STRATA_STORE_UINT1(PackFloat2ToUInt(SLAB_GLINT_UVDDY(BSDF)));
 							}
+
+							if (BSDF_GETHASSPECPROFILE(BSDF))
+							{
+								STRATA_STORE_UINT1(uint(SLAB_SPECPROFILEID(BSDF)));
+							}
 							#endif // STRATA_GLINTS
 
 							// 8-28 bytes
diff --git a/Engine/Shaders/Private/Strata/StrataLegacyConversion.ush b/Engine/Shaders/Private/Strata/StrataLegacyConversion.ush
--- a/Engine/Shaders/Private/Strata/StrataLegacyConversion.ush
+++ b/Engine/Shaders/Private/Strata/StrataLegacyConversion.ush
@@ -121,6 +121,7 @@ FStrataData CreateLegacySlab(
 		In.FuzzColor,										// FuzzColor
 		In.Roughness,										// FuzzRoughness
 		0.0f, 0.0f.xx,										// GlintData
+		0.0f,												// Specular profile
 		In.Thickness,
 		In.bIsThinSurface,
 		In.SharedLocalBasisIndex, SharedLocalBasisTypes);
diff --git a/Engine/Shaders/Private/Strata/StrataVisualize.usf b/Engine/Shaders/Private/Strata/StrataVisualize.usf
--- a/Engine/Shaders/Private/Strata/StrataVisualize.usf
+++ b/Engine/Shaders/Private/Strata/StrataVisualize.usf
@@ -334,6 +334,10 @@ void PrintBSDF(inout FShaderPrintContext Context, inout FStrataAddressing Strata
 				Print(Context, TEXT(" Glint Value "), FontBSDFPropName);	Print(Context, SLAB_GLINT_VALUE(BSDF), FontBSDFPropValu);				Newline(Context, RectMax);
 				Print(Context, TEXT(" Glint UVs   "), FontBSDFPropName);	Print(Context, SLAB_GLINT_UV(BSDF), FontBSDFPropValu);					Newline(Context, RectMax);
 			}
+			if (BSDF_GETHASSPECPROFILE(BSDF))
+			{
+				Print(Context, TEXT(" SpecID      "), FontBSDFPropName);	Print(Context, SLAB_SPECPROFILEID(BSDF), FontBSDFPropValu);				Newline(Context, RectMax);
+			}
 #endif
 		}
 		break;
diff --git a/Engine/Source/Editor/EditorStyle/Private/SlateEditorStyle.cpp b/Engine/Source/Editor/EditorStyle/Private/SlateEditorStyle.cpp
--- a/Engine/Source/Editor/EditorStyle/Private/SlateEditorStyle.cpp
+++ b/Engine/Source/Editor/EditorStyle/Private/SlateEditorStyle.cpp
@@ -6497,6 +6497,7 @@ void FSlateEditorStyle::FStyle::SetupClassIconsAndThumbnails()
 			TEXT("SlateBrushAsset"),
 			TEXT("SlateWidgetStyleAsset"),
 			TEXT("StringTable"),
+			TEXT("SpecularProfile"),
 			TEXT("SphereReflectionCapture"),
 			TEXT("SpotLight"),
 			TEXT("SpotLightMovable"),
diff --git a/Engine/Source/Editor/EditorStyle/Private/StarshipStyle.cpp b/Engine/Source/Editor/EditorStyle/Private/StarshipStyle.cpp
--- a/Engine/Source/Editor/EditorStyle/Private/StarshipStyle.cpp
+++ b/Engine/Source/Editor/EditorStyle/Private/StarshipStyle.cpp
@@ -5950,6 +5950,7 @@ void FStarshipEditorStyle::FStyle::SetupClassIconsAndThumbnails()
 			{TEXT("SoundMix")},
 			{TEXT("SoundSubmix") },
 			{TEXT("SphereReflectionCapture")},
+			{TEXT("SpecularProfile")},
 			{TEXT("SpotLight")},
 			{TEXT("SpringArmComponent")},
 			{TEXT("StaticMesh")},
diff --git a/Engine/Source/Editor/MaterialEditor/Private/SMaterialEditorStrataWidget.cpp b/Engine/Source/Editor/MaterialEditor/Private/SMaterialEditorStrataWidget.cpp
--- a/Engine/Source/Editor/MaterialEditor/Private/SMaterialEditorStrataWidget.cpp
+++ b/Engine/Source/Editor/MaterialEditor/Private/SMaterialEditorStrataWidget.cpp
@@ -377,7 +377,7 @@ void SMaterialEditorStrataWidget::Tick(const FGeometry& AllottedGeometry, const
 							case STRATA_OPERATOR_BSDF_LEGACY:	// legacy BSDF should have been converted to BSDF already.
 							case STRATA_OPERATOR_BSDF:
 							{
-								FString BSDFDesc = FString::Printf(TEXT("BSDF (%s%s%s%s%s%s%s)")
+								FString BSDFDesc = FString::Printf(TEXT("BSDF (%s%s%s%s%s%s%s%s)")
 									, Op.bBSDFHasEdgeColor ? TEXT("F90 ") : TEXT("")
 									, Op.bBSDFHasSSS ? TEXT("SSS ") : TEXT("")
 									, Op.bBSDFHasMFPPluggedIn ? TEXT("MFP ") : TEXT("")
@@ -385,6 +385,7 @@ void SMaterialEditorStrataWidget::Tick(const FGeometry& AllottedGeometry, const
 									, Op.bBSDFHasSecondRoughnessOrSimpleClearCoat ? TEXT("2Ro ") : TEXT("")
 									, Op.bBSDFHasFuzz ? TEXT("Fuz ") : TEXT("")
 									, Op.bBSDFHasGlint ? TEXT("Gli ") : TEXT("")
+									, Op.bBSDFHasSpecularProfile ? TEXT("Spc ") : TEXT("")
 								);
 
 								static FString ToolTip;
@@ -396,7 +397,8 @@ void SMaterialEditorStrataWidget::Tick(const FGeometry& AllottedGeometry, const
 									ToolTip += TEXT("Fuz means the BSDF fuzz layer is enabled.\n");
 									ToolTip += TEXT("2Ro means the BSDF either uses a second specular lob with a second roughness, or the legacy simple clear coat.\n");
 									ToolTip += TEXT("Ani means the BSDF anisotropic specular lighting is used.\n");
-									ToolTip += TEXT("Gli means the BSDF features glints.");
+									ToolTip += TEXT("Gli means the BSDF features glints.\n");
+									ToolTip += TEXT("Spc means the BSDF features specular profile.\n");
 								}
 
 								auto BSDF = SNew(SErrorText)
diff --git a/Engine/Source/Editor/UnrealEd/Classes/Factories/SpecularProfileFactory.h b/Engine/Source/Editor/UnrealEd/Classes/Factories/SpecularProfileFactory.h
new file mode 100644
--- /dev/null
+++ b/Engine/Source/Editor/UnrealEd/Classes/Factories/SpecularProfileFactory.h
@@ -0,0 +1,21 @@
+// Copyright Epic Games, Inc. All Rights Reserved.
+
+#pragma once
+
+#include "CoreMinimal.h"
+#include "UObject/ObjectMacros.h"
+#include "Factories/Factory.h"
+#include "SpecularProfileFactory.generated.h"
+
+UCLASS(hidecategories=Object)
+class USpecularProfileFactory : public UFactory
+{
+	GENERATED_UCLASS_BODY()
+
+	//~ Begin UFactory Interface
+	virtual UObject* FactoryCreateNew(UClass* Class,UObject* InParent,FName Name,EObjectFlags Flags,UObject* Context,FFeedbackContext* Warn) override;
+	//~ Begin UFactory Interface	
+};
+
+
+
diff --git a/Engine/Source/Editor/UnrealEd/Classes/ThumbnailRendering/SpecularProfileRenderer.h b/Engine/Source/Editor/UnrealEd/Classes/ThumbnailRendering/SpecularProfileRenderer.h
new file mode 100644
--- /dev/null
+++ b/Engine/Source/Editor/UnrealEd/Classes/ThumbnailRendering/SpecularProfileRenderer.h
@@ -0,0 +1,28 @@
+// Copyright Epic Games, Inc. All Rights Reserved.
+
+/**
+ *
+ * This thumbnail renderer displays the texture for the object in question
+ */
+
+#pragma once
+
+#include "CoreMinimal.h"
+#include "UObject/ObjectMacros.h"
+#include "ThumbnailRendering/TextureThumbnailRenderer.h"
+#include "SpecularProfileRenderer.generated.h"
+
+class FCanvas;
+class FRenderTarget;
+
+UCLASS()
+class USpecularProfileRenderer : public UTextureThumbnailRenderer
+{
+	GENERATED_UCLASS_BODY()
+
+	// Begin UThumbnailRenderer Object
+	virtual void GetThumbnailSize(UObject* Object, float Zoom, uint32& OutWidth, uint32& OutHeight) const override;
+	virtual void Draw(UObject* Object, int32 X, int32 Y, uint32 Width, uint32 Height, FRenderTarget*, FCanvas* Canvas, bool bAdditionalViewFamily) override;
+	// End UThumbnailRenderer Object
+};
+
diff --git a/Engine/Source/Editor/UnrealEd/Private/Factories/EditorFactories.cpp b/Engine/Source/Editor/UnrealEd/Private/Factories/EditorFactories.cpp
--- a/Engine/Source/Editor/UnrealEd/Private/Factories/EditorFactories.cpp
+++ b/Engine/Source/Editor/UnrealEd/Private/Factories/EditorFactories.cpp
@@ -110,6 +110,7 @@
 #include "Factories/StructureFactory.h"
 #include "Factories/StringTableFactory.h"
 #include "Factories/SubsurfaceProfileFactory.h"
+#include "Factories/SpecularProfileFactory.h"
 #include "Factories/Texture2dFactoryNew.h"
 #include "Engine/Texture.h"
 #include "Factories/TextureFactory.h"
@@ -145,6 +146,7 @@
 #include "Sound/SoundWave.h"
 #include "GameFramework/DefaultPhysicsVolume.h"
 #include "Engine/SubsurfaceProfile.h"
+#include "Engine/SpecularProfile.h"
 #include "Misc/ConfigCacheIni.h"
 #include "Misc/FeedbackContext.h"
 #include "GameFramework/WorldSettings.h"
@@ -7857,6 +7859,26 @@ UObject* USubsurfaceProfileFactory::FactoryCreateNew(UClass* InClass, UObject* I
 	return Object;
 }
 
+
+/*-----------------------------------------------------------------------------
+	USpecularProfileFactory implementation.
+	-----------------------------------------------------------------------------*/
+USpecularProfileFactory::USpecularProfileFactory(const FObjectInitializer& ObjectInitializer)
+: Super(ObjectInitializer)
+{
+
+	SupportedClass = USpecularProfile::StaticClass();
+	bCreateNew = true;
+	bEditorImport = false;
+	bEditAfterNew = true;
+}
+
+UObject* USpecularProfileFactory::FactoryCreateNew(UClass* InClass, UObject* InParent, FName InName, EObjectFlags Flags, UObject* Context, FFeedbackContext* Warn)
+{
+	USpecularProfile* Object = NewObject<USpecularProfile>(InParent, InName, Flags);	
+	return Object;
+}
+
 /*-----------------------------------------------------------------------------
 	UTouchInterfaceFactory implementation.
 -----------------------------------------------------------------------------*/
diff --git a/Engine/Source/Editor/UnrealEd/Private/ThumbnailRendering/SpecularProfileRenderer.cpp b/Engine/Source/Editor/UnrealEd/Private/ThumbnailRendering/SpecularProfileRenderer.cpp
new file mode 100644
--- /dev/null
+++ b/Engine/Source/Editor/UnrealEd/Private/ThumbnailRendering/SpecularProfileRenderer.cpp
@@ -0,0 +1,38 @@
+// Copyright Epic Games, Inc. All Rights Reserved.
+
+#include "ThumbnailRendering/SpecularProfileRenderer.h"
+#include "CanvasItem.h"
+#include "Engine/Engine.h"
+#include "EngineGlobals.h"
+#include "Engine/SpecularProfile.h"
+#include "CanvasTypes.h"
+
+USpecularProfileRenderer::USpecularProfileRenderer(const FObjectInitializer& ObjectInitializer)
+	: Super(ObjectInitializer)
+{
+}
+
+void USpecularProfileRenderer::GetThumbnailSize(UObject* Object, float Zoom, uint32& OutWidth, uint32& OutHeight) const
+{
+	OutWidth = 128;
+	OutHeight = 128;
+}
+
+void USpecularProfileRenderer::Draw(UObject* Object, int32 X, int32 Y, uint32 Width, uint32 Height, FRenderTarget*, FCanvas* Canvas, bool bAdditionalViewFamily)
+{
+	USpecularProfile* LocalSpecularProfile = Cast<USpecularProfile>(Object);
+	if (LocalSpecularProfile)
+	{
+		FLinearColor Col;
+
+		Col = LocalSpecularProfile->Settings.ViewColor; Col.A = 1;
+		Canvas->DrawTile(0,          0, Width, Height / 2, 0, 0, 1, 1, Col);
+		Col = LocalSpecularProfile->Settings.LightColor; Col.A = 1;
+		Canvas->DrawTile(0, Height / 2, Width, Height / 2, 0, 0, 1, 1, Col);
+
+		FCanvasTextItem TextItem(FVector2D(5.0f, 5.0f), FText::FromString(LocalSpecularProfile->Settings.IsProcedural() ? TEXT("Procedural") : TEXT("Texture")), GEngine->GetLargeFont(), FLinearColor::White);
+		TextItem.EnableShadow(FLinearColor::Black);
+		TextItem.Scale = FVector2D(Width / 128.0f, Height / 128.0f);
+		TextItem.Draw(Canvas);
+	}
+}
diff --git a/Engine/Source/Editor/UnrealEd/Public/UnrealEd.h b/Engine/Source/Editor/UnrealEd/Public/UnrealEd.h
--- a/Engine/Source/Editor/UnrealEd/Public/UnrealEd.h
+++ b/Engine/Source/Editor/UnrealEd/Public/UnrealEd.h
@@ -155,6 +155,7 @@ MONOLITHIC_HEADER_BOILERPLATE()
 #include "Factories/SkeletonFactory.h"
 #include "Factories/SlateBrushAssetFactory.h"
 #include "Factories/SlateWidgetStyleAssetFactory.h"
+#include "Factories/SpecularProfileFactory.h"
 #include "Factories/StructureFactory.h"
 #include "Factories/SubsurfaceProfileFactory.h"
 #include "Factories/Texture2dFactoryNew.h"
@@ -267,6 +268,7 @@ MONOLITHIC_HEADER_BOILERPLATE()
 #include "ThumbnailRendering/FontThumbnailRenderer.h"
 #include "ThumbnailRendering/ParticleSystemThumbnailRenderer.h"
 #include "ThumbnailRendering/SubsurfaceProfileRenderer.h"
+#include "ThumbnailRendering/SpecularProfileRenderer.h"
 #include "Factories/TextureCubeThumbnailRenderer.h"
 #include "Factories/Texture2dArrayThumbnailRenderer.h"
 #include "Editor/TransBuffer.h"
diff --git a/Engine/Source/Runtime/Engine/Classes/Engine/EngineTypes.h b/Engine/Source/Runtime/Engine/Classes/Engine/EngineTypes.h
--- a/Engine/Source/Runtime/Engine/Classes/Engine/EngineTypes.h
+++ b/Engine/Source/Runtime/Engine/Classes/Engine/EngineTypes.h
@@ -29,6 +29,7 @@ class UPhysicalMaterialMask;
 class UPrimitiveComponent;
 class USceneComponent;
 class USubsurfaceProfile;
+class USpecularProfile;
 
 /**
  * Default number of components to expect in TInlineAllocators used with AActor component arrays.
@@ -739,6 +740,10 @@ public:
 	int32 CountSubsurfaceProfiles() const { return SubsurfaceProfiles.Num(); }
 	USubsurfaceProfile* GetSubsurfaceProfile() const { return SubsurfaceProfiles.Num() > 0 ? SubsurfaceProfiles[0] : nullptr; }
 
+	void AddSpecularProfile(USpecularProfile* InProfile) { if (InProfile) SpecularProfiles.Add(InProfile); }
+	int32 CountSpecularProfiles() const { return SpecularProfiles.Num(); }
+	USpecularProfile* GetSpecularProfile() const { return SpecularProfiles.Num() > 0 ? SpecularProfiles[0] : nullptr; }
+
 	// Shading model from expression
 	void SetShadingModelFromExpression(bool bIn) { bHasShadingModelFromExpression = bIn ? 1u : 0u; }
 	bool HasShadingModelFromExpression() const { return bHasShadingModelFromExpression > 0u; }
@@ -789,6 +794,9 @@ private:
 	UPROPERTY()
 	TArray<TObjectPtr<USubsurfaceProfile>> SubsurfaceProfiles;
 
+	UPROPERTY()
+	TArray<TObjectPtr<USpecularProfile>> SpecularProfiles;
+
 #if WITH_EDITOR
 	// A simple way to detect and prevent node re-entry due to cycling graph; stop the compilation and avoid crashing.
 	bool bOutOfStackDepthWhenParsing = false;
diff --git a/Engine/Source/Runtime/Engine/Classes/Engine/SpecularProfile.h b/Engine/Source/Runtime/Engine/Classes/Engine/SpecularProfile.h
new file mode 100644
--- /dev/null
+++ b/Engine/Source/Runtime/Engine/Classes/Engine/SpecularProfile.h
@@ -0,0 +1,93 @@
+// Copyright Epic Games, Inc. All Rights Reserved.
+#pragma once
+
+#include "CoreMinimal.h"
+#include "UObject/ObjectMacros.h"
+#include "UObject/Object.h"
+#include "RenderResource.h"
+#include "SpecularProfile.generated.h"
+
+#define MAX_SPECULAR_PROFILE_COUNT 256
+
+class UTexture2D;
+class FRDGBuilder;
+class FTextureReference;
+
+// struct with all the settings we want in USpecularProfile, separate to make it easer to pass this data around in the engine.
+USTRUCT(BlueprintType)
+struct FSpecularProfileStruct
+{
+	GENERATED_USTRUCT_BODY()
+	
+	/**
+	* Define the view facing color
+	*/
+	UPROPERTY(Category = "Procedural", EditAnywhere, BlueprintReadOnly, meta = (ClampMin = "0.01", UIMax = "1.0", ClampMax = "1.0", HideAlphaChannel))
+	FLinearColor ViewColor;
+
+	/**
+	* Define the light facing color
+	*/
+	UPROPERTY(Category = "Procedural", EditAnywhere, BlueprintReadOnly, meta = (ClampMin = "0.001", UIMax = "1.0", ClampMax = "1.0", HideAlphaChannel))
+	FLinearColor LightColor;
+
+	/**
+	 * Define the texture used as a specular profile
+	 */
+	UPROPERTY(Category  ="Texture", EditAnywhere, BlueprintReadOnly, meta=(DisplayName="Texture"))
+	TObjectPtr<UTexture2D> Texture;
+
+	// constructor
+	FSpecularProfileStruct()
+	{
+		ViewColor = FLinearColor(1.0f, 1.0f, 1.0f);
+		LightColor = FLinearColor(1.0f, 1.0f, 1.0f);
+		Texture = nullptr;
+	}
+
+	bool IsProcedural() const { return Texture == nullptr; }
+
+	void Invalidate()
+	{
+		*this = FSpecularProfileStruct();
+	}
+};
+
+/**
+ * Specular profile asset, can be specified at a material. 
+ * Don't change at runtime. All properties in here are per material.
+ */
+UCLASS(autoexpandcategories = SpecularProfile, MinimalAPI)
+class USpecularProfile : public UObject
+{
+	GENERATED_UCLASS_BODY()
+
+	UPROPERTY(Category = USpecularProfile, EditAnywhere, meta = (ShowOnlyInnerProperties))
+	struct FSpecularProfileStruct Settings;
+
+	//~ Begin UObject Interface
+	virtual void BeginDestroy();
+	virtual void PostEditChangeProperty(struct FPropertyChangedEvent& PropertyChangedEvent);
+	//~ End UObject Interface
+};
+
+namespace SpecularProfileAtlas
+{
+// Atlas - Initializes or updates the contents of the specular profile texture.
+ENGINE_API void UpdateSpecularProfileTextureAtlas(FRDGBuilder& GraphBuilder, EShaderPlatform ShaderPlatform);
+
+// Atlas - Returns the specular profile texture if it exists, or null.
+ENGINE_API FRHITexture* GetSpecularProfileTextureAtlas();
+
+// Atlas - Returns the specular profile texture if it exists, or black.
+ENGINE_API FRHITexture* GetSpecularProfileTextureAtlasWithFallback();
+
+// Profile - Initializes or updates the contents of the specular profile texture.
+ENGINE_API int32 AddOrUpdateProfile(const USpecularProfile* InProfile, const FSpecularProfileStruct InSettings, const FTextureReference* InTexture);
+
+// Profile - Returns the specular profile ID shader parameter name
+ENGINE_API FName GetSpecularProfileParameterName();
+
+// Profile - Returns the specular profile ID for a given Specular Profile object
+ENGINE_API float GetSpecularProfileId(const USpecularProfile* In);
+}
\ No newline at end of file
diff --git a/Engine/Source/Runtime/Engine/Classes/Engine/SubsurfaceProfile.h b/Engine/Source/Runtime/Engine/Classes/Engine/SubsurfaceProfile.h
--- a/Engine/Source/Runtime/Engine/Classes/Engine/SubsurfaceProfile.h
+++ b/Engine/Source/Runtime/Engine/Classes/Engine/SubsurfaceProfile.h
@@ -7,6 +7,8 @@
 #include "RenderResource.h"
 #include "SubsurfaceProfile.generated.h"
 
+#define MAX_SUBSURFACE_PROFILE_COUNT 256
+
 // struct with all the settings we want in USubsurfaceProfile, separate to make it easer to pass this data around in the engine.
 USTRUCT(BlueprintType)
 struct FSubsurfaceProfileStruct
diff --git a/Engine/Source/Runtime/Engine/Classes/Materials/Material.h b/Engine/Source/Runtime/Engine/Classes/Materials/Material.h
--- a/Engine/Source/Runtime/Engine/Classes/Materials/Material.h
+++ b/Engine/Source/Runtime/Engine/Classes/Materials/Material.h
@@ -38,6 +38,7 @@ class UMaterialExpressionComment;
 class UPhysicalMaterial;
 class UPhysicalMaterialMask;
 class USubsurfaceProfile;
+class USpecularProfile;
 class UTexture;
 struct FODSCRequestPayload;
 
@@ -1125,6 +1126,7 @@ public:
 	ENGINE_API virtual bool IsUIMaterial() const;
 	ENGINE_API virtual bool IsPostProcessMaterial() const;
 	ENGINE_API virtual USubsurfaceProfile* GetSubsurfaceProfile_Internal() const override;
+	ENGINE_API virtual USpecularProfile* GetSpecularProfile_Internal() const override;
 	ENGINE_API virtual bool CastsRayTracedShadows() const override;
 	ENGINE_API virtual FDisplacementScaling GetDisplacementScaling() const override;
 	ENGINE_API virtual float GetMaxWorldPositionOffsetDisplacement() const override;
diff --git a/Engine/Source/Runtime/Engine/Classes/Materials/MaterialExpressionStrata.h b/Engine/Source/Runtime/Engine/Classes/Materials/MaterialExpressionStrata.h
--- a/Engine/Source/Runtime/Engine/Classes/Materials/MaterialExpressionStrata.h
+++ b/Engine/Source/Runtime/Engine/Classes/Materials/MaterialExpressionStrata.h
@@ -312,6 +312,10 @@ class UMaterialExpressionStrataSlabBSDF : public UMaterialExpressionStrataBSDF
 	UPROPERTY(EditAnywhere, BlueprintReadOnly, Category = Material, meta = (DisplayName = "Subsurface Profile"))
 	TObjectPtr<class USubsurfaceProfile> SubsurfaceProfile;
 
+	/** SpecularProfile, for modulating specular appearance and simulating more complex visuals such as iridescence. */
+	UPROPERTY(EditAnywhere, BlueprintReadOnly, Category = Material, meta = (DisplayName = "Specular Profile"))
+	TObjectPtr<class USpecularProfile> SpecularProfile;
+
 	/** Whether to use light diffusion (i.e., SSS diffusion) or wrap-approximation for material with scattering behavior. This option trades quality over performance and will result into visual differences. */
 	UPROPERTY(EditAnywhere, Category = Mode, meta = (DisplayName = "Use Subsurface Diffusion"))
 	uint32 bUseSSSDiffusion : 1;
@@ -338,6 +342,7 @@ class UMaterialExpressionStrataSlabBSDF : public UMaterialExpressionStrataBSDF
 	bool HasMFPPluggedIn() const;
 	bool HasAnisotropy() const;
 	bool HasGlint() const;
+	bool HasSpecularProfile() const;
 #endif
 	//~ End UMaterialExpression Interface
 };
diff --git a/Engine/Source/Runtime/Engine/Classes/Materials/MaterialInterface.h b/Engine/Source/Runtime/Engine/Classes/Materials/MaterialInterface.h
--- a/Engine/Source/Runtime/Engine/Classes/Materials/MaterialInterface.h
+++ b/Engine/Source/Runtime/Engine/Classes/Materials/MaterialInterface.h
@@ -49,6 +49,7 @@ class UMaterial;
 class UPhysicalMaterial;
 class UPhysicalMaterialMask;
 class USubsurfaceProfile;
+class USpecularProfile;
 class UTexture;
 class UMaterialInstance;
 struct FDebugShaderTypeInfo;
@@ -271,10 +272,14 @@ public:
 #endif // WITH_EDITORONLY_DATA
 
 public:
-	/** SubsurfaceProfile, for Screen Space Subsurface Scattering. Disabled when root node IsThin is true. */
+	/** SubsurfaceProfile, for Screen Space Subsurface Scattering.. */
 	UPROPERTY(EditAnywhere, BlueprintReadOnly, Category = Material, meta = (DisplayName = "Subsurface Profile"))
 	TObjectPtr<class USubsurfaceProfile> SubsurfaceProfile;
 
+	/** Specular Profile. */
+	UPROPERTY(EditAnywhere, BlueprintReadOnly, Category = Material)
+	TObjectPtr<class USpecularProfile> SpecularProfile;	
+
 	/** Whether this material interface is included in the base game (and not in a DLC) */
 	UPROPERTY()
 	uint8 bIncludedInBaseGame : 1;
@@ -910,6 +915,7 @@ public:
 	ENGINE_API virtual FDisplacementScaling GetDisplacementScaling() const;
 	ENGINE_API virtual float GetMaxWorldPositionOffsetDisplacement() const;
 	ENGINE_API virtual USubsurfaceProfile* GetSubsurfaceProfile_Internal() const;
+	ENGINE_API virtual USpecularProfile* GetSpecularProfile_Internal() const;
 	ENGINE_API virtual bool CastsRayTracedShadows() const;
 
 	/**
diff --git a/Engine/Source/Runtime/Engine/Private/Materials/HLSLMaterialTranslator.cpp b/Engine/Source/Runtime/Engine/Private/Materials/HLSLMaterialTranslator.cpp
--- a/Engine/Source/Runtime/Engine/Private/Materials/HLSLMaterialTranslator.cpp
+++ b/Engine/Source/Runtime/Engine/Private/Materials/HLSLMaterialTranslator.cpp
@@ -2381,9 +2381,9 @@ void FHLSLMaterialTranslator::GetMaterialEnvironment(EShaderPlatform InPlatform,
 						{
 							if (!It.IsDiscarded() && It.MaxDistanceFromLeaves == DistanceToLeaves)
 							{
-								StrataMaterialContextDescription += FString::Printf(TEXT("\tIdx=%d Op=%s ParentIdx=%d LeftIndex=%d RightIndex=%d BSDFIdx=%d LayerDepth=%d IsTop=%d IsBot=%d BSDFType=%s SSS=%d MFP=%d F90=%d Rough2=%d Fuzz=%d Aniso=%d Glint=%d \r\n"),
+								StrataMaterialContextDescription += FString::Printf(TEXT("\tIdx=%d Op=%s ParentIdx=%d LeftIndex=%d RightIndex=%d BSDFIdx=%d LayerDepth=%d IsTop=%d IsBot=%d BSDFType=%s SSS=%d MFP=%d F90=%d Rough2=%d Fuzz=%d Aniso=%d Glint=%d SpecularProfile=%d\r\n"),
 									It.Index, GetStrataOperatorStr(It.OperatorType), It.ParentIndex, It.LeftIndex, It.RightIndex, It.BSDFIndex, It.LayerDepth, It.bIsTop, It.bIsBottom,
-									*GetStrataBSDFName(It.BSDFType), It.bBSDFHasSSS, It.bBSDFHasMFPPluggedIn, It.bBSDFHasEdgeColor, It.bBSDFHasSecondRoughnessOrSimpleClearCoat, It.bBSDFHasFuzz, It.bBSDFHasAnisotropy, It.bBSDFHasGlint);
+									*GetStrataBSDFName(It.BSDFType), It.bBSDFHasSSS, It.bBSDFHasMFPPluggedIn, It.bBSDFHasEdgeColor, It.bBSDFHasSecondRoughnessOrSimpleClearCoat, It.bBSDFHasFuzz, It.bBSDFHasAnisotropy, It.bBSDFHasGlint, It.bBSDFHasSpecularProfile);
 							}
 						}
 					}
@@ -11577,8 +11577,8 @@ bool FHLSLMaterialTranslator::FStrataCompilationContext::StrataGenerateDerivedMa
 					{
 					case STRATA_BSDF_TYPE_SLAB:
 					{
-						bStrataMaterialIsSimple = bStrataMaterialIsSimple && !bMayHaveColoredWeight && !It.bBSDFHasAnisotropy && !It.bBSDFHasEdgeColor && !It.bBSDFHasFuzz && !It.bBSDFHasSecondRoughnessOrSimpleClearCoat && !It.bBSDFHasMFPPluggedIn && !It.bBSDFHasSSS && !It.bBSDFHasGlint;
-						bStrataMaterialIsSingle = bStrataMaterialIsSingle && !bMayHaveColoredWeight && !It.bBSDFHasAnisotropy && !It.bBSDFHasGlint;
+						bStrataMaterialIsSimple = bStrataMaterialIsSimple && !bMayHaveColoredWeight && !It.bBSDFHasAnisotropy && !It.bBSDFHasEdgeColor && !It.bBSDFHasFuzz && !It.bBSDFHasSecondRoughnessOrSimpleClearCoat && !It.bBSDFHasMFPPluggedIn && !It.bBSDFHasSSS && !It.bBSDFHasGlint && !It.bBSDFHasSpecularProfile;
+						bStrataMaterialIsSingle = bStrataMaterialIsSingle && !bMayHaveColoredWeight && !It.bBSDFHasAnisotropy && !It.bBSDFHasGlint && !It.bBSDFHasSpecularProfile;
 						break;
 					}
 					case STRATA_BSDF_TYPE_HAIR:
@@ -11722,6 +11722,10 @@ bool FHLSLMaterialTranslator::FStrataCompilationContext::StrataGenerateDerivedMa
 							StrataMaterialRequestedSizeByte += UintByteSize;
 							StrataMaterialRequestedSizeByte += UintByteSize;
 						}
+						if (It.bBSDFHasSpecularProfile)
+						{
+							StrataMaterialRequestedSizeByte += UintByteSize;
+						}
 						break;
 					}
 					case STRATA_BSDF_TYPE_HAIR:
@@ -12007,6 +12011,7 @@ int32 FHLSLMaterialTranslator::StrataSlabBSDF(
 	int32 FuzzAmount, int32 FuzzColor, int32 FuzzRoughness,
 	int32 Thickness,
 	int32 GlintValue, int32 GlintUV,
+	int32 SpecularProfileId,
 	int32 Normal, int32 Tangent, const FString& SharedLocalBasisIndexMacro,
 	FStrataOperator* PromoteToOperator)
 {
@@ -12024,7 +12029,7 @@ int32 FHLSLMaterialTranslator::StrataSlabBSDF(
 		}
 
 		return AddCodeChunk(
-			MCT_Strata, TEXT("Parameters.%s.PromoteParameterBlendedBSDFToOperator(GetStrataSlabBSDF(Parameters.StrataPixelFootprint, %s, %s, %s, %s, %s, %s, %s, %s, %s, %s, %s, %s, %s, %s, %s, %s, %s, %s, %s, %s, %s, %s, Parameters.%s.Types) /* Normal = %s ; Tangent = %s ; Thickness = %s */, %u, %u, %u, %u)"),
+			MCT_Strata, TEXT("Parameters.%s.PromoteParameterBlendedBSDFToOperator(GetStrataSlabBSDF(Parameters.StrataPixelFootprint, %s, %s, %s, %s, %s, %s, %s, %s, %s, %s, %s, %s, %s, %s, %s, %s, %s, %s, %s, %s, %s, %s, %s, Parameters.%s.Types) /* Normal = %s ; Tangent = %s ; Thickness = %s */, %u, %u, %u, %u)"),
 			*GetParametersStrataTreeName(CurrentStrataCompilationContext),
 			*StrataGetCastParameterCode(DiffuseAlbedo,			MCT_Float3),
 			*StrataGetCastParameterCode(F0,						MCT_Float3),
@@ -12045,6 +12050,7 @@ int32 FHLSLMaterialTranslator::StrataSlabBSDF(
 			*StrataGetCastParameterCode(FuzzRoughness,			MCT_Float1),
 			*StrataGetCastParameterCode(GlintValue,				MCT_Float),
 			*StrataGetCastParameterCode(GlintUV,				MCT_Float2),
+			*StrataGetCastParameterCode(SpecularProfileId,		MCT_Float),
 			*StrataGetCastParameterCode(Thickness,				MCT_Float),
 			bIsThinSurface ? TEXT("true") : TEXT("false"),
 			*SharedLocalBasisIndexMacro,
@@ -12060,7 +12066,7 @@ int32 FHLSLMaterialTranslator::StrataSlabBSDF(
 	}
 	
 	return AddCodeChunk(
-		MCT_Strata, TEXT("GetStrataSlabBSDF(Parameters.StrataPixelFootprint, %s, %s, %s, %s, %s, %s, %s, %s, %s, %s, %s, %s, %s, %s, %s, %s, %s, %s, %s, %s, %s, %s, Parameters.%s.Types) /* Normal = %s ; Tangent = %s ; Thickness = %s */"),
+		MCT_Strata, TEXT("GetStrataSlabBSDF(Parameters.StrataPixelFootprint, %s, %s, %s, %s, %s, %s, %s, %s, %s, %s, %s, %s, %s, %s, %s, %s, %s, %s, %s, %s, %s, %s, %s, Parameters.%s.Types) /* Normal = %s ; Tangent = %s ; Thickness = %s */"),
 		*StrataGetCastParameterCode(DiffuseAlbedo,			MCT_Float3),
 		*StrataGetCastParameterCode(F0,						MCT_Float3),
 		*StrataGetCastParameterCode(F90,					MCT_Float3),
@@ -12079,7 +12085,8 @@ int32 FHLSLMaterialTranslator::StrataSlabBSDF(
 		*StrataGetCastParameterCode(FuzzColor,				MCT_Float3),
 		*StrataGetCastParameterCode(FuzzRoughness,			MCT_Float1),
 		*StrataGetCastParameterCode(GlintValue,				MCT_Float),
-		*StrataGetCastParameterCode(GlintUV, MCT_Float2),
+		*StrataGetCastParameterCode(GlintUV, 				MCT_Float2),
+		*StrataGetCastParameterCode(SpecularProfileId,		MCT_Float),
 		*StrataGetCastParameterCode(Thickness, MCT_Float),
 		bIsThinSurface ? TEXT("true") : TEXT("false"),
 		*SharedLocalBasisIndexMacro,
diff --git a/Engine/Source/Runtime/Engine/Private/Materials/HLSLMaterialTranslator.h b/Engine/Source/Runtime/Engine/Private/Materials/HLSLMaterialTranslator.h
--- a/Engine/Source/Runtime/Engine/Private/Materials/HLSLMaterialTranslator.h
+++ b/Engine/Source/Runtime/Engine/Private/Materials/HLSLMaterialTranslator.h
@@ -1169,6 +1169,7 @@ protected:
 		int32 FuzzAmount, int32 FuzzColor, int32 FuzzRoughness,
 		int32 Thickness,
 		int32 GlintValue, int32 GlintUV,
+		int32 SpecularProfileId,
 		int32 Normal, int32 Tangent, const FString& SharedLocalBasisIndexMacro,
 		FStrataOperator* PromoteToOperator) override;
 	virtual int32 StrataConversionFromLegacy(
diff --git a/Engine/Source/Runtime/Engine/Private/Materials/Material.cpp b/Engine/Source/Runtime/Engine/Private/Materials/Material.cpp
--- a/Engine/Source/Runtime/Engine/Private/Materials/Material.cpp
+++ b/Engine/Source/Runtime/Engine/Private/Materials/Material.cpp
@@ -50,6 +50,7 @@
 #include "SceneView.h"
 #include "Materials/MaterialUniformExpressions.h"
 #include "Engine/SubsurfaceProfile.h"
+#include "Engine/SpecularProfile.h"
 #include "EditorSupportDelegates.h"
 #include "ComponentRecreateRenderStateContext.h"
 #include "ShaderCompiler.h"
@@ -448,6 +449,11 @@ public:
 				OutValue = GetSubsurfaceProfileId(GetSubsurfaceProfileRT());
 				return true;
 			}
+			else if (Type == EMaterialParameterType::Scalar && ParameterInfo.Name == SpecularProfileAtlas::GetSpecularProfileParameterName())
+			{
+				OutValue = SpecularProfileAtlas::GetSpecularProfileId(GetSpecularProfileRT());
+				return true;
+			}
 
 			return false;
 		}
@@ -5064,6 +5070,9 @@ void UMaterial::RebuildShadingModelField()
 			}
 			SubsurfaceProfile = StrataMaterialInfo.GetSubsurfaceProfile();
 		}
+
+		// Set specular profile if any
+		SpecularProfile = StrataMaterialInfo.GetSpecularProfile();
 	}
 	// If using shading model from material expression, go through the expressions and look for the ShadingModel expression to figure out what shading models need to be supported in this material.
 	// This might not be the same as what is actually compiled in to the shader, since there might be feature switches, static switches etc. that skip certain shading models.
@@ -6612,6 +6621,12 @@ USubsurfaceProfile* UMaterial::GetSubsurfaceProfile_Internal() const
 	return SubsurfaceProfile; 
 }
 
+USpecularProfile* UMaterial::GetSpecularProfile_Internal() const
+{
+	checkSlow(IsInGameThread());
+	return SpecularProfile; 
+}
+
 bool UMaterial::CastsRayTracedShadows() const
 {
 	return bCastRayTracedShadows;
diff --git a/Engine/Source/Runtime/Engine/Private/Materials/MaterialExpressions.cpp b/Engine/Source/Runtime/Engine/Private/Materials/MaterialExpressions.cpp
--- a/Engine/Source/Runtime/Engine/Private/Materials/MaterialExpressions.cpp
+++ b/Engine/Source/Runtime/Engine/Private/Materials/MaterialExpressions.cpp
@@ -35,6 +35,7 @@
 #include "Engine/TextureCubeArray.h"
 #include "Engine/VolumeTexture.h"
 #include "Engine/SubsurfaceProfile.h"
+#include "Engine/SpecularProfile.h"
 #include "Serialization/ObjectWriter.h"
 #include "Serialization/ObjectReader.h"
 #include "VT/RuntimeVirtualTexture.h"
@@ -23342,12 +23343,20 @@ int32 UMaterialExpressionStrataSlabBSDF::Compile(class FMaterialCompiler* Compil
 	const bool bHasSecondRoughness = HasSecondRoughness();
 	const bool bHasMFPPluggedIn = HasMFPPluggedIn();
 	const bool bHasSSS = HasSSS();
+	const bool bHasSpecularProfile = HasSpecularProfile();
+
 	int32 SSSProfileCodeChunk = INDEX_NONE;
 	if (bHasSSS)
 	{
 		SSSProfileCodeChunk = Compiler->ForceCast(Compiler->ScalarParameter(GetSubsurfaceProfileParameterName(), 1.0f), MCT_Float1);
 	}
 
+	int32 SpecularProfileCodeChunk = INDEX_NONE;
+	if (bHasSpecularProfile)
+	{
+		SpecularProfileCodeChunk = Compiler->ForceCast(Compiler->ScalarParameter(SpecularProfileAtlas::GetSpecularProfileParameterName(), 1.0f), MCT_Float1);
+	}
+
 	const float DefaultSpecular = 0.5f;
 	const float DefaultF0 = DielectricSpecularToF0(DefaultSpecular);
 
@@ -23391,6 +23400,7 @@ int32 UMaterialExpressionStrataSlabBSDF::Compile(class FMaterialCompiler* Compil
 		ThicknesCodeChunk,
 		CompileWithDefaultFloat1(Compiler, GlintValue, 0.0f),
 		CompileWithDefaultFloat2(Compiler, GlintUV, 0.0f, 0.0f),
+		SpecularProfileCodeChunk != INDEX_NONE ? SpecularProfileCodeChunk : Compiler->Constant(0.0f),
 		NormalCodeChunk,
 		TangentCodeChunk,
 		Compiler->GetStrataSharedLocalBasisIndexMacro(NewRegisteredSharedLocalBasis),
@@ -23632,6 +23642,11 @@ void UMaterialExpressionStrataSlabBSDF::GatherStrataMaterialInfo(FStrataMaterial
 	{
 		StrataMaterialInfo.AddShadingModel(SSM_DefaultLit);
 	}
+
+	if (HasSpecularProfile())
+	{
+		StrataMaterialInfo.AddSpecularProfile(SpecularProfile);
+	}
 }
 
 FStrataOperator* UMaterialExpressionStrataSlabBSDF::StrataGenerateMaterialTopologyTree(class FMaterialCompiler* Compiler, class UMaterialExpression* Parent, int32 OutputIndex)
@@ -23645,6 +23660,7 @@ FStrataOperator* UMaterialExpressionStrataSlabBSDF::StrataGenerateMaterialTopolo
 	StrataOperator.bBSDFHasMFPPluggedIn = HasMFPPluggedIn();
 	StrataOperator.bBSDFHasAnisotropy = HasAnisotropy();
 	StrataOperator.bBSDFHasGlint = HasGlint();
+	StrataOperator.bBSDFHasSpecularProfile = HasSpecularProfile();
 	StrataOperator.ThicknessIndex = Compiler->StrataThicknessStackGetThicknessIndex();
 	return &StrataOperator;
 }
@@ -23693,6 +23709,12 @@ bool UMaterialExpressionStrataSlabBSDF::HasGlint() const
 {
 	return GlintValue.IsConnected();
 }
+
+bool UMaterialExpressionStrataSlabBSDF::HasSpecularProfile() const
+{
+	return SpecularProfile != nullptr;
+}
+
 #endif // WITH_EDITOR
 
 
@@ -23750,6 +23772,7 @@ int32 UMaterialExpressionStrataSimpleClearCoatBSDF::Compile(class FMaterialCompi
 		ThicknessCodeChunk,										// Thickness
 		Compiler->Constant(0.0f),								// GlintValue
 		Compiler->Constant2(0.0f, 0.0f),						// GlintUV
+		Compiler->Constant(0.0f),								// SpecularProfile
 		NormalCodeChunk,
 		NullTangentCodeChunk,
 		Compiler->GetStrataSharedLocalBasisIndexMacro(NewRegisteredSharedLocalBasis),
@@ -23880,6 +23903,7 @@ FStrataOperator* UMaterialExpressionStrataSimpleClearCoatBSDF::StrataGenerateMat
 	StrataOperator.bBSDFHasMFPPluggedIn = false;
 	StrataOperator.bBSDFHasAnisotropy = false;
 	StrataOperator.bBSDFHasGlint = false;
+	StrataOperator.bBSDFHasSpecularProfile = false;
 	StrataOperator.ThicknessIndex = Compiler->StrataThicknessStackGetThicknessIndex();
 	return &StrataOperator;
 }
diff --git a/Engine/Source/Runtime/Engine/Private/Materials/MaterialInstance.cpp b/Engine/Source/Runtime/Engine/Private/Materials/MaterialInstance.cpp
--- a/Engine/Source/Runtime/Engine/Private/Materials/MaterialInstance.cpp
+++ b/Engine/Source/Runtime/Engine/Private/Materials/MaterialInstance.cpp
@@ -32,6 +32,7 @@
 #include "Materials/MaterialInstanceUpdateParameterSet.h"
 #include "Materials/MaterialInstanceSupport.h"
 #include "Engine/SubsurfaceProfile.h"
+#include "Engine/SpecularProfile.h"
 #include "ProfilingDebugging/LoadTimeTracker.h"
 #include "ObjectCacheEventSink.h"
 #include "Interfaces/ITargetPlatform.h"
@@ -280,7 +281,13 @@ bool FMaterialInstanceResource::GetParameterValue(EMaterialParameterType Type, c
 		OutValue = GetSubsurfaceProfileId(MySubsurfaceProfileRT);
 		bResult = true;
 	}
-
+	else if (Type == EMaterialParameterType::Scalar && ParameterInfo.Name == SpecularProfileAtlas::GetSpecularProfileParameterName())
+	{
+		check(ParameterInfo.Association == EMaterialParameterAssociation::GlobalParameter);
+		const USpecularProfile* MySpecularProfileRT = GetSpecularProfileRT();
+		OutValue = SpecularProfileAtlas::GetSpecularProfileId(MySpecularProfileRT);
+		bResult = true;
+	}
 	if (!bResult)
 	{
 		// Check for instances overrides
@@ -3110,7 +3117,7 @@ void UMaterialInstance::PostLoad()
 	UpdateCachedData();
 #endif
 
-	// called before we cache the uniform expression as a call to SubsurfaceProfileRT affects the dta in there
+	// called before we cache the uniform expression as a call to SubsurfaceProfileRT/SpecularProfileRT affects the data in there
 	PropagateDataToMaterialProxy();
 
 	STAT(double MaterialLoadTime = 0);
diff --git a/Engine/Source/Runtime/Engine/Private/Materials/MaterialInterface.cpp b/Engine/Source/Runtime/Engine/Private/Materials/MaterialInterface.cpp
--- a/Engine/Source/Runtime/Engine/Private/Materials/MaterialInterface.cpp
+++ b/Engine/Source/Runtime/Engine/Private/Materials/MaterialInterface.cpp
@@ -22,6 +22,7 @@
 #include "Engine/World.h"
 #include "ObjectCacheEventSink.h"
 #include "Engine/SubsurfaceProfile.h"
+#include "Engine/SpecularProfile.h"
 #include "Interfaces/ITargetPlatform.h"
 #include "Components/PrimitiveComponent.h"
 #include "ContentStreaming.h"
@@ -1311,6 +1312,11 @@ USubsurfaceProfile* UMaterialInterface::GetSubsurfaceProfile_Internal() const
 	return NULL;
 }
 
+USpecularProfile* UMaterialInterface::GetSpecularProfile_Internal() const
+{
+	return NULL;
+}
+
 bool UMaterialInterface::CastsRayTracedShadows() const
 {
 	return true;
@@ -1371,15 +1377,45 @@ void UMaterialInterface::UpdateMaterialRenderProxy(FMaterialRenderProxy& Proxy)
 
 		FMaterialRenderProxy* InProxy = &Proxy;
 		ENQUEUE_RENDER_COMMAND(UpdateMaterialRenderProxySubsurface)(
-			[Settings, LocalSubsurfaceProfile, InProxy](FRHICommandListImmediate& RHICmdList)
+		[Settings, LocalSubsurfaceProfile, InProxy](FRHICommandListImmediate& RHICmdList)
+		{
+			if (LocalSubsurfaceProfile)
 			{
-				if (LocalSubsurfaceProfile)
-				{
-					const uint32 AllocationId = GSubsurfaceProfileTextureObject.AddOrUpdateProfile(Settings, LocalSubsurfaceProfile);
-					check(AllocationId >= 0 && AllocationId <= 255);
-				}
-				InProxy->SetSubsurfaceProfileRT(LocalSubsurfaceProfile);
-			});
+				const uint32 AllocationId = GSubsurfaceProfileTextureObject.AddOrUpdateProfile(Settings, LocalSubsurfaceProfile);
+				check(AllocationId >= 0 && AllocationId < MAX_SUBSURFACE_PROFILE_COUNT);
+			}
+			InProxy->SetSubsurfaceProfileRT(LocalSubsurfaceProfile);
+		});
+	}
+
+	if (Strata::IsStrataEnabled())
+	{
+		USpecularProfile* LocalSpecularProfile = GetSpecularProfile_Internal();
+
+		FSpecularProfileStruct LocalSettings;
+
+		const FTextureReference* LocalTexture = nullptr;
+		if (LocalSpecularProfile)
+		{
+			LocalSettings = LocalSpecularProfile->Settings;
+
+			if (!LocalSpecularProfile->Settings.IsProcedural())
+			{
+				LocalTexture = &LocalSpecularProfile->Settings.Texture->TextureReference;
+			}
+		}
+
+		FMaterialRenderProxy* InProxy = &Proxy;
+		ENQUEUE_RENDER_COMMAND(UpdateMaterialRenderProxySpecular)(
+		[LocalSettings, LocalSpecularProfile, InProxy, LocalTexture](FRHICommandListImmediate& RHICmdList)
+		{
+			if (LocalSpecularProfile)
+			{
+				const uint32 AllocationId = SpecularProfileAtlas::AddOrUpdateProfile(LocalSpecularProfile, LocalSettings, LocalTexture);
+				check(AllocationId >= 0 && AllocationId < MAX_SPECULAR_PROFILE_COUNT);
+			}
+			InProxy->SetSpecularProfileRT(LocalSpecularProfile);
+		});
 	}
 }
 
diff --git a/Engine/Source/Runtime/Engine/Private/Materials/MaterialRenderProxy.cpp b/Engine/Source/Runtime/Engine/Private/Materials/MaterialRenderProxy.cpp
--- a/Engine/Source/Runtime/Engine/Private/Materials/MaterialRenderProxy.cpp
+++ b/Engine/Source/Runtime/Engine/Private/Materials/MaterialRenderProxy.cpp
@@ -526,7 +526,8 @@ void FMaterialRenderProxy::UpdateUniformExpressionCacheIfNeeded(ERHIFeatureLevel
 }
 
 FMaterialRenderProxy::FMaterialRenderProxy(FString InMaterialName)
-	: SubsurfaceProfileRT(0)
+	: SubsurfaceProfileRT(nullptr)
+	, SpecularProfileRT(nullptr)
 	, MaterialName(MoveTemp(InMaterialName))
 	, MarkedForGarbageCollection(0)
 	, DeletedFlag(0)
diff --git a/Engine/Source/Runtime/Engine/Private/Rendering/SpecularProfile.cpp b/Engine/Source/Runtime/Engine/Private/Rendering/SpecularProfile.cpp
new file mode 100644
--- /dev/null
+++ b/Engine/Source/Runtime/Engine/Private/Rendering/SpecularProfile.cpp
@@ -0,0 +1,436 @@
+// Copyright Epic Games, Inc. All Rights Reserved.
+
+#include "Engine/SpecularProfile.h"
+#include "Engine/Texture2D.h"
+#include "DataDrivenShaderPlatformInfo.h"
+#include "Math/Float16.h"
+#include "Rendering/BurleyNormalizedSSS.h"
+#include "EngineModule.h"
+#include "RenderTargetPool.h"
+#include "PixelShaderUtils.h"
+#include "RenderingThread.h"
+#include "Rendering/Texture2DResource.h"
+#include "RenderGraphResources.h"
+#include "RenderGraphBuilder.h"
+
+#include UE_INLINE_GENERATED_CPP_BY_NAME(SpecularProfile)
+
+DEFINE_LOG_CATEGORY_STATIC(LogSpecularProfile, Log, All);
+
+///////////////////////////////////////////////////////////////////////////////////////////////////
+
+static TAutoConsoleVariable<int32> CVarSpecularProfileResolution(
+	TEXT("r.Substrate.SpecularProfile.Resolution"),
+	64,
+	TEXT("The resolution of the specular profile texture.\n"),
+	ECVF_RenderThreadSafe
+);
+
+#if !(UE_BUILD_SHIPPING || UE_BUILD_TEST)
+static TAutoConsoleVariable<int32> CVarSpecularProfileForceUpdate(
+	TEXT("r.Substrate.SpecularProfile.ForceUpdate"),
+	0,
+	TEXT("0: Only update the specular profile as needed.\n")
+	TEXT("1: Force to update the specular profile every frame for debugging.\n"),
+	ECVF_Cheat | ECVF_RenderThreadSafe);
+#endif
+
+static bool ForceUpdateSpecularProfile()
+{
+#if !(UE_BUILD_SHIPPING || UE_BUILD_TEST)
+	return (CVarSpecularProfileForceUpdate.GetValueOnAnyThread() == 1);
+#else
+	return false;
+#endif
+}
+
+///////////////////////////////////////////////////////////////////////////////////////////////////
+
+class FSpecularProfileCopyCS : public FGlobalShader
+{
+public:
+	DECLARE_GLOBAL_SHADER(FSpecularProfileCopyCS);
+	SHADER_USE_PARAMETER_STRUCT(FSpecularProfileCopyCS, FGlobalShader);
+
+	BEGIN_SHADER_PARAMETER_STRUCT(FParameters, )
+		SHADER_PARAMETER_TEXTURE(Texture2D, SourceTexture)
+		SHADER_PARAMETER_SAMPLER(SamplerState, SourceSampler)
+		SHADER_PARAMETER(FIntPoint, SourceResolution)
+		SHADER_PARAMETER(FIntPoint, TargetResolution)
+		SHADER_PARAMETER(uint32, SourceMipCount)
+		SHADER_PARAMETER(uint32, TargetIndex)
+		SHADER_PARAMETER_RDG_TEXTURE_UAV(RWTexture2D, TargetTexture)
+	END_SHADER_PARAMETER_STRUCT()
+
+		static bool ShouldCompilePermutation(const FGlobalShaderPermutationParameters& Parameters)
+	{
+		return true;
+	}
+};
+
+IMPLEMENT_GLOBAL_SHADER(FSpecularProfileCopyCS, "/Engine/Private/SpecularProfile.usf", "MainCS", SF_Compute);
+
+///////////////////////////////////////////////////////////////////////////////////////////////////
+// FSpecularProfileTextureManager
+
+// render thread
+class FSpecularProfileTextureManager : public FRenderResource
+{
+public:
+	// constructor
+	FSpecularProfileTextureManager();
+
+	// destructor
+	~FSpecularProfileTextureManager();
+
+	// convenience, can be optimized 
+	// @param Profile must not be 0, game thread pointer, do not dereference, only for comparison
+	int32 AddOrUpdateProfile(const USpecularProfile* InProfile, const FSpecularProfileStruct InSettings, const FTextureReference* InTexture)
+	{
+		check(InProfile);
+
+		int32 AllocationId = FindAllocationId(InProfile); 
+		
+		if (AllocationId != -1)
+		{
+			UpdateProfile(AllocationId, InSettings, InTexture);
+		}
+		else
+		{
+			AllocationId = AddProfile(InSettings, InProfile, InTexture);
+		}
+
+		return AllocationId;
+	}
+
+	// O(n) n is a small number
+	// @param InProfile must not be 0, game thread pointer, do not dereference, only for comparison
+	// @return AllocationId -1: no allocation, should be deallocated with DeallocateSpecularProfile()
+	int32 AddProfile(const FSpecularProfileStruct Settings, const USpecularProfile* InProfile, const FTextureReference* InTexture);
+
+	// O(n) to find the element, n is the SSProfile count and usually quite small
+	void RemoveProfile(const USpecularProfile* InProfile);
+
+	// @param InProfile must not be 0, game thread pointer, do not dereference, only for comparison
+	void UpdateProfile(const USpecularProfile* InProfile, const FSpecularProfileStruct InSettings, const FTextureReference* InTexture) { UpdateProfile(FindAllocationId(InProfile), InSettings, InTexture); }
+
+	// @param InProfile must not be 0, game thread pointer, do not dereference, only for comparison
+	void UpdateProfile(int32 AllocationId, const FSpecularProfileStruct Settings, const FTextureReference* InTexture);
+
+	// @return can be nullptr if there is no SpecularProfile
+	struct IPooledRenderTarget* GetTexture(FRDGBuilder& GraphBuilder, EShaderPlatform ShaderPlatform);
+	struct IPooledRenderTarget* GetTexture();
+
+	//~ Begin FRenderResource Interface.
+	/**
+		* Release textures when device is lost/destroyed.
+		*/
+	virtual void ReleaseDynamicRHI() override;
+
+	// @param InProfile must not be 0, game thread pointer, do not dereference, only for comparison
+	// @return -1 if not found
+	int32 FindAllocationId(const USpecularProfile* InProfile) const;
+
+private:
+
+	struct FProfileEntry
+	{
+		FSpecularProfileStruct Settings;
+		const USpecularProfile* Profile = nullptr; // Game thread pointer! Do not dereference, only for comparison.
+		const FTextureReference* Texture = nullptr;
+		FIntPoint CachedResolution = FIntPoint::ZeroValue;
+	};
+	TArray<FProfileEntry> SpecularProfileEntries;
+};
+
+// Global resources - lives on the render thread
+TGlobalResource<FSpecularProfileTextureManager> GSpecularProfileTextureManager;
+
+// SpecularProfile atlas storing several texture profiles or 0 if there is no user
+static TRefCountPtr<IPooledRenderTarget> GSpecularProfileTextureAtlas;
+
+FSpecularProfileTextureManager::FSpecularProfileTextureManager()
+{
+	check(IsInGameThread());
+
+	FSpecularProfileStruct DefaultProfile;
+
+	// add element 0, it is used as default profile
+	SpecularProfileEntries.Add( { DefaultProfile, nullptr, nullptr });
+}
+
+FSpecularProfileTextureManager::~FSpecularProfileTextureManager()
+{
+}
+
+int32 FSpecularProfileTextureManager::AddProfile(const FSpecularProfileStruct InSettings, const USpecularProfile* InProfile, const FTextureReference* InTexture)
+{
+	check(InProfile);
+	check(FindAllocationId(InProfile) == INDEX_NONE);
+
+	int32 RetAllocationId = INDEX_NONE;
+	{
+		for (int32 i = 1; i < SpecularProfileEntries.Num(); ++i)
+		{
+			if (SpecularProfileEntries[i].Profile == 0)
+			{
+				RetAllocationId = i;
+				SpecularProfileEntries[RetAllocationId].Profile = InProfile;
+				break;
+			}
+		}
+
+		if (RetAllocationId == INDEX_NONE)
+		{
+			RetAllocationId = SpecularProfileEntries.Num();
+			SpecularProfileEntries.Add({ InSettings, InProfile, InTexture});
+		}
+	}
+
+	UpdateProfile(RetAllocationId, InSettings, InTexture);
+
+	return RetAllocationId;
+}
+
+void FSpecularProfileTextureManager::RemoveProfile(const USpecularProfile* InProfile)
+{
+	int32 AllocationId = FindAllocationId(InProfile);
+
+	if (AllocationId != INDEX_NONE)
+	{
+		// >0 as 0 is used as default profile which should never be removed
+		check(AllocationId > 0);
+		check(SpecularProfileEntries[AllocationId].Profile == InProfile);
+
+		// make it available for reuse
+		SpecularProfileEntries[AllocationId].Profile = 0;
+		SpecularProfileEntries[AllocationId].Settings.Invalidate();
+		SpecularProfileEntries[AllocationId].Texture = nullptr;
+	}
+}
+
+void FSpecularProfileTextureManager::UpdateProfile(int32 AllocationId, const FSpecularProfileStruct Settings, const FTextureReference* InTexture)
+{
+	check(IsInRenderingThread());
+
+	if (AllocationId != INDEX_NONE)
+	{
+		check(AllocationId < SpecularProfileEntries.Num());
+		SpecularProfileEntries[AllocationId].Settings = Settings;
+		SpecularProfileEntries[AllocationId].Texture = InTexture;
+		GSpecularProfileTextureAtlas.SafeRelease();
+	}
+}
+
+IPooledRenderTarget* FSpecularProfileTextureManager::GetTexture()
+{
+	return GSpecularProfileTextureAtlas;
+}
+
+IPooledRenderTarget* FSpecularProfileTextureManager::GetTexture(FRDGBuilder& GraphBuilder, EShaderPlatform ShaderPlatform)
+{
+	if (!Strata::IsStrataEnabled())
+	{
+		return nullptr;
+	}
+
+	const uint32 LayerCount = SpecularProfileEntries.Num();
+
+	// Since reference texture can be streamed/loaded progressively, we track if textures have been updated to update the LUT
+	bool bForceUpdate = ForceUpdateSpecularProfile();
+	for (uint32 LayerIt = 0; LayerIt < LayerCount && !bForceUpdate; ++LayerIt)
+	{
+		if (SpecularProfileEntries[LayerIt].Texture)
+		{
+			const FTextureReferenceRHIRef& TextureRHI = SpecularProfileEntries[LayerIt].Texture->TextureReferenceRHI;
+			if (TextureRHI->GetDesc().Extent != SpecularProfileEntries[LayerIt].CachedResolution)
+			{
+				bForceUpdate = true;
+				break;
+			}
+		}
+	}
+
+	if (bForceUpdate)
+	{
+		GSpecularProfileTextureAtlas.SafeRelease();
+		GSpecularProfileTextureAtlas = nullptr;
+	}
+
+	if (GSpecularProfileTextureAtlas == nullptr)
+	{
+		FRHICommandListImmediate& RHICmdList = GraphBuilder.RHICmdList;
+
+		// Each row of the texture contains SSS parameters, followed by 3 precomputed kernels. Texture must be wide enough to fit all data.
+		const uint32 Resolution = CVarSpecularProfileResolution.GetValueOnRenderThread();
+		check(LayerCount);
+
+		FPooledRenderTargetDesc Desc(FPooledRenderTargetDesc::Create2DArrayDesc(FIntPoint(Resolution, Resolution), PF_B8G8R8A8, FClearValueBinding::None, TexCreate_None, TexCreate_UAV | TexCreate_ShaderResource, false, LayerCount));
+		GRenderTargetPool.FindFreeElement(RHICmdList, Desc, GSpecularProfileTextureAtlas, TEXT("SpecularProfileTexture"));
+
+		// 1. Fill in all procedural profile
+		FLinearColor DefaultColor(1.f, 1.f, 1.f);
+		for (uint32 LayerIt = 0; LayerIt < LayerCount; ++LayerIt)
+		{
+			if (SpecularProfileEntries[LayerIt].Texture == nullptr)
+			{
+				const FSpecularProfileStruct Data = SpecularProfileEntries[LayerIt].Settings;
+
+				uint32 DestStride;
+				uint8* DestBuffer = (uint8*)RHICmdList.LockTexture2DArray(GSpecularProfileTextureAtlas->GetRHI(), LayerIt, 0, RLM_WriteOnly, DestStride, false);
+				for (uint32 y = 0; y < Resolution; ++y)
+				{
+					FColor* Dest = (FColor*)(DestBuffer + DestStride * y);
+					for (uint32 x = 0; x < Resolution; ++x)
+					{
+						const float NoL = float(y) / float(Resolution-1);
+						const float NoV = float(x) / float(Resolution-1);
+						const FLinearColor ViewColor = FMath::Lerp(DefaultColor, Data.ViewColor, NoV);
+						const FLinearColor LightColor = FMath::Lerp(DefaultColor, Data.LightColor, NoL);
+						const FLinearColor C = ViewColor * LightColor;
+	
+						Dest[x] = FColor(FMath::Quantize8UnsignedByte(C.R), FMath::Quantize8UnsignedByte(C.G), FMath::Quantize8UnsignedByte(C.B));
+					}
+				}
+				RHICmdList.UnlockTexture2DArray(GSpecularProfileTextureAtlas->GetRHI(), LayerIt, 0, false);
+			}
+		}
+
+		// 2. Fill in texture based profile
+		const auto GlobalShaderMap = GetGlobalShaderMap(ShaderPlatform);
+		FRDGTextureRef SpecularProfileTexture = GraphBuilder.RegisterExternalTexture(GSpecularProfileTextureAtlas, TEXT("SpecularProfileTexture"));
+		FRDGTextureUAVRef SpecularProfileUAV = GraphBuilder.CreateUAV(SpecularProfileTexture);
+		for (uint32 LayerIt = 0; LayerIt < LayerCount; ++LayerIt)
+		{
+			if (SpecularProfileEntries[LayerIt].Texture)
+			{
+				const FTextureReferenceRHIRef& TextureRHI = SpecularProfileEntries[LayerIt].Texture->TextureReferenceRHI;
+				FSpecularProfileCopyCS::FParameters* PassParameters = GraphBuilder.AllocParameters<FSpecularProfileCopyCS::FParameters>();
+				PassParameters->SourceMipCount = FMath::Max(1u, uint32(TextureRHI->GetDesc().NumMips));
+				PassParameters->SourceTexture = TextureRHI->GetReferencedTexture();
+				PassParameters->SourceSampler = TStaticSamplerState<SF_Point, AM_Clamp, AM_Clamp, AM_Clamp>::GetRHI();
+				PassParameters->SourceResolution = SpecularProfileTexture->Desc.Extent;
+				PassParameters->TargetResolution = SpecularProfileTexture->Desc.Extent;
+				PassParameters->TargetTexture = SpecularProfileUAV;
+				PassParameters->TargetIndex = LayerIt;
+				TShaderMapRef<FSpecularProfileCopyCS> Shader(GlobalShaderMap);
+				FComputeShaderUtils::AddPass(GraphBuilder, 
+											RDG_EVENT_NAME("SpecularProfile::CopyTexture"), 
+											Shader, 
+											PassParameters, 
+											FIntVector(FMath::DivideAndRoundUp(PassParameters->TargetResolution.X, 8), FMath::DivideAndRoundUp(PassParameters->TargetResolution.Y, 8), 1));
+	
+				SpecularProfileEntries[LayerIt].CachedResolution = TextureRHI->GetDesc().Extent;
+			}
+		}
+	}
+	return GSpecularProfileTextureAtlas;
+}
+
+void FSpecularProfileTextureManager::ReleaseDynamicRHI()
+{
+	GSpecularProfileTextureAtlas.SafeRelease();
+}
+
+int32 FSpecularProfileTextureManager::FindAllocationId(const USpecularProfile* InProfile) const
+{
+	// we start at 1 because [0] is the default profile and always [0].Profile = 0 so we don't need to iterate that one
+	for (int32 i = 1; i < SpecularProfileEntries.Num(); ++i)
+	{
+		if (SpecularProfileEntries[i].Profile == InProfile)
+		{
+			return i;
+		}
+	}
+	return INDEX_NONE;
+}
+
+///////////////////////////////////////////////////////////////////////////////////////////////////
+// USpecularProfile
+
+USpecularProfile::USpecularProfile(const FObjectInitializer& ObjectInitializer)
+	: Super(ObjectInitializer)
+{
+}
+
+void USpecularProfile::BeginDestroy()
+{
+	USpecularProfile* Ref = this;
+	ENQUEUE_RENDER_COMMAND(RemoveSpecularProfile)(
+	[Ref](FRHICommandList& RHICmdList)
+	{
+		GSpecularProfileTextureManager.RemoveProfile(Ref);
+	});
+
+	Super::BeginDestroy();
+}
+
+void USpecularProfile::PostEditChangeProperty(struct FPropertyChangedEvent& PropertyChangedEvent)
+{
+	const FSpecularProfileStruct LocalSettings = this->Settings;
+	USpecularProfile* LocalProfile = this;
+	GetRendererModule().InvalidatePathTracedOutput();
+
+	const FTextureReference* LocalTextureResource = nullptr;
+	if (!LocalSettings.IsProcedural())
+	{
+		LocalTextureResource = &LocalSettings.Texture->TextureReference;
+	}
+
+	ENQUEUE_RENDER_COMMAND(UpdateSpecularProfile)(
+	[LocalSettings, LocalProfile, LocalTextureResource](FRHICommandListImmediate& RHICmdList)
+	{
+		// any changes to the setting require an update of the texture
+		GSpecularProfileTextureManager.UpdateProfile(LocalProfile, LocalSettings, LocalTextureResource);
+	});
+}
+
+///////////////////////////////////////////////////////////////////////////////////////////////////
+// Public API
+
+namespace SpecularProfileAtlas
+{
+
+	FName GetSpecularProfileParameterName()
+	{
+		static FName NameSpecularProfile(TEXT("__SpecularProfile"));
+		return NameSpecularProfile;
+	}
+
+	float GetSpecularProfileId(const USpecularProfile* In)
+	{
+		int32 AllocationId = 0;
+		if (In)
+		{
+			// can be optimized (cached)
+			AllocationId = GSpecularProfileTextureManager.FindAllocationId(In);
+		}
+		else
+		{
+			// no profile specified means we use the default one stored at [0] which is human skin
+			AllocationId = 0;
+		}
+		return AllocationId / 255.0f;
+	}
+	
+	int32 AddOrUpdateProfile(const USpecularProfile* InProfile, const FSpecularProfileStruct InSettings, const FTextureReference* InTexture)
+	{
+		return GSpecularProfileTextureManager.AddOrUpdateProfile(InProfile, InSettings, InTexture);
+	}
+
+	FRHITexture* GetSpecularProfileTextureAtlas()
+	{
+		return GSpecularProfileTextureAtlas ? GSpecularProfileTextureAtlas->GetRHI() : nullptr;
+	}
+
+	FRHITexture* GetSpecularProfileTextureAtlasWithFallback()
+	{
+		return GSpecularProfileTextureAtlas ? GSpecularProfileTextureAtlas->GetRHI() : static_cast<FRHITexture*>(GBlackTexture->TextureRHI);
+	}
+
+	void UpdateSpecularProfileTextureAtlas(FRDGBuilder& GraphBuilder, EShaderPlatform ShaderPlatform)
+	{
+		GSpecularProfileTextureManager.GetTexture(GraphBuilder, ShaderPlatform);
+	}
+
+} // namespace SpecularProfileAtlas
diff --git a/Engine/Source/Runtime/Engine/Private/SceneManagement.cpp b/Engine/Source/Runtime/Engine/Private/SceneManagement.cpp
--- a/Engine/Source/Runtime/Engine/Private/SceneManagement.cpp
+++ b/Engine/Source/Runtime/Engine/Private/SceneManagement.cpp
@@ -1010,16 +1010,22 @@ FViewUniformShaderParameters::FViewUniformShaderParameters()
 	IESAtlasTexture = GBlackTextureWithSRV->TextureRHI;
 	IESAtlasSampler = TStaticSamplerState<SF_Bilinear, AM_Clamp, AM_Clamp, AM_Clamp>::GetRHI();
 
-	// Subsurface profiles/pre-intregrated
+	// Subsurface profiles
 	SSProfilesTextureSizeAndInvSize = FVector4f(1.f,1.f,1.f,1.f);
 	SSProfilesTexture = GBlackTextureWithSRV->TextureRHI;
 	SSProfilesSampler = TStaticSamplerState<SF_Point, AM_Clamp, AM_Clamp, AM_Clamp>::GetRHI();;
 	SSProfilesTransmissionSampler = TStaticSamplerState<SF_Bilinear, AM_Clamp, AM_Clamp, AM_Clamp>::GetRHI();
 
+	// Subsurface pre-intregrated profiles
 	SSProfilesPreIntegratedTextureSizeAndInvSize = FVector4f(1.f,1.f,1.f,1.f);
 	SSProfilesPreIntegratedTexture = GBlackArrayTexture->TextureRHI;
 	SSProfilesPreIntegratedSampler = TStaticSamplerState<SF_Bilinear, AM_Clamp, AM_Clamp, AM_Clamp>::GetRHI();
 
+	// Specular profiles
+	SpecularProfileTextureSizeAndInvSize = FVector4f(1.f,1.f,1.f,1.f);
+	SpecularProfileTexture = GBlackArrayTexture->TextureRHI;
+	SpecularProfileSampler = TStaticSamplerState<SF_Bilinear, AM_Clamp, AM_Clamp, AM_Clamp>::GetRHI();
+
 	VTFeedbackBuffer = GEmptyStructuredBufferWithUAV->UnorderedAccessViewRHI;
 }
 
diff --git a/Engine/Source/Runtime/Engine/Public/MaterialCompiler.h b/Engine/Source/Runtime/Engine/Public/MaterialCompiler.h
--- a/Engine/Source/Runtime/Engine/Public/MaterialCompiler.h
+++ b/Engine/Source/Runtime/Engine/Public/MaterialCompiler.h
@@ -555,6 +555,7 @@ public:
 		int32 FuzzAmount, int32 FuzzColor, int32 FuzzRoughness,
 		int32 Thickness,
 		int32 GlintValue, int32 GlintUV,
+		int32 SpecularProfileId,
 		int32 Normal, int32 Tangent, const FString& SharedLocalBasisIndexMacro,
 		FStrataOperator* PromoteToOperator) = 0;
 	virtual int32 StrataConversionFromLegacy(
@@ -1175,6 +1176,7 @@ public:
 		int32 FuzzAmount, int32 FuzzColor, int32 FuzzRoughness,
 		int32 Thickness,
 		int32 GlintValue, int32 GlintUV,
+		int32 SpecularProfileId,
 		int32 Normal, int32 Tangent, const FString& SharedLocalBasisIndexMacro,
 		FStrataOperator* PromoteToOperator) override
 	{
@@ -1187,6 +1189,7 @@ public:
 			FuzzAmount, FuzzColor, FuzzRoughness,
 			Thickness,
 			GlintValue, GlintUV,
+			SpecularProfileId,
 			Normal, Tangent, SharedLocalBasisIndexMacro,
 			PromoteToOperator);
 	}
diff --git a/Engine/Source/Runtime/Engine/Public/Materials/MaterialRenderProxy.h b/Engine/Source/Runtime/Engine/Public/Materials/MaterialRenderProxy.h
--- a/Engine/Source/Runtime/Engine/Public/Materials/MaterialRenderProxy.h
+++ b/Engine/Source/Runtime/Engine/Public/Materials/MaterialRenderProxy.h
@@ -22,6 +22,7 @@ class UMaterialInterface;
 class URuntimeVirtualTexture;
 class USparseVolumeTexture;
 class USubsurfaceProfile;
+class USpecularProfile;
 class UTexture;
 
 struct FMaterialParameterValue;
@@ -206,6 +207,9 @@ public:
 	void SetSubsurfaceProfileRT(const USubsurfaceProfile* Ptr) { SubsurfaceProfileRT = Ptr; }
 	const USubsurfaceProfile* GetSubsurfaceProfileRT() const { return SubsurfaceProfileRT; }
 
+	void SetSpecularProfileRT(const USpecularProfile* Ptr) { SpecularProfileRT = Ptr; }
+	const USpecularProfile* GetSpecularProfileRT() const { return SpecularProfileRT; }
+
 	static void UpdateDeferredCachedUniformExpressions();
 
 	static bool HasDeferredUniformExpressionCacheRequests();
@@ -223,6 +227,7 @@ private:
 
 	/** 0 if not set, game thread pointer, do not dereference, only for comparison */
 	const USubsurfaceProfile* SubsurfaceProfileRT;
+	const USpecularProfile* SpecularProfileRT;
 	FString MaterialName;
 
 	/** Incremented each time UniformExpressionCache is modified */
diff --git a/Engine/Source/Runtime/Engine/Public/Rendering/StrataMaterialShared.cpp b/Engine/Source/Runtime/Engine/Public/Rendering/StrataMaterialShared.cpp
--- a/Engine/Source/Runtime/Engine/Public/Rendering/StrataMaterialShared.cpp
+++ b/Engine/Source/Runtime/Engine/Public/Rendering/StrataMaterialShared.cpp
@@ -52,6 +52,7 @@ FStrataOperator::FStrataOperator()
 	bBSDFHasSecondRoughnessOrSimpleClearCoat = false;
 	bBSDFHasAnisotropy = false;
 	bBSDFHasGlint = false;
+	bBSDFHasSpecularProfile = false;
 
 	MaxDistanceFromLeaves = 0;
 	LayerDepth = 0;
@@ -72,6 +73,7 @@ void FStrataOperator::CombineFlagsForParameterBlending(FStrataOperator& A, FStra
 	bBSDFHasSecondRoughnessOrSimpleClearCoat = A.bBSDFHasSecondRoughnessOrSimpleClearCoat || B.bBSDFHasSecondRoughnessOrSimpleClearCoat;
 	bBSDFHasAnisotropy = A.bBSDFHasAnisotropy || B.bBSDFHasAnisotropy;
 	bBSDFHasGlint = A.bBSDFHasGlint || B.bBSDFHasGlint;
+	bBSDFHasSpecularProfile = A.bBSDFHasSpecularProfile || B.bBSDFHasSpecularProfile;
 #endif
 }
 
@@ -85,6 +87,7 @@ void FStrataOperator::CopyFlagsForParameterBlending(FStrataOperator& A)
 	bBSDFHasSecondRoughnessOrSimpleClearCoat = A.bBSDFHasSecondRoughnessOrSimpleClearCoat;
 	bBSDFHasAnisotropy = A.bBSDFHasAnisotropy;
 	bBSDFHasGlint = A.bBSDFHasGlint;
+	bBSDFHasSpecularProfile = A.bBSDFHasSpecularProfile;
 #endif
 }
 
diff --git a/Engine/Source/Runtime/Engine/Public/Rendering/StrataMaterialShared.h b/Engine/Source/Runtime/Engine/Public/Rendering/StrataMaterialShared.h
--- a/Engine/Source/Runtime/Engine/Public/Rendering/StrataMaterialShared.h
+++ b/Engine/Source/Runtime/Engine/Public/Rendering/StrataMaterialShared.h
@@ -55,6 +55,7 @@ public:
 	LAYOUT_FIELD_EDITORONLY(uint8, bBSDFHasSecondRoughnessOrSimpleClearCoat);
 	LAYOUT_FIELD_EDITORONLY(uint8, bBSDFHasAnisotropy);
 	LAYOUT_FIELD_EDITORONLY(uint8, bBSDFHasGlint);
+	LAYOUT_FIELD_EDITORONLY(uint8, bBSDFHasSpecularProfile);
 
 	// Data derived after the tree has been built.
 	LAYOUT_FIELD_EDITORONLY(int32, MaxDistanceFromLeaves);
diff --git a/Engine/Source/Runtime/Engine/Public/SceneView.h b/Engine/Source/Runtime/Engine/Public/SceneView.h
--- a/Engine/Source/Runtime/Engine/Public/SceneView.h
+++ b/Engine/Source/Runtime/Engine/Public/SceneView.h
@@ -881,6 +881,7 @@ enum ETranslucencyVolumeCascade
 	VIEW_UNIFORM_BUFFER_MEMBER(float, bSubsurfacePostprocessEnabled) \
 	VIEW_UNIFORM_BUFFER_MEMBER(FVector4f, SSProfilesTextureSizeAndInvSize) \
 	VIEW_UNIFORM_BUFFER_MEMBER(FVector4f, SSProfilesPreIntegratedTextureSizeAndInvSize) \
+	VIEW_UNIFORM_BUFFER_MEMBER(FVector4f, SpecularProfileTextureSizeAndInvSize) \
 	VIEW_UNIFORM_BUFFER_MEMBER(FVector3f, PhysicsFieldClipmapCenter) \
 	VIEW_UNIFORM_BUFFER_MEMBER(float, PhysicsFieldClipmapDistance) \
 	VIEW_UNIFORM_BUFFER_MEMBER(int, PhysicsFieldClipmapResolution) \
@@ -1010,6 +1011,9 @@ BEGIN_GLOBAL_SHADER_PARAMETER_STRUCT_WITH_CONSTRUCTOR(FViewUniformShaderParamete
 	SHADER_PARAMETER_SAMPLER(SamplerState, SSProfilesTransmissionSampler)
 	SHADER_PARAMETER_TEXTURE(Texture2DArray, SSProfilesPreIntegratedTexture)
 	SHADER_PARAMETER_SAMPLER(SamplerState, SSProfilesPreIntegratedSampler)
+	// Specular Profile
+	SHADER_PARAMETER_TEXTURE(Texture2DArray, SpecularProfileTexture)
+	SHADER_PARAMETER_SAMPLER(SamplerState, SpecularProfileSampler)
 	// Water
 	SHADER_PARAMETER_SRV(Buffer<float4>, WaterIndirection)
 	SHADER_PARAMETER_SRV(Buffer<float4>, WaterData)
diff --git a/Engine/Source/Runtime/Renderer/Private/DeferredShadingRenderer.cpp b/Engine/Source/Runtime/Renderer/Private/DeferredShadingRenderer.cpp
--- a/Engine/Source/Runtime/Renderer/Private/DeferredShadingRenderer.cpp
+++ b/Engine/Source/Runtime/Renderer/Private/DeferredShadingRenderer.cpp
@@ -70,6 +70,7 @@
 #include "InstanceCulling/InstanceCullingManager.h"
 #include "ProfilingDebugging/CpuProfilerTrace.h"
 #include "Engine/SubsurfaceProfile.h"
+#include "Engine/SpecularProfile.h"
 #include "SceneCaptureRendering.h"
 #include "NaniteSceneProxy.h"
 #include "Nanite/NaniteRayTracing.h"
@@ -1971,9 +1972,11 @@ bool FDeferredShadingSceneRenderer::SetupRayTracingPipelineStates(FRDGBuilder& G
 	// Initialize common resources used for lighting in ray tracing effects
 
 	ReferenceView.RayTracingSubSurfaceProfileTexture = GetSubsurfaceProfileTextureWithFallback();
-
 	ReferenceView.RayTracingSubSurfaceProfileSRV = RHICreateShaderResourceView(ReferenceView.RayTracingSubSurfaceProfileTexture, 0);
 
+	ReferenceView.RayTracingSpecularProfileTexture = SpecularProfileAtlas::GetSpecularProfileTextureAtlasWithFallback();
+	ReferenceView.RayTracingSpecularProfileSRV = RHICreateShaderResourceView(ReferenceView.RayTracingSpecularProfileTexture, 0);
+
 	for (int32 ViewIndex = 0; ViewIndex < AllFamilyViews.Num(); ++ViewIndex)
 	{
 		// TODO:  It would make more sense for common ray tracing resources to be in a shared structure, rather than copied into each FViewInfo.
@@ -1987,6 +1990,8 @@ bool FDeferredShadingSceneRenderer::SetupRayTracingPipelineStates(FRDGBuilder& G
 			View->RayTracingSubSurfaceProfileTexture = ReferenceView.RayTracingSubSurfaceProfileTexture;
 			View->RayTracingSubSurfaceProfileSRV = ReferenceView.RayTracingSubSurfaceProfileSRV;
 			View->RayTracingMaterialPipeline = ReferenceView.RayTracingMaterialPipeline;
+			View->RayTracingSpecularProfileTexture = ReferenceView.RayTracingSpecularProfileTexture;
+			View->RayTracingSpecularProfileSRV = ReferenceView.RayTracingSpecularProfileSRV;
 		}
 	}
 
@@ -2233,6 +2238,8 @@ static void ReleaseRaytracingResources(FRDGBuilder& GraphBuilder, TArrayView<FVi
 			// is finished using them (where multiple view families are rendered).
 			View.RayTracingSubSurfaceProfileSRV.SafeRelease();
 			View.RayTracingSubSurfaceProfileTexture.SafeRelease();
+			View.RayTracingSpecularProfileSRV.SafeRelease();
+			View.RayTracingSpecularProfileTexture.SafeRelease();
 		}
 	});
 }
@@ -2730,6 +2737,7 @@ void FDeferredShadingSceneRenderer::Render(FRDGBuilder& GraphBuilder)
 
 		// Force the subsurface profile texture to be updated.
 		UpdateSubsurfaceProfileTexture(GraphBuilder, ShaderPlatform);
+		SpecularProfileAtlas::UpdateSpecularProfileTextureAtlas(GraphBuilder, ShaderPlatform);
 
 		// Force the rect light texture & IES texture to be updated.
 		RectLightAtlas::UpdateAtlasTexture(GraphBuilder, FeatureLevel);
diff --git a/Engine/Source/Runtime/Renderer/Private/MobileBasePassRendering.cpp b/Engine/Source/Runtime/Renderer/Private/MobileBasePassRendering.cpp
--- a/Engine/Source/Runtime/Renderer/Private/MobileBasePassRendering.cpp
+++ b/Engine/Source/Runtime/Renderer/Private/MobileBasePassRendering.cpp
@@ -18,6 +18,7 @@
 #include "FramePro/FrameProProfiler.h"
 #include "PostProcess/PostProcessPixelProjectedReflectionMobile.h"
 #include "Engine/SubsurfaceProfile.h"
+#include "Engine/SpecularProfile.h"
 #include "LocalLightSceneProxy.h"
 #include "ReflectionEnvironment.h"
 #include "RenderCore.h"
diff --git a/Engine/Source/Runtime/Renderer/Private/MobileShadingRenderer.cpp b/Engine/Source/Runtime/Renderer/Private/MobileShadingRenderer.cpp
--- a/Engine/Source/Runtime/Renderer/Private/MobileShadingRenderer.cpp
+++ b/Engine/Source/Runtime/Renderer/Private/MobileShadingRenderer.cpp
@@ -65,6 +65,7 @@
 #include "RectLightTextureManager.h"
 #include "IESTextureManager.h"
 #include "SceneUniformBuffer.h"
+#include "Engine/SpecularProfile.h"
 
 uint32 GetShadowQuality();
 
@@ -874,6 +875,7 @@ void FMobileSceneRenderer::Render(FRDGBuilder& GraphBuilder)
 
 	// Force the subsurface profile texture to be updated.
 	UpdateSubsurfaceProfileTexture(GraphBuilder, ShaderPlatform);
+	SpecularProfileAtlas::UpdateSpecularProfileTextureAtlas(GraphBuilder, ShaderPlatform);
 	
 	if (bDeferredShading)
 	{
diff --git a/Engine/Source/Runtime/Renderer/Private/SceneRendering.cpp b/Engine/Source/Runtime/Renderer/Private/SceneRendering.cpp
--- a/Engine/Source/Runtime/Renderer/Private/SceneRendering.cpp
+++ b/Engine/Source/Runtime/Renderer/Private/SceneRendering.cpp
@@ -89,6 +89,7 @@
 #include "VolumetricCloudRendering.h"
 #include "VolumetricFog.h"
 #include "PrimitiveSceneShaderData.h"
+#include "Engine/SpecularProfile.h"
 
 /*-----------------------------------------------------------------------------
 	Globals
@@ -1662,6 +1663,15 @@ void FViewInfo::SetupUniformBufferParameters(
 		}
 	}
 
+	// Specular Profiles
+	{
+		FRHITexture* Texture = SpecularProfileAtlas::GetSpecularProfileTextureAtlasWithFallback();
+		FIntVector TextureSize = Texture->GetSizeXYZ();
+		ViewUniformShaderParameters.SpecularProfileTextureSizeAndInvSize = FVector4f(TextureSize.X, TextureSize.Y, 1.0f / TextureSize.X, 1.0f / TextureSize.Y);
+		ViewUniformShaderParameters.SpecularProfileTexture = Texture;
+		ViewUniformShaderParameters.SpecularProfileSampler = TStaticSamplerState<SF_Point, AM_Clamp, AM_Clamp, AM_Clamp>::GetRHI();
+	}
+
 	{
 		// This is the CVar default
 		float Value = 1.0f;
diff --git a/Engine/Source/Runtime/Renderer/Private/SceneRendering.h b/Engine/Source/Runtime/Renderer/Private/SceneRendering.h
--- a/Engine/Source/Runtime/Renderer/Private/SceneRendering.h
+++ b/Engine/Source/Runtime/Renderer/Private/SceneRendering.h
@@ -1632,6 +1632,10 @@ public:
 	TRefCountPtr<FRHITexture>						RayTracingSubSurfaceProfileTexture;
 	FShaderResourceViewRHIRef						RayTracingSubSurfaceProfileSRV;
 
+	// Common resources used for lighting in ray tracing effects
+	TRefCountPtr<FRHITexture>						RayTracingSpecularProfileTexture;
+	FShaderResourceViewRHIRef						RayTracingSpecularProfileSRV;
+
 	TRDGUniformBufferRef<FRaytracingLightDataPacked>	RayTracingLightDataUniformBuffer;
 	TRDGUniformBufferRef<FRayTracingDecals>				RayTracingDecalUniformBuffer;
 	bool												bHasRayTracingDecals = false;
