Add a GetOwningLocalPlayer function to the UPlayerInput object. This will make it a little easier to write triggers/modifiers in Enhanced Input that want to access the Local Player if they want access to the User Settings object to check aim sensativity or other things.

diff --git a/Engine/Source/Runtime/Engine/Classes/GameFramework/PlayerInput.h b/Engine/Source/Runtime/Engine/Classes/GameFramework/PlayerInput.h
--- a/Engine/Source/Runtime/Engine/Classes/GameFramework/PlayerInput.h
+++ b/Engine/Source/Runtime/Engine/Classes/GameFramework/PlayerInput.h
@@ -20,6 +20,7 @@
 
 class FDebugDisplayInfo;
 class UInputComponent;
+class ULocalPlayer;
 struct FDelegateDispatchDetails;
 struct FInputActionBinding;
 struct FInputAxisBinding;
@@ -517,6 +518,12 @@ public:
 	UFUNCTION(BlueprintCallable, BlueprintPure, Category = "Input")
 	APlayerController* GetOuterAPlayerController() const;
 
+	/**
+	* Returns the owning local player of this Input Object. This will only be valid if there is a 
+	* current outer player controller.
+	*/
+	ULocalPlayer* GetOwningLocalPlayer() const;
+	
 private:
 
 	/** Runtime struct that caches the list of mappings for a given Action Name and the capturing chord if applicable */
diff --git a/Engine/Source/Runtime/Engine/Private/UserInterface/PlayerInput.cpp b/Engine/Source/Runtime/Engine/Private/UserInterface/PlayerInput.cpp
--- a/Engine/Source/Runtime/Engine/Private/UserInterface/PlayerInput.cpp
+++ b/Engine/Source/Runtime/Engine/Private/UserInterface/PlayerInput.cpp
@@ -759,6 +759,16 @@ APlayerController* UPlayerInput::GetOuterAPlayerController() const
 	return Cast<APlayerController>(GetOuter());
 }
 
+ULocalPlayer* UPlayerInput::GetOwningLocalPlayer() const
+{
+	if (const APlayerController* PC = GetOuterAPlayerController())
+	{
+		return PC->GetLocalPlayer();
+	}
+
+	return nullptr;
+}
+
 void UPlayerInput::ConditionalBuildKeyMappings_Internal() const
 {
 	if (ActionKeyMap.Num() == 0)
