Add a getter for the current Mouse Lock Mode on the Game Viewport Client

#rnx

diff --git a/Engine/Source/Runtime/Engine/Classes/Engine/GameViewportClient.h b/Engine/Source/Runtime/Engine/Classes/Engine/GameViewportClient.h
--- a/Engine/Source/Runtime/Engine/Classes/Engine/GameViewportClient.h
+++ b/Engine/Source/Runtime/Engine/Classes/Engine/GameViewportClient.h
@@ -798,6 +798,11 @@ public:
 	*/
 	void SetMouseLockMode(EMouseLockMode InMouseLockMode);
 
+	/**
+	 * Returns the current mouse cursor lock mode when the viewport is clicked.
+	 */
+	EMouseLockMode GetMouseLockMode() const;
+
 	/**
 	 * Sets whether or not the cursor is hidden when the viewport captures the mouse
 	 */
diff --git a/Engine/Source/Runtime/Engine/Private/GameViewportClient.cpp b/Engine/Source/Runtime/Engine/Private/GameViewportClient.cpp
--- a/Engine/Source/Runtime/Engine/Private/GameViewportClient.cpp
+++ b/Engine/Source/Runtime/Engine/Private/GameViewportClient.cpp
@@ -2941,6 +2941,11 @@ void UGameViewportClient::SetMouseLockMode(EMouseLockMode InMouseLockMode)
 	}
 }
 
+EMouseLockMode UGameViewportClient::GetMouseLockMode() const
+{
+	return MouseLockMode;
+}
+
 void UGameViewportClient::SetHideCursorDuringCapture(bool InHideCursorDuringCapture)
 {
 	if (bHideCursorDuringCapture != InHideCursorDuringCapture)
