Fix crash setting invalid parameter type on a soundcue

diff --git a/Engine/Source/Runtime/Engine/Private/SoundCue.cpp b/Engine/Source/Runtime/Engine/Private/SoundCue.cpp
--- a/Engine/Source/Runtime/Engine/Private/SoundCue.cpp
+++ b/Engine/Source/Runtime/Engine/Private/SoundCue.cpp
@@ -632,6 +632,11 @@ void USoundCue::Parse(FAudioDevice* AudioDevice, const UPTRINT NodeWaveInstanceH
 
 		for (const FWaveInstance* Instance : WaveInstances)
 		{
+			if (!Instance)
+			{
+				continue;
+			}
+			
 			if (TSharedPtr<Audio::IParameterTransmitter>* ChildTransmitterPtr = Transmitter->Transmitters.Find(Instance->WaveInstanceHash))
 			{
 				TSharedPtr<Audio::IParameterTransmitter>& ChildTransmitter = *ChildTransmitterPtr;
@@ -639,7 +644,17 @@ void USoundCue::Parse(FAudioDevice* AudioDevice, const UPTRINT NodeWaveInstanceH
 				if (ChildTransmitter.IsValid())
 				{
 					TArray<FAudioParameter> Params = Transmitter->ParamsToSet;
-					ChildTransmitter->SetParameters(MoveTemp(Params));
+
+					if (USoundBase* Sound = Instance->WaveData)
+					{
+						static const FName ProxyFeatureName("SoundCueParameterTransmitter");
+						Sound->InitParameters(Params, ProxyFeatureName);
+					}
+
+					if (!Params.IsEmpty())
+					{
+						ChildTransmitter->SetParameters(MoveTemp(Params));
+					}
 				}
 			}
 		}
