Fix string and non-soundwave object params not being settable on sound cues

diff --git a/Engine/Source/Runtime/Engine/Private/SoundBase.cpp b/Engine/Source/Runtime/Engine/Private/SoundBase.cpp
--- a/Engine/Source/Runtime/Engine/Private/SoundBase.cpp
+++ b/Engine/Source/Runtime/Engine/Private/SoundBase.cpp
@@ -306,35 +306,7 @@ bool USoundBase::IsParameterValid(const FAudioParameter& InParameter) const
 		return false;
 	}
 
-	switch (InParameter.ParamType)
-	{
-		case EAudioParameterType::BooleanArray:
-		case EAudioParameterType::FloatArray:
-		case EAudioParameterType::IntegerArray:
-		case EAudioParameterType::NoneArray:
-		case EAudioParameterType::ObjectArray:
-		case EAudioParameterType::StringArray:
-		case EAudioParameterType::String:
-		{
-			return false;
-		}
-
-		case EAudioParameterType::Object:
-		case EAudioParameterType::None:
-		{
-			if (InParameter.ObjectParam)
-			{
-				return InParameter.ObjectParam->GetClass()->IsChildOf(USoundWave::StaticClass());
-			}
-		}
-
-		default:
-		{
-			break;
-		}
-	}
-
-	return true;
+	return !(InParameter.ParamType == EAudioParameterType::None || InParameter.ParamType == EAudioParameterType::NoneArray);
 }
 
 #if WITH_EDITORONLY_DATA
