Fixed issue where world-time-dilated and pawn-time-dilated intervals were being compared, affecting the timing of server-forced pawn position updates.

diff --git a/Engine/Source/Runtime/Engine/Private/PlayerController.cpp b/Engine/Source/Runtime/Engine/Private/PlayerController.cpp
--- a/Engine/Source/Runtime/Engine/Private/PlayerController.cpp
+++ b/Engine/Source/Runtime/Engine/Private/PlayerController.cpp
@@ -5027,7 +5027,9 @@ void APlayerController::TickActor( float DeltaSeconds, ELevelTick TickType, FAct
 						if (ServerData->bTriggeringForcedUpdates)
 						{
 							const float PawnTimeSinceForcingUpdates = (WorldTimeStamp - ServerData->ServerTimeBeginningForcedUpdates) * GetPawn()->CustomTimeDilation;
-							if (PawnTimeSinceForcingUpdates > ForcedUpdateMaxDuration * GetPawn()->GetActorTimeDilation())
+							const float PawnTimeForcedUpdateMaxDuration = ForcedUpdateMaxDuration * GetPawn()->GetActorTimeDilation();
+							
+							if (PawnTimeSinceForcingUpdates > PawnTimeForcedUpdateMaxDuration)
 							{
 								if (ServerData->ServerTimeStamp > ServerData->ServerTimeLastForcedUpdate)
 								{
@@ -5057,7 +5059,10 @@ void APlayerController::TickActor( float DeltaSeconds, ELevelTick TickType, FAct
 						LastMovementUpdateTime = CurrentRealTime;
 
 						// Trigger forced update if allowed
-						if (!bRecentHitch && ForcedUpdateInterval > 0.f && PawnTimeSinceUpdate > FMath::Max<float>(DeltaSeconds+0.06f, ForcedUpdateInterval * GetPawn()->GetActorTimeDilation()))
+						const float PawnTimeMinForcedUpdateInterval = (DeltaSeconds + 0.06f) * GetPawn()->CustomTimeDilation;
+						const float PawnTimeForcedUpdateInterval = FMath::Max<float>(PawnTimeMinForcedUpdateInterval, ForcedUpdateInterval * GetPawn()->GetActorTimeDilation());
+
+						if (!bRecentHitch && ForcedUpdateInterval > 0.f && (PawnTimeSinceUpdate > PawnTimeForcedUpdateInterval))
 						{
 							//UE_LOG(LogPlayerController, Warning, TEXT("ForcedMovementTick. PawnTimeSinceUpdate: %f, DeltaSeconds: %f, DeltaSeconds+: %f"), PawnTimeSinceUpdate, DeltaSeconds, DeltaSeconds+0.06f);
 							const USkeletalMeshComponent* PawnMesh = GetPawn()->FindComponentByClass<USkeletalMeshComponent>();
