Add support for multi specular profile per material.

diff --git a/Engine/Source/Editor/UnrealEd/Private/Factories/EditorFactories.cpp b/Engine/Source/Editor/UnrealEd/Private/Factories/EditorFactories.cpp
--- a/Engine/Source/Editor/UnrealEd/Private/Factories/EditorFactories.cpp
+++ b/Engine/Source/Editor/UnrealEd/Private/Factories/EditorFactories.cpp
@@ -7876,6 +7876,7 @@ USpecularProfileFactory::USpecularProfileFactory(const FObjectInitializer& Objec
 UObject* USpecularProfileFactory::FactoryCreateNew(UClass* InClass, UObject* InParent, FName InName, EObjectFlags Flags, UObject* Context, FFeedbackContext* Warn)
 {
 	USpecularProfile* Object = NewObject<USpecularProfile>(InParent, InName, Flags);	
+	Object->Guid = FGuid::NewGuid();
 	return Object;
 }
 
diff --git a/Engine/Source/Runtime/Engine/Classes/Engine/EngineTypes.h b/Engine/Source/Runtime/Engine/Classes/Engine/EngineTypes.h
--- a/Engine/Source/Runtime/Engine/Classes/Engine/EngineTypes.h
+++ b/Engine/Source/Runtime/Engine/Classes/Engine/EngineTypes.h
@@ -740,9 +740,10 @@ public:
 	int32 CountSubsurfaceProfiles() const { return SubsurfaceProfiles.Num(); }
 	USubsurfaceProfile* GetSubsurfaceProfile() const { return SubsurfaceProfiles.Num() > 0 ? SubsurfaceProfiles[0] : nullptr; }
 
+	// Specular profiles
 	void AddSpecularProfile(USpecularProfile* InProfile) { if (InProfile) SpecularProfiles.Add(InProfile); }
 	int32 CountSpecularProfiles() const { return SpecularProfiles.Num(); }
-	USpecularProfile* GetSpecularProfile() const { return SpecularProfiles.Num() > 0 ? SpecularProfiles[0] : nullptr; }
+	USpecularProfile* GetSpecularProfile(int32 Index) const { return Index < SpecularProfiles.Num() ? SpecularProfiles[Index] : nullptr; }
 
 	// Shading model from expression
 	void SetShadingModelFromExpression(bool bIn) { bHasShadingModelFromExpression = bIn ? 1u : 0u; }
diff --git a/Engine/Source/Runtime/Engine/Classes/Engine/SpecularProfile.h b/Engine/Source/Runtime/Engine/Classes/Engine/SpecularProfile.h
--- a/Engine/Source/Runtime/Engine/Classes/Engine/SpecularProfile.h
+++ b/Engine/Source/Runtime/Engine/Classes/Engine/SpecularProfile.h
@@ -4,6 +4,7 @@
 #include "CoreMinimal.h"
 #include "UObject/ObjectMacros.h"
 #include "UObject/Object.h"
+#include "Misc/Guid.h"
 #include "RenderResource.h"
 #include "SpecularProfile.generated.h"
 
@@ -65,6 +66,9 @@ class USpecularProfile : public UObject
 	UPROPERTY(Category = USpecularProfile, EditAnywhere, meta = (ShowOnlyInnerProperties))
 	struct FSpecularProfileStruct Settings;
 
+	UPROPERTY()
+	FGuid Guid;
+
 	//~ Begin UObject Interface
 	virtual void BeginDestroy();
 	virtual void PostEditChangeProperty(struct FPropertyChangedEvent& PropertyChangedEvent);
@@ -83,10 +87,10 @@ ENGINE_API FRHITexture* GetSpecularProfileTextureAtlas();
 ENGINE_API FRHITexture* GetSpecularProfileTextureAtlasWithFallback();
 
 // Profile - Initializes or updates the contents of the specular profile texture.
-ENGINE_API int32 AddOrUpdateProfile(const USpecularProfile* InProfile, const FSpecularProfileStruct InSettings, const FTextureReference* InTexture);
+ENGINE_API int32 AddOrUpdateProfile(const USpecularProfile* InProfile, const FGuid& InGuid, const FSpecularProfileStruct InSettings, const FTextureReference* InTexture);
 
 // Profile - Returns the specular profile ID shader parameter name
-ENGINE_API FName GetSpecularProfileParameterName();
+ENGINE_API FName GetSpecularProfileParameterName(const USpecularProfile* InProfile);
 
 // Profile - Returns the specular profile ID for a given Specular Profile object
 ENGINE_API float GetSpecularProfileId(const USpecularProfile* In);
diff --git a/Engine/Source/Runtime/Engine/Classes/Materials/Material.h b/Engine/Source/Runtime/Engine/Classes/Materials/Material.h
--- a/Engine/Source/Runtime/Engine/Classes/Materials/Material.h
+++ b/Engine/Source/Runtime/Engine/Classes/Materials/Material.h
@@ -1126,7 +1126,8 @@ public:
 	ENGINE_API virtual bool IsUIMaterial() const;
 	ENGINE_API virtual bool IsPostProcessMaterial() const;
 	ENGINE_API virtual USubsurfaceProfile* GetSubsurfaceProfile_Internal() const override;
-	ENGINE_API virtual USpecularProfile* GetSpecularProfile_Internal() const override;
+	ENGINE_API virtual uint32 NumSpecularProfile_Internal() const override;
+	ENGINE_API virtual USpecularProfile* GetSpecularProfile_Internal(uint32 Index) const override;
 	ENGINE_API virtual bool CastsRayTracedShadows() const override;
 	ENGINE_API virtual FDisplacementScaling GetDisplacementScaling() const override;
 	ENGINE_API virtual float GetMaxWorldPositionOffsetDisplacement() const override;
diff --git a/Engine/Source/Runtime/Engine/Classes/Materials/MaterialInterface.h b/Engine/Source/Runtime/Engine/Classes/Materials/MaterialInterface.h
--- a/Engine/Source/Runtime/Engine/Classes/Materials/MaterialInterface.h
+++ b/Engine/Source/Runtime/Engine/Classes/Materials/MaterialInterface.h
@@ -276,9 +276,9 @@ public:
 	UPROPERTY(EditAnywhere, BlueprintReadOnly, Category = Material, meta = (DisplayName = "Subsurface Profile"))
 	TObjectPtr<class USubsurfaceProfile> SubsurfaceProfile;
 
-	/** Specular Profile. */
-	UPROPERTY(EditAnywhere, BlueprintReadOnly, Category = Material)
-	TObjectPtr<class USpecularProfile> SpecularProfile;	
+	/** Specular Profile. For internal usage, not editable/visible */
+	UPROPERTY()
+	TArray<TObjectPtr<class USpecularProfile>> SpecularProfiles;
 
 	/** Whether this material interface is included in the base game (and not in a DLC) */
 	UPROPERTY()
@@ -915,7 +915,8 @@ public:
 	ENGINE_API virtual FDisplacementScaling GetDisplacementScaling() const;
 	ENGINE_API virtual float GetMaxWorldPositionOffsetDisplacement() const;
 	ENGINE_API virtual USubsurfaceProfile* GetSubsurfaceProfile_Internal() const;
-	ENGINE_API virtual USpecularProfile* GetSpecularProfile_Internal() const;
+	ENGINE_API virtual uint32 NumSpecularProfile_Internal() const;
+	ENGINE_API virtual USpecularProfile* GetSpecularProfile_Internal(uint32 Index) const;
 	ENGINE_API virtual bool CastsRayTracedShadows() const;
 
 	/**
diff --git a/Engine/Source/Runtime/Engine/Private/Materials/Material.cpp b/Engine/Source/Runtime/Engine/Private/Materials/Material.cpp
--- a/Engine/Source/Runtime/Engine/Private/Materials/Material.cpp
+++ b/Engine/Source/Runtime/Engine/Private/Materials/Material.cpp
@@ -449,10 +449,16 @@ public:
 				OutValue = GetSubsurfaceProfileId(GetSubsurfaceProfileRT());
 				return true;
 			}
-			else if (Type == EMaterialParameterType::Scalar && ParameterInfo.Name == SpecularProfileAtlas::GetSpecularProfileParameterName())
+			else if (Type == EMaterialParameterType::Scalar && NumSpecularProfileRT() > 0)
 			{
-				OutValue = SpecularProfileAtlas::GetSpecularProfileId(GetSpecularProfileRT());
-				return true;
+				for (uint32 It=0,Count=NumSpecularProfileRT();It<Count;++It)
+				{
+					if (ParameterInfo.Name == SpecularProfileAtlas::GetSpecularProfileParameterName(GetSpecularProfileRT(It)))
+					{
+						OutValue = SpecularProfileAtlas::GetSpecularProfileId(GetSpecularProfileRT(It));
+						return true;
+					}
+				}
 			}
 
 			return false;
@@ -5072,7 +5078,11 @@ void UMaterial::RebuildShadingModelField()
 		}
 
 		// Set specular profile if any
-		SpecularProfile = StrataMaterialInfo.GetSpecularProfile();
+		SpecularProfiles.SetNum(StrataMaterialInfo.CountSpecularProfiles());
+		for (int32 It = 0, Count = StrataMaterialInfo.CountSpecularProfiles(); It<Count; ++It)
+		{
+			SpecularProfiles[It] = StrataMaterialInfo.GetSpecularProfile(It);
+		}
 	}
 	// If using shading model from material expression, go through the expressions and look for the ShadingModel expression to figure out what shading models need to be supported in this material.
 	// This might not be the same as what is actually compiled in to the shader, since there might be feature switches, static switches etc. that skip certain shading models.
@@ -6621,10 +6631,16 @@ USubsurfaceProfile* UMaterial::GetSubsurfaceProfile_Internal() const
 	return SubsurfaceProfile; 
 }
 
-USpecularProfile* UMaterial::GetSpecularProfile_Internal() const
+uint32 UMaterial::NumSpecularProfile_Internal() const
+{
+	return SpecularProfiles.Num();
+}
+
+USpecularProfile* UMaterial::GetSpecularProfile_Internal(uint32 Index) const
 {
 	checkSlow(IsInGameThread());
-	return SpecularProfile; 
+	check(Index<uint32(SpecularProfiles.Num()));
+	return SpecularProfiles[Index];
 }
 
 bool UMaterial::CastsRayTracedShadows() const
diff --git a/Engine/Source/Runtime/Engine/Private/Materials/MaterialExpressions.cpp b/Engine/Source/Runtime/Engine/Private/Materials/MaterialExpressions.cpp
--- a/Engine/Source/Runtime/Engine/Private/Materials/MaterialExpressions.cpp
+++ b/Engine/Source/Runtime/Engine/Private/Materials/MaterialExpressions.cpp
@@ -23327,6 +23327,7 @@ UMaterialExpressionStrataSlabBSDF::UMaterialExpressionStrataSlabBSDF(const FObje
 }
 
 #if WITH_EDITOR
+FName CreateSpecularProfileParameterName(USpecularProfile* InProfile);
 int32 UMaterialExpressionStrataSlabBSDF::Compile(class FMaterialCompiler* Compiler, int32 OutputIndex)
 {
 	int32 RoughnessCodeChunk = CompileWithDefaultFloat1(Compiler, Roughness, 0.5f);
@@ -23354,7 +23355,8 @@ int32 UMaterialExpressionStrataSlabBSDF::Compile(class FMaterialCompiler* Compil
 	int32 SpecularProfileCodeChunk = INDEX_NONE;
 	if (bHasSpecularProfile)
 	{
-		SpecularProfileCodeChunk = Compiler->ForceCast(Compiler->ScalarParameter(SpecularProfileAtlas::GetSpecularProfileParameterName(), 1.0f), MCT_Float1);
+		const FName SpecularProfileParameterName = CreateSpecularProfileParameterName(SpecularProfile);
+		SpecularProfileCodeChunk = Compiler->ForceCast(Compiler->ScalarParameter(SpecularProfileParameterName, 1.0f), MCT_Float1);
 	}
 
 	const float DefaultSpecular = 0.5f;
diff --git a/Engine/Source/Runtime/Engine/Private/Materials/MaterialInstance.cpp b/Engine/Source/Runtime/Engine/Private/Materials/MaterialInstance.cpp
--- a/Engine/Source/Runtime/Engine/Private/Materials/MaterialInstance.cpp
+++ b/Engine/Source/Runtime/Engine/Private/Materials/MaterialInstance.cpp
@@ -281,12 +281,18 @@ bool FMaterialInstanceResource::GetParameterValue(EMaterialParameterType Type, c
 		OutValue = GetSubsurfaceProfileId(MySubsurfaceProfileRT);
 		bResult = true;
 	}
-	else if (Type == EMaterialParameterType::Scalar && ParameterInfo.Name == SpecularProfileAtlas::GetSpecularProfileParameterName())
+	else if (Type == EMaterialParameterType::Scalar && NumSpecularProfileRT() > 0)
 	{
-		check(ParameterInfo.Association == EMaterialParameterAssociation::GlobalParameter);
-		const USpecularProfile* MySpecularProfileRT = GetSpecularProfileRT();
-		OutValue = SpecularProfileAtlas::GetSpecularProfileId(MySpecularProfileRT);
-		bResult = true;
+		for (uint32 It=0,Count=NumSpecularProfileRT();It<Count;++It)
+		{
+			if (ParameterInfo.Name == SpecularProfileAtlas::GetSpecularProfileParameterName(GetSpecularProfileRT(It)))
+			{
+				check(ParameterInfo.Association == EMaterialParameterAssociation::GlobalParameter);
+				OutValue = SpecularProfileAtlas::GetSpecularProfileId(GetSpecularProfileRT(It));
+				bResult = true;
+				break;
+			}
+		}
 	}
 	if (!bResult)
 	{
diff --git a/Engine/Source/Runtime/Engine/Private/Materials/MaterialInterface.cpp b/Engine/Source/Runtime/Engine/Private/Materials/MaterialInterface.cpp
--- a/Engine/Source/Runtime/Engine/Private/Materials/MaterialInterface.cpp
+++ b/Engine/Source/Runtime/Engine/Private/Materials/MaterialInterface.cpp
@@ -1312,9 +1312,14 @@ USubsurfaceProfile* UMaterialInterface::GetSubsurfaceProfile_Internal() const
 	return NULL;
 }
 
-USpecularProfile* UMaterialInterface::GetSpecularProfile_Internal() const
+USpecularProfile* UMaterialInterface::GetSpecularProfile_Internal(uint32 Index) const
 {
-	return NULL;
+	return nullptr;
+}
+
+uint32 UMaterialInterface::NumSpecularProfile_Internal() const
+{
+	return 0u;
 }
 
 bool UMaterialInterface::CastsRayTracedShadows() const
@@ -1384,37 +1389,48 @@ void UMaterialInterface::UpdateMaterialRenderProxy(FMaterialRenderProxy& Proxy)
 				const uint32 AllocationId = GSubsurfaceProfileTextureObject.AddOrUpdateProfile(Settings, LocalSubsurfaceProfile);
 				check(AllocationId >= 0 && AllocationId < MAX_SUBSURFACE_PROFILE_COUNT);
 			}
-			InProxy->SetSubsurfaceProfileRT(LocalSubsurfaceProfile);
+			InProxy->SetSubsurfaceProfileRT(LocalSubsurfaceProfile/*, ParameterName */); // how to have a unique identifier?
 		});
 	}
 
 	if (Strata::IsStrataEnabled())
 	{
-		USpecularProfile* LocalSpecularProfile = GetSpecularProfile_Internal();
-
-		FSpecularProfileStruct LocalSettings;
-
-		const FTextureReference* LocalTexture = nullptr;
-		if (LocalSpecularProfile)
+		struct FEntry
 		{
-			LocalSettings = LocalSpecularProfile->Settings;
-
-			if (!LocalSpecularProfile->Settings.IsProcedural())
+			USpecularProfile* Profile = nullptr;
+			FSpecularProfileStruct Settings;
+			const FTextureReference* Texture = nullptr;
+			FGuid Guid;
+		};
+		TArray<FEntry> Entries;
+		for (int32 It = 0, Count = NumSpecularProfile_Internal(); It<Count; ++It)
+		{
+			FEntry& Entry = Entries.AddDefaulted_GetRef();
+			Entry.Profile = GetSpecularProfile_Internal(It);
+			if (Entry.Profile)
 			{
-				LocalTexture = &LocalSpecularProfile->Settings.Texture->TextureReference;
+				Entry.Settings 	= Entry.Profile->Settings;
+				Entry.Guid 		= Entry.Profile->Guid;	
+				if (!Entry.Settings.IsProcedural())
+				{
+					Entry.Texture = &Entry.Profile->Settings.Texture->TextureReference;
+				}
 			}
 		}
 
 		FMaterialRenderProxy* InProxy = &Proxy;
 		ENQUEUE_RENDER_COMMAND(UpdateMaterialRenderProxySpecular)(
-		[LocalSettings, LocalSpecularProfile, InProxy, LocalTexture](FRHICommandListImmediate& RHICmdList)
+		[InProxy, Entries](FRHICommandListImmediate& RHICmdList)
 		{
-			if (LocalSpecularProfile)
+			for (const FEntry& Entry : Entries)
 			{
-				const uint32 AllocationId = SpecularProfileAtlas::AddOrUpdateProfile(LocalSpecularProfile, LocalSettings, LocalTexture);
-				check(AllocationId >= 0 && AllocationId < MAX_SPECULAR_PROFILE_COUNT);
+				if (Entry.Profile)
+				{
+					const uint32 AllocationId = SpecularProfileAtlas::AddOrUpdateProfile(Entry.Profile, Entry.Guid, Entry.Settings, Entry.Texture);
+					check(AllocationId >= 0 && AllocationId < MAX_SPECULAR_PROFILE_COUNT);
+				}
+				InProxy->AddSpecularProfileRT(Entry.Profile);
 			}
-			InProxy->SetSpecularProfileRT(LocalSpecularProfile);
 		});
 	}
 }
diff --git a/Engine/Source/Runtime/Engine/Private/Materials/MaterialRenderProxy.cpp b/Engine/Source/Runtime/Engine/Private/Materials/MaterialRenderProxy.cpp
--- a/Engine/Source/Runtime/Engine/Private/Materials/MaterialRenderProxy.cpp
+++ b/Engine/Source/Runtime/Engine/Private/Materials/MaterialRenderProxy.cpp
@@ -527,7 +527,6 @@ void FMaterialRenderProxy::UpdateUniformExpressionCacheIfNeeded(ERHIFeatureLevel
 
 FMaterialRenderProxy::FMaterialRenderProxy(FString InMaterialName)
 	: SubsurfaceProfileRT(nullptr)
-	, SpecularProfileRT(nullptr)
 	, MaterialName(MoveTemp(InMaterialName))
 	, MarkedForGarbageCollection(0)
 	, DeletedFlag(0)
diff --git a/Engine/Source/Runtime/Engine/Private/Rendering/SpecularProfile.cpp b/Engine/Source/Runtime/Engine/Private/Rendering/SpecularProfile.cpp
--- a/Engine/Source/Runtime/Engine/Private/Rendering/SpecularProfile.cpp
+++ b/Engine/Source/Runtime/Engine/Private/Rendering/SpecularProfile.cpp
@@ -62,7 +62,7 @@ public:
 		SHADER_PARAMETER_RDG_TEXTURE_UAV(RWTexture2D, TargetTexture)
 	END_SHADER_PARAMETER_STRUCT()
 
-		static bool ShouldCompilePermutation(const FGlobalShaderPermutationParameters& Parameters)
+	static bool ShouldCompilePermutation(const FGlobalShaderPermutationParameters& Parameters)
 	{
 		return true;
 	}
@@ -85,41 +85,28 @@ public:
 
 	// convenience, can be optimized 
 	// @param Profile must not be 0, game thread pointer, do not dereference, only for comparison
-	int32 AddOrUpdateProfile(const USpecularProfile* InProfile, const FSpecularProfileStruct InSettings, const FTextureReference* InTexture)
-	{
-		check(InProfile);
-
-		int32 AllocationId = FindAllocationId(InProfile); 
-		
-		if (AllocationId != -1)
-		{
-			UpdateProfile(AllocationId, InSettings, InTexture);
-		}
-		else
-		{
-			AllocationId = AddProfile(InSettings, InProfile, InTexture);
-		}
-
-		return AllocationId;
-	}
+	int32 AddOrUpdateProfile(const USpecularProfile* InProfile, const FGuid& InGuid, const FSpecularProfileStruct InSettings, const FTextureReference* InTexture);
 
 	// O(n) n is a small number
 	// @param InProfile must not be 0, game thread pointer, do not dereference, only for comparison
-	// @return AllocationId -1: no allocation, should be deallocated with DeallocateSpecularProfile()
-	int32 AddProfile(const FSpecularProfileStruct Settings, const USpecularProfile* InProfile, const FTextureReference* InTexture);
+	// @return AllocationId INDEX_NONE: no allocation, should be deallocated with DeallocateSpecularProfile()
+	int32 AddProfile(const USpecularProfile* InProfile, const FGuid& InGuid, const FSpecularProfileStruct Settings, const FTextureReference* InTexture);
 
 	// O(n) to find the element, n is the SSProfile count and usually quite small
 	void RemoveProfile(const USpecularProfile* InProfile);
 
 	// @param InProfile must not be 0, game thread pointer, do not dereference, only for comparison
-	void UpdateProfile(const USpecularProfile* InProfile, const FSpecularProfileStruct InSettings, const FTextureReference* InTexture) { UpdateProfile(FindAllocationId(InProfile), InSettings, InTexture); }
+	void UpdateProfile(const USpecularProfile* InProfile, const FSpecularProfileStruct InSettings, const FTextureReference* InTexture);
 
 	// @param InProfile must not be 0, game thread pointer, do not dereference, only for comparison
 	void UpdateProfile(int32 AllocationId, const FSpecularProfileStruct Settings, const FTextureReference* InTexture);
 
+	// return the parameter name for a given profile
+	FName GetParameterName(const USpecularProfile* InProfile) const;
+
 	// @return can be nullptr if there is no SpecularProfile
-	struct IPooledRenderTarget* GetTexture(FRDGBuilder& GraphBuilder, EShaderPlatform ShaderPlatform);
-	struct IPooledRenderTarget* GetTexture();
+	struct IPooledRenderTarget* GetAtlasTexture(FRDGBuilder& GraphBuilder, EShaderPlatform ShaderPlatform);
+	struct IPooledRenderTarget* GetAtlasTexture();
 
 	//~ Begin FRenderResource Interface.
 	/**
@@ -128,7 +115,7 @@ public:
 	virtual void ReleaseDynamicRHI() override;
 
 	// @param InProfile must not be 0, game thread pointer, do not dereference, only for comparison
-	// @return -1 if not found
+	// @return INDEX_NONE if not found
 	int32 FindAllocationId(const USpecularProfile* InProfile) const;
 
 private:
@@ -139,6 +126,7 @@ private:
 		const USpecularProfile* Profile = nullptr; // Game thread pointer! Do not dereference, only for comparison.
 		const FTextureReference* Texture = nullptr;
 		FIntPoint CachedResolution = FIntPoint::ZeroValue;
+		FName ParameterName;
 	};
 	TArray<FProfileEntry> SpecularProfileEntries;
 };
@@ -149,21 +137,50 @@ TGlobalResource<FSpecularProfileTextureManager> GSpecularProfileTextureManager;
 // SpecularProfile atlas storing several texture profiles or 0 if there is no user
 static TRefCountPtr<IPooledRenderTarget> GSpecularProfileTextureAtlas;
 
+static FName CreateSpecularProfileParameterName(const FGuid& InGuid)
+{
+	return FName(TEXT("__SpecularProfile") + InGuid.ToString());
+}
+
+FName CreateSpecularProfileParameterName(USpecularProfile* InProfile)
+{
+	return InProfile ? CreateSpecularProfileParameterName(InProfile->Guid) : FName();
+}
+
 FSpecularProfileTextureManager::FSpecularProfileTextureManager()
 {
 	check(IsInGameThread());
 
-	FSpecularProfileStruct DefaultProfile;
-
 	// add element 0, it is used as default profile
-	SpecularProfileEntries.Add( { DefaultProfile, nullptr, nullptr });
+	FProfileEntry& Entry = SpecularProfileEntries.AddDefaulted_GetRef(); 
+	Entry.Settings = FSpecularProfileStruct(); 
+	Entry.Profile = nullptr; 
+	Entry.Texture = nullptr; 
 }
 
 FSpecularProfileTextureManager::~FSpecularProfileTextureManager()
 {
 }
 
-int32 FSpecularProfileTextureManager::AddProfile(const FSpecularProfileStruct InSettings, const USpecularProfile* InProfile, const FTextureReference* InTexture)
+int32 FSpecularProfileTextureManager::AddOrUpdateProfile(const USpecularProfile* InProfile, const FGuid& Guid, const FSpecularProfileStruct InSettings, const FTextureReference* InTexture)
+{
+	check(InProfile);
+
+	int32 AllocationId = FindAllocationId(InProfile); 
+		
+	if (AllocationId != INDEX_NONE)
+	{
+		UpdateProfile(AllocationId, InSettings, InTexture);
+	}
+	else
+	{
+		AllocationId = AddProfile(InProfile, Guid, InSettings, InTexture);
+	}
+
+	return AllocationId;
+}
+
+int32 FSpecularProfileTextureManager::AddProfile(const USpecularProfile* InProfile, const FGuid& Guid, const FSpecularProfileStruct InSettings, const FTextureReference* InTexture)
 {
 	check(InProfile);
 	check(FindAllocationId(InProfile) == INDEX_NONE);
@@ -172,10 +189,13 @@ int32 FSpecularProfileTextureManager::AddProfile(const FSpecularProfileStruct In
 	{
 		for (int32 i = 1; i < SpecularProfileEntries.Num(); ++i)
 		{
-			if (SpecularProfileEntries[i].Profile == 0)
+			if (SpecularProfileEntries[i].Profile == nullptr)
 			{
 				RetAllocationId = i;
-				SpecularProfileEntries[RetAllocationId].Profile = InProfile;
+
+				FProfileEntry& Entry= SpecularProfileEntries[RetAllocationId];
+				Entry.Profile 		= InProfile;
+				Entry.ParameterName = FName(TEXT("__SpecularProfile") + Guid.ToString());
 				break;
 			}
 		}
@@ -183,7 +203,10 @@ int32 FSpecularProfileTextureManager::AddProfile(const FSpecularProfileStruct In
 		if (RetAllocationId == INDEX_NONE)
 		{
 			RetAllocationId = SpecularProfileEntries.Num();
-			SpecularProfileEntries.Add({ InSettings, InProfile, InTexture});
+
+			FProfileEntry& Entry= SpecularProfileEntries.AddDefaulted_GetRef(); 			
+			Entry.Profile 		= InProfile;
+			Entry.ParameterName = FName(TEXT("__SpecularProfile") + Guid.ToString());
 		}
 	}
 
@@ -203,12 +226,18 @@ void FSpecularProfileTextureManager::RemoveProfile(const USpecularProfile* InPro
 		check(SpecularProfileEntries[AllocationId].Profile == InProfile);
 
 		// make it available for reuse
-		SpecularProfileEntries[AllocationId].Profile = 0;
+		SpecularProfileEntries[AllocationId].Profile = nullptr;
 		SpecularProfileEntries[AllocationId].Settings.Invalidate();
 		SpecularProfileEntries[AllocationId].Texture = nullptr;
+		SpecularProfileEntries[AllocationId].ParameterName = FName();
 	}
 }
 
+void FSpecularProfileTextureManager::UpdateProfile(const USpecularProfile* InProfile, const FSpecularProfileStruct InSettings, const FTextureReference* InTexture) 
+{ 
+	UpdateProfile(FindAllocationId(InProfile), InSettings, InTexture); 
+}
+
 void FSpecularProfileTextureManager::UpdateProfile(int32 AllocationId, const FSpecularProfileStruct Settings, const FTextureReference* InTexture)
 {
 	check(IsInRenderingThread());
@@ -222,12 +251,23 @@ void FSpecularProfileTextureManager::UpdateProfile(int32 AllocationId, const FSp
 	}
 }
 
-IPooledRenderTarget* FSpecularProfileTextureManager::GetTexture()
+FName FSpecularProfileTextureManager::GetParameterName(const USpecularProfile* InProfile) const
+{
+	const int32 AllocationId = FindAllocationId(InProfile);
+	if (AllocationId != INDEX_NONE)
+	{
+		// make it available for reuse
+		return SpecularProfileEntries[AllocationId].ParameterName;
+	}
+	return FName();
+}
+
+IPooledRenderTarget* FSpecularProfileTextureManager::GetAtlasTexture()
 {
 	return GSpecularProfileTextureAtlas;
 }
 
-IPooledRenderTarget* FSpecularProfileTextureManager::GetTexture(FRDGBuilder& GraphBuilder, EShaderPlatform ShaderPlatform)
+IPooledRenderTarget* FSpecularProfileTextureManager::GetAtlasTexture(FRDGBuilder& GraphBuilder, EShaderPlatform ShaderPlatform)
 {
 	if (!Strata::IsStrataEnabled())
 	{
@@ -390,32 +430,26 @@ void USpecularProfile::PostEditChangeProperty(struct FPropertyChangedEvent& Prop
 
 namespace SpecularProfileAtlas
 {
-
-	FName GetSpecularProfileParameterName()
+	FName GetSpecularProfileParameterName(const USpecularProfile* In)
 	{
-		static FName NameSpecularProfile(TEXT("__SpecularProfile"));
-		return NameSpecularProfile;
+		return GSpecularProfileTextureManager.GetParameterName(In);
 	}
 
 	float GetSpecularProfileId(const USpecularProfile* In)
 	{
+		// No profile specified means we use the default one (constant one)
 		int32 AllocationId = 0;
 		if (In)
 		{
 			// can be optimized (cached)
 			AllocationId = GSpecularProfileTextureManager.FindAllocationId(In);
 		}
-		else
-		{
-			// no profile specified means we use the default one stored at [0] which is human skin
-			AllocationId = 0;
-		}
 		return AllocationId / 255.0f;
 	}
 	
-	int32 AddOrUpdateProfile(const USpecularProfile* InProfile, const FSpecularProfileStruct InSettings, const FTextureReference* InTexture)
+	int32 AddOrUpdateProfile(const USpecularProfile* InProfile, const FGuid& InGuid, const FSpecularProfileStruct InSettings, const FTextureReference* InTexture)
 	{
-		return GSpecularProfileTextureManager.AddOrUpdateProfile(InProfile, InSettings, InTexture);
+		return GSpecularProfileTextureManager.AddOrUpdateProfile(InProfile, InGuid, InSettings, InTexture);
 	}
 
 	FRHITexture* GetSpecularProfileTextureAtlas()
@@ -430,7 +464,7 @@ namespace SpecularProfileAtlas
 
 	void UpdateSpecularProfileTextureAtlas(FRDGBuilder& GraphBuilder, EShaderPlatform ShaderPlatform)
 	{
-		GSpecularProfileTextureManager.GetTexture(GraphBuilder, ShaderPlatform);
+		GSpecularProfileTextureManager.GetAtlasTexture(GraphBuilder, ShaderPlatform);
 	}
 
 } // namespace SpecularProfileAtlas
diff --git a/Engine/Source/Runtime/Engine/Public/Materials/MaterialRenderProxy.h b/Engine/Source/Runtime/Engine/Public/Materials/MaterialRenderProxy.h
--- a/Engine/Source/Runtime/Engine/Public/Materials/MaterialRenderProxy.h
+++ b/Engine/Source/Runtime/Engine/Public/Materials/MaterialRenderProxy.h
@@ -204,11 +204,14 @@ public:
 	}
 #endif
 
+	// Subsurface profiles
 	void SetSubsurfaceProfileRT(const USubsurfaceProfile* Ptr) { SubsurfaceProfileRT = Ptr; }
 	const USubsurfaceProfile* GetSubsurfaceProfileRT() const { return SubsurfaceProfileRT; }
 
-	void SetSpecularProfileRT(const USpecularProfile* Ptr) { SpecularProfileRT = Ptr; }
-	const USpecularProfile* GetSpecularProfileRT() const { return SpecularProfileRT; }
+	// Specular profiles
+	void AddSpecularProfileRT(const USpecularProfile* Ptr) { SpecularProfilesRT.Add(Ptr); }
+	const USpecularProfile* GetSpecularProfileRT(uint32 Index) const { check(Index<uint32(SpecularProfilesRT.Num())); return SpecularProfilesRT[Index]; }
+	const uint32 NumSpecularProfileRT() const { return SpecularProfilesRT.Num(); }
 
 	static void UpdateDeferredCachedUniformExpressions();
 
@@ -227,7 +230,7 @@ private:
 
 	/** 0 if not set, game thread pointer, do not dereference, only for comparison */
 	const USubsurfaceProfile* SubsurfaceProfileRT;
-	const USpecularProfile* SpecularProfileRT;
+	TArray<const USpecularProfile*> SpecularProfilesRT;
 	FString MaterialName;
 
 	/** Incremented each time UniformExpressionCache is modified */
