Fixed mismatched copy constructor/assignment operator in FComponentSpacePoseContext.

diff --git a/Engine/Source/Runtime/Engine/Classes/Animation/AnimNodeBase.h b/Engine/Source/Runtime/Engine/Classes/Animation/AnimNodeBase.h
--- a/Engine/Source/Runtime/Engine/Classes/Animation/AnimNodeBase.h
+++ b/Engine/Source/Runtime/Engine/Classes/Animation/AnimNodeBase.h
@@ -574,6 +574,9 @@ public:
 		PreviousNodeId = SourceContext.PreviousNodeId;
 	}
 
+	// Note: this copy assignment operator copies the whole object but the copy constructor only copies part of the object.
+	FComponentSpacePoseContext& operator=(const FComponentSpacePoseContext&) = default;
+
 	ENGINE_API void ResetToRefPose();
 
 	ENGINE_API bool ContainsNaN() const;
