Added experimental DeadBlending animation graph node.

diff --git a/Engine/Source/Editor/AnimGraph/Private/AnimGraphNode_DeadBlending.cpp b/Engine/Source/Editor/AnimGraph/Private/AnimGraphNode_DeadBlending.cpp
new file mode 100644
--- /dev/null
+++ b/Engine/Source/Editor/AnimGraph/Private/AnimGraphNode_DeadBlending.cpp
@@ -0,0 +1,33 @@
+// Copyright Epic Games, Inc. All Rights Reserved.
+
+#include "AnimGraphNode_DeadBlending.h"
+
+#define LOCTEXT_NAMESPACE "AnimGraphNode_DeadBlending"
+
+
+FLinearColor UAnimGraphNode_DeadBlending::GetNodeTitleColor() const
+{
+	return FLinearColor(0.0f, 0.1f, 0.2f);
+}
+
+FText UAnimGraphNode_DeadBlending::GetTooltipText() const
+{
+	return LOCTEXT("NodeToolTip", "Performs inertialization using the Dead Blending algorithm.");
+}
+
+FText UAnimGraphNode_DeadBlending::GetNodeTitle(ENodeTitleType::Type TitleType) const
+{
+	return LOCTEXT("NodeTitle", "Dead Blending");
+}
+
+FText UAnimGraphNode_DeadBlending::GetMenuCategory() const
+{
+	return LOCTEXT("NodeCategory", "Animation|Misc.");
+}
+
+void UAnimGraphNode_DeadBlending::GetInputLinkAttributes(FNodeAttributeArray& OutAttributes) const
+{
+	OutAttributes.Add(UE::Anim::IInertializationRequester::Attribute);
+}
+
+#undef LOCTEXT_NAMESPACE
diff --git a/Engine/Source/Editor/AnimGraph/Public/AnimGraphNode_DeadBlending.h b/Engine/Source/Editor/AnimGraph/Public/AnimGraphNode_DeadBlending.h
new file mode 100644
--- /dev/null
+++ b/Engine/Source/Editor/AnimGraph/Public/AnimGraphNode_DeadBlending.h
@@ -0,0 +1,25 @@
+// Copyright Epic Games, Inc. All Rights Reserved.
+
+#pragma once
+
+#include "CoreMinimal.h"
+#include "UObject/ObjectMacros.h"
+#include "AnimGraphNode_Base.h"
+#include "Animation/AnimNode_DeadBlending.h"
+#include "AnimGraphNode_DeadBlending.generated.h"
+
+
+UCLASS(Experimental, MinimalAPI)
+class UAnimGraphNode_DeadBlending : public UAnimGraphNode_Base
+{
+	GENERATED_BODY()
+
+	UPROPERTY(EditAnywhere, Category = Settings)
+	FAnimNode_DeadBlending Node;
+
+	virtual FLinearColor GetNodeTitleColor() const override;
+	virtual FText GetTooltipText() const override;
+	virtual FText GetNodeTitle(ENodeTitleType::Type TitleType) const override;
+	virtual FText GetMenuCategory() const override;
+	virtual void GetInputLinkAttributes(FNodeAttributeArray& OutAttributes) const override;
+};
diff --git a/Engine/Source/Runtime/AnimGraphRuntime/Private/AnimNodes/AnimNode_BlendListBase.cpp b/Engine/Source/Runtime/AnimGraphRuntime/Private/AnimNodes/AnimNode_BlendListBase.cpp
--- a/Engine/Source/Runtime/AnimGraphRuntime/Private/AnimNodes/AnimNode_BlendListBase.cpp
+++ b/Engine/Source/Runtime/AnimGraphRuntime/Private/AnimNodes/AnimNode_BlendListBase.cpp
@@ -125,7 +125,7 @@ void FAnimNode_BlendListBase::Update_AnyThread(const FAnimationUpdateContext& Co
 				UE::Anim::IInertializationRequester* InertializationRequester = Context.GetMessage<UE::Anim::IInertializationRequester>();
 				if (InertializationRequester)
 				{
-					InertializationRequester->RequestInertialization(CurrentBlendTimes[ChildIndex], CurrentBlendProfile);
+					InertializationRequester->RequestInertializationWithBlendMode(CurrentBlendTimes[ChildIndex], CurrentBlendProfile, true, GetBlendType(), GetCustomBlendCurve());
 					InertializationRequester->AddDebugRecord(*Context.AnimInstanceProxy, Context.GetCurrentNodeId());
 				}
 				else
diff --git a/Engine/Source/Runtime/Engine/Classes/Animation/AnimNode_DeadBlending.h b/Engine/Source/Runtime/Engine/Classes/Animation/AnimNode_DeadBlending.h
new file mode 100644
--- /dev/null
+++ b/Engine/Source/Runtime/Engine/Classes/Animation/AnimNode_DeadBlending.h
@@ -0,0 +1,210 @@
+// Copyright Epic Games, Inc. All Rights Reserved.
+
+#pragma once
+
+#include "AnimNode_Inertialization.h"
+#include "AnimNode_DeadBlending.generated.h"
+
+/**
+ * Dead Blending Node
+ *
+ * Dead blending is an alternative method of inertialization that extrapolates the animation being transitioned from 
+ * forward in time and then performs a normal cross-fade blend between this extrapolated animation and the new animation
+ * being transitioned to.
+ * 
+ * For more background see: https://theorangeduck.com/page/dead-blending
+ * 
+ * This node works by extrapolating forward the animation being transition from using the animation's velocities at
+ * the point of transition, with an exponential decay which reduces those velocities over time to avoid the pose 
+ * becoming invalid.
+ * 
+ * The rate of this decay is set automatically based on how much the velocities of the animation being transitioned 
+ * from are moving toward the pose of the animation being transitioned to. If they are moving in the wrong direction or
+ * too quickly they will have a larger decay rate, while if they are in the correct direction and moving slowly relative 
+ * to the difference they will have a smaller decay rate.
+ * 
+ * These decay rates can be controlled by the `ExtrapolationHalfLife`, `ExtrapolationHalfLifeMin` and 
+ * `ExtrapolationHalfLifeMax` parameters, which specify the approximate average, min, and max decay periods.
+ * More specifically they specify the "half-life" - or how it takes for the velocities to be decayed by half.
+ */
+USTRUCT(Experimental, BlueprintInternalUseOnly)
+struct ENGINE_API FAnimNode_DeadBlending : public FAnimNode_Base
+{
+	GENERATED_BODY()
+
+	// Input Pose
+	UPROPERTY(EditAnywhere, Category = Links)
+	FPoseLink Source;
+
+private:
+
+	// If the default blend settings should always be used rather than those coming from the inertialization request.
+	UPROPERTY(EditAnywhere, Category = Blending)
+	bool bAlwaysUseDefaultBlendSettings = false;
+
+	// The default blend duration to use when "Always Use Default Blend Settings" is set to true.
+	UPROPERTY(EditAnywhere, Category = Blending)
+	float DefaultBlendDuration = 0.25f;
+
+	// Default blend profile to use when no blend profile is supplied with the inertialization request.
+	UPROPERTY(EditAnywhere, Category = Blending, meta = (UseAsBlendProfile = true))
+	TObjectPtr<UBlendProfile> DefaultBlendProfile = nullptr;
+
+	// Default blend mode to use when no blend mode is supplied with the inertialization request.
+	UPROPERTY(EditAnywhere, Category = Blending, DisplayName = "Default Blend Mode")
+	EAlphaBlendOption DefaultBlendMode = EAlphaBlendOption::Linear;
+
+	// Default custom blend curve to use along with the default blend mode.
+	UPROPERTY(EditAnywhere, Category = Blending)
+	TObjectPtr<UCurveFloat> DefaultCustomBlendCurve = nullptr;
+
+	// Multiplier that can be used to scale the overall blend durations coming from inertialization requests.
+	UPROPERTY(EditAnywhere, Category = Blending, meta = (Min = "0.0", UIMin = "0.0"))
+	float BlendTimeMultiplier = 1.0f;
+
+	// List of curves that should not use inertial blending. These curves will change instantly when the animation switches.
+	UPROPERTY(EditAnywhere, Category = Filter)
+	TArray<FName> FilteredCurves;
+
+	/**
+	 * The average half-life of decay in seconds to use when extrapolating the animation. To get the final half-life of 
+	 * decay, this value will be scaled by the amount to which the velocities of the animation being transitioned from 
+	 * are moving toward the animation being transitioned to.
+	 */
+	UPROPERTY(EditAnywhere, Category = Extrapolation, meta = (Min = "0.0", UIMin = "0.0"), DisplayName = "Extrapolation Half Life")
+	float ExtrapolationHalfLife = 1.0f;
+
+	/**
+	 * The minimum half-life of decay in seconds to use when extrapolating the animation. This will be used when the 
+	 * velocities of the animation being transitioned from are very small or moving away from the animation being 
+	 * transitioned to.
+	 */
+	UPROPERTY(EditAnywhere, Category = Extrapolation, meta = (Min = "0.0", UIMin = "0.0"), DisplayName = "Minimum Extrapolation Half Life")
+	float ExtrapolationHalfLifeMin = 0.05f;
+
+	/**
+	 * The maximum half-life of decay in seconds to use when extrapolating the animation. This will dictate the longest 
+	 * decay duration possible when velocities of the animation being transitioned from are small and moving towards the 
+	 * animation being transitioned to.
+	 */
+	UPROPERTY(EditAnywhere, Category = Extrapolation, meta = (Min = "0.0", UIMin = "0.0"), DisplayName = "Maximum Extrapolation Half Life")
+	float ExtrapolationHalfLifeMax = 1.0f;
+
+#if WITH_EDITORONLY_DATA
+	
+	// This setting can be used to show what the extrapolation of the animation looks like.
+	UPROPERTY(EditAnywhere, Category = Debug)
+	bool bShowExtrapolations = false;
+
+#endif // WITH_EDITORONLY_DATA
+
+public: // FAnimNode_DeadBlending
+
+	FAnimNode_DeadBlending();
+
+	/**
+	 * Request to activate inertialization. If multiple requests are made on the same inertialization node, the minimum 
+	 * requested time will be used. If multiple blend modes are given, the blend mode of the final request will be used.
+	 */
+	virtual void RequestInertialization(
+		float Duration,
+		const UBlendProfile* BlendProfile,
+		const bool bInUseBlendMode,
+		const EAlphaBlendOption InBlendMode,
+		UCurveFloat* InCustomBlendCurve);
+
+public: // FAnimNode_Base
+
+	virtual void Initialize_AnyThread(const FAnimationInitializeContext& Context) override;
+	virtual void CacheBones_AnyThread(const FAnimationCacheBonesContext& Context) override;
+	virtual void Update_AnyThread(const FAnimationUpdateContext& Context) override;
+	virtual void Evaluate_AnyThread(FPoseContext& Output) override;
+
+	virtual bool NeedsDynamicReset() const override;
+
+private:
+	
+	/**
+	 * Records the pose and velocity of the animation being transitioned from, and computes the extrapolation half-lives.
+	 * 
+	 * @param InPose		The current pose for the animation being transitioned to.
+	 * @param InCurves		The current curves for the animation being transitioned to.
+	 * @param SrcPosePrev	The pose recorded as output of the inertializer two frames ago.
+	 * @param SrcPoseCurr	The pose recorded as output of the inertializer on the previous frame.
+	 */
+	void InitFrom(
+		const FCompactPose& InPose,
+		const FBlendedCurve& InCurves,
+		const FInertializationPose& SrcPosePrev,
+		const FInertializationPose& SrcPoseCurr);
+
+	/**
+	 * Computes the extrapolated pose and blends it with the input pose.
+	 * 
+	 * @param InOutPose		The current pose to blend with the extrapolated pose.
+	 * @param InOutCurves	The current curves to blend with the extrapolated curves.
+	 */
+	void ApplyTo(FCompactPose& InOutPose, FBlendedCurve& InOutCurves);
+
+private:
+
+	// Set of UIDs for valid curves in FilteredCurves.
+	TSet<SmartName::UID_Type> CachedFilteredCurvesUIDs;
+
+	// Snapshots of the actor pose generated as output.
+	TArray<FInertializationPose, TInlineAllocator<2>> PoseSnapshots;
+
+	// Pending inertialization requests.
+	UPROPERTY(Transient)
+	TArray<FInertializationRequest> RequestQueue;
+
+private:
+
+	// Recorded pose state at point of transition.
+	TBitArray<TInlineAllocator<512>> BoneValid;
+	TArray<FVector> BoneTranslations;
+	TArray<FQuat> BoneRotations;
+	TArray<FQuat> BoneRotationDirections;
+	TArray<FVector> BoneScales;
+
+	TArray<FVector> BoneTranslationVelocities;
+	TArray<FVector> BoneRotationVelocities;
+	TArray<FVector> BoneScaleVelocities;
+
+	TArray<FVector> BoneTranslationDecayHalfLives;
+	TArray<FVector> BoneRotationDecayHalfLives;
+	TArray<FVector> BoneScaleDecayHalfLives;
+
+	// Recorded curve state at the point of transition.
+	TBitArray<TInlineAllocator<256>> CurveValid;
+	TArray<float, TInlineAllocator<32>> CurveValues;
+	TArray<float, TInlineAllocator<32>> CurveVelocities;
+	TArray<float, TInlineAllocator<32>> CurveDecayHalfLives;
+
+private:
+
+	// Variable used to store the elapsed delta time between calls to evaluate.
+	float DeltaTime = 0.0f;
+
+	// Current Inertialization state.
+	EInertializationState InertializationState = EInertializationState::Inactive;
+
+	// Time since the last request for inertialization.
+	float InertializationTime = 0.0f;
+
+	// Current inertialization duration (used for curves).
+	float InertializationDuration = 0.0f;
+
+	// Current inertialization durations for each bone, indexed by skeleton bone index (used for per-bone blending).
+	TCustomBoneIndexArray<float, FSkeletonPoseBoneIndex> InertializationDurationPerBone;
+
+	// Maximum of InertializationDuration and all entries in InertializationDurationPerBone (used for knowing when to shutdown the inertialization).
+	float InertializationMaxDuration = 0.0f;
+
+	// Current blend mode being used for Inertialization.
+	EAlphaBlendOption InertializationBlendMode = EAlphaBlendOption::Linear;
+
+	// Custom blend curve being used by the current blend mode.
+	UPROPERTY(Transient)
+	TObjectPtr<UCurveFloat> InertializationCustomBlendCurve = nullptr;
+};
diff --git a/Engine/Source/Runtime/Engine/Classes/Animation/AnimNode_Inertialization.h b/Engine/Source/Runtime/Engine/Classes/Animation/AnimNode_Inertialization.h
--- a/Engine/Source/Runtime/Engine/Classes/Animation/AnimNode_Inertialization.h
+++ b/Engine/Source/Runtime/Engine/Classes/Animation/AnimNode_Inertialization.h
@@ -7,6 +7,8 @@
 #include "Animation/AnimNodeBase.h"
 #include "Animation/AnimCurveTypes.h"
 #include "Animation/AnimNodeMessages.h"
+#include "AlphaBlend.h" // Required for EAlphaBlendOption
+
 #include "AnimNode_Inertialization.generated.h"
 
 
@@ -32,6 +34,22 @@ public:
 	// If multiple requests are made on the same inertialization node, the minimum requested time will be used.
 	virtual void RequestInertialization(float InRequestedDuration, const UBlendProfile* InBlendProfile = nullptr) = 0;
 
+	// Request to activate inertialization for a duration with the given blend mode.
+	// If multiple requests are made on the same inertialization node, the minimum requested time will be used.
+	// 
+	// InRequestedDuration				Duration of the blend
+	// InBlendProfile					Blend profile to use
+	// bUseBlendMode					If to use the provided Blend Mode
+	// InBlendMode						Blend mode to use
+	// InCustomBlendCurve				Custom Blend Curve to use for the Blend Mode
+	//
+	virtual void RequestInertializationWithBlendMode(
+		float InRequestedDuration,
+		const UBlendProfile* InBlendProfile,
+		const bool bUseBlendMode,
+		const EAlphaBlendOption InBlendMode,
+		UCurveFloat* InCustomBlendCurve) = 0;
+
 	// Add a record of this request
 	virtual void AddDebugRecord(const FAnimInstanceProxy& InSourceProxy, int32 InSourceNodeId) = 0;
 };
@@ -107,27 +125,46 @@ struct FInertializationRequest
 {
 	GENERATED_BODY()
 
-	FInertializationRequest()
-		: Duration(-1.0f)
-		, BlendProfile(nullptr)
+	FInertializationRequest() {}
+
+	FInertializationRequest(float InDuration, const UBlendProfile* InBlendProfile)
+		: Duration(InDuration)
+		, BlendProfile(InBlendProfile)
 	{
 	}
 
-	FInertializationRequest(float InDuration, const UBlendProfile* InBlendProfile)
+	FInertializationRequest(
+		float InDuration, 
+		const UBlendProfile* InBlendProfile,
+		const bool bInUseBlendMode,
+		const EAlphaBlendOption InBlendMode, 
+		UCurveFloat* InCustomBlendCurve)
 		: Duration(InDuration)
 		, BlendProfile(InBlendProfile)
+		, bUseBlendMode(bInUseBlendMode)
+		, BlendMode(InBlendMode)
+		, CustomBlendCurve(InCustomBlendCurve)
 	{
 	}
 
+
 	void Clear()
 	{
 		Duration = -1.0f;
 		BlendProfile = nullptr;
+		bUseBlendMode = false;
+		BlendMode = EAlphaBlendOption::Linear;
+		CustomBlendCurve = nullptr;
 	}
 
 	friend bool operator==(const FInertializationRequest& A, const FInertializationRequest& B)
 	{
-		return A.Duration == B.Duration && A.BlendProfile == B.BlendProfile;
+		return 
+			(A.Duration == B.Duration) && 
+			(A.BlendProfile == B.BlendProfile) &&
+			(A.bUseBlendMode == B.bUseBlendMode) &&
+			(A.BlendMode == B.BlendMode) &&
+			(A.CustomBlendCurve == B.CustomBlendCurve);
 	}
 
 	friend bool operator!=(const FInertializationRequest& A, const FInertializationRequest& B)
@@ -136,10 +173,19 @@ struct FInertializationRequest
 	}
 
 	UPROPERTY(Transient)
-	float Duration;
+	float Duration = -1.0f;
+
+	UPROPERTY(Transient)
+	TObjectPtr<const UBlendProfile> BlendProfile = nullptr;
+
+	UPROPERTY(Transient)
+	bool bUseBlendMode = false;
+
+	UPROPERTY(Transient)
+	EAlphaBlendOption BlendMode = EAlphaBlendOption::Linear;
 
 	UPROPERTY(Transient)
-	TObjectPtr<const UBlendProfile> BlendProfile;
+	TObjectPtr<UCurveFloat> CustomBlendCurve = nullptr;
 };
 
 
diff --git a/Engine/Source/Runtime/Engine/Private/Animation/AnimNode_DeadBlending.cpp b/Engine/Source/Runtime/Engine/Private/Animation/AnimNode_DeadBlending.cpp
new file mode 100644
--- /dev/null
+++ b/Engine/Source/Runtime/Engine/Private/Animation/AnimNode_DeadBlending.cpp
@@ -0,0 +1,852 @@
+// Copyright Epic Games, Inc. All Rights Reserved.
+
+#include "Animation/AnimNode_DeadBlending.h"
+#include "Animation/AnimInstanceProxy.h"
+#include "Animation/AnimNode_SaveCachedPose.h"
+#include "Animation/BlendProfile.h"
+#include "Algo/MaxElement.h"
+#include "Components/SkeletalMeshComponent.h"
+#include "Animation/AnimBlueprintGeneratedClass.h"
+#include "Logging/TokenizedMessage.h"
+
+#include UE_INLINE_GENERATED_CPP_BY_NAME(AnimNode_DeadBlending)
+
+LLM_DEFINE_TAG(Animation_DeadBlending);
+
+#define LOCTEXT_NAMESPACE "AnimNode_DeadBlending"
+
+namespace UE::Anim {
+
+	// Inertialization request event bound to a node
+	class FDeadBlendingRequester : public IInertializationRequester
+	{
+	public:
+		FDeadBlendingRequester(const FAnimationBaseContext& InContext, FAnimNode_DeadBlending* InNode)
+			: Node(*InNode)
+			, NodeId(InContext.GetCurrentNodeId())
+			, Proxy(*InContext.AnimInstanceProxy)
+		{}
+
+	private:
+		// IInertializationRequester interface
+		virtual void RequestInertialization(
+			float InRequestedDuration,
+			const UBlendProfile* InBlendProfile) override
+		{
+			Node.RequestInertialization(InRequestedDuration, InBlendProfile, false, EAlphaBlendOption::Linear, nullptr);
+		}
+
+		virtual void RequestInertializationWithBlendMode(
+			float InRequestedDuration,
+			const UBlendProfile* InBlendProfile,
+			const bool bUseBlendMode,
+			const EAlphaBlendOption InBlendMode,
+			UCurveFloat* InCustomBlendCurve)
+		{
+			Node.RequestInertialization(InRequestedDuration, InBlendProfile, bUseBlendMode, InBlendMode, InCustomBlendCurve);
+		}
+
+		virtual void AddDebugRecord(const FAnimInstanceProxy& InSourceProxy, int32 InSourceNodeId)
+		{
+#if WITH_EDITORONLY_DATA
+			Proxy.RecordNodeAttribute(InSourceProxy, NodeId, InSourceNodeId, IInertializationRequester::Attribute);
+#endif
+			TRACE_ANIM_NODE_ATTRIBUTE(Proxy, InSourceProxy, NodeId, InSourceNodeId, IInertializationRequester::Attribute);
+		}
+
+		// Node to target
+		FAnimNode_DeadBlending& Node;
+
+		// Node index
+		int32 NodeId;
+
+		// Proxy currently executing
+		FAnimInstanceProxy& Proxy;
+	};
+
+}	// namespace UE::Anim
+
+namespace UE::Anim::DeadBlending::Private
+{
+	static constexpr int32 MaxPoseSnapShotNum = 2;
+
+	static constexpr float Ln2 = 0.69314718056f;
+
+	static int32 GetNumSkeletonBones(const FBoneContainer& BoneContainer)
+	{
+		const USkeleton* SkeletonAsset = BoneContainer.GetSkeletonAsset();
+		check(SkeletonAsset);
+
+		const FReferenceSkeleton& RefSkeleton = SkeletonAsset->GetReferenceSkeleton();
+		return RefSkeleton.GetNum();
+	}
+
+	static inline FVector VectorDivMax(const float V, const FVector W, const float Epsilon = UE_SMALL_NUMBER)
+	{
+		return FVector(
+			V / FMath::Max(W.X, Epsilon),
+			V / FMath::Max(W.Y, Epsilon),
+			V / FMath::Max(W.Z, Epsilon));
+	}
+
+	static inline FVector VectorDivMax(const FVector V, const FVector W, const float Epsilon = UE_SMALL_NUMBER)
+	{
+		return FVector(
+			V.X / FMath::Max(W.X, Epsilon),
+			V.Y / FMath::Max(W.Y, Epsilon),
+			V.Z / FMath::Max(W.Z, Epsilon));
+	}
+
+	static inline FVector VectorInvExpApprox(const FVector V)
+	{
+		return FVector(
+			FMath::InvExpApprox(V.X),
+			FMath::InvExpApprox(V.Y),
+			FMath::InvExpApprox(V.Z));
+	}
+
+	static inline FVector VectorEerp(const FVector V, const FVector W, const float Alpha)
+	{
+		return FVector(
+			FMath::Pow(V.X, (1.0f - Alpha)) * FMath::Pow(W.X, Alpha),
+			FMath::Pow(V.Y, (1.0f - Alpha)) * FMath::Pow(W.Y, Alpha),
+			FMath::Pow(V.Z, (1.0f - Alpha)) * FMath::Pow(W.Z, Alpha));
+	}
+
+	static inline FVector VectorExp(const FVector V)
+	{
+		return FVector(
+			FMath::Exp(V.X),
+			FMath::Exp(V.Y),
+			FMath::Exp(V.Z));
+	}
+
+	static inline FVector VectorLogSafe(const FVector V, const float Epsilon = UE_SMALL_NUMBER)
+	{
+		return FVector(
+			FMath::Loge(FMath::Max(V.X, Epsilon)),
+			FMath::Loge(FMath::Max(V.Y, Epsilon)),
+			FMath::Loge(FMath::Max(V.Z, Epsilon)));
+	}
+
+	static inline FVector ExtrapolateTranslation(
+		const FVector Translation,
+		const FVector Velocity,
+		const float Time,
+		const FVector DecayHalflife,
+		const float Epsilon = UE_SMALL_NUMBER)
+	{
+		const FVector C = VectorDivMax(Ln2, DecayHalflife, Epsilon);
+		return Translation + VectorDivMax(Velocity, C, Epsilon) * (FVector::OneVector - VectorInvExpApprox(C * Time));
+	}
+
+	static inline FQuat ExtrapolateRotation(
+		const FQuat Rotation,
+		const FVector Velocity,
+		const float Time,
+		const FVector DecayHalflife,
+		const float Epsilon = UE_SMALL_NUMBER)
+	{
+		const FVector C = VectorDivMax(Ln2, DecayHalflife, Epsilon);
+		return FQuat::MakeFromRotationVector(VectorDivMax(Velocity, C, Epsilon) * (FVector::OneVector - VectorInvExpApprox(C * Time))) * Rotation;
+	}
+
+	static inline FVector ExtrapolateScale(
+		const FVector Scale,
+		const FVector Velocity,
+		const float Time,
+		const FVector DecayHalflife,
+		const float Epsilon = UE_SMALL_NUMBER)
+	{
+		const FVector C = VectorDivMax(Ln2, DecayHalflife, Epsilon);
+		return VectorExp(VectorDivMax(Velocity, C, Epsilon) * (FVector::OneVector - VectorInvExpApprox(C * Time))) * Scale;
+	}
+
+	static inline float ExtrapolateCurve(
+		const float Curve,
+		const float Velocity,
+		const float Time,
+		const float DecayHalflife,
+		const float Epsilon = UE_SMALL_NUMBER)
+	{
+		const float C = Ln2 / FMath::Max(DecayHalflife, Epsilon);
+		return Curve + FMath::Max(Velocity / C, Epsilon) * (1.0f - FMath::InvExpApprox(C * Time));
+	}
+
+	static inline float ClipMagnitudeToGreaterThanEpsilon(const float X, const float Epsilon = UE_KINDA_SMALL_NUMBER)
+	{
+		return
+			X >= 0.0f && X <  Epsilon ?  Epsilon :
+			X <  0.0f && X > -Epsilon ? -Epsilon : X;
+	}
+
+	static inline float ComputeDecayHalfLifeFromDiffAndVelocity(
+		const float SrcDstDiff,
+		const float SrcVelocity,
+		const float HalfLife,
+		const float HalfLifeMin,
+		const float HalfLifeMax,
+		const float Epsilon = UE_KINDA_SMALL_NUMBER)
+	{
+		// Essentially what this function does is compute a half-life based on the ratio between the velocity vector and
+		// the vector from the source to the destination. This is then clamped to some min and max. If the signs are
+		// different (i.e. the velocity and the vector from source to destination are in opposite directions) this will
+		// produce a negative number that will get clamped to HalfLifeMin. If the signs match, this will produce a large
+		// number when the velocity is small and the vector from source to destination is large, and a small number when
+		// the velocity is large and the vector from source to destination is small. This will be clamped either way to 
+		// be in the range given by HalfLifeMin and HalfLifeMax. Finally, since the velocity can be close to zero we 
+		// have to clamp it to always be greater than some given magnitude (preserving the sign).
+
+		return FMath::Clamp(HalfLife * (SrcDstDiff / ClipMagnitudeToGreaterThanEpsilon(SrcVelocity, Epsilon)), HalfLifeMin, HalfLifeMax);
+	}
+
+	static inline FVector ComputeDecayHalfLifeFromDiffAndVelocity(
+		const FVector SrcDstDiff,
+		const FVector SrcVelocity,
+		const float HalfLife,
+		const float HalfLifeMin,
+		const float HalfLifeMax,
+		const float Epsilon = UE_KINDA_SMALL_NUMBER)
+	{
+		return FVector(
+			ComputeDecayHalfLifeFromDiffAndVelocity(SrcDstDiff.X, SrcVelocity.X, HalfLife, HalfLifeMin, HalfLifeMax, Epsilon),
+			ComputeDecayHalfLifeFromDiffAndVelocity(SrcDstDiff.Y, SrcVelocity.Y, HalfLife, HalfLifeMin, HalfLifeMax, Epsilon),
+			ComputeDecayHalfLifeFromDiffAndVelocity(SrcDstDiff.Z, SrcVelocity.Z, HalfLife, HalfLifeMin, HalfLifeMax, Epsilon));
+	}
+}
+
+void FAnimNode_DeadBlending::InitFrom(
+	const FCompactPose& InPose,
+	const FBlendedCurve& InCurves,
+	const FInertializationPose& SrcPosePrev,
+	const FInertializationPose& SrcPoseCurr)
+{
+	TRACE_CPUPROFILER_EVENT_SCOPE(FAnimNode_DeadBlending::InitFrom);
+
+	const FBoneContainer& BoneContainer = InPose.GetBoneContainer();
+
+	const int32 NumSkeletonBones = UE::Anim::DeadBlending::Private::GetNumSkeletonBones(BoneContainer);
+
+	BoneValid.Init(false, NumSkeletonBones);
+	BoneTranslations.Init(FVector::ZeroVector, NumSkeletonBones);
+	BoneRotations.Init(FQuat::Identity, NumSkeletonBones);
+	BoneRotationDirections.Init(FQuat::Identity, NumSkeletonBones);
+	BoneScales.Init(FVector::OneVector, NumSkeletonBones);
+
+	BoneTranslationVelocities.Init(FVector::ZeroVector, NumSkeletonBones);
+	BoneRotationVelocities.Init(FVector::ZeroVector, NumSkeletonBones);
+	BoneScaleVelocities.Init(FVector::ZeroVector, NumSkeletonBones);
+
+	BoneTranslationDecayHalfLives.Init(ExtrapolationHalfLifeMin * FVector::OneVector, NumSkeletonBones);
+	BoneRotationDecayHalfLives.Init(ExtrapolationHalfLifeMin * FVector::OneVector, NumSkeletonBones);
+	BoneScaleDecayHalfLives.Init(ExtrapolationHalfLifeMin * FVector::OneVector, NumSkeletonBones);
+
+	// Record bone state
+
+	for (FCompactPoseBoneIndex BoneIndex : InPose.ForEachBoneIndex())
+	{
+		const int32 SkeletonPoseBoneIndex = BoneContainer.GetSkeletonIndex(BoneIndex);
+
+		if (SkeletonPoseBoneIndex == INDEX_NONE ||
+			SrcPosePrev.BoneStates[SkeletonPoseBoneIndex] != EInertializationBoneState::Valid ||
+			SrcPoseCurr.BoneStates[SkeletonPoseBoneIndex] != EInertializationBoneState::Valid)
+		{
+			continue;
+		}
+
+		// Mark bone as valid
+
+		BoneValid[SkeletonPoseBoneIndex] = true;
+
+		// Get Source Animation Transform
+
+		const FTransform SrcTransformCurr = SrcPoseCurr.BoneTransforms[SkeletonPoseBoneIndex];
+
+		BoneTranslations[SkeletonPoseBoneIndex] = SrcTransformCurr.GetTranslation();
+		BoneRotations[SkeletonPoseBoneIndex] = SrcTransformCurr.GetRotation();
+		BoneScales[SkeletonPoseBoneIndex] = SrcTransformCurr.GetScale3D();
+
+		if (SrcPoseCurr.DeltaTime > UE_SMALL_NUMBER)
+		{
+			// Get Source Animation Velocity
+
+			const FTransform SrcTransformPrev = SrcPosePrev.BoneTransforms[SkeletonPoseBoneIndex];
+
+			const FVector TranslationDiff = SrcTransformCurr.GetTranslation() - SrcTransformPrev.GetTranslation();
+
+			FQuat RotationDiff = SrcTransformCurr.GetRotation() * SrcTransformPrev.GetRotation().Inverse();
+			RotationDiff.EnforceShortestArcWith(FQuat::Identity);
+
+			const FVector ScaleDiff = UE::Anim::DeadBlending::Private::VectorDivMax(SrcTransformCurr.GetScale3D(), SrcTransformPrev.GetScale3D());
+
+			BoneTranslationVelocities[SkeletonPoseBoneIndex] = TranslationDiff / SrcPoseCurr.DeltaTime;
+			BoneRotationVelocities[SkeletonPoseBoneIndex] = RotationDiff.ToRotationVector() / SrcPoseCurr.DeltaTime;
+			BoneScaleVelocities[SkeletonPoseBoneIndex] = UE::Anim::DeadBlending::Private::VectorLogSafe(ScaleDiff) / SrcPoseCurr.DeltaTime;
+
+			// Compute Decay HalfLives
+
+			const FTransform DstTransform = InPose[BoneIndex];
+
+			const FVector TranslationSrcDstDiff = DstTransform.GetTranslation() - SrcTransformCurr.GetTranslation();
+
+			FQuat RotationSrcDstDiff = DstTransform.GetRotation() * SrcTransformCurr.GetRotation().Inverse();
+			RotationSrcDstDiff.EnforceShortestArcWith(FQuat::Identity);
+
+			const FVector ScaleSrcDstDiff = UE::Anim::DeadBlending::Private::VectorDivMax(DstTransform.GetScale3D(), SrcTransformCurr.GetScale3D());
+
+			BoneTranslationDecayHalfLives[SkeletonPoseBoneIndex] = UE::Anim::DeadBlending::Private::ComputeDecayHalfLifeFromDiffAndVelocity(
+				TranslationSrcDstDiff,
+				BoneTranslationVelocities[SkeletonPoseBoneIndex],
+				ExtrapolationHalfLife,
+				ExtrapolationHalfLifeMin,
+				ExtrapolationHalfLifeMax);
+
+			BoneRotationDecayHalfLives[SkeletonPoseBoneIndex] = UE::Anim::DeadBlending::Private::ComputeDecayHalfLifeFromDiffAndVelocity(
+				RotationSrcDstDiff.ToRotationVector(),
+				BoneRotationVelocities[SkeletonPoseBoneIndex],
+				ExtrapolationHalfLife,
+				ExtrapolationHalfLifeMin,
+				ExtrapolationHalfLifeMax);
+
+			BoneScaleDecayHalfLives[SkeletonPoseBoneIndex] = UE::Anim::DeadBlending::Private::ComputeDecayHalfLifeFromDiffAndVelocity(
+				ScaleSrcDstDiff,
+				BoneScaleVelocities[SkeletonPoseBoneIndex],
+				ExtrapolationHalfLife,
+				ExtrapolationHalfLifeMin,
+				ExtrapolationHalfLifeMax);
+		}
+	}
+
+	const int32 CurveNum = InCurves.IsValid() ? InCurves.UIDToArrayIndexLUT->Num() : 0;
+
+	CurveValid.Init(false, CurveNum);
+	CurveValues.Init(0.0f, CurveNum);
+	CurveVelocities.Init(0.0f, CurveNum);
+	CurveDecayHalfLives.Init(ExtrapolationHalfLifeMin, CurveNum);
+
+	// Record curve state
+
+	for (int32 CurveUID = 0; CurveUID < CurveNum; CurveUID++)
+	{
+		// Check if the curve is in our filter set. Leave Curves zeroed if it is.
+		if (CachedFilteredCurvesUIDs.Contains((SmartName::UID_Type)CurveUID))
+		{
+			continue;
+		}
+
+		const int32 SrcCurveIdxCurr = SrcPoseCurr.Curves.BlendedCurve.GetArrayIndexByUID(CurveUID);
+		
+		if (SrcCurveIdxCurr == INDEX_NONE)
+		{
+			continue;
+		}
+
+		// Mark curve as valid
+
+		CurveValid[SrcCurveIdxCurr] = true;
+
+		// Get Curve Value
+
+		const float SrcCurveCurr = SrcPoseCurr.Curves.BlendedCurve.CurveWeights[SrcCurveIdxCurr];
+
+		CurveValues[SrcCurveIdxCurr] = SrcCurveCurr;
+
+		const int32 SrcCurveIdxPrev = SrcPosePrev.Curves.BlendedCurve.GetArrayIndexByUID(CurveUID);
+		const int32 DstCurveIdx = InCurves.GetArrayIndexByUID(CurveUID);
+
+		if (SrcPoseCurr.DeltaTime > UE_SMALL_NUMBER && SrcCurveIdxPrev != INDEX_NONE && DstCurveIdx != INDEX_NONE)
+		{
+			// Get Curve Velocity
+
+			const float SrcCurvePrev = SrcPosePrev.Curves.BlendedCurve.CurveWeights[SrcCurveIdxPrev];
+
+			CurveVelocities[SrcCurveIdxCurr] = (SrcCurveCurr - SrcCurvePrev) / SrcPoseCurr.DeltaTime;
+
+			// Compute Decay Halflife
+
+			CurveDecayHalfLives[SrcCurveIdxCurr] = UE::Anim::DeadBlending::Private::ComputeDecayHalfLifeFromDiffAndVelocity(
+				InCurves.CurveWeights[DstCurveIdx] - SrcCurveCurr,
+				CurveVelocities[SrcCurveIdxCurr],
+				ExtrapolationHalfLife,
+				ExtrapolationHalfLifeMin,
+				ExtrapolationHalfLifeMax);
+		}
+	}
+}
+
+void FAnimNode_DeadBlending::ApplyTo(FCompactPose& InOutPose, FBlendedCurve& InOutCurves)
+{
+	TRACE_CPUPROFILER_EVENT_SCOPE(FAnimNode_DeadBlending::ApplyTo);
+
+	const FBoneContainer& BoneContainer = InOutPose.GetBoneContainer();
+
+	for (FCompactPoseBoneIndex BoneIndex : InOutPose.ForEachBoneIndex())
+	{
+		const int32 SkeletonPoseBoneIndex = BoneContainer.GetSkeletonIndex(BoneIndex);
+
+		if (SkeletonPoseBoneIndex == INDEX_NONE || !BoneValid[SkeletonPoseBoneIndex])
+		{
+			continue;
+		}
+
+		// Compute Extrapolated Bone State
+
+		const FVector ExtrapolatedTranslation = UE::Anim::DeadBlending::Private::ExtrapolateTranslation(
+			BoneTranslations[SkeletonPoseBoneIndex],
+			BoneTranslationVelocities[SkeletonPoseBoneIndex],
+			InertializationTime,
+			BoneTranslationDecayHalfLives[SkeletonPoseBoneIndex]);
+
+		const FQuat ExtrapolatedRotation = UE::Anim::DeadBlending::Private::ExtrapolateRotation(
+			BoneRotations[SkeletonPoseBoneIndex],
+			BoneRotationVelocities[SkeletonPoseBoneIndex],
+			InertializationTime,
+			BoneRotationDecayHalfLives[SkeletonPoseBoneIndex]);
+
+		const FVector ExtrapolatedScale = UE::Anim::DeadBlending::Private::ExtrapolateScale(
+			BoneScales[SkeletonPoseBoneIndex],
+			BoneScaleVelocities[SkeletonPoseBoneIndex],
+			InertializationTime,
+			BoneScaleDecayHalfLives[SkeletonPoseBoneIndex]);
+
+#if WITH_EDITORONLY_DATA
+		if (bShowExtrapolations)
+		{
+			InOutPose[BoneIndex].SetTranslation(ExtrapolatedTranslation);
+			InOutPose[BoneIndex].SetRotation(ExtrapolatedRotation);
+			InOutPose[BoneIndex].SetScale3D(ExtrapolatedScale);
+			continue;
+		}
+#endif
+		// We need to enforce that the blend of the rotation doesn't suddenly "switch sides"
+		// given that the extrapolated rotation can become quite far from the destination
+		// animation. To do this we keep track of the blend "direction" and ensure that the
+		// delta we are applying to the destination animation always remains on the same
+		// side of this rotation.
+
+		FQuat RotationDiff = ExtrapolatedRotation * InOutPose[BoneIndex].GetRotation().Inverse();
+		RotationDiff.EnforceShortestArcWith(BoneRotationDirections[SkeletonPoseBoneIndex]);
+
+		// Update BoneRotationDirections to match our current path
+		BoneRotationDirections[SkeletonPoseBoneIndex] = RotationDiff;
+
+		// Compute Blend Alpha
+
+		const float Alpha = 1.0f - FAlphaBlend::AlphaToBlendOption(
+			InertializationTime / FMath::Max(InertializationDurationPerBone[SkeletonPoseBoneIndex], UE_SMALL_NUMBER),
+			InertializationBlendMode, InertializationCustomBlendCurve);
+
+		// Perform Blend
+
+		if (Alpha != 0.0f)
+		{
+			InOutPose[BoneIndex].SetTranslation(FMath::Lerp(InOutPose[BoneIndex].GetTranslation(), ExtrapolatedTranslation, Alpha));
+			InOutPose[BoneIndex].SetRotation(FQuat::MakeFromRotationVector(RotationDiff.ToRotationVector() * Alpha) * InOutPose[BoneIndex].GetRotation());
+
+			// Here we use `Eerp` rather than `Lerp` to interpolate scales correctly (see: https://theorangeduck.com/page/scalar-velocity).
+			// This is inconsistent with the rest of Unreal which (mostly) uses `Lerp` on scales. The decision to use `Eerp` here is partially 
+			// due to the fact we are also dealing properly with scalar velocities in this node, and partially to try and not to lock this node 
+			// into having the same incorrect behavior by default. We can add an option to interpolate scales with `Lerp` later down the line if 
+			// users want this, but we will not be able to change the default behavior again, so the decision is to opt for the most correct by 
+			// default this time even if it is somewhat less performant and inconsistent with other parts of Unreal.
+			InOutPose[BoneIndex].SetScale3D(UE::Anim::DeadBlending::Private::VectorEerp(InOutPose[BoneIndex].GetScale3D(), ExtrapolatedScale, Alpha));
+		}
+	}
+
+	if (InOutCurves.IsValid())
+	{
+		// Compute Blend Alpha
+
+		const float CurveAlpha = 1.0f - FAlphaBlend::AlphaToBlendOption(
+			InertializationTime / FMath::Max(InertializationDuration, UE_SMALL_NUMBER),
+			InertializationBlendMode, InertializationCustomBlendCurve);
+
+		const int32 CurveNum = FMath::Min(InOutCurves.UIDToArrayIndexLUT->Num(), CurveValues.Num());
+
+		for (int32 CurveUID = 0; CurveUID != CurveNum; ++CurveUID)
+		{
+			const int32 CurveIdx = InOutCurves.GetArrayIndexByUID(CurveUID);
+
+			if (CurveIdx == INDEX_NONE || !CurveValid[CurveIdx])
+			{
+				continue;
+			}
+
+			// Compute Extrapolated Curve Value
+
+			const float ExtrapolatedCurve = UE::Anim::DeadBlending::Private::ExtrapolateCurve(
+				CurveValues[CurveIdx],
+				CurveVelocities[CurveIdx],
+				InertializationTime,
+				CurveDecayHalfLives[CurveIdx]);
+
+			if (CurveAlpha != 0.0f)
+			{
+
+#if WITH_EDITORONLY_DATA
+				if (bShowExtrapolations)
+				{
+					InOutCurves.CurveWeights[CurveIdx] = ExtrapolatedCurve;
+					InOutCurves.ValidCurveWeights[CurveIdx] = true;
+					continue;
+				}
+#endif
+				InOutCurves.CurveWeights[CurveIdx] = FMath::Lerp(InOutCurves.CurveWeights[CurveIdx], ExtrapolatedCurve, CurveAlpha);
+				InOutCurves.ValidCurveWeights[CurveIdx] = true;
+			}
+		}
+	}
+}
+
+FAnimNode_DeadBlending::FAnimNode_DeadBlending() {}
+
+void FAnimNode_DeadBlending::RequestInertialization(
+	float Duration,
+	const UBlendProfile* BlendProfile,
+	const bool bInUseBlendMode,
+	const EAlphaBlendOption InBlendMode,
+	UCurveFloat* InCustomBlendCurve)
+{
+	if (Duration >= 0.0f)
+	{
+		RequestQueue.AddUnique(FInertializationRequest(Duration, BlendProfile, bInUseBlendMode, InBlendMode, InCustomBlendCurve));
+	}
+}
+
+void FAnimNode_DeadBlending::Initialize_AnyThread(const FAnimationInitializeContext& Context)
+{
+	LLM_SCOPE_BYNAME(TEXT("Animation/DeadBlending"));
+	DECLARE_SCOPE_HIERARCHICAL_COUNTER_ANIMNODE(Initialize_AnyThread);
+
+	FAnimNode_Base::Initialize_AnyThread(Context);
+	Source.Initialize(Context);
+
+	const USkeleton* Skeleton = Context.AnimInstanceProxy->GetSkeleton();
+	check(Skeleton);
+
+	CachedFilteredCurvesUIDs.Empty(FilteredCurves.Num());
+	for (const FName& CurveName : FilteredCurves)
+	{
+		SmartName::UID_Type NameUID = Skeleton->GetUIDByName(USkeleton::AnimCurveMappingName, CurveName);
+		if (NameUID != SmartName::MaxUID)
+		{
+			// Grab UIDs of filtered curves to avoid lookup later
+			CachedFilteredCurvesUIDs.Add(NameUID);
+		}
+	}
+
+	PoseSnapshots.Empty(UE::Anim::DeadBlending::Private::MaxPoseSnapShotNum);
+
+	RequestQueue.Reserve(8);
+
+	const int32 NumSkeletonBones = UE::Anim::DeadBlending::Private::GetNumSkeletonBones(Context.AnimInstanceProxy->GetRequiredBones());
+
+	BoneValid.Empty(NumSkeletonBones);
+	BoneTranslations.Empty(NumSkeletonBones);
+	BoneRotations.Empty(NumSkeletonBones);
+	BoneRotationDirections.Empty(NumSkeletonBones);
+	BoneScales.Empty(NumSkeletonBones);
+
+	BoneTranslationVelocities.Empty(NumSkeletonBones);
+	BoneRotationVelocities.Empty(NumSkeletonBones);
+	BoneScaleVelocities.Empty(NumSkeletonBones);
+
+	BoneTranslationDecayHalfLives.Empty(NumSkeletonBones);
+	BoneRotationDecayHalfLives.Empty(NumSkeletonBones);
+	BoneScaleDecayHalfLives.Empty(NumSkeletonBones);
+
+	CurveValid.Empty();
+	CurveValues.Empty();
+	CurveVelocities.Empty();
+	CurveDecayHalfLives.Empty();
+
+	DeltaTime = 0.0f;
+
+	InertializationState = EInertializationState::Inactive;
+	InertializationTime = 0.0f;
+
+	InertializationDuration = 0.0f;
+	InertializationDurationPerBone.Empty(NumSkeletonBones);
+	InertializationMaxDuration = 0.0f;
+
+	InertializationBlendMode = DefaultBlendMode;
+	InertializationCustomBlendCurve = DefaultCustomBlendCurve;
+}
+
+
+void FAnimNode_DeadBlending::CacheBones_AnyThread(const FAnimationCacheBonesContext& Context)
+{
+	DECLARE_SCOPE_HIERARCHICAL_COUNTER_ANIMNODE(CacheBones_AnyThread);
+
+	FAnimNode_Base::CacheBones_AnyThread(Context);
+	Source.CacheBones(Context);
+}
+
+
+void FAnimNode_DeadBlending::Update_AnyThread(const FAnimationUpdateContext& Context)
+{
+	LLM_SCOPE_BYNAME(TEXT("Animation/DeadBlending"));
+	DECLARE_SCOPE_HIERARCHICAL_COUNTER_ANIMNODE(Update_AnyThread);
+
+	const int32 NodeId = Context.GetCurrentNodeId();
+	const FAnimInstanceProxy& Proxy = *Context.AnimInstanceProxy;
+
+	// Allow nodes further towards the leaves to inertialize using this node
+	UE::Anim::TScopedGraphMessage<UE::Anim::FDeadBlendingRequester> Inertialization(Context, Context, this);
+
+	// Handle skipped updates for cached poses by forwarding to inertialization nodes in those residual stacks
+	UE::Anim::TScopedGraphMessage<UE::Anim::FCachedPoseSkippedUpdateHandler> CachedPoseSkippedUpdate(Context, [this, NodeId, &Proxy](TArrayView<const UE::Anim::FMessageStack> InSkippedUpdates)
+	{
+		// If we have a pending request forward the request to other Inertialization nodes
+		// that were skipped due to pose caching.
+		if (RequestQueue.Num() > 0)
+		{
+			// Cached poses have their Update function called once even though there may be multiple UseCachedPose nodes for the same pose.
+			// Because of this, there may be Inertialization ancestors of the UseCachedPose nodes that missed out on requests.
+			// So here we forward 'this' node's requests to the ancestors of those skipped UseCachedPose nodes.
+			// Note that in some cases, we may be forwarding the requests back to this same node.  Those duplicate requests will ultimately
+			// be ignored by the 'AddUnique' in the body of FAnimNode_DeadBlending::RequestInertialization.
+			for (const UE::Anim::FMessageStack& Stack : InSkippedUpdates)
+			{
+				Stack.ForEachMessage<UE::Anim::IInertializationRequester>([this, NodeId, &Proxy](UE::Anim::IInertializationRequester& InMessage)
+				{
+					for (const FInertializationRequest& Request : RequestQueue)
+					{
+						InMessage.RequestInertializationWithBlendMode(Request.Duration, Request.BlendProfile, Request.bUseBlendMode, Request.BlendMode, Request.CustomBlendCurve);
+					}
+					InMessage.AddDebugRecord(Proxy, NodeId);
+
+					return UE::Anim::FMessageStack::EEnumerate::Stop;
+				});
+			}
+		}
+	});
+
+	Source.Update(Context);
+
+	// Accumulate delta time between calls to Evaluate_AnyThread
+	DeltaTime += Context.GetDeltaTime();
+}
+
+void FAnimNode_DeadBlending::Evaluate_AnyThread(FPoseContext& Output)
+{
+	LLM_SCOPE_BYNAME(TEXT("Animation/DeadBlending"));
+	DECLARE_SCOPE_HIERARCHICAL_COUNTER_ANIMNODE(Evaluate_AnyThread);
+
+	// Evaluate the Input and write it to the Output
+
+	Source.Evaluate(Output);
+
+	// Automatically detect teleports... note that we do the teleport distance check against the root bone's location (world space) rather
+	// than the mesh component's location because we still want to inertialize instances where the skeletal mesh component has been moved
+	// while simultaneously counter-moving the root bone (as is the case when mounting and dismounting vehicles for example)
+
+	const FTransform ComponentTransform = Output.AnimInstanceProxy->GetComponentTransform();
+
+	bool bTeleported = false;
+
+	const float TeleportDistanceThreshold = Output.AnimInstanceProxy->GetSkelMeshComponent()->GetTeleportDistanceThreshold();
+
+	if (PoseSnapshots.Num() > 0 && TeleportDistanceThreshold > 0.0f)
+	{
+		const FVector RootWorldSpaceLocation = ComponentTransform.TransformPosition(Output.Pose[FCompactPoseBoneIndex(0)].GetTranslation());
+		const FVector PrevRootWorldSpaceLocation = PoseSnapshots.Last().ComponentTransform.TransformPosition(PoseSnapshots.Last().BoneTransforms[0].GetTranslation());
+
+		if (FVector::DistSquared(RootWorldSpaceLocation, PrevRootWorldSpaceLocation) > FMath::Square(TeleportDistanceThreshold))
+		{
+			bTeleported = true;
+		}
+	}
+
+	// If teleported we simply reset the inertialization
+
+	if (bTeleported)
+	{
+		InertializationState = EInertializationState::Inactive;
+	}
+	
+	// If we don't have any Pose Snapshots recorded it means this is the first time this node has been evaluated in 
+	// which case there shouldn't be any discontinuity to remove, so no inertialization needs to be done, and we can 
+	// discard any requests.
+
+	if (PoseSnapshots.IsEmpty())
+	{
+		RequestQueue.Reset();
+	}
+
+	// Process Inertialization Requests
+
+	if (!RequestQueue.IsEmpty())
+	{
+		const int32 NumSkeletonBones = UE::Anim::DeadBlending::Private::GetNumSkeletonBones(Output.AnimInstanceProxy->GetRequiredBones());
+
+		InertializationTime = 0.0f;
+
+		if (bAlwaysUseDefaultBlendSettings)
+		{
+			InertializationDuration = BlendTimeMultiplier * DefaultBlendDuration;
+			InertializationDurationPerBone.Init(BlendTimeMultiplier * DefaultBlendDuration, NumSkeletonBones);
+			InertializationMaxDuration = BlendTimeMultiplier * DefaultBlendDuration;
+			InertializationBlendMode = DefaultBlendMode;
+			InertializationCustomBlendCurve = DefaultCustomBlendCurve;
+		}
+		else
+		{
+			// Process Request Durations by taking min of all requests
+			// For Blend Mode and Custom Curve we will just take whichever
+			// request is last in the stack (if they are provided).
+
+			InertializationDuration = UE_MAX_FLT;
+			InertializationDurationPerBone.Init(UE_MAX_FLT, NumSkeletonBones);
+			InertializationMaxDuration = UE_MAX_FLT;
+			InertializationBlendMode = DefaultBlendMode;
+			InertializationCustomBlendCurve = DefaultCustomBlendCurve;
+
+			UE::Anim::TTypedIndexArray<FSkeletonPoseBoneIndex, float, FAnimStackAllocator> RequestDurationPerBone;
+
+			for (const FInertializationRequest Request : RequestQueue)
+			{
+				// Duration is min of requests
+
+				InertializationDuration = BlendTimeMultiplier * FMath::Min(InertializationDuration, Request.Duration);
+
+				// Per-bone durations as min of requests accounting for blend profile
+
+				if (Request.BlendProfile)
+				{
+					Request.BlendProfile->FillSkeletonBoneDurationsArray(RequestDurationPerBone, Request.Duration);
+				}
+				else if (DefaultBlendProfile)
+				{
+					DefaultBlendProfile->FillSkeletonBoneDurationsArray(RequestDurationPerBone, Request.Duration);
+				}
+				else
+				{
+					RequestDurationPerBone.Init(Request.Duration, NumSkeletonBones);
+				}
+
+				check(RequestDurationPerBone.Num() == InertializationDurationPerBone.Num());
+
+				for (int32 BoneIndex = 0; BoneIndex < NumSkeletonBones; BoneIndex++)
+				{
+					InertializationDurationPerBone[BoneIndex] = BlendTimeMultiplier * FMath::Min(InertializationDurationPerBone[BoneIndex], RequestDurationPerBone[BoneIndex]);
+				}
+
+				// Process blend mode - will take the last one given
+
+				if (Request.bUseBlendMode)
+				{
+					InertializationBlendMode = Request.BlendMode;
+					InertializationCustomBlendCurve = Request.CustomBlendCurve;
+				}
+			}
+		}
+
+		InertializationMaxDuration = FMath::Max(InertializationDuration, *Algo::MaxElement(InertializationDurationPerBone));
+
+		check(InertializationDuration != UE_MAX_FLT);
+		check(InertializationMaxDuration != UE_MAX_FLT);
+
+		// Reset Request Queue
+
+		RequestQueue.Reset();
+
+		// Initialize the recorded pose state at the point of transition
+
+		if (PoseSnapshots.Num() > 1)
+		{
+			// We have two previous poses and so can initialize as normal.
+
+			InitFrom(
+				Output.Pose,
+				Output.Curve,
+				PoseSnapshots[PoseSnapshots.Num() - 2],
+				PoseSnapshots[PoseSnapshots.Num() - 1]);
+		}
+		else if (PoseSnapshots.Num() > 0)
+		{
+			// We only have a single previous pose. Repeat this pose assuming zero velocity.
+
+			InitFrom(
+				Output.Pose,
+				Output.Curve,
+				PoseSnapshots.Last(),
+				PoseSnapshots.Last());
+		}
+		else
+		{
+			// This should never happen because we are not able to issue an inertialization 
+			// requested until we have at least one pose recorded in the snapshots.
+			check(false);
+		}
+
+		// Set state to active
+
+		InertializationState = EInertializationState::Active;
+	}
+
+	// Update Time Since Transition and deactivate if blend is over
+
+	if (InertializationState == EInertializationState::Active)
+	{
+		InertializationTime += DeltaTime;
+
+		if (InertializationTime >= InertializationMaxDuration)
+		{
+			InertializationState = EInertializationState::Inactive;
+		}
+	}
+
+	// Apply inertialization
+
+	if (InertializationState == EInertializationState::Active)
+	{
+		ApplyTo(Output.Pose, Output.Curve);
+	}
+
+	// Find AttachParentName
+
+	FName AttachParentName = NAME_None;
+	if (AActor* Owner = Output.AnimInstanceProxy->GetSkelMeshComponent()->GetOwner())
+	{
+		if (AActor* AttachParentActor = Owner->GetAttachParentActor())
+		{
+			AttachParentName = AttachParentActor->GetFName();
+		}
+	}
+
+	// Record Pose Snapshot
+
+	if (PoseSnapshots.Num() < UE::Anim::DeadBlending::Private::MaxPoseSnapShotNum)
+	{
+		// Add the pose to the end of the buffer
+		PoseSnapshots.AddDefaulted_GetRef().InitFrom(Output.Pose, Output.Curve, ComponentTransform, AttachParentName, DeltaTime);
+	}
+	else
+	{
+		// Bubble the old poses forward in the buffer (using swaps to avoid allocations and copies)
+		for (int32 SnapshotIndex = 0; SnapshotIndex < UE::Anim::DeadBlending::Private::MaxPoseSnapShotNum - 1; ++SnapshotIndex)
+		{
+			Swap(PoseSnapshots[SnapshotIndex], PoseSnapshots[SnapshotIndex + 1]);
+
+			// Verify the swap operation used move assignment to properly fix up curve LUT pointers
+			checkSlow(PoseSnapshots[SnapshotIndex + 0].Curves.BlendedCurve.UIDToArrayIndexLUT == &PoseSnapshots[SnapshotIndex + 0].Curves.CurveUIDToArrayIndexLUT);
+			checkSlow(PoseSnapshots[SnapshotIndex + 1].Curves.BlendedCurve.UIDToArrayIndexLUT == &PoseSnapshots[SnapshotIndex + 1].Curves.CurveUIDToArrayIndexLUT);
+		}
+
+		// Overwrite the (now irrelevant) pose in the last slot with the new post snapshot
+		// (thereby avoiding the reallocation costs we would have incurred had we simply added a new pose at the end)
+		PoseSnapshots.Last().InitFrom(Output.Pose, Output.Curve, ComponentTransform, AttachParentName, DeltaTime);
+	}
+
+	// Reset Delta Time
+
+	DeltaTime = 0.0f;
+}
+
+bool FAnimNode_DeadBlending::NeedsDynamicReset() const
+{
+	return true;
+}
+
+#undef LOCTEXT_NAMESPACE
diff --git a/Engine/Source/Runtime/Engine/Private/Animation/AnimNode_Inertialization.cpp b/Engine/Source/Runtime/Engine/Private/Animation/AnimNode_Inertialization.cpp
--- a/Engine/Source/Runtime/Engine/Private/Animation/AnimNode_Inertialization.cpp
+++ b/Engine/Source/Runtime/Engine/Private/Animation/AnimNode_Inertialization.cpp
@@ -47,6 +47,17 @@ private:
 		Node.RequestInertialization(InRequestedDuration, InBlendProfile); 
 	}
 
+	virtual void RequestInertializationWithBlendMode(
+		float InRequestedDuration,
+		const UBlendProfile* InBlendProfile,
+		const bool bUseBlendMode,
+		const EAlphaBlendOption InBlendMode,
+		UCurveFloat* InCustomBlendCurve)
+	{
+		// The Blend Mode parameters are ignored as FAnimNode_Inertialization does not support them.
+		Node.RequestInertialization(InRequestedDuration, InBlendProfile);
+	}
+
 	virtual void AddDebugRecord(const FAnimInstanceProxy& InSourceProxy, int32 InSourceNodeId)
 	{
 #if WITH_EDITORONLY_DATA
@@ -175,7 +186,7 @@ void FAnimNode_Inertialization::Update_AnyThread(const FAnimationUpdateContext&
 				{
 					for (const FInertializationRequest& Request : RequestQueue)
 					{
-						InMessage.RequestInertialization(Request.Duration, Request.BlendProfile);
+						InMessage.RequestInertializationWithBlendMode(Request.Duration, Request.BlendProfile, Request.bUseBlendMode, Request.BlendMode, Request.CustomBlendCurve);
 					}
  					InMessage.AddDebugRecord(Proxy, NodeId);
 
diff --git a/Engine/Source/Runtime/Engine/Private/Animation/AnimNode_StateMachine.cpp b/Engine/Source/Runtime/Engine/Private/Animation/AnimNode_StateMachine.cpp
--- a/Engine/Source/Runtime/Engine/Private/Animation/AnimNode_StateMachine.cpp
+++ b/Engine/Source/Runtime/Engine/Private/Animation/AnimNode_StateMachine.cpp
@@ -1176,7 +1176,7 @@ void FAnimNode_StateMachine::TransitionToState(const FAnimationUpdateContext& Co
 				UE::Anim::IInertializationRequester* InertializationRequester = Context.GetMessage<UE::Anim::IInertializationRequester>();
 				if (InertializationRequester)
 				{
-					InertializationRequester->RequestInertialization(TransitionInfo.CrossfadeDuration, TransitionInfo.BlendProfile);
+					InertializationRequester->RequestInertializationWithBlendMode(TransitionInfo.CrossfadeDuration, TransitionInfo.BlendProfile, true, TransitionInfo.BlendMode, TransitionInfo.CustomCurve);
 					InertializationRequester->AddDebugRecord(*Context.AnimInstanceProxy, Context.GetCurrentNodeId());
 				}
 				else
