[Backout] - CL25578176
Original CL Desc
-----------------------------------------------------------------
Added experimental DeadBlending animation graph node.

diff --git a/Engine/Source/Editor/AnimGraph/Private/AnimGraphNode_DeadBlending.cpp b/Engine/Source/Editor/AnimGraph/Private/AnimGraphNode_DeadBlending.cpp
deleted file mode 100644
diff --git a/Engine/Source/Editor/AnimGraph/Public/AnimGraphNode_DeadBlending.h b/Engine/Source/Editor/AnimGraph/Public/AnimGraphNode_DeadBlending.h
deleted file mode 100644
diff --git a/Engine/Source/Runtime/AnimGraphRuntime/Private/AnimNodes/AnimNode_BlendListBase.cpp b/Engine/Source/Runtime/AnimGraphRuntime/Private/AnimNodes/AnimNode_BlendListBase.cpp
--- a/Engine/Source/Runtime/AnimGraphRuntime/Private/AnimNodes/AnimNode_BlendListBase.cpp
+++ b/Engine/Source/Runtime/AnimGraphRuntime/Private/AnimNodes/AnimNode_BlendListBase.cpp
@@ -125,7 +125,7 @@ void FAnimNode_BlendListBase::Update_AnyThread(const FAnimationUpdateContext& Co
 				UE::Anim::IInertializationRequester* InertializationRequester = Context.GetMessage<UE::Anim::IInertializationRequester>();
 				if (InertializationRequester)
 				{
-					InertializationRequester->RequestInertializationWithBlendMode(CurrentBlendTimes[ChildIndex], CurrentBlendProfile, true, GetBlendType(), GetCustomBlendCurve());
+					InertializationRequester->RequestInertialization(CurrentBlendTimes[ChildIndex], CurrentBlendProfile);
 					InertializationRequester->AddDebugRecord(*Context.AnimInstanceProxy, Context.GetCurrentNodeId());
 				}
 				else
diff --git a/Engine/Source/Runtime/Engine/Classes/Animation/AnimNode_DeadBlending.h b/Engine/Source/Runtime/Engine/Classes/Animation/AnimNode_DeadBlending.h
deleted file mode 100644
diff --git a/Engine/Source/Runtime/Engine/Classes/Animation/AnimNode_Inertialization.h b/Engine/Source/Runtime/Engine/Classes/Animation/AnimNode_Inertialization.h
--- a/Engine/Source/Runtime/Engine/Classes/Animation/AnimNode_Inertialization.h
+++ b/Engine/Source/Runtime/Engine/Classes/Animation/AnimNode_Inertialization.h
@@ -7,8 +7,6 @@
 #include "Animation/AnimNodeBase.h"
 #include "Animation/AnimCurveTypes.h"
 #include "Animation/AnimNodeMessages.h"
-#include "AlphaBlend.h" // Required for EAlphaBlendOption
-
 #include "AnimNode_Inertialization.generated.h"
 
 
@@ -34,22 +32,6 @@ public:
 	// If multiple requests are made on the same inertialization node, the minimum requested time will be used.
 	virtual void RequestInertialization(float InRequestedDuration, const UBlendProfile* InBlendProfile = nullptr) = 0;
 
-	// Request to activate inertialization for a duration with the given blend mode.
-	// If multiple requests are made on the same inertialization node, the minimum requested time will be used.
-	// 
-	// InRequestedDuration				Duration of the blend
-	// InBlendProfile					Blend profile to use
-	// bUseBlendMode					If to use the provided Blend Mode
-	// InBlendMode						Blend mode to use
-	// InCustomBlendCurve				Custom Blend Curve to use for the Blend Mode
-	//
-	virtual void RequestInertializationWithBlendMode(
-		float InRequestedDuration,
-		const UBlendProfile* InBlendProfile,
-		const bool bUseBlendMode,
-		const EAlphaBlendOption InBlendMode,
-		UCurveFloat* InCustomBlendCurve) = 0;
-
 	// Add a record of this request
 	virtual void AddDebugRecord(const FAnimInstanceProxy& InSourceProxy, int32 InSourceNodeId) = 0;
 };
@@ -125,46 +107,27 @@ struct FInertializationRequest
 {
 	GENERATED_BODY()
 
-	FInertializationRequest() {}
-
-	FInertializationRequest(float InDuration, const UBlendProfile* InBlendProfile)
-		: Duration(InDuration)
-		, BlendProfile(InBlendProfile)
+	FInertializationRequest()
+		: Duration(-1.0f)
+		, BlendProfile(nullptr)
 	{
 	}
 
-	FInertializationRequest(
-		float InDuration, 
-		const UBlendProfile* InBlendProfile,
-		const bool bInUseBlendMode,
-		const EAlphaBlendOption InBlendMode, 
-		UCurveFloat* InCustomBlendCurve)
+	FInertializationRequest(float InDuration, const UBlendProfile* InBlendProfile)
 		: Duration(InDuration)
 		, BlendProfile(InBlendProfile)
-		, bUseBlendMode(bInUseBlendMode)
-		, BlendMode(InBlendMode)
-		, CustomBlendCurve(InCustomBlendCurve)
 	{
 	}
 
-
 	void Clear()
 	{
 		Duration = -1.0f;
 		BlendProfile = nullptr;
-		bUseBlendMode = false;
-		BlendMode = EAlphaBlendOption::Linear;
-		CustomBlendCurve = nullptr;
 	}
 
 	friend bool operator==(const FInertializationRequest& A, const FInertializationRequest& B)
 	{
-		return 
-			(A.Duration == B.Duration) && 
-			(A.BlendProfile == B.BlendProfile) &&
-			(A.bUseBlendMode == B.bUseBlendMode) &&
-			(A.BlendMode == B.BlendMode) &&
-			(A.CustomBlendCurve == B.CustomBlendCurve);
+		return A.Duration == B.Duration && A.BlendProfile == B.BlendProfile;
 	}
 
 	friend bool operator!=(const FInertializationRequest& A, const FInertializationRequest& B)
@@ -173,19 +136,10 @@ struct FInertializationRequest
 	}
 
 	UPROPERTY(Transient)
-	float Duration = -1.0f;
-
-	UPROPERTY(Transient)
-	TObjectPtr<const UBlendProfile> BlendProfile = nullptr;
-
-	UPROPERTY(Transient)
-	bool bUseBlendMode = false;
-
-	UPROPERTY(Transient)
-	EAlphaBlendOption BlendMode = EAlphaBlendOption::Linear;
+	float Duration;
 
 	UPROPERTY(Transient)
-	TObjectPtr<UCurveFloat> CustomBlendCurve = nullptr;
+	TObjectPtr<const UBlendProfile> BlendProfile;
 };
 
 
diff --git a/Engine/Source/Runtime/Engine/Private/Animation/AnimNode_DeadBlending.cpp b/Engine/Source/Runtime/Engine/Private/Animation/AnimNode_DeadBlending.cpp
deleted file mode 100644
diff --git a/Engine/Source/Runtime/Engine/Private/Animation/AnimNode_Inertialization.cpp b/Engine/Source/Runtime/Engine/Private/Animation/AnimNode_Inertialization.cpp
--- a/Engine/Source/Runtime/Engine/Private/Animation/AnimNode_Inertialization.cpp
+++ b/Engine/Source/Runtime/Engine/Private/Animation/AnimNode_Inertialization.cpp
@@ -47,17 +47,6 @@ private:
 		Node.RequestInertialization(InRequestedDuration, InBlendProfile); 
 	}
 
-	virtual void RequestInertializationWithBlendMode(
-		float InRequestedDuration,
-		const UBlendProfile* InBlendProfile,
-		const bool bUseBlendMode,
-		const EAlphaBlendOption InBlendMode,
-		UCurveFloat* InCustomBlendCurve)
-	{
-		// The Blend Mode parameters are ignored as FAnimNode_Inertialization does not support them.
-		Node.RequestInertialization(InRequestedDuration, InBlendProfile);
-	}
-
 	virtual void AddDebugRecord(const FAnimInstanceProxy& InSourceProxy, int32 InSourceNodeId)
 	{
 #if WITH_EDITORONLY_DATA
@@ -186,7 +175,7 @@ void FAnimNode_Inertialization::Update_AnyThread(const FAnimationUpdateContext&
 				{
 					for (const FInertializationRequest& Request : RequestQueue)
 					{
-						InMessage.RequestInertializationWithBlendMode(Request.Duration, Request.BlendProfile, Request.bUseBlendMode, Request.BlendMode, Request.CustomBlendCurve);
+						InMessage.RequestInertialization(Request.Duration, Request.BlendProfile);
 					}
  					InMessage.AddDebugRecord(Proxy, NodeId);
 
diff --git a/Engine/Source/Runtime/Engine/Private/Animation/AnimNode_StateMachine.cpp b/Engine/Source/Runtime/Engine/Private/Animation/AnimNode_StateMachine.cpp
--- a/Engine/Source/Runtime/Engine/Private/Animation/AnimNode_StateMachine.cpp
+++ b/Engine/Source/Runtime/Engine/Private/Animation/AnimNode_StateMachine.cpp
@@ -1176,7 +1176,7 @@ void FAnimNode_StateMachine::TransitionToState(const FAnimationUpdateContext& Co
 				UE::Anim::IInertializationRequester* InertializationRequester = Context.GetMessage<UE::Anim::IInertializationRequester>();
 				if (InertializationRequester)
 				{
-					InertializationRequester->RequestInertializationWithBlendMode(TransitionInfo.CrossfadeDuration, TransitionInfo.BlendProfile, true, TransitionInfo.BlendMode, TransitionInfo.CustomCurve);
+					InertializationRequester->RequestInertialization(TransitionInfo.CrossfadeDuration, TransitionInfo.BlendProfile);
 					InertializationRequester->AddDebugRecord(*Context.AnimInstanceProxy, Context.GetCurrentNodeId());
 				}
 				else
