* Change Specular profile atlas format to reduce banding.
* Move procedural profile generation into shader.

diff --git a/Engine/Shaders/Private/SpecularProfile.usf b/Engine/Shaders/Private/SpecularProfile.usf
--- a/Engine/Shaders/Private/SpecularProfile.usf
+++ b/Engine/Shaders/Private/SpecularProfile.usf
@@ -2,6 +2,9 @@
 
 #include "Common.ush"
 
+float3 ViewColor;
+float3 LightColor;
+
 int2 SourceResolution;
 uint SourceMipCount;
 int2 TargetResolution;
@@ -16,8 +19,23 @@ void MainCS(uint3 DispatchThreadId : SV_DispatchThreadID)
 	const uint2 PixelPos = DispatchThreadId.xy;
 	if (all(PixelPos < uint2(TargetResolution)))
 	{
-		const float SourceMip = clamp(log2(SourceResolution) - log2(TargetResolution), 0.f, SourceMipCount-1);
-		const float4 Value = SourceTexture.SampleLevel(SourceSampler, float2(PixelPos+0.5f) / float2(TargetResolution), SourceMip);
-		TargetTexture[uint3(PixelPos,TargetIndex)] = Value;
+		const float2 TexCoord = float2(PixelPos + 0.5f) / float2(TargetResolution);
+		#if PERMUTATION_PROCEDURAL
+		{
+			const float NoL = TexCoord.y;
+			const float NoV = TexCoord.x;
+
+			const float3 DefaultColor = 1.f;
+			const float3 C0 = lerp(DefaultColor, ViewColor, NoV);
+			const float3 C1 = lerp(DefaultColor, LightColor, NoL);
+			TargetTexture[uint3(PixelPos,TargetIndex)] = float4(C0 * C1, 1);
+		}
+		#else
+		{
+			const float SourceMip = clamp(log2(SourceResolution) - log2(TargetResolution), 0.f, SourceMipCount-1);
+			const float4 Value = SourceTexture.SampleLevel(SourceSampler, TexCoord, SourceMip);
+			TargetTexture[uint3(PixelPos,TargetIndex)] = Value;
+		}
+		#endif
 	}
 }
\ No newline at end of file
diff --git a/Engine/Shaders/Private/Strata/StrataEvaluation.ush b/Engine/Shaders/Private/Strata/StrataEvaluation.ush
--- a/Engine/Shaders/Private/Strata/StrataEvaluation.ush
+++ b/Engine/Shaders/Private/Strata/StrataEvaluation.ush
@@ -687,9 +687,7 @@ FStrataEvaluateResult StrataEvaluateBSDFCommon(FStrataBSDFContext BSDFContext, F
 			if (BSDF_GETHASSPECPROFILE(BSDFContext.BSDF))
 			{
 				const uint ProfileID = SLAB_SPECPROFILEID(BSDFContext.BSDF);
-				const float NoL = saturate(dot(BSDFContext.L, BSDFContext.N));
-				const float NoV = saturate(dot(BSDFContext.V, BSDFContext.N));
-				Sample.IntegratedSpecularValue *= View.SpecularProfileTexture.SampleLevel(View.SpecularProfileSampler, float3(NoV, NoL, ProfileID), 0);
+				Sample.IntegratedSpecularValue *= View.SpecularProfileTexture.SampleLevel(View.SpecularProfileSampler, float3(BSDFContext.SatNoV, BSDFContext.SatNoL, ProfileID), 0);
 			}
 			#endif
 
@@ -1540,8 +1538,7 @@ FStrataEnvLightResult StrataEvaluateForEnvLight(FStrataBSDFContext BSDFContext,
 				{
 					const uint ProfileID = SLAB_SPECPROFILEID(BSDFContext.BSDF);
 					const float NoR = saturate(dot(BSDFContext.R, BSDFContext.N));
-					const float NoV = saturate(dot(BSDFContext.V, BSDFContext.N));
-					StrataEnvLightResult.SpecularWeight *= View.SpecularProfileTexture.SampleLevel(View.SpecularProfileSampler, float3(NoV, NoR, ProfileID), 0);
+					StrataEnvLightResult.SpecularWeight *= View.SpecularProfileTexture.SampleLevel(View.SpecularProfileSampler, float3(BSDFContext.SatNoV, NoR, ProfileID), 0);
 				}
 			#endif // STRATA_COMPLEXPATH
 
diff --git a/Engine/Source/Runtime/Engine/Private/Rendering/SpecularProfile.cpp b/Engine/Source/Runtime/Engine/Private/Rendering/SpecularProfile.cpp
--- a/Engine/Source/Runtime/Engine/Private/Rendering/SpecularProfile.cpp
+++ b/Engine/Source/Runtime/Engine/Private/Rendering/SpecularProfile.cpp
@@ -52,6 +52,9 @@ public:
 	DECLARE_GLOBAL_SHADER(FSpecularProfileCopyCS);
 	SHADER_USE_PARAMETER_STRUCT(FSpecularProfileCopyCS, FGlobalShader);
 
+	class FProcedural : SHADER_PERMUTATION_BOOL("PERMUTATION_PROCEDURAL");
+	using FPermutationDomain = TShaderPermutationDomain<FProcedural>;
+
 	BEGIN_SHADER_PARAMETER_STRUCT(FParameters, )
 		SHADER_PARAMETER_TEXTURE(Texture2D, SourceTexture)
 		SHADER_PARAMETER_SAMPLER(SamplerState, SourceSampler)
@@ -59,6 +62,8 @@ public:
 		SHADER_PARAMETER(FIntPoint, TargetResolution)
 		SHADER_PARAMETER(uint32, SourceMipCount)
 		SHADER_PARAMETER(uint32, TargetIndex)
+		SHADER_PARAMETER(FVector3f, ViewColor)
+		SHADER_PARAMETER(FVector3f, LightColor)
 		SHADER_PARAMETER_RDG_TEXTURE_UAV(RWTexture2D, TargetTexture)
 	END_SHADER_PARAMETER_STRUCT()
 
@@ -105,8 +110,8 @@ public:
 	FName GetParameterName(const USpecularProfile* InProfile) const;
 
 	// @return can be nullptr if there is no SpecularProfile
-	struct IPooledRenderTarget* GetAtlasTexture(FRDGBuilder& GraphBuilder, EShaderPlatform ShaderPlatform);
-	struct IPooledRenderTarget* GetAtlasTexture();
+	IPooledRenderTarget* GetAtlasTexture(FRDGBuilder& GraphBuilder, EShaderPlatform ShaderPlatform);
+	IPooledRenderTarget* GetAtlasTexture();
 
 	//~ Begin FRenderResource Interface.
 	/**
@@ -305,63 +310,40 @@ IPooledRenderTarget* FSpecularProfileTextureManager::GetAtlasTexture(FRDGBuilder
 		const uint32 Resolution = CVarSpecularProfileResolution.GetValueOnRenderThread();
 		check(LayerCount);
 
-		FPooledRenderTargetDesc Desc(FPooledRenderTargetDesc::Create2DArrayDesc(FIntPoint(Resolution, Resolution), PF_B8G8R8A8, FClearValueBinding::None, TexCreate_None, TexCreate_UAV | TexCreate_ShaderResource, false, LayerCount));
+		// 1. Create atlas texture
+		FPooledRenderTargetDesc Desc(FPooledRenderTargetDesc::Create2DArrayDesc(FIntPoint(Resolution, Resolution), PF_FloatR11G11B10, FClearValueBinding::None, TexCreate_None, TexCreate_UAV | TexCreate_ShaderResource, false, LayerCount));
 		GRenderTargetPool.FindFreeElement(RHICmdList, Desc, GSpecularProfileTextureAtlas, TEXT("SpecularProfileTexture"));
 
-		// 1. Fill in all procedural profile
-		FLinearColor DefaultColor(1.f, 1.f, 1.f);
-		for (uint32 LayerIt = 0; LayerIt < LayerCount; ++LayerIt)
-		{
-			if (SpecularProfileEntries[LayerIt].Texture == nullptr)
-			{
-				const FSpecularProfileStruct Data = SpecularProfileEntries[LayerIt].Settings;
-
-				uint32 DestStride;
-				uint8* DestBuffer = (uint8*)RHICmdList.LockTexture2DArray(GSpecularProfileTextureAtlas->GetRHI(), LayerIt, 0, RLM_WriteOnly, DestStride, false);
-				for (uint32 y = 0; y < Resolution; ++y)
-				{
-					FColor* Dest = (FColor*)(DestBuffer + DestStride * y);
-					for (uint32 x = 0; x < Resolution; ++x)
-					{
-						const float NoL = float(y) / float(Resolution-1);
-						const float NoV = float(x) / float(Resolution-1);
-						const FLinearColor ViewColor = FMath::Lerp(DefaultColor, Data.ViewColor, NoV);
-						const FLinearColor LightColor = FMath::Lerp(DefaultColor, Data.LightColor, NoL);
-						const FLinearColor C = ViewColor * LightColor;
-	
-						Dest[x] = FColor(FMath::Quantize8UnsignedByte(C.R), FMath::Quantize8UnsignedByte(C.G), FMath::Quantize8UnsignedByte(C.B));
-					}
-				}
-				RHICmdList.UnlockTexture2DArray(GSpecularProfileTextureAtlas->GetRHI(), LayerIt, 0, false);
-			}
-		}
-
-		// 2. Fill in texture based profile
+		// 2. Fill in profiles
 		const auto GlobalShaderMap = GetGlobalShaderMap(ShaderPlatform);
 		FRDGTextureRef SpecularProfileTexture = GraphBuilder.RegisterExternalTexture(GSpecularProfileTextureAtlas, TEXT("SpecularProfileTexture"));
 		FRDGTextureUAVRef SpecularProfileUAV = GraphBuilder.CreateUAV(SpecularProfileTexture);
 		for (uint32 LayerIt = 0; LayerIt < LayerCount; ++LayerIt)
 		{
-			if (SpecularProfileEntries[LayerIt].Texture)
-			{
-				const FTextureReferenceRHIRef& TextureRHI = SpecularProfileEntries[LayerIt].Texture->TextureReferenceRHI;
-				FSpecularProfileCopyCS::FParameters* PassParameters = GraphBuilder.AllocParameters<FSpecularProfileCopyCS::FParameters>();
-				PassParameters->SourceMipCount = FMath::Max(1u, uint32(TextureRHI->GetDesc().NumMips));
-				PassParameters->SourceTexture = TextureRHI->GetReferencedTexture();
-				PassParameters->SourceSampler = TStaticSamplerState<SF_Point, AM_Clamp, AM_Clamp, AM_Clamp>::GetRHI();
-				PassParameters->SourceResolution = SpecularProfileTexture->Desc.Extent;
-				PassParameters->TargetResolution = SpecularProfileTexture->Desc.Extent;
-				PassParameters->TargetTexture = SpecularProfileUAV;
-				PassParameters->TargetIndex = LayerIt;
-				TShaderMapRef<FSpecularProfileCopyCS> Shader(GlobalShaderMap);
-				FComputeShaderUtils::AddPass(GraphBuilder, 
-											RDG_EVENT_NAME("SpecularProfile::CopyTexture"), 
-											Shader, 
-											PassParameters, 
-											FIntVector(FMath::DivideAndRoundUp(PassParameters->TargetResolution.X, 8), FMath::DivideAndRoundUp(PassParameters->TargetResolution.Y, 8), 1));
-	
-				SpecularProfileEntries[LayerIt].CachedResolution = TextureRHI->GetDesc().Extent;
-			}
+			const FSpecularProfileStruct Data = SpecularProfileEntries[LayerIt].Settings;
+			const FTextureReference* Texture = SpecularProfileEntries[LayerIt].Texture;
+
+			FSpecularProfileCopyCS::FParameters* PassParameters = GraphBuilder.AllocParameters<FSpecularProfileCopyCS::FParameters>();
+			PassParameters->SourceMipCount = Texture ? FMath::Max(1u, uint32(Texture->TextureReferenceRHI->GetDesc().NumMips)) : 0;
+			PassParameters->SourceTexture = Texture ? Texture->TextureReferenceRHI->GetReferencedTexture() : nullptr;
+			PassParameters->SourceSampler = TStaticSamplerState<SF_Point, AM_Clamp, AM_Clamp, AM_Clamp>::GetRHI();
+			PassParameters->SourceResolution = SpecularProfileTexture->Desc.Extent;
+			PassParameters->TargetResolution = SpecularProfileTexture->Desc.Extent;
+			PassParameters->TargetTexture = SpecularProfileUAV;
+			PassParameters->TargetIndex = LayerIt;
+			PassParameters->ViewColor = FVector3f(Data.ViewColor);
+			PassParameters->LightColor = FVector3f(Data.LightColor);
+
+			FSpecularProfileCopyCS::FPermutationDomain PermutationVector;
+			PermutationVector.Set<FSpecularProfileCopyCS::FProcedural>(Data.IsProcedural());
+			TShaderMapRef<FSpecularProfileCopyCS> Shader(GlobalShaderMap, PermutationVector);
+			FComputeShaderUtils::AddPass(GraphBuilder, 
+										RDG_EVENT_NAME("SpecularProfile::CopyTexture"), 
+										Shader, 
+										PassParameters, 
+										FIntVector(FMath::DivideAndRoundUp(PassParameters->TargetResolution.X, 8), FMath::DivideAndRoundUp(PassParameters->TargetResolution.Y, 8), 1));
+
+			SpecularProfileEntries[LayerIt].CachedResolution = Texture ? Texture->TextureReferenceRHI->GetDesc().Extent : FIntPoint(0,0);
 		}
 	}
 	return GSpecularProfileTextureAtlas;
