UAnimSequence::IsCachedCookedPlatformDataLoaded Can Cause Cooks To Soft Lock
    - Add safety net in AnimSequence::PreSave to generate compressed data if necessary
    - Only try and update animation sequence with skeleton data (if it is actually out-of-date)
    - Deprecated ClearCompressedBoneData and ClearCompressedCurveData
    - Only early out of AnimSequence::BeginCacheDerivedData if hash matched running platform, if the data is _actually_ valid
    - AnimSequence::TryCancelAsyncTasks now only cancels the task, and leaves AnimationSequenceCompiler to clean up
    - AnimSequence::IsAsyncTaskComplete only returns true if _all_ tasks have finished (not just one)
    - AnimSequence::FinishAsyncTasks now handles cancelled tasks -> clear out compressed data, map entry and reset running platform hash (if applicable)

diff --git a/Engine/Plugins/Animation/AnimationData/Source/AnimationData/Private/AnimSequencerController.cpp b/Engine/Plugins/Animation/AnimationData/Source/AnimationData/Private/AnimSequencerController.cpp
--- a/Engine/Plugins/Animation/AnimationData/Source/AnimationData/Private/AnimSequencerController.cpp
+++ b/Engine/Plugins/Animation/AnimationData/Source/AnimationData/Private/AnimSequencerController.cpp
@@ -54,6 +54,8 @@ void UAnimSequencerController::OpenBracket(const FText& InTitle, bool bShouldTra
 		FBracketPayload Payload;
 		Payload.Description = InTitle.ToString();
 
+		Model->LockEvaluationAndModification();
+
 		Model->GetNotifier().Notify(EAnimDataModelNotifyType::BracketOpened, Payload);
 	}
 
@@ -84,6 +86,8 @@ void UAnimSequencerController::CloseBracket(bool bShouldTransact /*= true*/)
 		}
 		
 		Model->GetNotifier().Notify(EAnimDataModelNotifyType::BracketClosed);
+		
+		Model->UnlockEvaluationAndModification();
 	}
 }
 
diff --git a/Engine/Source/Runtime/Engine/Classes/Animation/AnimationAsset.h b/Engine/Source/Runtime/Engine/Classes/Animation/AnimationAsset.h
--- a/Engine/Source/Runtime/Engine/Classes/Animation/AnimationAsset.h
+++ b/Engine/Source/Runtime/Engine/Classes/Animation/AnimationAsset.h
@@ -1185,5 +1185,6 @@ public:
 
 	FGuid GetSkeletonVirtualBoneGuid() const { return SkeletonVirtualBoneGuid; }
 	void SetSkeletonVirtualBoneGuid(FGuid Guid) { SkeletonVirtualBoneGuid = Guid; }
+	FGuid GetSkeletonGuid() const { return SkeletonGuid; }
 };
 
diff --git a/Engine/Source/Runtime/Engine/Private/Animation/AnimSequence.cpp b/Engine/Source/Runtime/Engine/Private/Animation/AnimSequence.cpp
--- a/Engine/Source/Runtime/Engine/Private/Animation/AnimSequence.cpp
+++ b/Engine/Source/Runtime/Engine/Private/Animation/AnimSequence.cpp
@@ -466,10 +466,7 @@ void UAnimSequence::WillNeverCacheCookedPlatformDataAgain()
 {
 	Super::WillNeverCacheCookedPlatformDataAgain();
 
-	UE::Anim::FAnimSequenceCompilingManager::Get().FinishCompilation({this});	
-	ClearCompressedCurveData();
-	ClearCompressedBoneData();
-
+	UE::Anim::FAnimSequenceCompilingManager::Get().FinishCompilation({this});
 	// Clear out current platform, and any target platform data
 	CompressedData.Reset();
 		
@@ -483,7 +480,7 @@ void UAnimSequence::WillNeverCacheCookedPlatformDataAgain()
 void UAnimSequence::ClearAllCachedCookedPlatformData()
 {
 	Super::ClearAllCachedCookedPlatformData();
-	
+
 	// Delete any cache tasks first because the destructor will cancel the cache and build tasks,
 	// and drop their pointers to the data.
 	CacheTasksByKeyHash.Empty();
@@ -905,6 +902,13 @@ void UAnimSequence::PreSave(FObjectPreSaveContext ObjectSaveContext)
 	if (ObjectSaveContext.IsCooking())
 	{
 		check(IsCachedCookedPlatformDataLoaded(ObjectSaveContext.GetTargetPlatform()));
+
+		// In case compressed data was cleared between ::IsCachedCookedPlatformDataLoaded and ::PreSave being called, synchronously run compression
+		if(!IsCachedCookedPlatformDataLoaded(ObjectSaveContext.GetTargetPlatform()))
+		{
+			CacheDerivedData(ObjectSaveContext.GetTargetPlatform());
+			check(IsCachedCookedPlatformDataLoaded(ObjectSaveContext.GetTargetPlatform()));
+		}
 	}
 
 	if (!ObjectSaveContext.IsProceduralSave())
@@ -1037,8 +1041,7 @@ void UAnimSequence::BeginDestroy()
 	Super::BeginDestroy();
 
 	PRAGMA_DISABLE_DEPRECATION_WARNINGS
-	ClearCompressedCurveData();
-	ClearCompressedBoneData();
+	CompressedData.Reset();
 	PRAGMA_ENABLE_DEPRECATION_WARNINGS
 }
 
@@ -1050,10 +1053,7 @@ bool UAnimSequence::IsReadyForFinishDestroy()
 	}
 
 #if WITH_EDITOR
-	if (!TryCancelAsyncTasks())
-	{
-		return false;
-	}
+	WaitOnExistingCompression(false);
 #endif
 
 	return true;
@@ -1150,7 +1150,8 @@ bool UAnimSequence::IsCachedCookedPlatformDataLoaded(const ITargetPlatform* Targ
 	const FIoHash KeyHash = CreateDerivedDataKeyHash(TargetPlatform);
 	if (KeyHash.IsZero())
 	{
-		return true;
+		UE_LOG(LogAnimation, Warning, TEXT("Zero key hash compressed animation data for %s requested platform %s"), *GetName(), *TargetPlatform->PlatformName());
+		return false;
 	}
 
 	if (PollCacheDerivedData(KeyHash))
@@ -1162,15 +1163,20 @@ bool UAnimSequence::IsCachedCookedPlatformDataLoaded(const ITargetPlatform* Targ
 	{
 		if (IsInGameThread())
 		{
-			// This is required until we can ensure inter-animsequence dependencies are handled correctly (see UE-184634)
-			BeginCacheForCookedPlatformData(TargetPlatform);
-			UE_LOG(LogAnimation, Warning, TEXT("Missing compressed animation data for %s requested platform %s and hash %s"), *GetName(), *TargetPlatform->PlatformName(), *LexToString(KeyHash));
+			if (!HasAnyFlags(EObjectFlags::RF_NeedPostLoad) && GetPackage()->GetHasBeenEndLoaded())
+			{
+				BeginCacheForCookedPlatformData(TargetPlatform);
+			}
 		}
 	};
 
 	if (KeyHash == DataKeyHash)
 	{
-		if (!CacheTasksByKeyHash.Contains(KeyHash) && CompressedData.IsValid(this))
+		if (!CacheTasksByKeyHash.Contains(KeyHash) && CompressedData.IsValid(this) &&
+			PRAGMA_DISABLE_DEPRECATION_WARNINGS
+			!bUseRawDataOnly
+			PRAGMA_ENABLE_DEPRECATION_WARNINGS
+			)
 		{
 			return true;
 		}
@@ -2709,8 +2715,11 @@ void UAnimSequence::RemapTracksToNewSkeleton( USkeleton* NewSkeleton, bool bConv
 	}
 	else
 	{
-		SetSkeleton(NewSkeleton);
-		Controller->UpdateWithSkeleton(NewSkeleton, false);
+		if (NewSkeleton != GetSkeleton() || NewSkeleton->GetGuid() != GetSkeletonGuid())
+		{
+			SetSkeleton(NewSkeleton);
+			Controller->UpdateWithSkeleton(NewSkeleton, false);
+		}
 	}
 
 	Super::RemapTracksToNewSkeleton(NewSkeleton, bConvertSpaces);
@@ -3861,7 +3870,7 @@ void UAnimSequence::EvaluateAttributes(FAnimationPoseData& OutAnimationPoseData,
 #if WITH_EDITOR
 void UAnimSequence::OnSetSkeleton(USkeleton* NewSkeleton)
 {
-	UE::Anim::FAnimSequenceCompilingManager::Get().FinishCompilation({this});
+	WaitOnExistingCompression(false);
 }
 
 void UAnimSequence::RemoveCustomAttribute(const FName& BoneName, const FName& AttributeName)
@@ -4636,16 +4645,6 @@ bool UAnimSequence::IsCurveCompressedDataValid() const
 	return true;
 }
 
-void UAnimSequence::ClearCompressedBoneData()
-{
-	CompressedData.ClearCompressedBoneData();
-}
-
-void UAnimSequence::ClearCompressedCurveData()
-{
-	CompressedData.ClearCompressedCurveData();
-}
-
 #if WITH_EDITOR
 void UAnimSequence::EnsureValidRawDataGuid()
 {
@@ -4918,7 +4917,7 @@ void UAnimSequence::OnModelModified(const EAnimDataModelNotifyType& NotifyType,
 				{
 					const bool bWasModelReset = NotifyCollector.Contains(EAnimDataModelNotifyType::Reset);
 					UpdateRawDataGuid(bWasModelReset ? GenerateNewGUID : RegenerateGUID);
-					ClearCompressedCurveData();
+					CompressedData.ClearCompressedCurveData();
 					HandleTrackDataChanged(bWasModelReset);
 				}
 			}
@@ -4958,7 +4957,7 @@ void UAnimSequence::OnModelModified(const EAnimDataModelNotifyType& NotifyType,
 		case EAnimDataModelNotifyType::CurveFlagsChanged:
 		case EAnimDataModelNotifyType::CurveScaled:
 		{
-			ClearCompressedCurveData();
+			CompressedData.ClearCompressedCurveData();
 				
 			if (NotifyCollector.IsNotWithinBracket())
 			{						
@@ -5135,7 +5134,7 @@ FIoHash UAnimSequence::BeginCacheDerivedData(const ITargetPlatform* TargetPlatfo
 		
 	const FIoHash KeyHash = CreateDerivedDataKeyHash(TargetPlatform);
 	// Early out if not valid, equal to running platform hash or any previously cached hashes
-	if (KeyHash.IsZero() || DataKeyHash == KeyHash || DataByPlatformKeyHash.Contains(KeyHash))
+	if (KeyHash.IsZero() || (DataKeyHash == KeyHash && (CompressedData.IsValid(this) || CacheTasksByKeyHash.Contains(KeyHash))) || DataByPlatformKeyHash.Contains(KeyHash))
 	{
 		return KeyHash;
 	}
@@ -5146,11 +5145,7 @@ FIoHash UAnimSequence::BeginCacheDerivedData(const ITargetPlatform* TargetPlatfo
 	// We should wait here for any previous compilations?
 	FCompressedAnimSequence* TargetData = nullptr;
 	if (TargetPlatform->IsRunningPlatform())
-	{
-		PRAGMA_DISABLE_DEPRECATION_WARNINGS
-		bUseRawDataOnly = true;
-		PRAGMA_ENABLE_DEPRECATION_WARNINGS
-		
+	{		
 		DataKeyHash = KeyHash;
 		TargetData = &CompressedData;
 	}
@@ -5158,8 +5153,19 @@ FIoHash UAnimSequence::BeginCacheDerivedData(const ITargetPlatform* TargetPlatfo
     {
     	TargetData = DataByPlatformKeyHash.Emplace(KeyHash, MakeUnique<FCompressedAnimSequence>()).Get();
     }
+
+	// Mark sequence to only sample raw animation data when recompressing for running platform (or equivalent hash)
+	if (KeyHash == DataKeyHash)
+	{
+		PRAGMA_DISABLE_DEPRECATION_WARNINGS
+		bUseRawDataOnly = true;
+		PRAGMA_ENABLE_DEPRECATION_WARNINGS
+	}
+	
 	check(TargetData);
 	TargetData->Reset();
+	
+	
 
 	bool bPerformFrameStripping  = false;
 	{
@@ -5267,14 +5273,7 @@ bool UAnimSequence::TryCancelAsyncTasks()
 {
 	for (auto It = CacheTasksByKeyHash.CreateIterator(); It; ++It)
 	{
-		if (It->Value->Poll())
-		{
-			It.RemoveCurrent();
-		}
-		else
-		{
-			It->Value->Cancel();
-		}
+		It->Value->Cancel();
 	}
 
 	return CacheTasksByKeyHash.IsEmpty();
@@ -5282,15 +5281,15 @@ bool UAnimSequence::TryCancelAsyncTasks()
 
 bool UAnimSequence::IsAsyncTaskComplete() const
 {
+	check(CacheTasksByKeyHash.Num());
+
+	bool bAllFinished = true;
 	for (auto& Pair : CacheTasksByKeyHash)
 	{
-		if (!Pair.Value->Poll())
-		{
-			return false;
-		}
+		bAllFinished &= Pair.Value->Poll();
 	}
 
-	return true;
+	return bAllFinished;
 }
 
 void UAnimSequence::FinishAsyncTasks()
@@ -5305,7 +5304,7 @@ void UAnimSequence::FinishAsyncTasks()
 		{
 			It->Value->Wait();
 
-			const FCompressedAnimSequence* CompressedAnimData = [Hash = It->Key, this]()
+			FCompressedAnimSequence* CompressedAnimData = [Hash = It->Key, this]()
 			{
 				if(Hash == DataKeyHash)
 				{
@@ -5315,30 +5314,46 @@ void UAnimSequence::FinishAsyncTasks()
 				return DataByPlatformKeyHash.FindChecked(Hash).Get();
 			}();
 
-			if (CompressedAnimData->IsValid(this))
+			if (It->Value.Get()->WasCancelled())
 			{
-				#if WITH_EDITOR
-				//This is only safe during sync anim compression
-				if (GetSkeleton())
-				{
-					SetSkeletonVirtualBoneGuid(GetSkeleton()->GetVirtualBoneGuid());
-				}
-#endif
-				if (It->Key == DataKeyHash)
+				CompressedAnimData->Reset();
+				DataByPlatformKeyHash.Remove(It->Key);
+
+				// Reset running platform hash (if it got cancelled)
+				if (DataKeyHash == It->Key)
 				{
-					
-					FAssetRegistryModule& AssetRegistryModule = FModuleManager::LoadModuleChecked<FAssetRegistryModule>(TEXT("AssetRegistry"));
-					AssetRegistryModule.Get().AssetTagsFinalized(*this);
-					
-					PRAGMA_DISABLE_DEPRECATION_WARNINGS
-					check(bUseRawDataOnly);
-					bUseRawDataOnly = false;
-                    PRAGMA_ENABLE_DEPRECATION_WARNINGS
+					DataKeyHash = FIoHash::Zero;
 				}
 			}
 			else
 			{
-				// Failed to compress
+				if (CompressedAnimData->IsValid(this))
+				{
+#if WITH_EDITOR
+					//This is only safe during sync anim compression
+					if (GetSkeleton())
+					{
+						SetSkeletonVirtualBoneGuid(GetSkeleton()->GetVirtualBoneGuid());
+					}
+#endif
+					if (It->Key == DataKeyHash)
+					{
+						FAssetRegistryModule& AssetRegistryModule = FModuleManager::LoadModuleChecked<FAssetRegistryModule>(TEXT("AssetRegistry"));
+						AssetRegistryModule.Get().AssetTagsFinalized(*this);
+					
+						PRAGMA_DISABLE_DEPRECATION_WARNINGS
+						check(bUseRawDataOnly);
+						bUseRawDataOnly = false;
+						PRAGMA_ENABLE_DEPRECATION_WARNINGS
+					}
+				
+					check(CompressedAnimData->IsValid(this, true));
+				}
+				else
+				{
+					// Failed to compress
+					UE_LOG(LogAnimation, Warning, TEXT("Failed to finish async Animation Compression task for %s, as the generated data is not valid."), *GetName());
+				}
 			}
 
 			It.RemoveCurrent();
diff --git a/Engine/Source/Runtime/Engine/Private/Animation/AnimationCompressionDerivedData.h b/Engine/Source/Runtime/Engine/Private/Animation/AnimationCompressionDerivedData.h
--- a/Engine/Source/Runtime/Engine/Private/Animation/AnimationCompressionDerivedData.h
+++ b/Engine/Source/Runtime/Engine/Private/Animation/AnimationCompressionDerivedData.h
@@ -83,6 +83,7 @@ namespace UE::Anim
 		void Wait(bool bPerformWork = true);
 		bool Poll() const;
 		void Reschedule(FQueuedThreadPool* InThreadPool, EQueuedWorkPriority InPriority) const;
+		bool WasCancelled() const { return CompressibleAnimPtr->IsCancelled() || Owner.IsCanceled(); }
 	private:
 		void BeginCache(const FIoHash& KeyHash);
 		void EndCache(UE::DerivedData::FCacheGetValueResponse&& Response);
diff --git a/Engine/Source/Runtime/Engine/Private/Animation/AnimationSequenceCompiler.cpp b/Engine/Source/Runtime/Engine/Private/Animation/AnimationSequenceCompiler.cpp
--- a/Engine/Source/Runtime/Engine/Private/Animation/AnimationSequenceCompiler.cpp
+++ b/Engine/Source/Runtime/Engine/Private/Animation/AnimationSequenceCompiler.cpp
@@ -309,6 +309,11 @@ namespace UE::Anim
 			{
 				InAnimSequence->FinishAsyncTasks();
 			}
+			else if (InAnimSequence)
+			{				
+				UE_LOG(LogAnimation, Display, TEXT("Failed to finish Async Animation Compression for %s, as it is either unreachable (%i) has a null skeleton (%i) or its skeleton (%s) is unreachable (%i)"), *InAnimSequence->GetName(),
+					InAnimSequence->IsUnreachable(), InAnimSequence->GetSkeleton() == nullptr,  InAnimSequence->GetSkeleton() == nullptr ? TEXT("null") :  *InAnimSequence->GetName(), InAnimSequence->GetSkeleton()->IsUnreachable());
+			}
 		}
 	}
 	void FAnimSequenceCompilingManager::UpdateCompilationNotification()
@@ -344,6 +349,9 @@ namespace UE::Anim
 						{
 							PendingAnimationSequences.Add(AnimSequence);
 						}
+						
+						UE_LOG(LogAnimation, Display, TEXT("Cancelling Async Animation Compression for %s, as it is either unreachable (%i) has a null skeleton (%i) or its skeleton (%s) is unreachable (%i)"), *AnimSequence->GetName(),
+							AnimSequence->IsUnreachable(), AnimSequence->GetSkeleton() == nullptr,  AnimSequence->GetSkeleton() == nullptr ? TEXT("null") :  *AnimSequence->GetName(), AnimSequence->GetSkeleton()->IsUnreachable());
 					}
 				}
 			}
