Global Event Relay for Collision, Removal and Crumbling events.

diff --git a/Engine/Plugins/Experimental/ChaosCaching/Source/ChaosCaching/Private/Chaos/Adapters/GeometryCollectionComponentCacheAdapter.cpp b/Engine/Plugins/Experimental/ChaosCaching/Source/ChaosCaching/Private/Chaos/Adapters/GeometryCollectionComponentCacheAdapter.cpp
--- a/Engine/Plugins/Experimental/ChaosCaching/Source/ChaosCaching/Private/Chaos/Adapters/GeometryCollectionComponentCacheAdapter.cpp
+++ b/Engine/Plugins/Experimental/ChaosCaching/Source/ChaosCaching/Private/Chaos/Adapters/GeometryCollectionComponentCacheAdapter.cpp
@@ -184,7 +184,7 @@ namespace Chaos
 			}
 		}
 
-		if (Proxy->GetSimParameters().bGenerateBreakingData && BreakingDataArray && ProxyCachedEventData.ProxyBreakingDataIndices)
+		if ((Proxy->GetSimParameters().bGenerateBreakingData || Proxy->GetSimParameters().bGenerateGlobalBreakingData) && BreakingDataArray && ProxyCachedEventData.ProxyBreakingDataIndices)
 		{
 			for (int32 Index : *ProxyCachedEventData.ProxyBreakingDataIndices)
 			{
@@ -200,7 +200,7 @@ namespace Chaos
 			}
 		}
 		
-		if (Proxy->GetSimParameters().bGenerateCollisionData && CollisionDataArray && ProxyCachedEventData.ProxyCollisionDataIndices)
+		if ((Proxy->GetSimParameters().bGenerateCollisionData || Proxy->GetSimParameters().bGenerateGlobalCollisionData) && CollisionDataArray && ProxyCachedEventData.ProxyCollisionDataIndices)
 		{
 			for (int32 Index : *ProxyCachedEventData.ProxyCollisionDataIndices)
 			{
@@ -347,7 +347,7 @@ namespace Chaos
 			}
 		}
 
-		if (Proxy->GetSimParameters().bGenerateBreakingData && BreakingEvents)
+		if ((Proxy->GetSimParameters().bGenerateBreakingData || Proxy->GetSimParameters().bGenerateGlobalBreakingData) && BreakingEvents)
 		{
 			const FSolverBreakingEventFilter* SolverBreakingEventFilter = Solver->GetEventFilters()->GetBreakingFilter();
 
@@ -375,7 +375,7 @@ namespace Chaos
 							CachedBreak.Mass = Event->Mass;
 							CachedBreak.BoundingBox = TAABB<FReal, 3>(Event->BoundingBoxMin, Event->BoundingBoxMax);
 							CachedBreak.BoundingBox = CachedBreak.BoundingBox.TransformedAABB(ComponentToWorld);
-							CachedBreak.bIsGlobal = Proxy->GetSimParameters().bDispatchGlobalBreakingData;
+							CachedBreak.SetEmitterFlag(Proxy->GetSimParameters().bGenerateBreakingData, Proxy->GetSimParameters().bGenerateGlobalBreakingData);
 
 							if (!SolverBreakingEventFilter->Enabled() || SolverBreakingEventFilter->Pass(CachedBreak))
 							{
@@ -443,7 +443,7 @@ namespace Chaos
 			}
 		}
 
-		if (Proxy->GetSimParameters().bGenerateCollisionData && CollisionEvents)
+		if ((Proxy->GetSimParameters().bGenerateCollisionData || Proxy->GetSimParameters().bGenerateGlobalCollisionData) && CollisionEvents)
 		{
 
 			const FSolverCollisionEventFilter* SolverCollisionEventFilter = Solver->GetEventFilters()->GetCollisionFilter();
@@ -727,12 +727,12 @@ namespace Chaos
 			Chaos::FEventManager* EventManager = Solver->GetEventManager();
 			if (EventManager)
 			{
-				if (SimulationParameters.bGenerateBreakingData)
+				if (SimulationParameters.bGenerateBreakingData || SimulationParameters.bGenerateGlobalBreakingData)
 				{ 
 					EventManager->RegisterHandler<Chaos::FBreakingEventData>(Chaos::EEventType::Breaking, const_cast<FGeometryCollectionCacheAdapter*>(this), &FGeometryCollectionCacheAdapter::HandleBreakingEvents);
 				}
 
-				if (SimulationParameters.bGenerateCollisionData)
+				if (SimulationParameters.bGenerateCollisionData || SimulationParameters.bGenerateGlobalCollisionData)
 				{ 
 					EventManager->RegisterHandler<Chaos::FCollisionEventData>(Chaos::EEventType::Collision, const_cast<FGeometryCollectionCacheAdapter*>(this), &FGeometryCollectionCacheAdapter::HandleCollisionEvents);
 				}
diff --git a/Engine/Source/Runtime/Engine/Private/PhysicsEngine/Experimental/ChaosEventRelay.cpp b/Engine/Source/Runtime/Engine/Private/PhysicsEngine/Experimental/ChaosEventRelay.cpp
--- a/Engine/Source/Runtime/Engine/Private/PhysicsEngine/Experimental/ChaosEventRelay.cpp
+++ b/Engine/Source/Runtime/Engine/Private/PhysicsEngine/Experimental/ChaosEventRelay.cpp
@@ -6,10 +6,35 @@ UChaosEventRelay::UChaosEventRelay()
 {
 }
 
+void UChaosEventRelay::DispatchPhysicsCollisionEvents(const TArray<FCollisionChaosEvent>& CollisionEvents)
+{
+	
+	if (OnCollisionEvent.IsBound())
+	{
+		OnCollisionEvent.Broadcast(CollisionEvents);
+	}
+}
+
 void UChaosEventRelay::DispatchPhysicsBreakEvents(const TArray<FBreakChaosEvent>& BreakEvents)
 {
 	if (OnBreakEvent.IsBound())
 	{
 		OnBreakEvent.Broadcast(BreakEvents);
 	}
+}
+
+void UChaosEventRelay::DispatchPhysicsRemovalEvents(const TArray<FRemovalChaosEvent>& RemovalEvents)
+{
+	if (OnRemovalEvent.IsBound())
+	{
+		OnRemovalEvent.Broadcast(RemovalEvents);
+	}
+}
+
+void UChaosEventRelay::DispatchPhysicsCrumblingEvents(const TArray<FCrumblingChaosEvent>& CrumblingEvents)
+{
+	if (OnCrumblingEvent.IsBound())
+	{
+		OnCrumblingEvent.Broadcast(CrumblingEvents);
+	}
 }
\ No newline at end of file
diff --git a/Engine/Source/Runtime/Engine/Private/PhysicsEngine/Experimental/ChaosEventType.cpp b/Engine/Source/Runtime/Engine/Private/PhysicsEngine/Experimental/ChaosEventType.cpp
--- a/Engine/Source/Runtime/Engine/Private/PhysicsEngine/Experimental/ChaosEventType.cpp
+++ b/Engine/Source/Runtime/Engine/Private/PhysicsEngine/Experimental/ChaosEventType.cpp
@@ -6,6 +6,34 @@
 
 #include UE_INLINE_GENERATED_CPP_BY_NAME(ChaosEventType)
 
+FCollisionChaosEventBodyInfo::FCollisionChaosEventBodyInfo()
+	: PhysMaterial(nullptr)
+	, Component(nullptr)
+{
+}
+
+FCollisionChaosEvent::FCollisionChaosEvent()
+{
+}
+
+FCollisionChaosEvent::FCollisionChaosEvent(const Chaos::FCollidingData& CollisionData)
+	: Location(CollisionData.Location)
+	, AccumulatedImpulse(CollisionData.AccumulatedImpulse)
+	, Normal(CollisionData.Normal)
+	, PenetrationDepth(CollisionData.PenetrationDepth)
+{
+	Body1.Velocity = CollisionData.Velocity1;
+	Body1.DeltaVelocity = CollisionData.DeltaVelocity1;
+	Body1.AngularVelocity = CollisionData.AngularVelocity1;
+	Body1.Mass = CollisionData.Mass1;
+
+	Body2.Velocity = CollisionData.Velocity2;
+	Body2.DeltaVelocity = CollisionData.DeltaVelocity2;
+	Body2.AngularVelocity = CollisionData.AngularVelocity2;
+	Body2.Mass = CollisionData.Mass2;
+
+}
+
 FBreakChaosEvent::FBreakChaosEvent()
 	: Component(nullptr)
 	, Location(FVector::ZeroVector)
@@ -29,3 +57,10 @@ FBreakChaosEvent::FBreakChaosEvent(const Chaos::FBreakingData& BreakingData)
 	, bFromCrumble(BreakingData.bFromCrumble)
 {
 }
+
+FRemovalChaosEvent::FRemovalChaosEvent()
+	: Component(nullptr)
+	, Location(FVector::ZeroVector)
+	, Mass(0.0f)
+{
+}
diff --git a/Engine/Source/Runtime/Engine/Private/PhysicsEngine/Experimental/PhysScene_Chaos.cpp b/Engine/Source/Runtime/Engine/Private/PhysicsEngine/Experimental/PhysScene_Chaos.cpp
--- a/Engine/Source/Runtime/Engine/Private/PhysicsEngine/Experimental/PhysScene_Chaos.cpp
+++ b/Engine/Source/Runtime/Engine/Private/PhysicsEngine/Experimental/PhysScene_Chaos.cpp
@@ -481,6 +481,8 @@ FPhysScene_Chaos::FPhysScene_Chaos(AActor* InSolverActor
 	, SolverActor(InSolverActor)
 	, LastEventDispatchTime(Chaos::FReal(-1))
 	, LastBreakEventDispatchTime(Chaos::FReal(-1))
+	, LastRemovalEventDispatchTime(Chaos::FReal(-1))
+	, LastCrumblingEventDispatchTime(Chaos::FReal(-1))
 #if WITH_EDITOR
 	, SingleStepCounter(0)
 #endif
@@ -503,6 +505,8 @@ FPhysScene_Chaos::FPhysScene_Chaos(AActor* InSolverActor
 	Chaos::FEventManager* EventManager = SceneSolver->GetEventManager();
 	EventManager->RegisterHandler<Chaos::FCollisionEventData>(Chaos::EEventType::Collision, this, &FPhysScene_Chaos::HandleCollisionEvents);
 	EventManager->RegisterHandler<Chaos::FBreakingEventData>(Chaos::EEventType::Breaking, this, &FPhysScene_Chaos::HandleBreakingEvents);
+	EventManager->RegisterHandler<Chaos::FRemovalEventData>(Chaos::EEventType::Removal, this, &FPhysScene_Chaos::HandleRemovalEvents);
+	EventManager->RegisterHandler<Chaos::FCrumblingEventData>(Chaos::EEventType::Crumbling, this, &FPhysScene_Chaos::HandleCrumblingEvents);
 
 
 	//Initialize unique ptrs that are just here to allow forward declare. This should be reworked todo(ocohen)
@@ -586,6 +590,26 @@ void FPhysScene_Chaos::UnRegisterForCollisionEvents(UPrimitiveComponent* Compone
 	CollisionEventRegistrations.Remove(Component);
 }
 
+void FPhysScene_Chaos::RegisterForGlobalCollisionEvents(UPrimitiveComponent* Component)
+{
+	GlobalCollisionEventRegistrations.Add(Component);
+}
+
+void FPhysScene_Chaos::UnRegisterForGlobalCollisionEvents(UPrimitiveComponent* Component)
+{
+	GlobalCollisionEventRegistrations.Remove(Component);
+}
+
+void FPhysScene_Chaos::RegisterForGlobalRemovalEvents(UPrimitiveComponent* Component)
+{
+	GlobalRemovalEventRegistrations.Add(Component);
+}
+
+void FPhysScene_Chaos::UnRegisterForGlobalRemovalEvents(UPrimitiveComponent* Component)
+{
+	GlobalRemovalEventRegistrations.Remove(Component);
+}
+
 template<typename ObjectType>
 void AddPhysicsProxy(ObjectType* InObject, Chaos::FPhysicsSolver* InSolver)
 {
@@ -840,9 +864,9 @@ FORCEINLINE void FPhysScene_Chaos::HandleEachCollisionEvent(const TArray<int32>&
 
 		// Are the proxies pending destruction? If they are no longer tracked by the PhysScene, the proxy is deleted or pending deletion.
 		UPrimitiveComponent* const Comp1 = GetOwningComponent<UPrimitiveComponent>(PhysicsProxy1);
-		if (Comp0 == nullptr || Comp1 == nullptr)
+		if (Comp1 == nullptr)
 		{
-			return;
+			continue;
 		}
 
 		bool bNewEntry = false;
@@ -967,6 +991,9 @@ void FPhysScene_Chaos::HandleCollisionEvents(const Chaos::FCollisionEventData& E
 
 	// Tell the world and actors about the collisions
 	DispatchPendingCollisionNotifies();
+
+	HandleGlobalCollisionEvent(CollisionData);
+
 }
 
 void FPhysScene_Chaos::DispatchPendingCollisionNotifies()
@@ -979,57 +1006,155 @@ void FPhysScene_Chaos::DispatchPendingCollisionNotifies()
 		OwningWorld->PhysicsCollisionHandler->HandlePhysicsCollisions_AssumesLocked(PendingCollisionNotifies);
 	}
 
-	if (!PendingCollisionNotifies.IsEmpty())
+	// Fire any collision notifies in the queue.
+	for (FCollisionNotifyInfo& NotifyInfo : PendingCollisionNotifies)
 	{
-		// Fire any collision notifies in the queue.
-		for (FCollisionNotifyInfo& NotifyInfo : PendingCollisionNotifies)
+		if (NotifyInfo.bCallEvent0)
 		{
-			//if (NotifyInfo.RigidCollisionData.ContactInfos.Num() > 0)
+			if (AActor* Actor = NotifyInfo.Info0.Actor.Get())
 			{
-				if (NotifyInfo.bCallEvent0 /*&& NotifyInfo.IsValidForNotify()*/)
-				{
-					if (AActor* Actor = NotifyInfo.Info0.Actor.Get())
-					{
-						QUICK_SCOPE_CYCLE_COUNTER(STAT_DispatchPhysicsCollisionHit);
-						SCOPE_CYCLE_UOBJECT(NotifyActor, Actor);
-						Actor->DispatchPhysicsCollisionHit(NotifyInfo.Info0, NotifyInfo.Info1, NotifyInfo.RigidCollisionData);
-					}
-				}
+				QUICK_SCOPE_CYCLE_COUNTER(STAT_DispatchPhysicsCollisionHit);
+				SCOPE_CYCLE_UOBJECT(NotifyActor, Actor);
+				Actor->DispatchPhysicsCollisionHit(NotifyInfo.Info0, NotifyInfo.Info1, NotifyInfo.RigidCollisionData);
+			}
+		}
+	}
+	PendingCollisionNotifies.Reset();
+}
+
+void FPhysScene_Chaos::HandleGlobalCollisionEvent(Chaos::FCollisionDataArray const& CollisionData)
+{
+	QUICK_SCOPE_CYCLE_COUNTER(STAT_HandleGlobalCollisionEvents);
+	TArray<FCollisionChaosEvent> CollisionEvents;
+	CollisionEvents.Reserve(CollisionEvents.Num());
+
+	// If iterating by proxy, then need to process the duplicate collision, 
+	// so iterating by collision data should be the faster when using only globals
+	for (const Chaos::FCollidingData& CollisionItem : CollisionData)
+	{
+		UPrimitiveComponent* Comp0 = GetOwningComponent<UPrimitiveComponent>(CollisionItem.Proxy1);
+		UPrimitiveComponent* Comp1 = GetOwningComponent<UPrimitiveComponent>(CollisionItem.Proxy2);
+
+		if (GlobalCollisionEventRegistrations.Contains(Comp0) || GlobalCollisionEventRegistrations.Contains(Comp1))
+		{
+			FCollisionChaosEvent& CollisionEvent = CollisionEvents.Emplace_GetRef(CollisionItem);
+			CollisionEvent.Body1.Component = GetOwningComponent<UPrimitiveComponent>(CollisionItem.Proxy1);
+			CollisionEvent.Body2.Component = GetOwningComponent<UPrimitiveComponent>(CollisionItem.Proxy2);
+
+			const Chaos::FChaosPhysicsMaterial* InternalMat1 = CollisionItem.Mat1.Get();
+			const Chaos::FChaosPhysicsMaterial* InternalMat2 = CollisionItem.Mat2.Get();
+			CollisionEvent.Body1.PhysMaterial = InternalMat1 ? FPhysicsUserData::Get<UPhysicalMaterial>(InternalMat1->UserData) : nullptr;
+			CollisionEvent.Body2.PhysMaterial = InternalMat2 ? FPhysicsUserData::Get<UPhysicalMaterial>(InternalMat2->UserData) : nullptr;
 
-				// CHAOS: don't call event 1, because the code below will generate the reflexive hit data as separate entries
+			const FBodyInstance* BodyInst1 = GetBodyInstanceFromProxyAndShape(CollisionItem.Proxy1, CollisionItem.ShapeIndex1);
+			if (BodyInst1 != nullptr)
+			{
+				CollisionEvent.Body1.BodyIndex = BodyInst1->InstanceBodyIndex;
+				CollisionEvent.Body1.BoneName = BodyInst1->BodySetup.IsValid() ? BodyInst1->BodySetup->BoneName : NAME_None;
+			}
+
+			const FBodyInstance* BodyInst2 = GetBodyInstanceFromProxyAndShape(CollisionItem.Proxy2, CollisionItem.ShapeIndex2);
+			if (BodyInst2 != nullptr)
+			{
+				CollisionEvent.Body2.BodyIndex = BodyInst2->InstanceBodyIndex;
+				CollisionEvent.Body2.BoneName = BodyInst2->BodySetup.IsValid() ? BodyInst2->BodySetup->BoneName : NAME_None;
 			}
 		}
-		PendingCollisionNotifies.Reset();
+	}
+	if (CollisionEvents.Num() > 0)
+	{
+		ChaosEventRelay->DispatchPhysicsCollisionEvents(CollisionEvents);
 	}
 }
 
 void FPhysScene_Chaos::HandleBreakingEvents(const Chaos::FBreakingEventData& Event)
 {
-	QUICK_SCOPE_CYCLE_COUNTER(STAT_HandleBreakingEvents);
-
-	TArray<FBreakChaosEvent> PendingBreakEvents;
-	Chaos::FBreakingDataArray const& BreakingDataArray = Event.BreakingData.AllBreakingsArray;
+	QUICK_SCOPE_CYCLE_COUNTER(STAT_HandleGlobalBreakingEvents);
 
 	if (Event.BreakingData.bHasGlobalEvent && LastBreakEventDispatchTime < Event.BreakingData.TimeCreated)
 	{
 		LastBreakEventDispatchTime = Event.BreakingData.TimeCreated;
+		Chaos::FBreakingDataArray const& BreakingDataArray = Event.BreakingData.AllBreakingsArray;
+		TArray<FBreakChaosEvent> PendingBreakEvents;
 		for (const Chaos::FBreakingData& BreakingData: BreakingDataArray)
 		{
-			if (BreakingData.bIsGlobal && BreakingData.Proxy)
+			if ((BreakingData.EmitterFlag & Chaos::EventEmitterFlag::GlobalDispatcher) && BreakingData.Proxy)
 			{
 				UPrimitiveComponent* Comp = GetOwningComponent<UPrimitiveComponent>(BreakingData.Proxy);
 				FBreakChaosEvent& BreakEvent = PendingBreakEvents.Emplace_GetRef(BreakingData);
 				BreakEvent.Component = Comp;
 			}
 		}
+		if (PendingBreakEvents.Num() > 0)
+		{
+			ChaosEventRelay->DispatchPhysicsBreakEvents(PendingBreakEvents);
+		}
 	}
+}
 
-	if (PendingBreakEvents.Num() > 0)
+void FPhysScene_Chaos::HandleRemovalEvents(const Chaos::FRemovalEventData& Event)
+{
+	QUICK_SCOPE_CYCLE_COUNTER(STAT_HandleGlobalRemovalEvents);
+	
+	const float RemovalDataTimestamp = Event.RemovalData.TimeCreated;
+	if (RemovalDataTimestamp > LastRemovalEventDispatchTime)
 	{
-		ChaosEventRelay->DispatchPhysicsBreakEvents(PendingBreakEvents);
+		LastRemovalEventDispatchTime = RemovalDataTimestamp;
+		TArray<FRemovalChaosEvent> PendingRemovalEvents;
+		Chaos::FRemovalDataArray const& RemovalData = Event.RemovalData.AllRemovalArray;
+		for (Chaos::FRemovalData const& RemovalDataItem : RemovalData)
+		{
+			if (RemovalDataItem.Proxy)
+			{
+				UPrimitiveComponent* Comp = GetOwningComponent<UPrimitiveComponent>(RemovalDataItem.Proxy);
+				if (GlobalRemovalEventRegistrations.Contains(Comp))
+				{
+					FRemovalChaosEvent& RemovalEvent = PendingRemovalEvents.AddZeroed_GetRef();
+					RemovalEvent.Component = Comp;
+					RemovalEvent.Location = RemovalDataItem.Location;
+					RemovalEvent.Mass = RemovalDataItem.Mass;
+				}
+			}
+		}
+		if (PendingRemovalEvents.Num() > 0)
+		{
+			ChaosEventRelay->DispatchPhysicsRemovalEvents(PendingRemovalEvents);
+		}
 	}
 }
 
+void FPhysScene_Chaos::HandleCrumblingEvents(const Chaos::FCrumblingEventData& Event)
+{
+	QUICK_SCOPE_CYCLE_COUNTER(STAT_HandleGlobalCrumblingEvents);
+
+	const float CrumblingDataTimestamp = Event.CrumblingData.TimeCreated;
+	if (Event.CrumblingData.bHasGlobalEvent && CrumblingDataTimestamp > LastCrumblingEventDispatchTime)
+	{
+		LastCrumblingEventDispatchTime = CrumblingDataTimestamp;
+		TArray<FCrumblingChaosEvent> PendingCrumblingEvent;
+		for (const Chaos::FCrumblingData& CrumblingDataItem : Event.CrumblingData.AllCrumblingsArray)
+		{
+			if ((CrumblingDataItem.EmitterFlag & Chaos::EventEmitterFlag::GlobalDispatcher) && CrumblingDataItem.Proxy)
+			{
+				FCrumblingChaosEvent& CrumblingEvent = PendingCrumblingEvent.AddZeroed_GetRef();
+				CrumblingEvent.Component = GetOwningComponent<UPrimitiveComponent>(CrumblingDataItem.Proxy);
+				CrumblingEvent.Location = CrumblingDataItem.Location;
+				CrumblingEvent.Orientation = CrumblingDataItem.Orientation;
+				CrumblingEvent.LinearVelocity = CrumblingDataItem.LinearVelocity;
+				CrumblingEvent.AngularVelocity = CrumblingDataItem.AngularVelocity;
+				CrumblingEvent.Mass = static_cast<float>(CrumblingDataItem.Mass);
+				CrumblingEvent.LocalBounds = FBox(CrumblingDataItem.LocalBounds.Min(), CrumblingDataItem.LocalBounds.Max());
+				CrumblingEvent.Children = CrumblingDataItem.Children;
+			}
+		}
+		if (PendingCrumblingEvent.Num() > 0)
+		{
+			ChaosEventRelay->DispatchPhysicsCrumblingEvents(PendingCrumblingEvent);
+		}
+	}
+}
+
+
 #if CHAOS_WITH_PAUSABLE_SOLVER
 void FPhysScene_Chaos::OnUpdateWorldPause()
 {
diff --git a/Engine/Source/Runtime/Engine/Public/Physics/Experimental/ChaosEventRelay.h b/Engine/Source/Runtime/Engine/Public/Physics/Experimental/ChaosEventRelay.h
--- a/Engine/Source/Runtime/Engine/Public/Physics/Experimental/ChaosEventRelay.h
+++ b/Engine/Source/Runtime/Engine/Public/Physics/Experimental/ChaosEventRelay.h
@@ -3,13 +3,18 @@
 #pragma once
 
 #include "Physics/Experimental/ChaosEventType.h"
+#include "PhysicsPublic.h"
+#include "Engine/EngineTypes.h"
 #include "UObject/Object.h"
 
 #include "ChaosEventRelay.generated.h"
 
 
 
+DECLARE_DYNAMIC_MULTICAST_DELEGATE_OneParam(FCollisionEventSignature, const TArray<FCollisionChaosEvent>&, CollisionEvents);
 DECLARE_DYNAMIC_MULTICAST_DELEGATE_OneParam(FBreakEventSignature, const TArray<FBreakChaosEvent>&, BreakEvents);
+DECLARE_DYNAMIC_MULTICAST_DELEGATE_OneParam(FRemovalEventSignature, const TArray<FRemovalChaosEvent>&, RemovalEvents);
+DECLARE_DYNAMIC_MULTICAST_DELEGATE_OneParam(FCrumblingEventSignature, const TArray<FCrumblingChaosEvent>&, CrumblingEvents);
 
 
 /**
@@ -24,9 +29,22 @@ public:
 
 	UChaosEventRelay();
 
+	void DispatchPhysicsCollisionEvents(const TArray<FCollisionChaosEvent>& CollisionEvents);
 	void DispatchPhysicsBreakEvents(const TArray<FBreakChaosEvent>& BreakEvents);
+	void DispatchPhysicsRemovalEvents(const TArray<FRemovalChaosEvent>& RemovalEvents);
+	void DispatchPhysicsCrumblingEvents(const TArray<FCrumblingChaosEvent>& CrumblingEvents);
+
+	
+	UPROPERTY(BlueprintAssignable, Category = "Physics Events")
+	FCollisionEventSignature OnCollisionEvent;
 
 	UPROPERTY(BlueprintAssignable, Category = "Physics Events")
 	FBreakEventSignature OnBreakEvent;
 
+	UPROPERTY(BlueprintAssignable, Category = "Physics Events")
+	FRemovalEventSignature OnRemovalEvent;
+
+	UPROPERTY(BlueprintAssignable, Category = "Physics Events")
+	FCrumblingEventSignature OnCrumblingEvent;
+
 };
diff --git a/Engine/Source/Runtime/Engine/Public/Physics/Experimental/ChaosEventType.h b/Engine/Source/Runtime/Engine/Public/Physics/Experimental/ChaosEventType.h
--- a/Engine/Source/Runtime/Engine/Public/Physics/Experimental/ChaosEventType.h
+++ b/Engine/Source/Runtime/Engine/Public/Physics/Experimental/ChaosEventType.h
@@ -11,8 +11,74 @@
 namespace Chaos
 {
 	struct FBreakingData;
+	struct FCollidingData;
 }
 
+USTRUCT(BlueprintType)
+struct ENGINE_API FCollisionChaosEventBodyInfo
+{
+	GENERATED_BODY()
+
+public:
+	FCollisionChaosEventBodyInfo();
+
+	UPROPERTY(BlueprintReadOnly, Category = "Collision Event")
+	FVector Velocity;
+
+	UPROPERTY(BlueprintReadOnly, Category = "Collision Event")
+	FVector DeltaVelocity;
+
+	UPROPERTY(BlueprintReadOnly, Category = "Collision Event")
+	FVector AngularVelocity;
+
+	UPROPERTY(BlueprintReadOnly, Category = "Collision Event")
+	float Mass;
+
+	UPROPERTY(BlueprintReadOnly, Category = "Collision Event")
+	TObjectPtr<class UPhysicalMaterial> PhysMaterial;
+
+	UPROPERTY(BlueprintReadOnly, Category = "Collision Event")
+	TWeakObjectPtr<UPrimitiveComponent> Component;
+
+	UPROPERTY(BlueprintReadOnly, Category = "Collision Event")
+	int32 BodyIndex;
+
+	UPROPERTY(BlueprintReadOnly, Category = "Collision Event")
+	FName BoneName;
+};
+
+
+
+USTRUCT(BlueprintType)
+struct ENGINE_API FCollisionChaosEvent
+{
+	GENERATED_BODY()
+
+public:
+	FCollisionChaosEvent();
+	FCollisionChaosEvent(const Chaos::FCollidingData& CollisionData);
+
+	UPROPERTY(BlueprintReadOnly, Category = "Collision Event")
+	FVector Location;
+
+	UPROPERTY(BlueprintReadOnly, Category = "Collision Event")
+	FVector AccumulatedImpulse;
+
+	UPROPERTY(BlueprintReadOnly, Category = "Collision Event")
+	FVector Normal;
+
+	UPROPERTY(BlueprintReadOnly, Category = "Collision Event")
+	float PenetrationDepth;
+
+	UPROPERTY(BlueprintReadOnly, Category = "Collision Event")
+	FCollisionChaosEventBodyInfo Body1;
+
+	UPROPERTY(BlueprintReadOnly, Category = "Collision Event")
+	FCollisionChaosEventBodyInfo Body2;
+
+};
+
+
 USTRUCT(BlueprintType)
 struct ENGINE_API FBreakChaosEvent
 {
@@ -56,3 +122,74 @@ public:
 	bool bFromCrumble;
 };
 
+USTRUCT(BlueprintType)
+struct ENGINE_API FRemovalChaosEvent
+{
+	GENERATED_BODY()
+
+public:
+
+	FRemovalChaosEvent();
+
+	UPROPERTY(BlueprintReadOnly, Category = "Removal Event")
+	TObjectPtr<UPrimitiveComponent> Component = nullptr;
+
+	UPROPERTY(BlueprintReadOnly, Category = "Removal Event")
+	FVector Location;
+
+	UPROPERTY(BlueprintReadOnly, Category = "Removal Event")
+	float Mass;
+};
+
+
+USTRUCT(BlueprintType)
+struct ENGINE_API FCrumblingChaosEvent
+{
+	GENERATED_BODY()
+
+public:
+	FCrumblingChaosEvent()
+		: Component(nullptr)
+		, Location(FVector::ZeroVector)
+		, Orientation(FQuat::Identity)
+		, LinearVelocity(FVector::ZeroVector)
+		, AngularVelocity(FVector::ZeroVector)
+		, Mass(0)
+		, LocalBounds(ForceInitToZero)
+	{}
+
+	/** primitive component involved in the crumble event */
+	UPROPERTY(BlueprintReadOnly, Category = "Crumble Event")
+		TObjectPtr<UPrimitiveComponent> Component = nullptr;
+
+	/** World location of the crumbling cluster */
+	UPROPERTY(BlueprintReadOnly, Category = "Crumble Event")
+		FVector Location;
+
+	/** World orientation of the crumbling cluster */
+	UPROPERTY(BlueprintReadOnly, Category = "Crumble Event")
+		FQuat Orientation;
+
+	/** Linear Velocity of the crumbling cluster */
+	UPROPERTY(BlueprintReadOnly, Category = "Crumble Event")
+		FVector LinearVelocity;
+
+	/** Angular Velocity of the crumbling cluster  */
+	UPROPERTY(BlueprintReadOnly, Category = "Crumble Event")
+		FVector AngularVelocity;
+
+	/** Mass of the crumbling cluster  */
+	UPROPERTY(BlueprintReadOnly, Category = "Crumble Event")
+		float Mass;
+
+	/** Local bounding box of the crumbling cluster  */
+	UPROPERTY(BlueprintReadOnly, Category = "Crumble Event")
+		FBox LocalBounds;
+
+	/** List of children indices released (optional : see geometry collection component bCrumblingEventIncludesChildren) */
+	UPROPERTY(BlueprintReadOnly, Category = "Crumble Event")
+		TArray<int32> Children;
+};
+
+
+
diff --git a/Engine/Source/Runtime/Engine/Public/Physics/Experimental/PhysScene_Chaos.h b/Engine/Source/Runtime/Engine/Public/Physics/Experimental/PhysScene_Chaos.h
--- a/Engine/Source/Runtime/Engine/Public/Physics/Experimental/PhysScene_Chaos.h
+++ b/Engine/Source/Runtime/Engine/Public/Physics/Experimental/PhysScene_Chaos.h
@@ -114,9 +114,14 @@ public:
 	AActor* GetSolverActor() const;
 
 	void RegisterForCollisionEvents(UPrimitiveComponent* Component);
-
 	void UnRegisterForCollisionEvents(UPrimitiveComponent* Component);
 
+	void RegisterForGlobalCollisionEvents(UPrimitiveComponent* Component);
+	void UnRegisterForGlobalCollisionEvents(UPrimitiveComponent* Component);
+
+	void RegisterForGlobalRemovalEvents(UPrimitiveComponent* Component);
+	void UnRegisterForGlobalRemovalEvents(UPrimitiveComponent* Component);
+
 	void RegisterAsyncPhysicsTickComponent(UActorComponent* Component);
 	void UnregisterAsyncPhysicsTickComponent(UActorComponent* Component);
 
@@ -249,8 +254,9 @@ public:
 	FPrimitiveComponentReplicationCache ReplicationCache;
 
 private:
-	UPROPERTY()
 	TSet<UPrimitiveComponent*> CollisionEventRegistrations;
+	TSet<UPrimitiveComponent*> GlobalCollisionEventRegistrations;
+	TSet<UPrimitiveComponent*> GlobalRemovalEventRegistrations;
 
 	// contains the set of properties that uniquely identifies a reported collision
 	// Note that order matters, { Body0, Body1 } is not the same as { Body1, Body0 }
@@ -279,10 +285,13 @@ private:
 
 	// Chaos Event Handlers
 	void HandleEachCollisionEvent(const TArray<int32>& CollisionIndices, IPhysicsProxyBase* PhysicsProxy0, UPrimitiveComponent* const Comp0, Chaos::FCollisionDataArray const& CollisionData, Chaos::FReal MinDeltaVelocityThreshold);
+	void HandleGlobalCollisionEvent(Chaos::FCollisionDataArray const& CollisionData);
 	void HandleCollisionEvents(const Chaos::FCollisionEventData& CollisionData);
 	void DispatchPendingCollisionNotifies();
 
 	void HandleBreakingEvents(const Chaos::FBreakingEventData& Event);
+	void HandleRemovalEvents(const Chaos::FRemovalEventData& Event);
+	void HandleCrumblingEvents(const Chaos::FCrumblingEventData& Event);
 
 	/** Replication manager that updates physics bodies towards replicated physics state */
 	FPhysicsReplication* PhysicsReplication;
@@ -372,6 +381,9 @@ private:
 
 	Chaos::FReal LastEventDispatchTime;
 	Chaos::FReal LastBreakEventDispatchTime;
+	Chaos::FReal LastRemovalEventDispatchTime;
+	Chaos::FReal LastCrumblingEventDispatchTime;
+	
 	class FAsyncPhysicsTickCallback* AsyncPhysicsTickCallback = nullptr;
 
 #if WITH_EDITOR
diff --git a/Engine/Source/Runtime/Experimental/Chaos/Private/Chaos/PBDRigidClustering.cpp b/Engine/Source/Runtime/Experimental/Chaos/Private/Chaos/PBDRigidClustering.cpp
--- a/Engine/Source/Runtime/Experimental/Chaos/Private/Chaos/PBDRigidClustering.cpp
+++ b/Engine/Source/Runtime/Experimental/Chaos/Private/Chaos/PBDRigidClustering.cpp
@@ -555,7 +555,7 @@ namespace Chaos
 		if (FGeometryCollectionPhysicsProxy* ConcreteProxy = GetConcreteProxy(ClusteredParticle))
 		{
 			const FSimulationParameters& SimParams = ConcreteProxy->GetSimParameters();
-			if (SimParams.bGenerateBreakingData)
+			if (SimParams.bGenerateBreakingData || SimParams.bGenerateGlobalBreakingData)
 			{
 				FBreakingData& ClusterBreak = MAllClusterBreakings.AddDefaulted_GetRef();
 				ClusterBreak.Proxy = ClusteredParticle->PhysicsProxy();
@@ -569,7 +569,7 @@ namespace Chaos
 				}
 				ClusterBreak.TransformGroupIndex = ConcreteProxy->GetTransformGroupIndexFromHandle(ClusteredParticle);
 				ClusterBreak.bFromCrumble = bFromCrumble;
-				ClusterBreak.bIsGlobal = SimParams.bDispatchGlobalBreakingData;
+				ClusterBreak.SetEmitterFlag(SimParams.bGenerateBreakingData, SimParams.bGenerateGlobalBreakingData);
 			}
 		}
 	}
@@ -581,7 +581,7 @@ namespace Chaos
 		if (FGeometryCollectionPhysicsProxy* ConcreteProxy = GetConcreteProxy(ClusteredParticle))
 		{
 			const FSimulationParameters& SimParams = ConcreteProxy->GetSimParameters(); 
-			if (SimParams.bGenerateCrumblingData)
+			if (SimParams.bGenerateCrumblingData || SimParams.bGenerateGlobalCrumblingData)
 			{
 				FCrumblingData& ClusterCrumbling = MAllClusterCrumblings.AddDefaulted_GetRef();
 				ClusterCrumbling.Proxy = ClusteredParticle->PhysicsProxy();
@@ -590,11 +590,12 @@ namespace Chaos
 				ClusterCrumbling.LinearVelocity = ClusteredParticle->V();
 				ClusterCrumbling.AngularVelocity = ClusteredParticle->W();
 				ClusterCrumbling.Mass = ClusteredParticle->M();
+				ClusterCrumbling.SetEmitterFlag(SimParams.bGenerateCrumblingData, SimParams.bGenerateGlobalCrumblingData);
 				if (ClusteredParticle->Geometry() && ClusteredParticle->Geometry()->HasBoundingBox())
 				{
 					ClusterCrumbling.LocalBounds = ClusteredParticle->Geometry()->BoundingBox();
 				}
-				if (SimParams.bGenerateCrumblingChildrenData)
+				if (SimParams.bGenerateCrumblingChildrenData || SimParams.bGenerateGlobalCrumblingChildrenData)
 				{
 					// when sending this event, children are still attached
 					if (const FRigidHandleArray* Children = MChildren.Find(ClusteredParticle))
diff --git a/Engine/Source/Runtime/Experimental/Chaos/Private/Chaos/PBDRigidsEvolution.cpp b/Engine/Source/Runtime/Experimental/Chaos/Private/Chaos/PBDRigidsEvolution.cpp
--- a/Engine/Source/Runtime/Experimental/Chaos/Private/Chaos/PBDRigidsEvolution.cpp
+++ b/Engine/Source/Runtime/Experimental/Chaos/Private/Chaos/PBDRigidsEvolution.cpp
@@ -1111,8 +1111,7 @@ namespace Chaos
 		}
 		
 		// Record removal for event generation
-		const int32 NewIdx = MAllRemovals.Add(FRemovalData());
-		FRemovalData& Removal = MAllRemovals[NewIdx];
+		FRemovalData& Removal = MAllRemovals.AddDefaulted_GetRef();
 		Removal.Proxy = Particle->PhysicsProxy();
 		Removal.Location = Particle->X();
 		
diff --git a/Engine/Source/Runtime/Experimental/Chaos/Private/EventDefaults.cpp b/Engine/Source/Runtime/Experimental/Chaos/Private/EventDefaults.cpp
--- a/Engine/Source/Runtime/Experimental/Chaos/Private/EventDefaults.cpp
+++ b/Engine/Source/Runtime/Experimental/Chaos/Private/EventDefaults.cpp
@@ -345,7 +345,7 @@ namespace Chaos
 					{
 						const int32 NewIndex = FilteredBreakingDataArray.Emplace(ClusterBreaking);
 						FilteredBreakingIndicesByPhysicsProxy.FindOrAdd(ClusterBreaking.Proxy).Add(FEventManager::EncodeCollisionIndex(NewIndex, false));
-						BreakingEventData.BreakingData.bHasGlobalEvent |= ClusterBreaking.bIsGlobal;
+						BreakingEventData.BreakingData.bHasGlobalEvent |= (ClusterBreaking.EmitterFlag & EventEmitterFlag::GlobalDispatcher) != 0;
 					}
 				}
 			}
@@ -531,10 +531,8 @@ namespace Chaos
 					int32 NewIdx = AllRemovalDataArray.Add(FRemovalData());
 					FRemovalData& RemovalDataArrayItem = AllRemovalDataArray[NewIdx];
 					RemovalDataArrayItem = RemovalData;
-
 					AllRemovalIndicesByPhysicsProxy.FindOrAdd(RemovalData.Proxy).Add(FEventManager::EncodeCollisionIndex(NewIdx, false));
 				}
-			
 			});
 	}
 
@@ -554,6 +552,7 @@ namespace Chaos
 			if (bResetData)
 			{
 				CrumblingEventData.Reset();
+				CrumblingEventData.CrumblingData.bHasGlobalEvent = false;
 			}
 			CrumblingEventData.SetTimeCreated(Solver->MTime);
 
@@ -564,6 +563,7 @@ namespace Chaos
 			{
 				// todo(chaos) : implement filtering only if necessary as crumbling event should be sparser than breaking ones
 				CrumblingEventData.AddCrumbling(Crumbling);
+				CrumblingEventData.CrumblingData.bHasGlobalEvent |= (Crumbling.EmitterFlag & EventEmitterFlag::GlobalDispatcher) != 0;
 			}
 		});
 	}
diff --git a/Engine/Source/Runtime/Experimental/Chaos/Private/PhysicsProxy/GeometryCollectionPhysicsProxy.cpp b/Engine/Source/Runtime/Experimental/Chaos/Private/PhysicsProxy/GeometryCollectionPhysicsProxy.cpp
--- a/Engine/Source/Runtime/Experimental/Chaos/Private/PhysicsProxy/GeometryCollectionPhysicsProxy.cpp
+++ b/Engine/Source/Runtime/Experimental/Chaos/Private/PhysicsProxy/GeometryCollectionPhysicsProxy.cpp
@@ -2464,6 +2464,20 @@ void FGeometryCollectionPhysicsProxy::SetNotifyBreakings_External(bool bNotify)
 	SetProxyDirty_External();
 }
 
+void FGeometryCollectionPhysicsProxy::SetNotifyRemovals_External(bool bNotify)
+{
+	check(IsInGameThread());
+	GameThreadPerFrameData.SetNotifyRemovals(bNotify);
+	SetProxyDirty_External();
+}
+
+void FGeometryCollectionPhysicsProxy::SetNotifyCrumblings_External(bool bNotify, bool bIncludeChildren)
+{
+	check(IsInGameThread());
+	GameThreadPerFrameData.SetNotifyCrumblings(bNotify, bIncludeChildren);
+	SetProxyDirty_External();
+}
+
 void FGeometryCollectionPhysicsProxy::SetNotifyGlobalBreakings_External(bool bNotify)
 {
 	check(IsInGameThread());
@@ -2471,17 +2485,17 @@ void FGeometryCollectionPhysicsProxy::SetNotifyGlobalBreakings_External(bool bNo
 	SetProxyDirty_External();
 }
 
-void FGeometryCollectionPhysicsProxy::SetNotifyRemovals_External(bool bNotify)
+void FGeometryCollectionPhysicsProxy::SetNotifyGlobalRemovals_External(bool bNotify)
 {
 	check(IsInGameThread());
-	GameThreadPerFrameData.SetNotifyRemovals(bNotify);
+	GameThreadPerFrameData.SetNotifyGlobalRemovals(bNotify);
 	SetProxyDirty_External();
 }
 
-void FGeometryCollectionPhysicsProxy::SetNotifyCrumblings_External(bool bNotify, bool bIncludeChildren)
+void FGeometryCollectionPhysicsProxy::SetNotifyGlobalCrumblings_External(bool bNotify, bool bIncludeChildren)
 {
 	check(IsInGameThread());
-	GameThreadPerFrameData.SetNotifyCrumblings(bNotify, bIncludeChildren);
+	GameThreadPerFrameData.SetNotifyGlobalCrumblings(bNotify, bIncludeChildren);
 	SetProxyDirty_External();
 }
 
@@ -2698,10 +2712,11 @@ void FGeometryCollectionPhysicsProxy::PushStateOnGameThread(Chaos::FPBDRigidsSol
 	if (GameThreadPerFrameData.GetIsNotificationDataDirty())
 	{
 		Parameters.bGenerateBreakingData = GameThreadPerFrameData.GetNotifyBreakings() || GameThreadPerFrameData.GetNotifyGlobalBreakings();
-		Parameters.bDispatchGlobalBreakingData = GameThreadPerFrameData.GetNotifyGlobalBreakings();
-		Parameters.bGenerateRemovalsData = GameThreadPerFrameData.GetNotifyRemovals();
 		Parameters.bGenerateCrumblingData = GameThreadPerFrameData.GetNotifyCrumblings();
 		Parameters.bGenerateCrumblingChildrenData = GameThreadPerFrameData.GetCrumblingEventIncludesChildren();
+		Parameters.bGenerateGlobalBreakingData = GameThreadPerFrameData.GetNotifyGlobalBreakings();
+		Parameters.bGenerateGlobalCrumblingData = GameThreadPerFrameData.GetNotifyGlobalCrumblings();
+		Parameters.bGenerateGlobalCrumblingChildrenData = GameThreadPerFrameData.GetGlobalCrumblingEventIncludesChildren();
 		GameThreadPerFrameData.ResetIsNotificationDataDirty();
 	}
 
@@ -4870,4 +4885,4 @@ void FRuntimeDataCollector::RemoveCollector(const FGuid& Guid)
 {
 	Collectors.Remove(Guid);
 	
-}
\ No newline at end of file
+}
diff --git a/Engine/Source/Runtime/Experimental/Chaos/Public/Chaos/ExternalCollisionData.h b/Engine/Source/Runtime/Experimental/Chaos/Public/Chaos/ExternalCollisionData.h
--- a/Engine/Source/Runtime/Experimental/Chaos/Public/Chaos/ExternalCollisionData.h
+++ b/Engine/Source/Runtime/Experimental/Chaos/Public/Chaos/ExternalCollisionData.h
@@ -12,11 +12,39 @@ class UPhysicalMaterial;
 
 namespace Chaos
 {
+	// Event Emitter flag
+	enum EventEmitterFlag
+	{
+		EmptyDispatcher = 0,
+		OwnDispatcher = 1,
+		GlobalDispatcher = 2,
+		BothDispatcher = 3,
+	};
 
 	struct FBaseEventFlag
 	{
-		FBaseEventFlag() : bIsGlobal(false) {}
-		bool bIsGlobal;
+		FBaseEventFlag() : EmitterFlag(OwnDispatcher) {}
+
+		void SetEmitterFlag(bool LocalEmitter, bool GlobalEmitter)
+		{
+			EmitterFlag = ComputeEmitterFlag(LocalEmitter, GlobalEmitter);
+		}
+
+		static EventEmitterFlag ComputeEmitterFlag(bool LocalEmitter, bool GlobalEmitter)
+		{
+			EventEmitterFlag EmitterFlagOut = EventEmitterFlag::EmptyDispatcher;
+			if (LocalEmitter)
+			{
+				EmitterFlagOut = EventEmitterFlag::OwnDispatcher;
+			}
+			if (GlobalEmitter)
+			{
+				EmitterFlagOut = EventEmitterFlag(EmitterFlagOut | EventEmitterFlag::GlobalDispatcher);
+			}
+			return EmitterFlagOut;
+		}
+
+		EventEmitterFlag EmitterFlag;
 	};
 
 	/**
@@ -211,10 +239,11 @@ namespace Chaos
 	/*
 	CrumblingData passed from the physics solver to subsystems
 	*/
-	struct FCrumblingData
+	struct FCrumblingData : public FBaseEventFlag
 	{
 		FCrumblingData()
-			: Proxy(nullptr)
+			: FBaseEventFlag()
+			, Proxy(nullptr)
 			, Location(FVec3::ZeroVector)
 			, Orientation(FRotation3::Identity)
 			, LinearVelocity(FVec3::ZeroVector)
@@ -223,17 +252,6 @@ namespace Chaos
 			, LocalBounds(FAABB3(FVec3((FReal)0.0), FVec3((FReal)0.0)))
 		{}
 
-		FCrumblingData(const FCrumblingData& Other)
-			: Proxy(Other.Proxy)
-			, Location(Other.Location)
-			, Orientation(Other.Orientation)
-			, LinearVelocity(Other.LinearVelocity)
-			, AngularVelocity(Other.AngularVelocity)
-			, Mass(Other.Mass)
-			, LocalBounds(Other.LocalBounds)
-			, Children(Other.Children)
-		{}
-		
 		// The pointer to the proxy should be used with caution on the Game Thread.
 		// Ideally we only ever use this as a table key when acquiring related structures.
 		// If we genuinely need to dereference the pointer for any reason, test if it is deleted (nullptr) or
diff --git a/Engine/Source/Runtime/Experimental/Chaos/Public/EventsData.h b/Engine/Source/Runtime/Experimental/Chaos/Public/EventsData.h
--- a/Engine/Source/Runtime/Experimental/Chaos/Public/EventsData.h
+++ b/Engine/Source/Runtime/Experimental/Chaos/Public/EventsData.h
@@ -158,7 +158,7 @@ namespace Chaos
 	*/
 	struct FAllCrumblingData : public FTimeResource
 	{
-		FAllCrumblingData() : AllCrumblingsArray(FCrumblingDataArray()) {}
+		FAllCrumblingData() : AllCrumblingsArray(FCrumblingDataArray()), bHasGlobalEvent(false) {}
 
 		void Reset()
 		{
@@ -166,6 +166,7 @@ namespace Chaos
 		}
 
 		FCrumblingDataArray AllCrumblingsArray;
+		bool bHasGlobalEvent;
 	};
 
 	struct FCrumblingEventData
diff --git a/Engine/Source/Runtime/Experimental/Chaos/Public/GeometryCollection/GeometryCollectionSimulationCoreTypes.h b/Engine/Source/Runtime/Experimental/Chaos/Public/GeometryCollection/GeometryCollectionSimulationCoreTypes.h
--- a/Engine/Source/Runtime/Experimental/Chaos/Public/GeometryCollection/GeometryCollectionSimulationCoreTypes.h
+++ b/Engine/Source/Runtime/Experimental/Chaos/Public/GeometryCollection/GeometryCollectionSimulationCoreTypes.h
@@ -202,10 +202,12 @@ struct FSimulationParameters
 		, bGenerateBreakingData(false)
 		, bGenerateCollisionData(false)
 		, bGenerateTrailingData(false)
-		, bGenerateRemovalsData(false)
 		, bGenerateCrumblingData(false)
 		, bGenerateCrumblingChildrenData(false)
-		, bDispatchGlobalBreakingData(false)
+		, bGenerateGlobalBreakingData(false)
+		, bGenerateGlobalCollisionData(false)
+		, bGenerateGlobalCrumblingData(false)
+		, bGenerateGlobalCrumblingChildrenData(false)
 		, EnableGravity(true)
 		, GravityGroupIndex(0)
 		, UseInertiaConditioning(true)
@@ -257,10 +259,12 @@ struct FSimulationParameters
 		, bGenerateBreakingData(Other.bGenerateBreakingData)
 		, bGenerateCollisionData(Other.bGenerateCollisionData)
 		, bGenerateTrailingData(Other.bGenerateTrailingData)
-		, bGenerateRemovalsData(Other.bGenerateRemovalsData)
 		, bGenerateCrumblingData(Other.bGenerateCrumblingData)
 		, bGenerateCrumblingChildrenData(Other.bGenerateCrumblingChildrenData)
-		, bDispatchGlobalBreakingData(Other.bDispatchGlobalBreakingData)
+		, bGenerateGlobalBreakingData(Other.bGenerateGlobalBreakingData)
+		, bGenerateGlobalCollisionData(Other.bGenerateGlobalCollisionData)
+		, bGenerateGlobalCrumblingData(Other.bGenerateGlobalCrumblingData)
+		, bGenerateGlobalCrumblingChildrenData(Other.bGenerateGlobalCrumblingChildrenData)
 		, Shared(Other.Shared)
 		, EnableGravity(Other.EnableGravity)
 		, GravityGroupIndex(Other.GravityGroupIndex)
@@ -340,11 +344,13 @@ struct FSimulationParameters
 	bool bGenerateBreakingData;
 	bool bGenerateCollisionData;
 	bool bGenerateTrailingData;
-	bool bGenerateRemovalsData;
 	bool bGenerateCrumblingData;
 	bool bGenerateCrumblingChildrenData;
 
-	bool bDispatchGlobalBreakingData;
+	bool bGenerateGlobalBreakingData;
+	bool bGenerateGlobalCollisionData;
+	bool bGenerateGlobalCrumblingData;
+	bool bGenerateGlobalCrumblingChildrenData;
 
 	FSharedSimulationParameters Shared;
 
diff --git a/Engine/Source/Runtime/Experimental/Chaos/Public/GeometryCollectionProxyData.h b/Engine/Source/Runtime/Experimental/Chaos/Public/GeometryCollectionProxyData.h
--- a/Engine/Source/Runtime/Experimental/Chaos/Public/GeometryCollectionProxyData.h
+++ b/Engine/Source/Runtime/Experimental/Chaos/Public/GeometryCollectionProxyData.h
@@ -163,6 +163,10 @@ public:
 		, bNotifyRemovals(false)
 		, bNotifyCrumblings(false)
 		, bCrumblingEventIncludesChildren(false)
+		, bNotifyGlobalBreakings(false)
+		, bNotifyGlobalRemovals(false)
+		, bNotifyGlobalCrumblings(false)
+		, bGlobalCrumblingEventIncludesChildren(false)
 		, bEnableStrainOnCollision(false)
 	{}
 
@@ -204,13 +208,6 @@ public:
 	}
 	bool GetNotifyBreakings() const { return bNotifyBreakings; }
 
-	void SetNotifyGlobalBreakings(bool bNotify)
-	{
-		bNotifyGlobalBreakings = bNotify;
-		bIsNotificationDataDirty = true;
-	}
-	bool GetNotifyGlobalBreakings() const { return bNotifyGlobalBreakings; }
-
 	void SetNotifyRemovals(bool bNotify)
 	{
 		bNotifyRemovals = bNotify;
@@ -229,6 +226,31 @@ public:
 	bool GetNotifyCrumblings() const { return bNotifyCrumblings; }
 	bool GetCrumblingEventIncludesChildren() const { return bCrumblingEventIncludesChildren; }
 
+	void SetNotifyGlobalBreakings(bool bNotify)
+	{
+		bNotifyGlobalBreakings = bNotify;
+		bIsNotificationDataDirty = true;
+	}
+	bool GetNotifyGlobalBreakings() const { return bNotifyGlobalBreakings; }
+
+	void SetNotifyGlobalRemovals(bool bNotify)
+	{
+		bNotifyGlobalRemovals = bNotify;
+		bIsNotificationDataDirty = true;
+	}
+	bool GetNotifyGlobalRemovals() const { return bNotifyGlobalRemovals; }
+
+	void SetNotifyGlobalCrumblings(bool bNotify, bool bIncludeChildren)
+	{
+		bNotifyGlobalCrumblings = bNotify;
+		bGlobalCrumblingEventIncludesChildren = bIncludeChildren;
+		bIsNotificationDataDirty = true;
+	}
+
+	bool GetNotifyGlobalCrumblings() const { return bNotifyGlobalCrumblings; }
+	bool GetGlobalCrumblingEventIncludesChildren() const { return bGlobalCrumblingEventIncludesChildren; }
+
+
 	bool GetIsNotificationDataDirty() const { return bIsNotificationDataDirty; }
 	void ResetIsNotificationDataDirty() { bIsNotificationDataDirty = false; }
 
@@ -255,6 +277,9 @@ private:
 	uint16 bNotifyCrumblings : 1;
 	uint16 bCrumblingEventIncludesChildren : 1;
 	uint16 bNotifyGlobalBreakings : 1;
+	uint16 bNotifyGlobalRemovals : 1;
+	uint16 bNotifyGlobalCrumblings : 1;
+	uint16 bGlobalCrumblingEventIncludesChildren : 1;
 
 	/** updated when bDamageSettingsDataDirty is set */
 	uint16 bEnableStrainOnCollision : 1;
diff --git a/Engine/Source/Runtime/Experimental/Chaos/Public/PhysicsProxy/GeometryCollectionPhysicsProxy.h b/Engine/Source/Runtime/Experimental/Chaos/Public/PhysicsProxy/GeometryCollectionPhysicsProxy.h
--- a/Engine/Source/Runtime/Experimental/Chaos/Public/PhysicsProxy/GeometryCollectionPhysicsProxy.h
+++ b/Engine/Source/Runtime/Experimental/Chaos/Public/PhysicsProxy/GeometryCollectionPhysicsProxy.h
@@ -424,6 +424,8 @@ public:
 	void SetNotifyCrumblings_External(bool bNotify, bool bIncludeChildren);
 
 	void SetNotifyGlobalBreakings_External(bool bNotify);
+	void SetNotifyGlobalRemovals_External(bool bNotify);
+	void SetNotifyGlobalCrumblings_External(bool bNotify, bool bIncludeChildren);
 
 	float ComputeMaterialBasedDamageThreshold_Internal(Chaos::FPBDRigidClusteredParticleHandle& ClusteredParticle) const;
 
diff --git a/Engine/Source/Runtime/Experimental/ChaosSolverEngine/Private/Chaos/ChaosGameplayEventDispatcher.cpp b/Engine/Source/Runtime/Experimental/ChaosSolverEngine/Private/Chaos/ChaosGameplayEventDispatcher.cpp
--- a/Engine/Source/Runtime/Experimental/ChaosSolverEngine/Private/Chaos/ChaosGameplayEventDispatcher.cpp
+++ b/Engine/Source/Runtime/Experimental/ChaosSolverEngine/Private/Chaos/ChaosGameplayEventDispatcher.cpp
@@ -19,13 +19,6 @@
 #include UE_INLINE_GENERATED_CPP_BY_NAME(ChaosGameplayEventDispatcher)
 
 
-FChaosRemovalEvent::FChaosRemovalEvent()
-	: Component(nullptr)
-	, Location(FVector::ZeroVector)
-	, Mass(0.0f)
-{
-
-}
 
 UChaosGameplayEventDispatcher::UChaosGameplayEventDispatcher()
 {
@@ -49,9 +42,9 @@ void UChaosGameplayEventDispatcher::OnUnregister()
 }
 
 // internal
-static void DispatchPendingBreakEvents(TArray<FChaosBreakEvent> const& Events, TMap<TObjectPtr<UPrimitiveComponent>, FBreakEventCallbackWrapper> const& Registrations)
+static void DispatchPendingBreakEvents(TArray<FBreakChaosEvent> const& Events, TMap<TObjectPtr<UPrimitiveComponent>, FBreakEventCallbackWrapper> const& Registrations)
 {
-	for (FChaosBreakEvent const& E : Events)
+	for (FBreakChaosEvent const& E : Events)
 	{
 		if (E.Component)
 		{
@@ -64,9 +57,9 @@ static void DispatchPendingBreakEvents(TArray<FChaosBreakEvent> const& Events, T
 	}
 }
 
-static void DispatchPendingRemovalEvents(TArray<FChaosRemovalEvent> const& Events, TMap<TObjectPtr<UPrimitiveComponent>, FRemovalEventCallbackWrapper> const& Registrations)
+static void DispatchPendingRemovalEvents(TArray<FRemovalChaosEvent> const& Events, TMap<TObjectPtr<UPrimitiveComponent>, FRemovalEventCallbackWrapper> const& Registrations)
 {
-	for (FChaosRemovalEvent const& E : Events)
+	for (FRemovalChaosEvent const& E : Events)
 	{
 		if (E.Component)
 		{
@@ -79,9 +72,9 @@ static void DispatchPendingRemovalEvents(TArray<FChaosRemovalEvent> const& Event
 	}
 }
 
-static void DispatchPendingCrumblingEvents(TArray<FChaosCrumblingEvent> const& Events, TMap<TObjectPtr<UPrimitiveComponent>, FCrumblingEventCallbackWrapper> const& Registrations)
+static void DispatchPendingCrumblingEvents(TArray<FCrumblingChaosEvent> const& Events, TMap<TObjectPtr<UPrimitiveComponent>, FCrumblingEventCallbackWrapper> const& Registrations)
 {
-	for (FChaosCrumblingEvent const& E : Events)
+	for (FCrumblingChaosEvent const& E : Events)
 	{
 		if (E.Component)
 		{
@@ -479,15 +472,13 @@ void UChaosGameplayEventDispatcher::HandleCollisionEvents(const Chaos::FCollisio
 										NewContact.ContactPenetration = CollisionDataItem.PenetrationDepth;
 										// NewContact.PhysMaterial[1] UPhysicalMaterial required here
 
-										if(bSwapOrder)
+										if (bSwapOrder)
 										{
 											NotifyInfo.RigidCollisionData.SwapContactOrders();
 										}
 									}
-
 								}
 
-
 								if (HandlerSet.ChaosHandlers.Num() > 0)
 								{
 									bool bNewEntry = false;
@@ -550,7 +541,7 @@ void UChaosGameplayEventDispatcher::HandleBreakingEvents(const Chaos::FBreakingE
 
 	// BREAK EVENTS
 
-	TArray<FChaosBreakEvent> PendingBreakEvents;
+	TArray<FBreakChaosEvent> PendingBreakEvents;
 
 	const float BreakingDataTimestamp = Event.BreakingData.TimeCreated;
 	if (BreakingDataTimestamp > LastBreakingDataTime)
@@ -567,13 +558,13 @@ void UChaosGameplayEventDispatcher::HandleBreakingEvents(const Chaos::FBreakingE
 
 			for (const Chaos::FBreakingData& BreakingData : BreakingDataArray)
 			{	
-				if (BreakingData.Proxy)
+				if ((BreakingData.EmitterFlag & Chaos::EventEmitterFlag::OwnDispatcher) && BreakingData.Proxy)
 				{
 					UPrimitiveComponent* const PrimComp = Scene.GetOwningComponent<UPrimitiveComponent>(BreakingData.Proxy);
 					if (PrimComp && BreakEventRegistrations.Contains(PrimComp))
 					{
 						// queue them up so we can release the physics data before trigging BP events
-						FChaosBreakEvent& BreakEvent = PendingBreakEvents.Emplace_GetRef(BreakingData);
+						FBreakChaosEvent& BreakEvent = PendingBreakEvents.Emplace_GetRef(BreakingData);
 						BreakEvent.Component = PrimComp;
 					}
 				}
@@ -614,7 +605,7 @@ void UChaosGameplayEventDispatcher::HandleRemovalEvents(const Chaos::FRemovalEve
 {
 	// REMOVAL EVENTS
 
-	TArray<FChaosRemovalEvent> PendingRemovalEvents;
+	TArray<FRemovalChaosEvent> PendingRemovalEvents;
 
 	const float RemovalDataTimestamp = Event.RemovalData.TimeCreated;
 	if (RemovalDataTimestamp > LastRemovalDataTime)
@@ -636,7 +627,7 @@ void UChaosGameplayEventDispatcher::HandleRemovalEvents(const Chaos::FRemovalEve
 					if (PrimComp && RemovalEventRegistrations.Contains(PrimComp))
 					{
 						// queue them up so we can release the physics data before trigging BP events
-						FChaosRemovalEvent& RemovalEvent = PendingRemovalEvents.AddZeroed_GetRef();
+						FRemovalChaosEvent& RemovalEvent = PendingRemovalEvents.AddZeroed_GetRef();
 						RemovalEvent.Component = PrimComp;
 						RemovalEvent.Location = RemovalDataItem.Location;
 						RemovalEvent.Mass = RemovalDataItem.Mass;
@@ -666,17 +657,17 @@ void UChaosGameplayEventDispatcher::HandleCrumblingEvents(const Chaos::FCrumblin
 		const FPhysScene& Scene = *(GetWorld()->GetPhysicsScene());
 
 		// let's assume crumbles are rare, so we will iterate breaks instead of registered components for now
-		TArray<FChaosCrumblingEvent> PendingCrumblingEvent;
+		TArray<FCrumblingChaosEvent> PendingCrumblingEvent;
 		for (const Chaos::FCrumblingData& CrumblingDataItem : Event.CrumblingData.AllCrumblingsArray)
 		{	
-			if (CrumblingDataItem.Proxy)
+			if ((CrumblingDataItem.EmitterFlag & Chaos::EventEmitterFlag::OwnDispatcher) && CrumblingDataItem.Proxy)
 			{
 				if (UPrimitiveComponent* const PrimComp = Scene.GetOwningComponent<UPrimitiveComponent>(CrumblingDataItem.Proxy))
 				{
 					if (CrumblingEventRegistrations.Contains(PrimComp))
 					{
 						// queue them up so we can release the physics data before triggering BP events
-						FChaosCrumblingEvent& CrumblingEvent = PendingCrumblingEvent.AddZeroed_GetRef();
+						FCrumblingChaosEvent& CrumblingEvent = PendingCrumblingEvent.AddZeroed_GetRef();
 						CrumblingEvent.Component = PrimComp;
 						CrumblingEvent.Location = CrumblingDataItem.Location;
 						CrumblingEvent.Orientation = CrumblingDataItem.Orientation;
diff --git a/Engine/Source/Runtime/Experimental/ChaosSolverEngine/Public/Chaos/ChaosGameplayEventDispatcher.h b/Engine/Source/Runtime/Experimental/ChaosSolverEngine/Public/Chaos/ChaosGameplayEventDispatcher.h
--- a/Engine/Source/Runtime/Experimental/ChaosSolverEngine/Public/Chaos/ChaosGameplayEventDispatcher.h
+++ b/Engine/Source/Runtime/Experimental/ChaosSolverEngine/Public/Chaos/ChaosGameplayEventDispatcher.h
@@ -26,79 +26,28 @@ struct CHAOSSOLVERENGINE_API FChaosBreakEvent : public FBreakChaosEvent
 {
 	GENERATED_BODY()
 	FChaosBreakEvent() : FBreakChaosEvent() {}
+	FChaosBreakEvent(const FBreakChaosEvent& BreakChaosEvent) : FBreakChaosEvent(BreakChaosEvent) {}
 	FChaosBreakEvent(const Chaos::FBreakingData& BreakingData) : FBreakChaosEvent(BreakingData) {}
 };
 
 USTRUCT(BlueprintType)
-struct CHAOSSOLVERENGINE_API FChaosRemovalEvent
+struct CHAOSSOLVERENGINE_API FChaosRemovalEvent : public FRemovalChaosEvent
 {
 	GENERATED_BODY()
-
-public:
-
-	FChaosRemovalEvent();
-
-	UPROPERTY(BlueprintReadOnly, Category = "Removal Event")
-	TObjectPtr<UPrimitiveComponent> Component = nullptr;
-
-	UPROPERTY(BlueprintReadOnly, Category = "Removal Event")
-	FVector Location;
-
-	UPROPERTY(BlueprintReadOnly, Category = "Removal Event")
-	float Mass;
+	FChaosRemovalEvent() : FRemovalChaosEvent() {}
+	FChaosRemovalEvent(const FRemovalChaosEvent& RemovalChaosEvent) : FRemovalChaosEvent(RemovalChaosEvent) {}
 };
 
-
 USTRUCT(BlueprintType)
-struct CHAOSSOLVERENGINE_API FChaosCrumblingEvent
+struct CHAOSSOLVERENGINE_API FChaosCrumblingEvent : public FCrumblingChaosEvent
 {
 	GENERATED_BODY()
-
-public:
-	FChaosCrumblingEvent()
-		: Component(nullptr)
-		, Location(FVector::ZeroVector)
-		, Orientation(FQuat::Identity)
-		, LinearVelocity(FVector::ZeroVector)
-		, AngularVelocity(FVector::ZeroVector)
-		, Mass(0)
-		, LocalBounds(ForceInitToZero)
-	{}
-
-	/** primitive component involved in the crumble event */
-	UPROPERTY(BlueprintReadOnly, Category = "Crumble Event")
-	TObjectPtr<UPrimitiveComponent> Component = nullptr;
-
-	/** World location of the crumbling cluster */
-	UPROPERTY(BlueprintReadOnly, Category = "Crumble Event")
-	FVector Location;
-
-	/** World orientation of the crumbling cluster */
-	UPROPERTY(BlueprintReadOnly, Category = "Crumble Event")
-	FQuat Orientation;
-
-	/** Linear Velocity of the crumbling cluster */
-	UPROPERTY(BlueprintReadOnly, Category = "Crumble Event")
-	FVector LinearVelocity;
-
-	/** Angular Velocity of the crumbling cluster  */
-	UPROPERTY(BlueprintReadOnly, Category = "Crumble Event")
-	FVector AngularVelocity;
-
-	/** Mass of the crumbling cluster  */
-	UPROPERTY(BlueprintReadOnly, Category = "Crumble Event")
-	float Mass;
-
-	/** Local bounding box of the crumbling cluster  */
-	UPROPERTY(BlueprintReadOnly, Category = "Crumble Event")
-	FBox LocalBounds;
-	
-	/** List of children indices released (optional : see geometry collection component bCrumblingEventIncludesChildren) */
-	UPROPERTY(BlueprintReadOnly, Category = "Crumble Event")
-	TArray<int32> Children;
+	FChaosCrumblingEvent() : FCrumblingChaosEvent() {}
+	FChaosCrumblingEvent(const FCrumblingChaosEvent& CrumbleChaosEvent) : FCrumblingChaosEvent(CrumbleChaosEvent) {}
 };
 
-typedef TFunction<void(const FChaosBreakEvent&)> FOnBreakEventCallback;
+
+typedef TFunction<void(const FBreakChaosEvent&)> FOnBreakEventCallback;
 
 /** UStruct wrapper so we can store the TFunction in a TMap */
 USTRUCT()
@@ -110,7 +59,7 @@ public:
 	FOnBreakEventCallback BreakEventCallback;
 };
 
-typedef TFunction<void(const FChaosRemovalEvent&)> FOnRemovalEventCallback;
+typedef TFunction<void(const FRemovalChaosEvent&)> FOnRemovalEventCallback;
 
 /** UStruct wrapper so we can store the TFunction in a TMap */
 USTRUCT()
@@ -122,7 +71,7 @@ public:
 	FOnRemovalEventCallback RemovalEventCallback;
 };
 
-typedef TFunction<void(const FChaosCrumblingEvent&)> FOnCrumblingEventCallback;
+typedef TFunction<void(const FCrumblingChaosEvent&)> FOnCrumblingEventCallback;
 
 /** UStruct wrapper so we can store the TFunction in a TMap */
 USTRUCT()
diff --git a/Engine/Source/Runtime/Experimental/GeometryCollectionEngine/Private/GeometryCollection/GeometryCollectionComponent.cpp b/Engine/Source/Runtime/Experimental/GeometryCollectionEngine/Private/GeometryCollection/GeometryCollectionComponent.cpp
--- a/Engine/Source/Runtime/Experimental/GeometryCollectionEngine/Private/GeometryCollection/GeometryCollectionComponent.cpp
+++ b/Engine/Source/Runtime/Experimental/GeometryCollectionEngine/Private/GeometryCollection/GeometryCollectionComponent.cpp
@@ -345,7 +345,12 @@ UGeometryCollectionComponent::UGeometryCollectionComponent(const FObjectInitiali
 	, bNotifyCollisions(false)
 	, bNotifyRemovals(false)
 	, bNotifyCrumblings(false)
+	, bCrumblingEventIncludesChildren(false)
 	, bNotifyGlobalBreaks(false)
+	, bNotifyGlobalCollisions(false)
+	, bNotifyGlobalRemovals(false)
+	, bNotifyGlobalCrumblings(false)
+	, bGlobalCrumblingEventIncludesChildren(false)
 	, bStoreVelocities(false)
 	, bShowBoneColors(false)
 	, bUseRootProxyForNavigation(false)
@@ -839,6 +844,42 @@ void UGeometryCollectionComponent::SetNotifyGlobalBreaks(bool bNewNotifyGlobalBr
 	}
 }
 
+void UGeometryCollectionComponent::SetNotifyGlobalCollision(bool bNewNotifyGlobalCollisions)
+{
+	if (bNotifyGlobalCollisions != bNewNotifyGlobalCollisions)
+	{
+		bNotifyGlobalCollisions = bNewNotifyGlobalCollisions;
+		UpdateGlobalCollisionEventRegistration();
+	}
+}
+
+void UGeometryCollectionComponent::SetNotifyGlobalRemovals(bool bNewNotifyGlobalRemovals)
+{
+	if (bNotifyGlobalRemovals != bNewNotifyGlobalRemovals)
+	{
+		if (PhysicsProxy)
+		{
+			PhysicsProxy->SetNotifyGlobalRemovals_External(bNewNotifyGlobalRemovals);
+		}
+		bNotifyGlobalRemovals = bNewNotifyGlobalRemovals;
+		UpdateGlobalRemovalEventRegistration();
+	}
+}
+
+void UGeometryCollectionComponent::SetNotifyGlobalCrumblings(bool bNewNotifyGlobalCrumblings, bool bGlobalNewCrumblingEventIncludesChildren)
+{
+	if (bNotifyGlobalCrumblings != bNewNotifyGlobalCrumblings ||
+		bGlobalCrumblingEventIncludesChildren != bGlobalNewCrumblingEventIncludesChildren)
+	{
+		if (PhysicsProxy)
+		{
+			PhysicsProxy->SetNotifyGlobalCrumblings_External(bNewNotifyGlobalCrumblings, bGlobalNewCrumblingEventIncludesChildren);
+		}
+		bNotifyCrumblings = bNewNotifyGlobalCrumblings;
+		bGlobalCrumblingEventIncludesChildren = bGlobalNewCrumblingEventIncludesChildren;
+	}
+}
+
 FBodyInstance* UGeometryCollectionComponent::GetBodyInstance(FName BoneName /*= NAME_None*/, bool bGetWelded /*= true*/, int32 Index /*=INDEX_NONE*/) const
 {
 	return nullptr;// const_cast<FBodyInstance*>(&DummyBodyInstance);
@@ -967,7 +1008,7 @@ void UGeometryCollectionComponent::AddTorqueInRadians(FVector Torque, FName Bone
 	DispatchFieldCommand(Command);
 }
 
-void UGeometryCollectionComponent::DispatchBreakEvent(const FChaosBreakEvent& Event)
+void UGeometryCollectionComponent::DispatchBreakEvent(const FBreakChaosEvent& Event)
 {
 	// native
 	NotifyBreak(Event);
@@ -975,11 +1016,11 @@ void UGeometryCollectionComponent::DispatchBreakEvent(const FChaosBreakEvent& Ev
 	// bp
 	if (OnChaosBreakEvent.IsBound())
 	{
-		OnChaosBreakEvent.Broadcast(Event);
+		OnChaosBreakEvent.Broadcast(FChaosBreakEvent(Event));
 	}
 }
 
-void UGeometryCollectionComponent::DispatchRemovalEvent(const FChaosRemovalEvent& Event)
+void UGeometryCollectionComponent::DispatchRemovalEvent(const FRemovalChaosEvent& Event)
 {
 	// native
 	NotifyRemoval(Event);
@@ -987,16 +1028,16 @@ void UGeometryCollectionComponent::DispatchRemovalEvent(const FChaosRemovalEvent
 	// bp
 	if (OnChaosRemovalEvent.IsBound())
 	{
-		OnChaosRemovalEvent.Broadcast(Event);
+		OnChaosRemovalEvent.Broadcast(FChaosRemovalEvent(Event));
 	}
 }
 
-void UGeometryCollectionComponent::DispatchCrumblingEvent(const FChaosCrumblingEvent& Event)
+void UGeometryCollectionComponent::DispatchCrumblingEvent(const FCrumblingChaosEvent& Event)
 {
 	// bp
 	if (OnChaosCrumblingEvent.IsBound())
 	{
-		OnChaosCrumblingEvent.Broadcast(Event);
+		OnChaosCrumblingEvent.Broadcast(FChaosCrumblingEvent(Event));
 	}
 }
 
@@ -1432,7 +1473,7 @@ void UGeometryCollectionComponent::PostEditChangeProperty(FPropertyChangedEvent&
 }
 #endif
 
-static void DispatchGeometryCollectionBreakEvent(const FChaosBreakEvent& Event)
+static void DispatchGeometryCollectionBreakEvent(const FBreakChaosEvent& Event)
 {
 	if (UGeometryCollectionComponent* const GC = Cast<UGeometryCollectionComponent>(Event.Component))
 	{
@@ -1440,7 +1481,7 @@ static void DispatchGeometryCollectionBreakEvent(const FChaosBreakEvent& Event)
 	}
 }
 
-static void DispatchGeometryCollectionRemovalEvent(const FChaosRemovalEvent& Event)
+static void DispatchGeometryCollectionRemovalEvent(const FRemovalChaosEvent& Event)
 {
 	if (UGeometryCollectionComponent* const GC = Cast<UGeometryCollectionComponent>(Event.Component))
 	{
@@ -1448,7 +1489,7 @@ static void DispatchGeometryCollectionRemovalEvent(const FChaosRemovalEvent& Eve
 	}
 }
 
-static void DispatchGeometryCollectionCrumblingEvent(const FChaosCrumblingEvent& Event)
+static void DispatchGeometryCollectionCrumblingEvent(const FCrumblingChaosEvent& Event)
 {
 	if (UGeometryCollectionComponent* const GC = Cast<UGeometryCollectionComponent>(Event.Component))
 	{
@@ -1530,13 +1571,35 @@ void UGeometryCollectionComponent::RegisterForEvents()
 				}
 			}
 		}
-		if (bNotifyGlobalBreaks)
+		if (bNotifyGlobalBreaks || bNotifyGlobalCrumblings)
 		{
 			Solver->EnqueueCommandImmediate([Solver]()
 				{
 					Solver->SetGenerateBreakingData(true);
 				});
 		}
+		if (bNotifyGlobalCollisions)
+		{
+			if (FPhysScene_Chaos* PhysicsScene = GetInnerChaosScene())
+			{
+				PhysicsScene->RegisterForGlobalCollisionEvents(this);
+				Solver->EnqueueCommandImmediate([Solver]()
+					{
+						Solver->SetGenerateCollisionData(true);
+					});
+			}
+		}
+		if (bNotifyGlobalRemovals)
+		{
+			if (FPhysScene_Chaos* PhysicsScene = GetInnerChaosScene())
+			{
+				PhysicsScene->RegisterForGlobalRemovalEvents(this);
+				Solver->EnqueueCommandImmediate([Solver]()
+					{
+						Solver->SetGenerateRemovalData(true);
+					});
+			}
+		}
 	}
 }
 
@@ -1555,6 +1618,20 @@ void UGeometryCollectionComponent::UpdateRBCollisionEventRegistration()
 	}
 }
 
+void UGeometryCollectionComponent::UpdateGlobalCollisionEventRegistration()
+{
+	if (FPhysScene_Chaos* PhysScene = GetInnerChaosScene())
+	{
+		if (bNotifyGlobalCollisions)
+		{
+			PhysScene->RegisterForGlobalCollisionEvents(this);
+		}
+		else
+		{
+			PhysScene->UnRegisterForGlobalCollisionEvents(this);
+		}
+	}
+}
 void UGeometryCollectionComponent::UpdateBreakEventRegistration()
 {
 	if (EventDispatcher)
@@ -1585,6 +1662,21 @@ void UGeometryCollectionComponent::UpdateRemovalEventRegistration()
 	}
 }
 
+void UGeometryCollectionComponent::UpdateGlobalRemovalEventRegistration()
+{
+	if (FPhysScene_Chaos* PhysicsScene = GetInnerChaosScene())
+	{
+		if (bNotifyGlobalRemovals)
+		{
+			PhysicsScene->RegisterForGlobalRemovalEvents(this);
+		}
+		else
+		{
+			PhysicsScene->UnRegisterForGlobalRemovalEvents(this);
+		}
+	}
+}
+
 void UGeometryCollectionComponent::UpdateCrumblingEventRegistration()
 {
 	if (EventDispatcher)
@@ -2936,13 +3028,15 @@ void UGeometryCollectionComponent::RegisterAndInitializePhysicsProxy()
 		SimulationParameters.StartAwake = BodyInstance.bStartAwake;
 		SimulationParameters.CacheType = CacheParameters.CacheMode;
 		SimulationParameters.ReverseCacheBeginTime = CacheParameters.ReverseCacheBeginTime;
-		SimulationParameters.bGenerateBreakingData = bNotifyBreaks || bNotifyGlobalBreaks;
-		SimulationParameters.bDispatchGlobalBreakingData = bNotifyGlobalBreaks;
+		SimulationParameters.bGenerateBreakingData = bNotifyBreaks;
 		SimulationParameters.bGenerateCollisionData = bNotifyCollisions;
 		SimulationParameters.bGenerateTrailingData = bNotifyTrailing;
-		SimulationParameters.bGenerateRemovalsData = bNotifyRemovals;
 		SimulationParameters.bGenerateCrumblingData = bNotifyCrumblings;
 		SimulationParameters.bGenerateCrumblingChildrenData = bCrumblingEventIncludesChildren;
+		SimulationParameters.bGenerateGlobalBreakingData = bNotifyGlobalBreaks;
+		SimulationParameters.bGenerateGlobalCollisionData = bNotifyGlobalCollisions;
+		SimulationParameters.bGenerateGlobalCrumblingData = bNotifyGlobalCrumblings;
+		SimulationParameters.bGenerateGlobalCrumblingChildrenData = bGlobalCrumblingEventIncludesChildren;
 		SimulationParameters.EnableGravity = BodyInstance.bEnableGravity;
 		SimulationParameters.GravityGroupIndex = GravityGroupIndex;
 		SimulationParameters.UseInertiaConditioning = BodyInstance.IsInertiaConditioningEnabled();
@@ -3300,7 +3394,7 @@ void UGeometryCollectionComponent::LoadCollisionProfiles()
 
 		Cache.QueryFilter = InitialQueryFilter;
 		Cache.SimFilter = InitialSimFilter;
-		CreateShapeFilterData(Cache.Template.ObjectType, BodyInstance.GetMaskFilter(), ActorID, Cache.Template.ResponseToChannels, CompID, INDEX_NONE, Cache.QueryFilter, Cache.SimFilter, BodyInstance.bUseCCD, bNotifyCollisions, false, false);
+		CreateShapeFilterData(Cache.Template.ObjectType, BodyInstance.GetMaskFilter(), ActorID, Cache.Template.ResponseToChannels, CompID, INDEX_NONE, Cache.QueryFilter, Cache.SimFilter, BodyInstance.bUseCCD, bNotifyCollisions || bNotifyGlobalCollisions, false, false);
 
 		// Maintain parity with the rest of the geometry collection filters.
 		Cache.QueryFilter.Word3 |= (EPDF_SimpleCollision | EPDF_ComplexCollision);
@@ -3489,7 +3583,7 @@ void UGeometryCollectionComponent::BuildInitialFilterData()
 	InitialQueryFilter.Word3 |= (EPDF_SimpleCollision | EPDF_ComplexCollision);
 	InitialSimFilter.Word3 |= (EPDF_SimpleCollision | EPDF_ComplexCollision);
 
-	if (bNotifyCollisions)
+	if (bNotifyCollisions || bNotifyGlobalCollisions)
 	{
 		InitialQueryFilter.Word3 |= EPDF_ContactNotify;
 		InitialSimFilter.Word3 |= EPDF_ContactNotify;
diff --git a/Engine/Source/Runtime/Experimental/GeometryCollectionEngine/Public/GeometryCollection/GeometryCollectionComponent.h b/Engine/Source/Runtime/Experimental/GeometryCollectionEngine/Public/GeometryCollection/GeometryCollectionComponent.h
--- a/Engine/Source/Runtime/Experimental/GeometryCollectionEngine/Public/GeometryCollection/GeometryCollectionComponent.h
+++ b/Engine/Source/Runtime/Experimental/GeometryCollectionEngine/Public/GeometryCollection/GeometryCollectionComponent.h
@@ -56,6 +56,7 @@ enum class EObjectStateTypeEnum : uint8;
 namespace Chaos { enum class EObjectStateType: int8; }
 template<class InElementType> class TManagedArray;
 
+// To deprecate:  to use FBreakChaosEvent from Engine instead of FChaosBreakEvent from ChaosSolverEngine
 DECLARE_DYNAMIC_MULTICAST_DELEGATE_OneParam(FOnChaosBreakEvent, const FChaosBreakEvent&, BreakEvent);
 
 DECLARE_DYNAMIC_MULTICAST_DELEGATE_OneParam(FOnChaosRemovalEvent, const FChaosRemovalEvent&, RemovalEvent);
@@ -987,19 +988,31 @@ public:
 	UFUNCTION(BlueprintCallable, Category = "Physics")
 	void SetNotifyRemovals(bool bNewNotifyRemovals);
 
-	/** Changes whether or not this component will get future break notifications. */
+	/** Changes whether or not this component will get future crumbling notifications. */
 	UFUNCTION(BlueprintCallable, Category = "Physics")
 	void SetNotifyCrumblings(bool bNewNotifyCrumblings, bool bNewCrumblingEventIncludesChildren = false);
 
 	/** Changes whether or not this component will get future global break notifications. */
 	UFUNCTION(BlueprintCallable, Category = "Physics")
 	void SetNotifyGlobalBreaks(bool bNewNotifyGlobalBreaks);
+
+	/** Changes whether or not this component will get future global collision notifications. */
+	UFUNCTION(BlueprintCallable, Category = "Physics")
+	void SetNotifyGlobalCollision(bool bNewNotifyGlobalCollisions);
+
+	/** Changes whether or not this component will get future global removal notifications. */
+	UFUNCTION(BlueprintCallable, Category = "Physics")
+	void SetNotifyGlobalRemovals(bool bNewNotifyGlobalRemovals);
+
+	/** Changes whether or not this component will get future global crumbling notifications. */
+	UFUNCTION(BlueprintCallable, Category = "Physics")
+	void SetNotifyGlobalCrumblings(bool bNewNotifyGlobalCrumblings, bool bGlobalNewCrumblingEventIncludesChildren);
 	
 	/** Overrideable native notification */
-	virtual void NotifyBreak(const FChaosBreakEvent& Event) {};
+	virtual void NotifyBreak(const FBreakChaosEvent& Event) {};
 
 	/** Overrideable native notification */
-	virtual void NotifyRemoval(const FChaosRemovalEvent& Event) {};
+	virtual void NotifyRemoval(const FRemovalChaosEvent& Event) {};
 
 	UPROPERTY(BlueprintAssignable, Category = "Chaos")
 	FOnChaosBreakEvent OnChaosBreakEvent;
@@ -1013,11 +1026,11 @@ public:
 	// todo(chaos) remove when no longer necessary
 	FOnChaosBreakEvent OnRootBreakEvent;
 
-	void DispatchBreakEvent(const FChaosBreakEvent& Event);
+	void DispatchBreakEvent(const FBreakChaosEvent& Event);
 
-	void DispatchRemovalEvent(const FChaosRemovalEvent& Event);
+	void DispatchRemovalEvent(const FRemovalChaosEvent& Event);
 
-	void DispatchCrumblingEvent(const FChaosCrumblingEvent& Event);
+	void DispatchCrumblingEvent(const FCrumblingChaosEvent& Event);
 
 	UPROPERTY(Transient, VisibleAnywhere, BlueprintReadWrite, Interp, Category = "Chaos")
 	float DesiredCacheTime;
@@ -1105,14 +1118,31 @@ public:
 	UPROPERTY(EditAnywhere, BlueprintReadOnly, Category = "ChaosPhysics|Events")
 	bool bNotifyCrumblings;
 
+	/** If this and bNotifyCrumblings are true, the crumbling events will contain released children indices. */
+	UPROPERTY(EditAnywhere, BlueprintReadOnly, Category = "ChaosPhysics|Events", meta = (EditCondition = "bNotifyCrumblings"))
+	bool bCrumblingEventIncludesChildren;
+
 	/** If true, this component will generate breaking events that will be listened by the global event relay. */
 	UPROPERTY(EditAnywhere, BlueprintReadOnly, Category = "ChaosPhysics|Events|Global")
 	bool bNotifyGlobalBreaks;
 
-	/** If this and bNotifyCrumblings are true, the crumbling events will contain released children indices. */
-	UPROPERTY(EditAnywhere, BlueprintReadOnly, Category = "ChaosPhysics|Events", meta = (EditCondition = "bNotifyCrumblings"))
-	bool bCrumblingEventIncludesChildren;
-	
+	/** If true, this component will generate collision events  that will be listened by the global event relay. */
+	UPROPERTY(EditAnywhere, BlueprintReadOnly, Category = "ChaosPhysics|Events|Global")
+	bool bNotifyGlobalCollisions;
+
+	/** If true, this component will generate removal events  that will be listened by the global event relay. */
+	UPROPERTY(EditAnywhere, BlueprintReadOnly, Category = "ChaosPhysics|Events|Global")
+	bool bNotifyGlobalRemovals;
+
+	/** If true, this component will generate crumbling events  that will be listened by the global event relay. */
+	UPROPERTY(EditAnywhere, BlueprintReadOnly, Category = "ChaosPhysics|Events|Global")
+	bool bNotifyGlobalCrumblings;
+
+	/** If this and bNotifyGlobalCrumblings are true, the crumbling events will contain released children indices. */
+	UPROPERTY(EditAnywhere, BlueprintReadOnly, Category = "ChaosPhysics|Events|Global", meta = (EditCondition = "bNotifyGlobalCrumblings"))
+	bool bGlobalCrumblingEventIncludesChildren;
+
+
 	/** If true, this component will save linear and angular velocities on its DynamicCollection. */
 	UPROPERTY(EditAnywhere, BlueprintReadOnly, Category = "ChaosPhysics|General")
 	bool bStoreVelocities;
@@ -1184,8 +1214,10 @@ protected:
 
 	void RegisterForEvents();
 	void UpdateRBCollisionEventRegistration();
+	void UpdateGlobalCollisionEventRegistration();
 	void UpdateBreakEventRegistration();
 	void UpdateRemovalEventRegistration();
+	void UpdateGlobalRemovalEventRegistration();
 	void UpdateCrumblingEventRegistration();
 	
 	/* Per-instance override to enable/disable replication for the geometry collection */
