Enable r.PurgeEditorSceneDuringPIE by default now, so that going in and out of PIE will tear down and restore the Editor FScene and memory allocations, avoiding serious overhead and rampant GPU reservation hitches for a lot of complex scenes/games.

diff --git a/Engine/Source/Runtime/Engine/Private/World.cpp b/Engine/Source/Runtime/Engine/Private/World.cpp
--- a/Engine/Source/Runtime/Engine/Private/World.cpp
+++ b/Engine/Source/Runtime/Engine/Private/World.cpp
@@ -146,7 +146,7 @@ FAutoConsoleVariableRef CVarDisableRemapScriptActors(TEXT("net.DisableRemapScrip
 
 static TAutoConsoleVariable<int32> CVarPurgeEditorSceneDuringPIE(
 	TEXT("r.PurgeEditorSceneDuringPIE"),
-	0,
+	1,
 	TEXT("0 to keep editor scene fully initialized during PIE (default)\n")
 	TEXT("1 to purge editor scene from memory during PIE and restore when the session finishes."));
 
