[USoundWave] Fixing data race.

diff --git a/Engine/Source/Runtime/Engine/Private/SoundWave.cpp b/Engine/Source/Runtime/Engine/Private/SoundWave.cpp
--- a/Engine/Source/Runtime/Engine/Private/SoundWave.cpp
+++ b/Engine/Source/Runtime/Engine/Private/SoundWave.cpp
@@ -2512,11 +2512,8 @@ void USoundWave::PostEditChangeProperty(FPropertyChangedEvent& PropertyChangedEv
 	static const FName InitialChunkSizeFName = GET_MEMBER_NAME_CHECKED(USoundWave, InitialChunkSize);
 	static const FName TransformationsFName = GET_MEMBER_NAME_CHECKED(USoundWave, Transformations);
 
-	// force proxy flags to be up to date
-	SoundWaveDataPtr->bIsSeekable = IsSeekable();
-	SoundWaveDataPtr->SoundAssetCompressionType = SoundAssetCompressionType;
-	SoundWaveDataPtr->bIsStreaming = IsStreaming(nullptr);
-	SoundWaveDataPtr->bShouldUseStreamCaching = ShouldUseStreamCaching();
+	// force proxy state to be up to date
+	SoundWaveDataPtr->InitializeDataFromSoundWave(*this);
 
 	if (FProperty* PropertyThatChanged = PropertyChangedEvent.Property)
 	{
@@ -3332,7 +3329,6 @@ FSoundWaveProxyPtr USoundWave::CreateSoundWaveProxy()
 #endif // #if WITH_EDITORONLY_DATA
 
 	check(SoundWaveDataPtr);
-	SoundWaveDataPtr->InitializeDataFromSoundWave(*this);
 
 	return MakeShared<FSoundWaveProxy, ESPMode::ThreadSafe>(this);
 }
