Add a Debug Display Name string to Player Controller input modes.

diff --git a/Engine/Source/Runtime/Engine/Classes/GameFramework/PlayerController.h b/Engine/Source/Runtime/Engine/Classes/GameFramework/PlayerController.h
--- a/Engine/Source/Runtime/Engine/Classes/GameFramework/PlayerController.h
+++ b/Engine/Source/Runtime/Engine/Classes/GameFramework/PlayerController.h
@@ -24,6 +24,7 @@
 #include "GenericPlatform/IInputInterface.h"
 #include "Physics/AsyncPhysicsData.h"
 #include "WorldPartition/WorldPartitionStreamingSource.h"
+#include "EngineDefines.h"		// For UE_ENABLE_DEBUG_DRAWING
 
 #if UE_ENABLE_INCLUDE_ORDER_DEPRECATED_IN_5_2
 #include "Widgets/SWidget.h"
@@ -116,6 +117,13 @@ protected:
 	void SetFocusAndLocking(FReply& SlateOperations, TSharedPtr<class SWidget> InWidgetToFocus, bool bLockMouseToViewport, TSharedRef<class SViewport> InViewportWidget) const;
 
 	friend class APlayerController;
+
+public:
+
+#if UE_ENABLE_DEBUG_DRAWING
+	/** Returns the name of this input mode for debug display when you call the "showdebug input" command. */
+	virtual const FString& GetDebugDisplayName() const;
+#endif	// UE_ENABLE_DEBUG_DRAWING
 };
 
 /** Data structure used to setup an input mode that allows only the UI to respond to user input. */
@@ -132,6 +140,10 @@ struct ENGINE_API FInputModeUIOnly : public FInputModeDataBase
 		, MouseLockMode(EMouseLockMode::LockInFullscreen)
 	{}
 
+#if UE_ENABLE_DEBUG_DRAWING
+	virtual const FString& GetDebugDisplayName() const override;
+#endif	// UE_ENABLE_DEBUG_DRAWING
+
 protected:
 	TSharedPtr<SWidget> WidgetToFocus;
 	EMouseLockMode MouseLockMode;
@@ -190,6 +202,10 @@ struct ENGINE_API FInputModeGameAndUI : public FInputModeDataBase
 		, bHideCursorDuringCapture(true)
 	{}
 
+#if UE_ENABLE_DEBUG_DRAWING
+	virtual const FString& GetDebugDisplayName() const override;
+#endif	// UE_ENABLE_DEBUG_DRAWING
+
 protected:
 
 	TSharedPtr<SWidget> WidgetToFocus;
@@ -209,6 +225,10 @@ struct ENGINE_API FInputModeGameOnly : public FInputModeDataBase
 		: bConsumeCaptureMouseDown(true)
 	{}
 
+#if UE_ENABLE_DEBUG_DRAWING
+	virtual const FString& GetDebugDisplayName() const override;
+#endif	// UE_ENABLE_DEBUG_DRAWING
+	
 protected:
 	bool bConsumeCaptureMouseDown;
 
@@ -2215,6 +2235,11 @@ public:
 	 * Player controllers should always be included in seamless travel
 	 */
 	virtual bool ShouldParticipateInSeamlessTravel() const override { return true; }
+	
+#if UE_ENABLE_DEBUG_DRAWING
+	/** Keep track of the current input mode debug string here. Set in APlayerController::SetInputMode */
+	const FString& GetCurrentInputModeDebugString() const;
+#endif
 
 private:
 	/** If true, prevent any haptic effects from playing */
@@ -2228,6 +2253,11 @@ private:
 	 */
 	bool bShouldFlushInputWhenViewportFocusChanges : 1;
 
+#if UE_ENABLE_DEBUG_DRAWING
+	/** Keep track of the current input mode debug string here. Set in APlayerController::SetInputMode */
+	FString CurrentInputModeDebugString;
+#endif
+
 public: 
 
 	/**
diff --git a/Engine/Source/Runtime/Engine/Private/PlayerController.cpp b/Engine/Source/Runtime/Engine/Private/PlayerController.cpp
--- a/Engine/Source/Runtime/Engine/Private/PlayerController.cpp
+++ b/Engine/Source/Runtime/Engine/Private/PlayerController.cpp
@@ -198,6 +198,9 @@ APlayerController::APlayerController(const FObjectInitializer& ObjectInitializer
 		bAsyncPhysicsTickEnabled = true;
 		AsyncPhysicsDataClass = UAsyncPhysicsData::StaticClass();
 	}
+#if UE_ENABLE_DEBUG_DRAWING
+	CurrentInputModeDebugString = TEXT("Default");
+#endif	// UE_ENABLE_DEBUG_DRAWING
 }
 
 float APlayerController::GetNetPriority(const FVector& ViewPos, const FVector& ViewDir, AActor* Viewer, AActor* ViewTarget, UActorChannel* InChannel, float Time, bool bLowBandwidth)
@@ -5777,6 +5780,14 @@ void FInputModeDataBase::SetFocusAndLocking(FReply& SlateOperations, TSharedPtr<
 	}
 }
 
+#if UE_ENABLE_DEBUG_DRAWING
+const FString& FInputModeDataBase::GetDebugDisplayName() const
+{
+	static const FString DisplayName = TEXT("Base");
+	return DisplayName;
+}
+#endif	// UE_ENABLE_DEBUG_DRAWING
+
 FInputModeUIOnly& FInputModeUIOnly::SetWidgetToFocus(TSharedPtr<SWidget> InWidgetToFocus)
 {
 #if !(UE_BUILD_SHIPPING || UE_BUILD_TEST)
@@ -5795,6 +5806,14 @@ FInputModeUIOnly& FInputModeUIOnly::SetLockMouseToViewportBehavior(EMouseLockMod
 	return *this;
 }
 
+#if UE_ENABLE_DEBUG_DRAWING
+const FString& FInputModeUIOnly::GetDebugDisplayName() const
+{
+	static const FString DebugName = TEXT("UI Only (Input will only be consumed by the UI, not the player!)");
+	return DebugName;
+}
+#endif	// UE_ENABLE_DEBUG_DRAWING
+
 void FInputModeUIOnly::ApplyInputMode(FReply& SlateOperations, class UGameViewportClient& GameViewportClient) const
 {
 	TSharedPtr<SViewport> ViewportWidget = GameViewportClient.GetGameViewportWidget();
@@ -5830,6 +5849,22 @@ void FInputModeGameAndUI::ApplyInputMode(FReply& SlateOperations, class UGameVie
 	}
 }
 
+#if UE_ENABLE_DEBUG_DRAWING
+const FString& FInputModeGameAndUI::GetDebugDisplayName() const
+{
+	static const FString DisplayName = TEXT("Game and UI");
+	return DisplayName;
+}
+#endif	// UE_ENABLE_DEBUG_DRAWING
+
+#if UE_ENABLE_DEBUG_DRAWING
+const FString& FInputModeGameOnly::GetDebugDisplayName() const
+{
+	static const FString DisplayName = TEXT("Game Only (Input will only be consumed by the player, not UI)");
+	return DisplayName;
+}
+#endif	// UE_ENABLE_DEBUG_DRAWING
+
 void FInputModeGameOnly::ApplyInputMode(FReply& SlateOperations, class UGameViewportClient& GameViewportClient) const
 {
 	TSharedPtr<SViewport> ViewportWidget = GameViewportClient.GetGameViewportWidget();
@@ -5853,9 +5888,21 @@ void APlayerController::SetInputMode(const FInputModeDataBase& InData)
 	{
 		InData.ApplyInputMode( LocalPlayer->GetSlateOperations(), *GameViewportClient );
 		bShouldFlushInputWhenViewportFocusChanges = InData.ShouldFlushInputOnViewportFocus();
+
+		// Keep track of the name of this input mode for debug purposes
+#if UE_ENABLE_DEBUG_DRAWING
+		CurrentInputModeDebugString = InData.GetDebugDisplayName();
+#endif
 	}
 }
 
+#if UE_ENABLE_DEBUG_DRAWING
+const FString& APlayerController::GetCurrentInputModeDebugString() const
+{
+	return CurrentInputModeDebugString;
+}
+#endif // #if UE_ENABLE_DEBUG_DRAWING
+
 void APlayerController::UpdateCameraManager(float DeltaSeconds)
 {
 	if (PlayerCameraManager != NULL)
diff --git a/Engine/Source/Runtime/Engine/Private/UserInterface/PlayerInput.cpp b/Engine/Source/Runtime/Engine/Private/UserInterface/PlayerInput.cpp
--- a/Engine/Source/Runtime/Engine/Private/UserInterface/PlayerInput.cpp
+++ b/Engine/Source/Runtime/Engine/Private/UserInterface/PlayerInput.cpp
@@ -15,6 +15,10 @@
 #include "Engine/LocalPlayer.h"
 #include "GameFramework/InputSettings.h"
 
+#if UE_ENABLE_DEBUG_DRAWING
+#include "Engine/GameViewportClient.h"	// For grabbing some info about mouse capture for the debug display
+#endif	// UE_ENABLE_DEBUG_DRAWING
+
 #include UE_INLINE_GENERATED_CPP_BY_NAME(PlayerInput)
 
 DECLARE_CYCLE_STAT(TEXT("    PC Gesture Recognition"), STAT_PC_GestureRecognition, STATGROUP_PlayerController);
@@ -1614,6 +1618,44 @@ void UPlayerInput::DisplayDebug(class UCanvas* Canvas, const FDebugDisplayInfo&
 	if (Canvas)
 	{
 		FDisplayDebugManager& DisplayDebugManager = Canvas->DisplayDebugManager;
+
+#if UE_ENABLE_DEBUG_DRAWING
+		// Show the UI mode of the owning player controller if there is one. This is a very common issue with folks
+		// debugging input. Often times users will set the UI mode to "UI Only", and then never set it back so the 
+		// game can consume input again. 
+		if (const APlayerController* PlayerController = GetOuterAPlayerController())
+		{		
+			DisplayDebugManager.SetDrawColor(FColor::Orange);
+			DisplayDebugManager.DrawString(FString::Printf(TEXT("\t Player Controller Input Settings:  %s"), *GetNameSafe(PlayerController)));
+			
+			// Display whether input is enabled on the player controller or not
+			const bool bIsInputEnabled = PlayerController->InputEnabled();
+			DisplayDebugManager.SetDrawColor(bIsInputEnabled ? FColor::White : FColor::Red);
+			DisplayDebugManager.DrawString(FString::Printf(TEXT("\t PlayerController bIsInputEnabled: %s"), bIsInputEnabled ? TEXT("True") : TEXT("False")));
+
+			DisplayDebugManager.SetDrawColor(FColor::White);
+			DisplayDebugManager.DrawString(FString::Printf(TEXT("\t PlayerController Input Mode: %s"), *PlayerController->GetCurrentInputModeDebugString()));
+
+			const ULocalPlayer* LP = GetOwningLocalPlayer();
+
+			// Display some info about the viewport, all of these settings can affect the input behavior of the mouse
+			if (LP && LP->ViewportClient)
+			{
+				// Show the capture mode of the viewport
+				DisplayDebugManager.DrawString(FString::Printf(TEXT("\t ViewportClient MouseCaptureMode: %s"), *UEnum::GetValueAsString(TEXT("Engine.EMouseCaptureMode"), LP->ViewportClient->GetMouseCaptureMode())));
+				DisplayDebugManager.DrawString(FString::Printf(TEXT("\t ViewportClient MouseLockMode: %s"), *UEnum::GetValueAsString(TEXT("Engine.EMouseLockMode"), LP->ViewportClient->GetMouseLockMode())));
+				
+				const bool bIsIgnoringInput = LP->ViewportClient->IgnoreInput();
+				DisplayDebugManager.SetDrawColor(bIsIgnoringInput ? FColor::Red : FColor::White);
+				DisplayDebugManager.DrawString(FString::Printf(TEXT("\t ViewportClient bIsIgnoringInput: %s"), bIsIgnoringInput ? TEXT("True") : TEXT("False")));
+				
+				const bool bIsUsingMouseForTouch = LP->ViewportClient->GetUseMouseForTouch();
+				DisplayDebugManager.SetDrawColor(bIsUsingMouseForTouch ? FColor::Green : FColor::White);
+				DisplayDebugManager.DrawString(FString::Printf(TEXT("\t ViewportClient bIsUsingMouseForTouch: %s"), bIsUsingMouseForTouch ? TEXT("True") : TEXT("False")));
+			}
+		}
+#endif // #if UE_ENABLE_DEBUG_DRAWING
+
 		DisplayDebugManager.SetDrawColor(FColor::Red);
 		DisplayDebugManager.DrawString(FString::Printf(TEXT("INPUT %s"), *GetName()));
 
