LevelInstanceSubsystem: ResetLoaders on World rename

Fixes failed renames when Level is being renamed while being loaded by a Level Instance

#rnx

diff --git a/Engine/Source/Runtime/Engine/Private/LevelInstance/LevelInstanceSubsystem.cpp b/Engine/Source/Runtime/Engine/Private/LevelInstance/LevelInstanceSubsystem.cpp
--- a/Engine/Source/Runtime/Engine/Private/LevelInstance/LevelInstanceSubsystem.cpp
+++ b/Engine/Source/Runtime/Engine/Private/LevelInstance/LevelInstanceSubsystem.cpp
@@ -86,6 +86,7 @@ void ULevelInstanceSubsystem::Initialize(FSubsystemCollectionBase& Collection)
 		ILevelInstanceEditorModule& EditorModule = FModuleManager::LoadModuleChecked<ILevelInstanceEditorModule>("LevelInstanceEditor");
 		FEditorDelegates::OnAssetsPreDelete.AddUObject(this, &ULevelInstanceSubsystem::OnAssetsPreDelete);
 		FEditorDelegates::PreSaveWorldWithContext.AddUObject(this, &ULevelInstanceSubsystem::OnPreSaveWorldWithContext);
+		FWorldDelegates::OnPreWorldRename.AddUObject(this, &ULevelInstanceSubsystem::OnPreWorldRename);
 	}
 #endif
 }
@@ -95,6 +96,7 @@ void ULevelInstanceSubsystem::Deinitialize()
 #if WITH_EDITOR
 	FEditorDelegates::OnAssetsPreDelete.RemoveAll(this);
 	FEditorDelegates::PreSaveWorldWithContext.RemoveAll(this);
+	FWorldDelegates::OnPreWorldRename.RemoveAll(this);
 #endif
 }
 
@@ -378,6 +380,18 @@ void ULevelInstanceSubsystem::OnPreSaveWorldWithContext(UWorld* InWorld, FObject
 	}
 }
 
+void ULevelInstanceSubsystem::OnPreWorldRename(UWorld* InWorld, const TCHAR* InName, UObject* NewOuter, ERenameFlags Flags, bool& bShouldFailRename)
+{
+	const bool bTestRename = (Flags & REN_Test) != 0;
+	if (!bTestRename)
+	{
+		if (const UPackage* WorldPackage = InWorld->GetPackage())
+		{
+			ResetLoadersForWorldAssetInternal(WorldPackage->GetName());
+		}
+	}
+}
+
 void ULevelInstanceSubsystem::RegisterLoadedLevelStreamingLevelInstanceEditor(ULevelStreamingLevelInstanceEditor* LevelStreaming)
 {
 	if (!bIsCreatingLevelInstance)
diff --git a/Engine/Source/Runtime/Engine/Public/LevelInstance/LevelInstanceSubsystem.h b/Engine/Source/Runtime/Engine/Public/LevelInstance/LevelInstanceSubsystem.h
--- a/Engine/Source/Runtime/Engine/Public/LevelInstance/LevelInstanceSubsystem.h
+++ b/Engine/Source/Runtime/Engine/Public/LevelInstance/LevelInstanceSubsystem.h
@@ -145,6 +145,7 @@ private:
 	void ResetLoadersForWorldAssetInternal(const FString& WorldAsset);
 	void OnAssetsPreDelete(const TArray<UObject*>& Objects);
 	void OnPreSaveWorldWithContext(UWorld* InWorld, FObjectPreSaveContext ObjectSaveContext);
+	void OnPreWorldRename(UWorld* InWorld, const TCHAR* InName, UObject* NewOuter, ERenameFlags Flags, bool& bShouldFailRename);
 
 	void RegisterLoadedLevelStreamingLevelInstanceEditor(ULevelStreamingLevelInstanceEditor* LevelStreaming);
 
