FExternalPackageHelper::LoadObjectsFromExternalPackages : Filter packages with PKG_PlayInEditor flag as they are considered even when specifying bIncludeOnlyOnDiskAssets = true

#rnx

diff --git a/Engine/Source/Runtime/Engine/Public/ExternalPackageHelper.h b/Engine/Source/Runtime/Engine/Public/ExternalPackageHelper.h
--- a/Engine/Source/Runtime/Engine/Public/ExternalPackageHelper.h
+++ b/Engine/Source/Runtime/Engine/Public/ExternalPackageHelper.h
@@ -7,6 +7,7 @@
 #include "UObject/LinkerLoad.h"
 #include "UObject/Package.h"
 #include "UObject/UObjectThreadContext.h"
+#include "UObject/ObjectMacros.h"
 #include "AssetRegistry/AssetData.h"
 #include "AssetRegistry/AssetRegistryModule.h"
 #include "Misc/ArchiveMD5.h"
@@ -99,7 +100,8 @@ void FExternalPackageHelper::LoadObjectsFromExternalPackages(UObject* InOuter, T
 	Filter.ClassPaths.Add(T::StaticClass()->GetClassPathName());
 	Filter.bRecursiveClasses = true;
 	Filter.PackagePaths.Add(*ExternalObjectsPath);
-
+	// @todo_ow: remove once bIncludeOnlyOnDiskAssets=true doesn't return PIE Packages anymore
+	Filter.WithoutPackageFlags = PKG_PlayInEditor;
 	TArray<FAssetData> Assets;
 	AssetRegistry.GetAssets(Filter, Assets);
 
