Fix conflicting changes from two different branches.

diff --git a/Engine/Source/Editor/WorldPartitionEditor/Private/WorldPartition/SWorldPartitionEditorGrid2D.cpp b/Engine/Source/Editor/WorldPartitionEditor/Private/WorldPartition/SWorldPartitionEditorGrid2D.cpp
--- a/Engine/Source/Editor/WorldPartitionEditor/Private/WorldPartition/SWorldPartitionEditorGrid2D.cpp
+++ b/Engine/Source/Editor/WorldPartitionEditor/Private/WorldPartition/SWorldPartitionEditorGrid2D.cpp
@@ -1072,7 +1072,7 @@ void SWorldPartitionEditorGrid2D::Tick(const FGeometry& AllottedGeometry, const
 
 	// Also include transient actor loader adapters that might have been spawned by blutilities, etc. Since these actors can't be saved because they are transient,
 	// they will never get an actor descriptor so they will never appear in the world partition editor. Also include unsaved, newly created actors for convenience.
-	for (auto& [Reference, Actor] : GetWorldPartition()->GetForceLoadedActors())
+	for (auto& [Reference, Actor] : GetWorldPartition()->GetDirtyActors())
 	{
 		if (!GetWorldPartition()->GetActorDesc(Actor->GetActorGuid()) && Actor->Implements<UWorldPartitionActorLoaderInterface>())
 		{
diff --git a/Engine/Source/Runtime/Engine/Private/WorldPartition/WorldPartition.cpp b/Engine/Source/Runtime/Engine/Private/WorldPartition/WorldPartition.cpp
--- a/Engine/Source/Runtime/Engine/Private/WorldPartition/WorldPartition.cpp
+++ b/Engine/Source/Runtime/Engine/Private/WorldPartition/WorldPartition.cpp
@@ -257,7 +257,6 @@ public:
 		bIncludeNonSpatiallyLoadedActors = true;
 	}
 };
-
 #endif
 
 UWorldPartition::UWorldPartition(const FObjectInitializer& ObjectInitializer)
@@ -353,7 +352,7 @@ void UWorldPartition::OnPackageDirtyStateChanged(UPackage* Package)
 		{
 			if (Package->IsDirty())
 			{
-				ForceLoadedActors.Add(ActorHandle.ToReference(), Actor);
+				DirtyActors.Add(ActorHandle.ToReference(), Actor);
 			}
 		}
 	}
@@ -744,7 +743,7 @@ void UWorldPartition::Uninitialize()
 			RegisteredEditorLoaderAdapters.Empty();
 		}
 
-		ForceLoadedActors.Empty();
+		DirtyActors.Empty();
 
 		UninitializeActorDescContainers();
 		ActorDescContainer = nullptr;
@@ -1178,7 +1177,7 @@ void UWorldPartition::OnActorDescAdded(FWorldPartitionActorDesc* NewActorDesc)
 
 	if (AActor* NewActor = NewActorDesc->GetActor())
 	{
-		ForceLoadedActors.Add(FWorldPartitionReference(NewActorDesc->GetContainer(), NewActorDesc->GetGuid()), NewActor);
+		DirtyActors.Add(FWorldPartitionReference(NewActorDesc->GetContainer(), NewActorDesc->GetGuid()), NewActor);
 	}
 
 	if (ForceLoadedActors)
@@ -1421,7 +1420,7 @@ void UWorldPartition::AddReferencedObjects(UObject* InThis, FReferenceCollector&
 	if (IsGarbageCollecting())
 	{
 		Collector.AllowEliminatingReferences(false);
-		for (auto& [ActorReference, Actor] : This->ForceLoadedActors)
+		for (auto& [ActorReference, Actor] : This->DirtyActors)
 		{
 			Collector.AddReferencedObject(Actor);
 		}
@@ -1445,23 +1444,23 @@ void UWorldPartition::Tick(float DeltaSeconds)
 		EditorHash->Tick(DeltaSeconds);
 	}
 
-	for (TMap<FWorldPartitionReference, AActor*>::TIterator ForceLoadedActorIt(ForceLoadedActors); ForceLoadedActorIt; ++ForceLoadedActorIt)
+	for (TMap<FWorldPartitionReference, AActor*>::TIterator DirtyActorIt(DirtyActors); DirtyActorIt; ++DirtyActorIt)
 	{
-		if (!ForceLoadedActorIt.Key().IsValid())
+		if (!DirtyActorIt.Key().IsValid())
 		{
-			ForceLoadedActorIt.RemoveCurrent();
+			DirtyActorIt.RemoveCurrent();
 		}
-		else if (!ForceLoadedActorIt.Value()->GetPackage()->IsDirty())
+		else if (!DirtyActorIt.Value()->GetPackage()->IsDirty())
 		{
 			// Transfer ownership of the last ref if actor can be pinned
-			if (ForceLoadedActorIt.Key()->GetHardRefCount() <= 1 && PinnedActors && FLoaderAdapterPinnedActors::SupportsPinning(ForceLoadedActorIt.Key().Get()))
+			if (DirtyActorIt.Key()->GetHardRefCount() <= 1 && PinnedActors && FLoaderAdapterPinnedActors::SupportsPinning(DirtyActorIt.Key().Get()))
 			{
-				PinnedActors->AddActors({ ForceLoadedActorIt.Key().ToHandle() });
-				ForceLoadedActorIt.RemoveCurrent();
+				PinnedActors->AddActors({ DirtyActorIt.Key().ToHandle() });
+				DirtyActorIt.RemoveCurrent();
 			} // Clean up if a loader took a reference on the actor (ILoaderAdapter::RefreshLoadedState was called since)
-			else if (ForceLoadedActorIt.Key()->GetHardRefCount() > 1)
+			else if (DirtyActorIt.Key()->GetHardRefCount() > 1)
 			{
-				ForceLoadedActorIt.RemoveCurrent();
+				DirtyActorIt.RemoveCurrent();
 			}
 		}
 	}
@@ -1897,11 +1896,11 @@ bool UWorldPartition::UnregisterActorDescContainer(UActorDescContainer* InActorD
 			{
 				ActorGuids.Add(It->GetGuid());
 
-				for (TMap<FWorldPartitionReference, AActor*>::TIterator ForceLoadedActorIt(ForceLoadedActors); ForceLoadedActorIt; ++ForceLoadedActorIt)
+				for (TMap<FWorldPartitionReference, AActor*>::TIterator DirtyActorIt(DirtyActors); DirtyActorIt; ++DirtyActorIt)
 				{
-					if (ForceLoadedActorIt.Key() == ActorHandle)
+					if (DirtyActorIt.Key() == ActorHandle)
 					{
-						ForceLoadedActorIt.RemoveCurrent();
+						DirtyActorIt.RemoveCurrent();
 					}
 				}
 			}
diff --git a/Engine/Source/Runtime/Engine/Public/WorldPartition/WorldPartition.h b/Engine/Source/Runtime/Engine/Public/WorldPartition/WorldPartition.h
--- a/Engine/Source/Runtime/Engine/Public/WorldPartition/WorldPartition.h
+++ b/Engine/Source/Runtime/Engine/Public/WorldPartition/WorldPartition.h
@@ -305,10 +305,7 @@ public:
 
 	bool IsEnablingStreamingJustified() const { return bEnablingStreamingJustified; }
 
-	UE_DEPRECATED(5.3, "Use GetForceLoadedActors instead.")
-	const TMap<FWorldPartitionReference, AActor*>& GetDirtyActors() const { return ForceLoadedActors; }
-	const TMap<FWorldPartitionReference, AActor*>& GetForceLoadedActors() const { return ForceLoadedActors; }
-
+	const TMap<FWorldPartitionReference, AActor*>& GetDirtyActors() const { return DirtyActors; }
 #endif
 
 public:
@@ -427,7 +424,7 @@ private:
 
 	TArray<FWorldPartitionReference> LoadedSubobjects;
 
-	TMap<FWorldPartitionReference, AActor*> ForceLoadedActors;
+	TMap<FWorldPartitionReference, AActor*> DirtyActors;
 
 	TSet<FString> GeneratedStreamingPackageNames;
 
