World Partiion: don't try to preserve FWorldTileInfo when loading a level in a partitioned world (this is only used for World Composition).

diff --git a/Engine/Source/Runtime/Engine/Private/WorldComposition.cpp b/Engine/Source/Runtime/Engine/Private/WorldComposition.cpp
--- a/Engine/Source/Runtime/Engine/Private/WorldComposition.cpp
+++ b/Engine/Source/Runtime/Engine/Private/WorldComposition.cpp
@@ -949,12 +949,15 @@ void UWorldComposition::OnLevelPostLoad(ULevel* InLevel)
 				Info = Tile->Info;
 			}
 		}
-		else
+		else if (!World->IsPartitionedWorld())
 		{
 #if WITH_EDITOR
-			// Preserve FWorldTileInfo in case sub-level was loaded in the editor outside of world composition
-			FString PackageFilename = FPackageName::LongPackageNameToFilename(LevelPackage->GetName(), FPackageName::GetMapPackageExtension());
-			FWorldTileInfo::Read(PackageFilename, Info);
+			// Try to preserve FWorldTileInfo in case sub-level was loaded in the editor outside of world composition
+			FString PackageFilename;
+			if (FPackageName::TryConvertLongPackageNameToFilename(LevelPackage->GetName(), PackageFilename, FPackageName::GetMapPackageExtension()))
+			{
+				FWorldTileInfo::Read(PackageFilename, Info);
+			}
 #endif //WITH_EDITOR
 		}
 		
