Wait for reporting task to complete during Deininitialize to avoid race condition with GC.

diff --git a/Engine/Source/Runtime/Engine/Private/ProfilingDebugging/LevelStreamingProfilingSubsystem.cpp b/Engine/Source/Runtime/Engine/Private/ProfilingDebugging/LevelStreamingProfilingSubsystem.cpp
--- a/Engine/Source/Runtime/Engine/Private/ProfilingDebugging/LevelStreamingProfilingSubsystem.cpp
+++ b/Engine/Source/Runtime/Engine/Private/ProfilingDebugging/LevelStreamingProfilingSubsystem.cpp
@@ -103,6 +103,12 @@ void ULevelStreamingProfilingSubsystem::Deinitialize()
 		StopTrackingAndReport();
 	}
 
+	if (!ReportWritingTask.IsCompleted())
+	{
+		UE_LOG(LogLevelStreamingProfiling, Log, TEXT("Waiting for report writing task to complete during Deinitialize for safety"));
+		ReportWritingTask.Wait();
+	}
+
 	FLevelStreamingDelegates::OnLevelStreamingTargetStateChanged.Remove(Handle_OnLevelStreamingTargetStateChanged);
 	FLevelStreamingDelegates::OnLevelStreamingStateChanged.Remove(Handle_OnLevelStreamingStateChanged);
 	FLevelStreamingDelegates::OnLevelBeginMakingVisible.Remove(Handle_OnLevelBeginAddToWorld);
