Fix for TransformMode, GroundAlbedo and GroundRadius on USkyAtmosphereComponent not working properly with sequencer.

#rnx

diff --git a/Engine/Source/Runtime/Engine/Classes/Components/SkyAtmosphereComponent.h b/Engine/Source/Runtime/Engine/Classes/Components/SkyAtmosphereComponent.h
--- a/Engine/Source/Runtime/Engine/Classes/Components/SkyAtmosphereComponent.h
+++ b/Engine/Source/Runtime/Engine/Classes/Components/SkyAtmosphereComponent.h
@@ -51,7 +51,7 @@ class USkyAtmosphereComponent : public USceneComponent
 
 
 	/** The ground albedo that will tint the atmosphere when the sun light will bounce on it. Only taken into account when MultiScattering>0.0. */
-	UPROPERTY(EditAnywhere, BlueprintReadOnly, interp, Category = "Planet", meta = (HideAlphaChannel))
+	UPROPERTY(EditAnywhere, BlueprintReadOnly, Category = "Planet", meta = (HideAlphaChannel))
 	ESkyAtmosphereTransformMode TransformMode;
 
 	/** The radius in kilometers from the center of the planet to the ground level. */
@@ -170,6 +170,11 @@ class USkyAtmosphereComponent : public USceneComponent
 	UFUNCTION(BlueprintCallable, Category = "Rendering")
 	ENGINE_API void ResetAtmosphereLightDirectionOverride(int32 AtmosphereLightIndex);
 
+	UFUNCTION(BlueprintCallable, Category = "Rendering", meta = (DisplayName = "Set Ground Radius"))
+	ENGINE_API void SetBottomRadius(float NewValue);
+	UFUNCTION(BlueprintCallable, Category = "Rendering")
+	ENGINE_API void SetGroundAlbedo(const FColor& NewValue);
+
 	UFUNCTION(BlueprintCallable, Category = "Rendering")
 	ENGINE_API void SetAtmosphereHeight(float NewValue);
 	UFUNCTION(BlueprintCallable, Category = "Rendering")
diff --git a/Engine/Source/Runtime/Engine/Private/Components/SkyAtmosphereComponent.cpp b/Engine/Source/Runtime/Engine/Private/Components/SkyAtmosphereComponent.cpp
--- a/Engine/Source/Runtime/Engine/Private/Components/SkyAtmosphereComponent.cpp
+++ b/Engine/Source/Runtime/Engine/Private/Components/SkyAtmosphereComponent.cpp
@@ -363,6 +363,9 @@ void USkyAtmosphereComponent::SetPositionToMatchDeprecatedAtmosphericFog()
 	}\
 }\
 
+SKY_DECLARE_BLUEPRINT_SETFUNCTION(float, BottomRadius);
+SKY_DECLARE_BLUEPRINT_SETFUNCTION(const FColor&, GroundAlbedo);
+
 SKY_DECLARE_BLUEPRINT_SETFUNCTION(float, AtmosphereHeight);
 SKY_DECLARE_BLUEPRINT_SETFUNCTION(float, MultiScatteringFactor);
 
