Fix bone LOD filtering automated tests

diff --git a/Engine/Source/Runtime/Engine/Private/BoneContainer.cpp b/Engine/Source/Runtime/Engine/Private/BoneContainer.cpp
--- a/Engine/Source/Runtime/Engine/Private/BoneContainer.cpp
+++ b/Engine/Source/Runtime/Engine/Private/BoneContainer.cpp
@@ -52,7 +52,7 @@ FBoneContainer::FBoneContainer(const TArray<FBoneIndexType>& InRequiredBoneIndex
 {
 	LLM_SCOPE_BYNAME(TEXT("Animation/BoneContainer"));
 
-	const UE::Anim::FCurveFilterSettings CurveFilterSettings(CurveEvalOption.bAllowCurveEvaluation ? UE::Anim::ECurveFilterMode::None : UE::Anim::ECurveFilterMode::DisallowAll, CurveEvalOption.DisallowedList, CurveEvalOption.LODIndex);
+	const UE::Anim::FCurveFilterSettings CurveFilterSettings(CurveEvalOption.bAllowCurveEvaluation ? UE::Anim::ECurveFilterMode::DisallowFiltered : UE::Anim::ECurveFilterMode::DisallowAll, CurveEvalOption.DisallowedList, CurveEvalOption.LODIndex);
 	Initialize(CurveFilterSettings);
 }
 
@@ -76,7 +76,7 @@ FBoneContainer::FBoneContainer(const TArray<FBoneIndexType>& InRequiredBoneIndex
 
 void FBoneContainer::InitializeTo(const TArray<FBoneIndexType>& InRequiredBoneIndexArray, const FCurveEvaluationOption& CurveEvalOption, UObject& InAsset)
 {
-	const UE::Anim::FCurveFilterSettings CurveFilterSettings(CurveEvalOption.bAllowCurveEvaluation ? UE::Anim::ECurveFilterMode::None : UE::Anim::ECurveFilterMode::DisallowAll, CurveEvalOption.DisallowedList, CurveEvalOption.LODIndex);
+	const UE::Anim::FCurveFilterSettings CurveFilterSettings(CurveEvalOption.bAllowCurveEvaluation ? UE::Anim::ECurveFilterMode::DisallowFiltered : UE::Anim::ECurveFilterMode::DisallowAll, CurveEvalOption.DisallowedList, CurveEvalOption.LODIndex);
 	InitializeTo(InRequiredBoneIndexArray, CurveFilterSettings, InAsset);
 }
 
diff --git a/Engine/Source/Runtime/Engine/Private/Components/SkeletalMeshComponent.cpp b/Engine/Source/Runtime/Engine/Private/Components/SkeletalMeshComponent.cpp
--- a/Engine/Source/Runtime/Engine/Private/Components/SkeletalMeshComponent.cpp
+++ b/Engine/Source/Runtime/Engine/Private/Components/SkeletalMeshComponent.cpp
@@ -4616,7 +4616,7 @@ UE::Anim::FCurveFilterSettings USkeletalMeshComponent::GetCurveFilterSettings(in
 		}
 		else
 		{
-			CurveFilterSettings.FilterMode = UE::Anim::ECurveFilterMode::None;
+			CurveFilterSettings.FilterMode = UE::Anim::ECurveFilterMode::DisallowFiltered;	// This applies LOD-based and bone-linked filtering
 		}
 	}
 	else
diff --git a/Engine/Source/Runtime/Engine/Public/BoneContainer.h b/Engine/Source/Runtime/Engine/Public/BoneContainer.h
--- a/Engine/Source/Runtime/Engine/Public/BoneContainer.h
+++ b/Engine/Source/Runtime/Engine/Public/BoneContainer.h
@@ -76,7 +76,7 @@ namespace UE::Anim
 /** Options for filtering curves. Used to initialize a bone container with a filter curve. */
 struct FCurveFilterSettings
 {
-	FCurveFilterSettings(UE::Anim::ECurveFilterMode InFilterMode = UE::Anim::ECurveFilterMode::None, const TArray<FName>* InFilterCurves = nullptr, int32 InLODIndex = 0)
+	FCurveFilterSettings(UE::Anim::ECurveFilterMode InFilterMode = UE::Anim::ECurveFilterMode::DisallowFiltered, const TArray<FName>* InFilterCurves = nullptr, int32 InLODIndex = 0)
 		: FilterCurves(InFilterCurves)
 		, FilterMode(InFilterMode)
 		, LODIndex(InLODIndex)
@@ -86,7 +86,7 @@ struct FCurveFilterSettings
 	const TArray<FName>* FilterCurves = nullptr;
 	
 	// Filtering mode - allows curves to be disable entirely or allow/deny-listed
-	UE::Anim::ECurveFilterMode FilterMode = UE::Anim::ECurveFilterMode::None;
+	UE::Anim::ECurveFilterMode FilterMode = UE::Anim::ECurveFilterMode::DisallowFiltered;
 
 	// LOD index
 	int32 LODIndex = 0;
