Output log contains error for ChaosEventType.h

diff --git a/Engine/Source/Runtime/Engine/Private/PhysicsEngine/Experimental/ChaosEventType.cpp b/Engine/Source/Runtime/Engine/Private/PhysicsEngine/Experimental/ChaosEventType.cpp
--- a/Engine/Source/Runtime/Engine/Private/PhysicsEngine/Experimental/ChaosEventType.cpp
+++ b/Engine/Source/Runtime/Engine/Private/PhysicsEngine/Experimental/ChaosEventType.cpp
@@ -7,12 +7,22 @@
 #include UE_INLINE_GENERATED_CPP_BY_NAME(ChaosEventType)
 
 FCollisionChaosEventBodyInfo::FCollisionChaosEventBodyInfo()
-	: PhysMaterial(nullptr)
+	: Velocity(FVector::ZeroVector)
+	, DeltaVelocity(FVector::ZeroVector)
+	, AngularVelocity(FVector::ZeroVector)
+	, Mass(0.0)
+	, PhysMaterial(nullptr)
 	, Component(nullptr)
+	, BodyIndex(0)
+	, BoneName(NAME_None)
 {
 }
 
 FCollisionChaosEvent::FCollisionChaosEvent()
+	: Location(FVector::ZeroVector)
+	, AccumulatedImpulse(FVector::ZeroVector)
+	, Normal(FVector::ZeroVector)
+	, PenetrationDepth(0.0)
 {
 }
 
