Cache CVarNetEnableDetailedScopeCounters into global variable in order to prevent unnecessary work getting the CVar value every conditional scope cycle counter. Only record ReplicateCustomDeltaPropTime if object has any custom delta properties.

diff --git a/Engine/Source/Runtime/Engine/Private/DataChannel.cpp b/Engine/Source/Runtime/Engine/Private/DataChannel.cpp
--- a/Engine/Source/Runtime/Engine/Private/DataChannel.cpp
+++ b/Engine/Source/Runtime/Engine/Private/DataChannel.cpp
@@ -14,6 +14,7 @@
 #include "Misc/MemStack.h"
 #include "Misc/ScopeExit.h"
 #include "Net/Core/Trace/Private/NetTraceInternal.h"
+#include "Net/Core/NetCoreModule.h"
 #include "UObject/UObjectIterator.h"
 #include "EngineStats.h"
 #include "Engine/Engine.h"
@@ -91,8 +92,6 @@ FAutoConsoleVariableRef CVarNetSkipReplicatorForDestructionInfos(
 	GSkipReplicatorForDestructionInfos,
 	TEXT("If enabled, skip creation of object replicator in SetChannelActor when we know there is no content payload and we're going to immediately destroy the actor."));
 
-extern NETCORE_API TAutoConsoleVariable<int32> CVarNetEnableDetailedScopeCounters;
-
 // Fairly large number, and probably a bad idea to even have a bunch this size, but want to be safe for now and not throw out legitimate data
 static int32 NetMaxConstructedPartialBunchSizeBytes = 1024 * 64;
 static FAutoConsoleVariableRef CVarNetMaxConstructedPartialBunchSizeBytes(
@@ -108,6 +107,8 @@ static FAutoConsoleVariableRef CVarDormancyHysteresis(
 	TEXT("When > 0, represents the time we'll wait before letting a channel become fully dormant (in seconds). This can prevent churn when objects are going in and out of dormant more frequently than normal.")
 );
 
+extern bool GUseDetailedScopeCounters;
+
 namespace UE::Net
 {
 	extern int32 FilterGuidRemapping;
@@ -5071,7 +5072,7 @@ TSharedRef<FObjectReplicator>* UActorChannel::FindReplicator(UObject* Obj)
 
 TSharedRef<FObjectReplicator>* UActorChannel::FindReplicator(UObject* Obj, bool* bOutFoundInvalid)
 {
-	CONDITIONAL_SCOPE_CYCLE_COUNTER(Stat_ActorChanFindOrCreateRep, CVarNetEnableDetailedScopeCounters.GetValueOnAnyThread() > 0);
+	CONDITIONAL_SCOPE_CYCLE_COUNTER(Stat_ActorChanFindOrCreateRep, GUseDetailedScopeCounters);
 
 	// First, try to find it on the channel replication map
 	TSharedRef<FObjectReplicator>* ReplicatorRefPtr = ReplicationMap.Find( Obj );
@@ -5106,7 +5107,7 @@ TSharedRef<FObjectReplicator>& UActorChannel::CreateReplicator(UObject* Obj)
 
 TSharedRef<FObjectReplicator>& UActorChannel::CreateReplicator(UObject* Obj, bool bCheckDormantReplicators)
 {
-	CONDITIONAL_SCOPE_CYCLE_COUNTER(Stat_ActorChanFindOrCreateRep, CVarNetEnableDetailedScopeCounters.GetValueOnAnyThread() > 0);
+	CONDITIONAL_SCOPE_CYCLE_COUNTER(Stat_ActorChanFindOrCreateRep, GUseDetailedScopeCounters);
 
 	// Try to find in the dormancy map if desired
 	TSharedPtr<FObjectReplicator> NewReplicator;
diff --git a/Engine/Source/Runtime/Engine/Private/DataReplication.cpp b/Engine/Source/Runtime/Engine/Private/DataReplication.cpp
--- a/Engine/Source/Runtime/Engine/Private/DataReplication.cpp
+++ b/Engine/Source/Runtime/Engine/Private/DataReplication.cpp
@@ -21,6 +21,7 @@
 #include "Engine/NetConnection.h"
 #include "Net/NetworkGranularMemoryLogging.h"
 #include "Net/Core/Trace/NetTrace.h"
+#include "Net/Core/NetCoreModule.h"
 #include "HAL/LowLevelMemStats.h"
 #include "Net/Core/PushModel/Types/PushModelPerNetDriverState.h"
 #include "Net/RPCDoSDetection.h"
@@ -83,7 +84,7 @@ FAutoConsoleVariableRef CVarPushModelSkipUndirtiedFastArrays(
 
 extern int32 GNumSkippedObjectEmptyUpdates;
 
-extern NETCORE_API TAutoConsoleVariable<int32> CVarNetEnableDetailedScopeCounters;
+extern bool GUseDetailedScopeCounters;
 
 class FNetSerializeCB : public INetSerializeCB
 {
@@ -345,7 +346,7 @@ bool FObjectReplicator::SendCustomDeltaProperty(UObject* InObject, uint16 Custom
 	check(!NewFullState.IsValid()); // NewState is passed in as nullptr and instantiated within this function if necessary
 	check(RepLayout);
 
-	CONDITIONAL_SCOPE_CYCLE_COUNTER(STAT_NetSerializeItemDeltaTime, CVarNetEnableDetailedScopeCounters.GetValueOnAnyThread() > 0);
+	CONDITIONAL_SCOPE_CYCLE_COUNTER(STAT_NetSerializeItemDeltaTime, GUseDetailedScopeCounters);
 
 	UNetDriver* const ConnectionDriver = Connection->GetDriver();
 	FNetSerializeCB NetSerializeCB(ConnectionDriver);
@@ -1555,8 +1556,6 @@ static FORCEINLINE FPropertyRetirement** UpdateAckedRetirements(
 
 void FObjectReplicator::ReplicateCustomDeltaProperties( FNetBitWriter & Bunch, FReplicationFlags RepFlags, bool& bSkippedPropertyCondition)
 {
-	CONDITIONAL_SCOPE_CYCLE_COUNTER(STAT_NetReplicateCustomDeltaPropTime, CVarNetEnableDetailedScopeCounters.GetValueOnAnyThread() > 0);
-
 	check(RepLayout);
 	const FRepLayout& LocalRepLayout = *RepLayout;
 
@@ -1570,6 +1569,8 @@ void FObjectReplicator::ReplicateCustomDeltaProperties( FNetBitWriter & Bunch, F
 		return;
 	}
 
+	CONDITIONAL_SCOPE_CYCLE_COUNTER(STAT_NetReplicateCustomDeltaPropTime, GUseDetailedScopeCounters);
+
 	// TODO: See comments in ReceivedBunch. This code should get merged into RepLayout, to help optimize
 	//			the receiving end, and make things more consistent.
 
diff --git a/Engine/Source/Runtime/Engine/Private/RepLayout.cpp b/Engine/Source/Runtime/Engine/Private/RepLayout.cpp
--- a/Engine/Source/Runtime/Engine/Private/RepLayout.cpp
+++ b/Engine/Source/Runtime/Engine/Private/RepLayout.cpp
@@ -21,6 +21,7 @@
 #include "Serialization/ArchiveCountMem.h"
 #include "Net/Core/PushModel/Types/PushModelPerNetDriverState.h"
 #include "Net/Core/Trace/NetTrace.h"
+#include "Net/Core/NetCoreModule.h"
 #include "Stats/StatsTrace.h"
 #include "UObject/EnumProperty.h"
 #if UE_WITH_IRIS
@@ -115,7 +116,7 @@ namespace UE::Net::Private
 extern int32 GNumSharedSerializationHit;
 extern int32 GNumSharedSerializationMiss;
 
-extern NETCORE_API TAutoConsoleVariable<int32> CVarNetEnableDetailedScopeCounters;
+extern bool GUseDetailedScopeCounters;
 
 /** 
 * Helper method to allow us to instrument FBitArchive using FNetTraceCollector
@@ -1750,7 +1751,7 @@ ERepLayoutResult FRepLayout::CompareProperties(
 	const FConstRepObjectDataBuffer Data,
 	const FReplicationFlags& RepFlags) const
 {
-	CONDITIONAL_SCOPE_CYCLE_COUNTER(STAT_NetReplicateDynamicPropCompareTime, CVarNetEnableDetailedScopeCounters.GetValueOnAnyThread() > 0);
+	CONDITIONAL_SCOPE_CYCLE_COUNTER(STAT_NetReplicateDynamicPropCompareTime, GUseDetailedScopeCounters);
 
 	if (IsEmpty())
 	{
@@ -1944,7 +1945,7 @@ bool FRepLayout::ReplicateProperties(
 	FNetBitWriter& Writer,
 	const FReplicationFlags& RepFlags) const
 {
-	CONDITIONAL_SCOPE_CYCLE_COUNTER(STAT_NetReplicateDynamicPropTime, CVarNetEnableDetailedScopeCounters.GetValueOnAnyThread() > 0);
+	CONDITIONAL_SCOPE_CYCLE_COUNTER(STAT_NetReplicateDynamicPropTime, GUseDetailedScopeCounters);
 
 	check(ObjectClass == Owner);
 
@@ -7538,7 +7539,7 @@ ERepLayoutResult FRepLayout::DeltaSerializeFastArrayProperty(FFastArrayDeltaSeri
 {
 	using namespace UE_RepLayout_Private;
 
-	CONDITIONAL_SCOPE_CYCLE_COUNTER(STAT_RepLayout_DeltaSerializeFastArray, CVarNetEnableDetailedScopeCounters.GetValueOnAnyThread() > 0);
+	CONDITIONAL_SCOPE_CYCLE_COUNTER(STAT_RepLayout_DeltaSerializeFastArray, GUseDetailedScopeCounters);
 
 	// A portion of this work could be shared across all Fast Array Properties for a given object,
 	// but that would be easier to do if the Custom Delta Serialization was completely encapsulated in FRepLayout.
diff --git a/Engine/Source/Runtime/Net/Core/Classes/Net/Serialization/FastArraySerializer.h b/Engine/Source/Runtime/Net/Core/Classes/Net/Serialization/FastArraySerializer.h
--- a/Engine/Source/Runtime/Net/Core/Classes/Net/Serialization/FastArraySerializer.h
+++ b/Engine/Source/Runtime/Net/Core/Classes/Net/Serialization/FastArraySerializer.h
@@ -17,6 +17,7 @@
 #include "Containers/ArrayView.h"
 #include "Net/Core/Misc/GuidReferences.h"
 #include "Net/Core/PushModel/PushModel.h"
+#include "Net/Core/NetCoreModule.h"
 #include "HAL/IConsoleManager.h"
 #include "Templates/EnableIf.h"
 #include "FastArraySerializer.generated.h"
@@ -31,7 +32,7 @@ DECLARE_CYCLE_STAT_EXTERN(TEXT("NetSerializeFast Array"), STAT_NetSerializeFastA
 DECLARE_CYCLE_STAT_EXTERN(TEXT("NetSerializeFast Array BuildMap"), STAT_NetSerializeFastArray_BuildMap, STATGROUP_ServerCPU, NETCORE_API);
 DECLARE_CYCLE_STAT_EXTERN(TEXT("NetSerializeFast Array Delta Struct"), STAT_NetSerializeFastArray_DeltaStruct, STATGROUP_ServerCPU, NETCORE_API);
 
-extern NETCORE_API TAutoConsoleVariable<int32> CVarNetEnableDetailedScopeCounters;
+extern bool GUseDetailedScopeCounters;
 
 /**
  *	===================== Fast TArray Replication ===================== 
@@ -1190,7 +1191,7 @@ bool FFastArraySerializer::FastArrayDeltaSerialize(TArray<Type> &Items, FNetDelt
 		return FastArrayDeltaSerialize_DeltaSerializeStructs(Items, Parms, ArraySerializer);
 	}
 
-	CONDITIONAL_SCOPE_CYCLE_COUNTER(STAT_NetSerializeFastArray, CVarNetEnableDetailedScopeCounters.GetValueOnAnyThread() > 0);
+	CONDITIONAL_SCOPE_CYCLE_COUNTER(STAT_NetSerializeFastArray, GUseDetailedScopeCounters);
 	class UScriptStruct* InnerStruct = Type::StaticStruct();
 
 	UE_LOG(LogNetFastTArray, Log, TEXT("FastArrayDeltaSerialize for %s. %s. %s"), *InnerStruct->GetName(), *InnerStruct->GetOwnerStruct()->GetName(), Parms.Reader ? TEXT("Reading") : TEXT("Writing"));
@@ -1640,7 +1641,7 @@ bool FFastArraySerializer::FastArrayDeltaSerialize_DeltaSerializeStructs(TArray<
 		}
 	};
 
-	CONDITIONAL_SCOPE_CYCLE_COUNTER(STAT_NetSerializeFastArray_DeltaStruct, CVarNetEnableDetailedScopeCounters.GetValueOnAnyThread() > 0);
+	CONDITIONAL_SCOPE_CYCLE_COUNTER(STAT_NetSerializeFastArray_DeltaStruct, GUseDetailedScopeCounters);
 
 	class UScriptStruct* InnerStruct = Type::StaticStruct();
 
diff --git a/Engine/Source/Runtime/Net/Core/Private/Net/Core/NetCoreModule.cpp b/Engine/Source/Runtime/Net/Core/Private/Net/Core/NetCoreModule.cpp
--- a/Engine/Source/Runtime/Net/Core/Private/Net/Core/NetCoreModule.cpp
+++ b/Engine/Source/Runtime/Net/Core/Private/Net/Core/NetCoreModule.cpp
@@ -1,5 +1,6 @@
 // Copyright Epic Games, Inc. All Rights Reserved.
 
+#include "Net/Core/NetCoreModule.h"
 #include "CoreMinimal.h"
 #include "Modules/ModuleInterface.h"
 #include "Modules/ModuleManager.h"
@@ -10,7 +11,12 @@
 #include "Net/Core/NetHandle/NetHandleManager.h"
 #include "Net/Serialization/FastArraySerializer.h"
 
-TAutoConsoleVariable<int32> CVarNetEnableDetailedScopeCounters(TEXT("net.EnableDetailedScopeCounters"), 1, TEXT("Enables detailed networking scope cycle counters. There are often lots of these which can negatively impact performance."));
+extern bool GUseDetailedScopeCounters;
+static FAutoConsoleVariableRef CVarNetEnableDetailedScopeCounters(
+	TEXT("net.EnableDetailedScopeCounters"),
+	GUseDetailedScopeCounters,
+	TEXT("Enables detailed networking scope cycle counters. There are often lots of these which can negatively impact performance."),
+	ECVF_Default);
 
 DEFINE_LOG_CATEGORY(LogNetFastTArray);
 
diff --git a/Engine/Source/Runtime/Net/Core/Public/Net/Core/NetCoreModule.h b/Engine/Source/Runtime/Net/Core/Public/Net/Core/NetCoreModule.h
new file mode 100644
--- /dev/null
+++ b/Engine/Source/Runtime/Net/Core/Public/Net/Core/NetCoreModule.h
@@ -0,0 +1,5 @@
+// Copyright Epic Games, Inc. All Rights Reserved.
+
+#pragma once
+
+static bool GUseDetailedScopeCounters = true;
\ No newline at end of file
