Substrate - added dialog box in order to be able to warn the user that this feature needsto be handled with care and give a ini file alternative.

diff --git a/Engine/Source/Runtime/Engine/Private/RendererSettings.cpp b/Engine/Source/Runtime/Engine/Private/RendererSettings.cpp
--- a/Engine/Source/Runtime/Engine/Private/RendererSettings.cpp
+++ b/Engine/Source/Runtime/Engine/Private/RendererSettings.cpp
@@ -232,6 +232,16 @@ void URendererSettings::PostEditChangeProperty(FPropertyChangedEvent& PropertyCh
 			}
 		}
 
+		if (PropertyChangedEvent.Property->GetFName() == GET_MEMBER_NAME_CHECKED(URendererSettings, StrataDebugAdvancedVisualizationShaders)
+			&& StrataDebugAdvancedVisualizationShaders)
+		{
+			if (FMessageDialog::Open(EAppMsgType::YesNo, LOCTEXT("Substrate advanced visualization shaders", "This will make the rendering slower and allocate more memory.\nPlease do not check-in.\nInstead you could add `r.Substrate.Debug.AdvancedVisualizationShaders=1` to your ConsoleVariables.ini.\nAre you sure you want to enable that feature now?")) == EAppReturnType::No)
+			{
+				StrataDebugAdvancedVisualizationShaders = false;
+				UpdateDependentPropertyInConfigFile(this, GET_MEMBER_NAME_CHECKED(URendererSettings, StrataDebugAdvancedVisualizationShaders));
+			}
+		}
+
 		if (PropertyChangedEvent.Property->GetFName() == GET_MEMBER_NAME_CHECKED(URendererSettings, VirtualTextureTileSize))
 		{
 			VirtualTextureTileSize = FMath::RoundUpToPowerOfTwo(VirtualTextureTileSize);
