Add profile markers for EditorViewport Tick and HitProxyMap readback

diff --git a/Engine/Source/Editor/UnrealEd/Private/EditorViewportClient.cpp b/Engine/Source/Editor/UnrealEd/Private/EditorViewportClient.cpp
--- a/Engine/Source/Editor/UnrealEd/Private/EditorViewportClient.cpp
+++ b/Engine/Source/Editor/UnrealEd/Private/EditorViewportClient.cpp
@@ -1375,6 +1375,8 @@ void FEditorViewportClient::LostFocus(FViewport* InViewport)
 
 void FEditorViewportClient::Tick(float DeltaTime)
 {
+	SCOPED_NAMED_EVENT(FEditorViewportClient_Tick, FColor::Red);
+	
 	ConditionalCheckHoveredHitProxy();
 
 	FViewportCameraTransform& ViewTransform = GetViewTransform();
diff --git a/Engine/Source/Runtime/Engine/Private/UnrealClient.cpp b/Engine/Source/Runtime/Engine/Private/UnrealClient.cpp
--- a/Engine/Source/Runtime/Engine/Private/UnrealClient.cpp
+++ b/Engine/Source/Runtime/Engine/Private/UnrealClient.cpp
@@ -1939,6 +1939,8 @@ const TArray<FColor>& FViewport::GetRawHitProxyData(FIntRect InRect)
 	// If the hit proxy map isn't up to date, render the viewport client's hit proxies to it.
 	else if (!bHitProxiesCached)
 	{
+		SCOPED_NAMED_EVENT(HitProxyMapGen, FColor::Red);
+		
 		RenderCaptureInterface::FScopedCapture RenderCapture(GHitProxyCaptureNextUpdate != 0, TEXT("Update Hit Proxies"));
 		GHitProxyCaptureNextUpdate = 0;
 
@@ -1985,6 +1987,7 @@ const TArray<FColor>& FViewport::GetRawHitProxyData(FIntRect InRect)
 	if (bFetchHitProxyBytes)
 	{
 		// Read the hit proxy map surface data back.
+		SCOPED_NAMED_EVENT(HitProxyReadback, FColor::Red);
 		FIntRect ViewportRect(0, 0, SizeX, SizeY);
 		struct FReadSurfaceContext
 		{
