Change PIE initialization to create game instance subsystems before it creates and initializes world subsystems.
This matches what happens in packaged games on startup and all cases of level transfers

diff --git a/Engine/Source/Runtime/Engine/Private/GameInstance.cpp b/Engine/Source/Runtime/Engine/Private/GameInstance.cpp
--- a/Engine/Source/Runtime/Engine/Private/GameInstance.cpp
+++ b/Engine/Source/Runtime/Engine/Private/GameInstance.cpp
@@ -350,16 +350,18 @@ FGameInstancePIEResult UGameInstance::InitializeForPlayInEditor(int32 PIEInstanc
 	WorldContext->AddRef(static_cast<UWorld*&>(EditorEngine->PlayWorld));	// Tie this context to this UEngine::PlayWorld*		// @fixme, needed still?
 	NewWorld->bKismetScriptError = Params.bAnyBlueprintErrors;
 
-	// Initialize the world after setting world context to be consistent with normal loads
-	EditorEngine->PostCreatePIEWorld(NewWorld);
-
 	// Do a GC pass if necessary to remove any potentially unreferenced objects
 	if(bNeedsGarbageCollection)
 	{
 		CollectGarbage(GARBAGE_COLLECTION_KEEPFLAGS);
 	}
 
+	// This creates the game instance subsystems
 	Init();
+	
+	// Initialize the world after setting world context and initializing the game instance to be consistent with normal loads.
+	// This creates the world subsystems and prepares to begin play
+	EditorEngine->PostCreatePIEWorld(NewWorld);
 
 	// Games can override this to return failure if PIE is not allowed for some reason
 	return FGameInstancePIEResult::Success();
