Reverted purge editor during PIE for now until other issues are resolved

diff --git a/Engine/Source/Runtime/Engine/Private/World.cpp b/Engine/Source/Runtime/Engine/Private/World.cpp
--- a/Engine/Source/Runtime/Engine/Private/World.cpp
+++ b/Engine/Source/Runtime/Engine/Private/World.cpp
@@ -146,7 +146,7 @@ FAutoConsoleVariableRef CVarDisableRemapScriptActors(TEXT("net.DisableRemapScrip
 
 static TAutoConsoleVariable<int32> CVarPurgeEditorSceneDuringPIE(
 	TEXT("r.PurgeEditorSceneDuringPIE"),
-	1,
+	0,
 	TEXT("0 to keep editor scene fully initialized during PIE (default)\n")
 	TEXT("1 to purge editor scene from memory during PIE and restore when the session finishes."));
 
