Path Tracing: Remove the Pixel Filter Width from PPV

Setting this smaller than the default value produces superficially sharp images which actually exhibit strong aliasing artifacts. Now that the post process sharpen settings are exposed in the post process volume, they are a much nicer way of improving image sharpness without introducing artifacts.

For backwards compatibility, the old cvar is still there but is mainly intended for debugging.

diff --git a/Engine/Source/Runtime/Engine/Classes/Engine/Scene.h b/Engine/Source/Runtime/Engine/Classes/Engine/Scene.h
--- a/Engine/Source/Runtime/Engine/Classes/Engine/Scene.h
+++ b/Engine/Source/Runtime/Engine/Classes/Engine/Scene.h
@@ -1274,9 +1274,6 @@ struct FPostProcessSettings
 	UPROPERTY(EditAnywhere, BlueprintReadWrite, Category = Overrides, meta = (PinHiddenByDefault, InlineEditConditionToggle))
 	uint32 bOverride_PathTracingSamplesPerPixel : 1;
 
-	UPROPERTY(EditAnywhere, BlueprintReadWrite, Category = Overrides, meta = (PinHiddenByDefault, InlineEditConditionToggle))
-	uint32 bOverride_PathTracingFilterWidth : 1;
-
 	UPROPERTY(EditAnywhere, BlueprintReadWrite, Category = Overrides, meta = (PinHiddenByDefault, InlineEditConditionToggle))
 	uint32 bOverride_PathTracingMaxPathExposure : 1;
 
@@ -2175,10 +2172,6 @@ struct FPostProcessSettings
 	UPROPERTY(EditAnywhere, BlueprintReadWrite, Category = "Path Tracing", meta = (ClampMin = "1", UIMax = "65536", editcondition = "bOverride_PathTracingSamplesPerPixel", DisplayName = "Samples Per Pixel"))
 	int32 PathTracingSamplesPerPixel;
 
-	/** Sets anti-aliasing filter width for the path tracer. Lower values are sharper (and more aliased), larger values are softer (and blurrier). */
-	UPROPERTY(interp, EditAnywhere, BlueprintReadWrite, Category = "Path Tracing", meta = (ClampMin = "1.0", ClampMax = "6.0", editcondition = "bOverride_PathTracingFilterWidth", DisplayName = "Filter Width"))
-	float PathTracingFilterWidth;
-
 	/** Sets the maximum exposure allowed in the path tracer to reduce fireflies. This should be set a few stops higher than the scene exposure. */
 	UPROPERTY(interp, EditAnywhere, BlueprintReadWrite, Category = "Path Tracing", meta = (ClampMin = "-10.0", ClampMax = "30.0", editcondition = "bOverride_PathTracingMaxPathExposure", DisplayName = "Max Path Exposure"))
 	float PathTracingMaxPathExposure;
diff --git a/Engine/Source/Runtime/Engine/Private/Scene.cpp b/Engine/Source/Runtime/Engine/Private/Scene.cpp
--- a/Engine/Source/Runtime/Engine/Private/Scene.cpp
+++ b/Engine/Source/Runtime/Engine/Private/Scene.cpp
@@ -632,7 +632,6 @@ FPostProcessSettings::FPostProcessSettings()
 
 	PathTracingMaxBounces = 32;
 	PathTracingSamplesPerPixel = 2048;
-	PathTracingFilterWidth = 3.0f;
 	PathTracingMaxPathExposure = 30.0f;
 	PathTracingEnableReferenceDOF = 0;
 	PathTracingEnableReferenceAtmosphere = 0;
diff --git a/Engine/Source/Runtime/Engine/Private/SceneView.cpp b/Engine/Source/Runtime/Engine/Private/SceneView.cpp
--- a/Engine/Source/Runtime/Engine/Private/SceneView.cpp
+++ b/Engine/Source/Runtime/Engine/Private/SceneView.cpp
@@ -1766,11 +1766,6 @@ void FSceneView::OverridePostProcessSettings(const FPostProcessSettings& Src, fl
 			Dest.PathTracingSamplesPerPixel = Src.PathTracingSamplesPerPixel;
 		}
 
-		if (Src.bOverride_PathTracingFilterWidth)
-		{
-			Dest.PathTracingFilterWidth = Src.PathTracingFilterWidth;
-		}
-
 		if (Src.bOverride_PathTracingMaxPathExposure)
 		{
 			Dest.PathTracingMaxPathExposure = Src.PathTracingMaxPathExposure;
diff --git a/Engine/Source/Runtime/Renderer/Private/PathTracing.cpp b/Engine/Source/Runtime/Renderer/Private/PathTracing.cpp
--- a/Engine/Source/Runtime/Renderer/Private/PathTracing.cpp
+++ b/Engine/Source/Runtime/Renderer/Private/PathTracing.cpp
@@ -82,8 +82,8 @@ TAutoConsoleVariable<int32> CVarPathTracingSamplesPerPixel(
 
 TAutoConsoleVariable<float> CVarPathTracingFilterWidth(
 	TEXT("r.PathTracing.FilterWidth"),
-	-1,
-	TEXT("Sets the anti-aliasing filter width (default = -1 (driven by postprocesing volume))"),
+	3.0,
+	TEXT("Sets the anti-aliasing filter width (default = 3.0 which corresponds to a gaussian with standard deviation of a 1/2 pixel)"),
 	ECVF_RenderThreadSafe
 );
 
@@ -611,12 +611,7 @@ static void PreparePathTracingData(const FScene* Scene, const FViewInfo& View, F
 	PathTracingData.SamplerType = CVarPathTracingSamplerType.GetValueOnRenderThread();
 	PathTracingData.VisualizeLightGrid = CVarPathTracingLightGridVisualize.GetValueOnRenderThread();
 	PathTracingData.VisualizeDecalGrid = CVarPathTracingDecalGridVisualize.GetValueOnRenderThread();
-	float FilterWidth = CVarPathTracingFilterWidth.GetValueOnRenderThread();
-	if (FilterWidth < 0)
-	{
-		FilterWidth = PPV.PathTracingFilterWidth;
-	}
-	PathTracingData.FilterWidth = FilterWidth;
+	PathTracingData.FilterWidth = CVarPathTracingFilterWidth.GetValueOnRenderThread();
 	PathTracingData.CameraFocusDistance = 0;
 	PathTracingData.CameraLensRadius = FVector2f::ZeroVector;
 	if (ShowFlags.DepthOfField &&
