Add LWC function usage statistics to material translator

diff --git a/Engine/Source/Editor/MaterialEditor/Private/MaterialEditor.cpp b/Engine/Source/Editor/MaterialEditor/Private/MaterialEditor.cpp
--- a/Engine/Source/Editor/MaterialEditor/Private/MaterialEditor.cpp
+++ b/Engine/Source/Editor/MaterialEditor/Private/MaterialEditor.cpp
@@ -2253,6 +2253,23 @@ void FMaterialEditor::DrawMaterialInfoStrings(
 			DrawPositionY += SpacingBetweenLines;
 		}
 
+		TStaticArray<uint16, (int)ELWCFunctionKind::Max> LWCFuncUsages = MaterialResource->GetEstimatedLWCFuncUsages();
+		for (int KindIndex = 0; KindIndex < (int)ELWCFunctionKind::Max; ++KindIndex)
+		{
+			int Usages = LWCFuncUsages[KindIndex];
+			if (LWCFuncUsages[KindIndex] > 0)
+			{
+				Canvas->DrawShadowedString(
+					5,
+					DrawPositionY,
+					*FString::Printf(TEXT("LWC %s usages (Est.): %u"), *UEnum::GetDisplayValueAsText((ELWCFunctionKind)KindIndex).ToString(), Usages),
+					FontToUse,
+					FLinearColor(1,1,0)
+				);
+				DrawPositionY += SpacingBetweenLines;
+			}
+		}
+
 		if (bGeneratedNewShaders)
 		{
 			int32 NumShaders = 0;
@@ -3029,6 +3046,21 @@ void FMaterialEditor::UpdateMaterialinfoList_Old()
 					Messages.Add(Line);
 				}
 
+			TStaticArray<uint16, (int)ELWCFunctionKind::Max> LWCFuncUsages = MaterialResource->GetEstimatedLWCFuncUsages();
+			for (int KindIndex = 0; KindIndex < (int)ELWCFunctionKind::Max; ++KindIndex)
+			{
+				int Usages = LWCFuncUsages[KindIndex];
+				if (LWCFuncUsages[KindIndex] > 0)
+				{
+					FString Message = FString::Printf(TEXT("LWC %s usages (Est.): %u"), *UEnum::GetDisplayValueAsText((ELWCFunctionKind)KindIndex).ToString(), Usages);
+						
+					TempMaterialInfoList.Add(MakeShareable(new FMaterialInfo(Message, FLinearColor::Yellow)));
+					TSharedRef<FTokenizedMessage> Line = FTokenizedMessage::Create(EMessageSeverity::Info);
+					Line->AddToken(FTextToken::Create(FText::FromString(Message)));
+					Messages.Add(Line);
+				}
+			}
+
 				if (FMaterialShaderMap* ShaderMap = MaterialResource->GetGameThreadShaderMap())
 				{
 					// Add shader count
diff --git a/Engine/Source/Runtime/Engine/Private/Materials/HLSLMaterialDerivativeAutogen.cpp b/Engine/Source/Runtime/Engine/Private/Materials/HLSLMaterialDerivativeAutogen.cpp
--- a/Engine/Source/Runtime/Engine/Private/Materials/HLSLMaterialDerivativeAutogen.cpp
+++ b/Engine/Source/Runtime/Engine/Private/Materials/HLSLMaterialDerivativeAutogen.cpp
@@ -484,6 +484,7 @@ int32 FMaterialDerivativeAutogen::GenerateExpressionFunc1(FHLSLMaterialTranslato
 		case EFunc1::Abs:
 			if (bIsLWC)
 			{
+				Translator.AddLWCFuncUsage(ELWCFunctionKind::Other);
 				DstToken = TEXT("LWCAbs(") + SrcToken + TEXT(")");
 			}
 			else
@@ -503,6 +504,7 @@ int32 FMaterialDerivativeAutogen::GenerateExpressionFunc1(FHLSLMaterialTranslato
 		case EFunc1::Sin:
 			if (bIsLWC)
 			{
+				Translator.AddLWCFuncUsage(ELWCFunctionKind::Other);
 				DstToken = TEXT("LWCSin(") + SrcToken + TEXT(")");
 			}
 			else
@@ -513,6 +515,7 @@ int32 FMaterialDerivativeAutogen::GenerateExpressionFunc1(FHLSLMaterialTranslato
 		case EFunc1::Cos:
 			if (bIsLWC)
 			{
+				Translator.AddLWCFuncUsage(ELWCFunctionKind::Other);
 				DstToken = TEXT("LWCCos(") + SrcToken + TEXT(")");
 			}
 			else
@@ -523,6 +526,7 @@ int32 FMaterialDerivativeAutogen::GenerateExpressionFunc1(FHLSLMaterialTranslato
 		case EFunc1::Tan:
 			if (bIsLWC)
 			{
+				Translator.AddLWCFuncUsage(ELWCFunctionKind::Other);
 				DstToken = TEXT("LWCTan(") + SrcToken + TEXT(")");
 			}
 			else
@@ -533,6 +537,7 @@ int32 FMaterialDerivativeAutogen::GenerateExpressionFunc1(FHLSLMaterialTranslato
 		case EFunc1::Asin:
 			if (bIsLWC)
 			{
+				Translator.AddLWCFuncUsage(ELWCFunctionKind::Other);
 				DstToken = TEXT("LWCASin(") + SrcToken + TEXT(")");
 			}
 			else
@@ -543,6 +548,7 @@ int32 FMaterialDerivativeAutogen::GenerateExpressionFunc1(FHLSLMaterialTranslato
 		case EFunc1::AsinFast:
 			if (bIsLWC)
 			{
+				Translator.AddLWCFuncUsage(ELWCFunctionKind::Other);
 				DstToken = TEXT("LWCASin(") + SrcToken + TEXT(")");
 			}
 			else
@@ -553,6 +559,7 @@ int32 FMaterialDerivativeAutogen::GenerateExpressionFunc1(FHLSLMaterialTranslato
 		case EFunc1::Acos:
 			if (bIsLWC)
 			{
+				Translator.AddLWCFuncUsage(ELWCFunctionKind::Other);
 				DstToken = TEXT("LWCACos(") + SrcToken + TEXT(")");
 			}
 			else
@@ -563,6 +570,7 @@ int32 FMaterialDerivativeAutogen::GenerateExpressionFunc1(FHLSLMaterialTranslato
 		case EFunc1::AcosFast:
 			if (bIsLWC)
 			{
+				Translator.AddLWCFuncUsage(ELWCFunctionKind::Other);
 				DstToken = TEXT("LWCACos(") + SrcToken + TEXT(")");
 			}
 			else
@@ -573,6 +581,7 @@ int32 FMaterialDerivativeAutogen::GenerateExpressionFunc1(FHLSLMaterialTranslato
 		case EFunc1::Atan:
 			if (bIsLWC)
 			{
+				Translator.AddLWCFuncUsage(ELWCFunctionKind::Other);
 				DstToken = TEXT("LWCATan(") + SrcToken + TEXT(")");
 			}
 			else
@@ -583,6 +592,7 @@ int32 FMaterialDerivativeAutogen::GenerateExpressionFunc1(FHLSLMaterialTranslato
 		case EFunc1::AtanFast:
 			if (bIsLWC)
 			{
+				Translator.AddLWCFuncUsage(ELWCFunctionKind::Other);
 				DstToken = TEXT("LWCATan(") + SrcToken + TEXT(")");
 			}
 			else
@@ -593,6 +603,7 @@ int32 FMaterialDerivativeAutogen::GenerateExpressionFunc1(FHLSLMaterialTranslato
 		case EFunc1::Sqrt:
 			if (bIsLWC)
 			{
+				Translator.AddLWCFuncUsage(ELWCFunctionKind::Other);
 				DstToken = TEXT("LWCSqrt(") + SrcToken + TEXT(")");
 			}
 			else
@@ -603,6 +614,7 @@ int32 FMaterialDerivativeAutogen::GenerateExpressionFunc1(FHLSLMaterialTranslato
 		case EFunc1::Rcp:
 			if (bIsLWC)
 			{
+				Translator.AddLWCFuncUsage(ELWCFunctionKind::Other);
 				DstToken = TEXT("LWCRcp(") + SrcToken + TEXT(")");
 			}
 			else
@@ -613,6 +625,7 @@ int32 FMaterialDerivativeAutogen::GenerateExpressionFunc1(FHLSLMaterialTranslato
 		case EFunc1::Rsqrt:
 			if (bIsLWC)
 			{
+				Translator.AddLWCFuncUsage(ELWCFunctionKind::Other);
 				DstToken = TEXT("LWCRsqrt(") + SrcToken + TEXT(")");
 			}
 			else
@@ -623,6 +636,7 @@ int32 FMaterialDerivativeAutogen::GenerateExpressionFunc1(FHLSLMaterialTranslato
 		case EFunc1::Saturate:
 			if (bIsLWC)
 			{
+				Translator.AddLWCFuncUsage(ELWCFunctionKind::Other);
 				DstToken = TEXT("LWCSaturate(") + SrcToken + TEXT(")");
 			}
 			else
@@ -633,6 +647,7 @@ int32 FMaterialDerivativeAutogen::GenerateExpressionFunc1(FHLSLMaterialTranslato
 		case EFunc1::Frac:
 			if (bIsLWC)
 			{
+				Translator.AddLWCFuncUsage(ELWCFunctionKind::Other);
 				DstToken = TEXT("LWCFrac(") + SrcToken + TEXT(")");
 			}
 			else
@@ -643,6 +658,7 @@ int32 FMaterialDerivativeAutogen::GenerateExpressionFunc1(FHLSLMaterialTranslato
 		case EFunc1::Length:
 			if (bIsLWC)
 			{
+				Translator.AddLWCFuncUsage(ELWCFunctionKind::Other);
 				DstToken = TEXT("LWCLength(") + SrcToken + TEXT(")");
 			}
 			else
@@ -653,6 +669,7 @@ int32 FMaterialDerivativeAutogen::GenerateExpressionFunc1(FHLSLMaterialTranslato
 		case EFunc1::InvLength:
 			if (bIsLWC)
 			{
+				Translator.AddLWCFuncUsage(ELWCFunctionKind::Other);
 				DstToken = TEXT("LWCRcpLength(") + SrcToken + TEXT(")");
 			}
 			else
@@ -663,6 +680,7 @@ int32 FMaterialDerivativeAutogen::GenerateExpressionFunc1(FHLSLMaterialTranslato
 		case EFunc1::Normalize:
 			if (bIsLWC)
 			{
+				Translator.AddLWCFuncUsage(ELWCFunctionKind::Other);
 				DstToken = TEXT("LWCNormalize(") + SrcToken + TEXT(")");
 			}
 			else if (OperationType.IntermediateType == EDerivativeType::Float1)
@@ -874,6 +892,7 @@ int32 FMaterialDerivativeAutogen::GenerateExpressionFunc2(FHLSLMaterialTranslato
 		case EFunc2::Add:
 			if (bIsLWC)
 			{
+				Translator.AddLWCFuncUsage(ELWCFunctionKind::Add);
 				DstToken = TEXT("LWCAdd(") + LhsToken + TEXT(", ") + RhsToken + TEXT(")");
 			}
 			else
@@ -884,6 +903,7 @@ int32 FMaterialDerivativeAutogen::GenerateExpressionFunc2(FHLSLMaterialTranslato
 		case EFunc2::Sub:
 			if (bIsLWC)
 			{
+				Translator.AddLWCFuncUsage(ELWCFunctionKind::Subtract);
 				DstToken = TEXT("LWCSubtract(") + LhsToken + TEXT(", ") + RhsToken + TEXT(")");
 			}
 			else
@@ -894,6 +914,7 @@ int32 FMaterialDerivativeAutogen::GenerateExpressionFunc2(FHLSLMaterialTranslato
 		case EFunc2::Mul:
 			if (bIsLWC)
 			{
+				Translator.AddLWCFuncUsage(ELWCFunctionKind::MultiplyVectorVector);
 				DstToken = TEXT("LWCMultiply(") + LhsToken + TEXT(", ") + RhsToken + TEXT(")");
 			}
 			else
@@ -904,6 +925,7 @@ int32 FMaterialDerivativeAutogen::GenerateExpressionFunc2(FHLSLMaterialTranslato
 		case EFunc2::Div:
 			if (bIsLWC)
 			{
+				Translator.AddLWCFuncUsage(ELWCFunctionKind::Divide);
 				DstToken = TEXT("LWCDivide(") + LhsToken + TEXT(", ") + RhsToken + TEXT(")");
 			}
 			else
@@ -914,6 +936,7 @@ int32 FMaterialDerivativeAutogen::GenerateExpressionFunc2(FHLSLMaterialTranslato
 		case EFunc2::Fmod:
 			if (bIsLWC)
 			{
+				Translator.AddLWCFuncUsage(ELWCFunctionKind::Other);
 				DstToken = TEXT("LWCFmod(") + LhsToken + TEXT(",") + RhsToken + TEXT(")");
 			}
 			else
@@ -924,6 +947,7 @@ int32 FMaterialDerivativeAutogen::GenerateExpressionFunc2(FHLSLMaterialTranslato
 		case EFunc2::Min:
 			if (bIsLWC)
 			{
+				Translator.AddLWCFuncUsage(ELWCFunctionKind::Other);
 				DstToken = TEXT("LWCMin(") + LhsToken + TEXT(",") + RhsToken + TEXT(")");
 			}
 			else
@@ -934,6 +958,7 @@ int32 FMaterialDerivativeAutogen::GenerateExpressionFunc2(FHLSLMaterialTranslato
 		case EFunc2::Max:
 			if (bIsLWC)
 			{
+				Translator.AddLWCFuncUsage(ELWCFunctionKind::Other);
 				DstToken = TEXT("LWCMax(") + LhsToken + TEXT(",") + RhsToken + TEXT(")");
 			}
 			else
@@ -944,6 +969,7 @@ int32 FMaterialDerivativeAutogen::GenerateExpressionFunc2(FHLSLMaterialTranslato
 		case EFunc2::Dot:
 			if (bIsLWC)
 			{
+				Translator.AddLWCFuncUsage(ELWCFunctionKind::Other);
 				DstToken = TEXT("LWCDot(") + LhsToken + TEXT(", ") + RhsToken + TEXT(")");
 			}
 			else
@@ -1066,7 +1092,9 @@ int32 FMaterialDerivativeAutogen::GenerateLerpFunc(FHLSLMaterialTranslator& Tran
 	const uint32 NumResultComponents = GetNumComponents(ResultType);
 
 	const bool bAllZeroDeriv = (ADerivInfo.DerivativeStatus == EDerivativeStatus::Zero && BDerivInfo.DerivativeStatus == EDerivativeStatus::Zero && SDerivInfo.DerivativeStatus == EDerivativeStatus::Zero);
-	const TCHAR* FunctionName = (ResultType & MCT_LWCType) ? TEXT("LWCLerp") : TEXT("lerp");
+	const bool bIsLWC = ResultType & MCT_LWCType;
+	if (bIsLWC) { Translator.AddLWCFuncUsage(ELWCFunctionKind::Other); }
+	const TCHAR* FunctionName = bIsLWC ? TEXT("LWCLerp") : TEXT("lerp");
 	FString FiniteString = FString::Printf(TEXT("%s(%s,%s,%s)"), FunctionName, *Translator.CoerceParameter(A, ResultType), *Translator.CoerceParameter(B, ResultType), *Translator.CoerceParameter(S, AlphaType));
 	
 	if (!bAllZeroDeriv && IsDerivativeValid(ADerivInfo.DerivativeStatus) && IsDerivativeValid(BDerivInfo.DerivativeStatus) && IsDerivativeValid(SDerivInfo.DerivativeStatus))
@@ -1181,6 +1209,7 @@ int32 FMaterialDerivativeAutogen::GenerateIfFunc(FHLSLMaterialTranslator& Transl
 	FString CompareNotEqual;
 	if (IsLWCType(CompareType))
 	{
+		Translator.AddLWCFuncUsage(ELWCFunctionKind::Other);
 		CompareGreaterEqual = FString::Printf(TEXT("LWCGreaterEqual(%s, %s)"), *AFinite, *BFinite);
 	}
 	else
@@ -1193,6 +1222,7 @@ int32 FMaterialDerivativeAutogen::GenerateIfFunc(FHLSLMaterialTranslator& Transl
 		ThresholdFinite = *Translator.GetParameterCode(Threshold);
 		if (IsLWCType(CompareType))
 		{
+			Translator.AddLWCFuncUsage(ELWCFunctionKind::Other);
 			CompareNotEqual = FString::Printf(TEXT("(!LWCEqualsApprox(%s, %s, %s))"), *AFinite, *BFinite, *ThresholdFinite);
 		}
 		else
@@ -1202,6 +1232,7 @@ int32 FMaterialDerivativeAutogen::GenerateIfFunc(FHLSLMaterialTranslator& Transl
 
 		if (IsLWCType(ResultType))
 		{
+			Translator.AddLWCFuncUsage(ELWCFunctionKind::Other);
 			CodeFinite = FString::Printf(
 				TEXT("LWCSelect(%s, LWCSelect(%s, %s, %s), %s)"),
 				*CompareNotEqual, *CompareGreaterEqual,
@@ -1231,6 +1262,7 @@ int32 FMaterialDerivativeAutogen::GenerateIfFunc(FHLSLMaterialTranslator& Transl
 	{
 		if (IsLWCType(ResultType))
 		{
+			Translator.AddLWCFuncUsage(ELWCFunctionKind::Other);
 			CodeFinite = FString::Printf(
 				TEXT("LWCSelect(%s, %s, %s)"),
 				*CompareGreaterEqual,
@@ -1482,6 +1514,7 @@ FString FMaterialDerivativeAutogen::GenerateUsedFunctions(FHLSLMaterialTranslato
 					Ret += TEXT("\t") + BaseName + TEXT(" Ret;") LINE_TERMINATOR;
 					if (IsLWCType(DerivType))
 					{
+						Translator.AddLWCFuncUsage(ELWCFunctionKind::Add);
 						Ret += TEXT("\tRet.Value = LWCAdd(A.Value, B.Value);") LINE_TERMINATOR;
 					}
 					else
@@ -1500,6 +1533,7 @@ FString FMaterialDerivativeAutogen::GenerateUsedFunctions(FHLSLMaterialTranslato
 					Ret += TEXT("\t") + BaseName + TEXT(" Ret;") LINE_TERMINATOR;
 					if (IsLWCType(DerivType))
 					{
+						Translator.AddLWCFuncUsage(ELWCFunctionKind::Subtract);
 						Ret += TEXT("\tRet.Value = LWCSubtract(A.Value, B.Value);") LINE_TERMINATOR;
 					}
 					else
@@ -1518,6 +1552,7 @@ FString FMaterialDerivativeAutogen::GenerateUsedFunctions(FHLSLMaterialTranslato
 					Ret += TEXT("\t") + BaseName + TEXT(" Ret;") LINE_TERMINATOR;
 					if (IsLWCType(DerivType))
 					{
+						Translator.AddLWCFuncUsage(ELWCFunctionKind::MultiplyVectorVector);
 						Ret += TEXT("\tRet.Value = LWCMultiply(A.Value, B.Value);") LINE_TERMINATOR;
 						Ret += TEXT("\tRet.Ddx = A.Ddx * LWCToFloat(B.Value) + LWCToFloat(A.Value) * B.Ddx;") LINE_TERMINATOR;
 						Ret += TEXT("\tRet.Ddy = A.Ddy * LWCToFloat(B.Value) + LWCToFloat(A.Value) * B.Ddy;") LINE_TERMINATOR;
@@ -1538,7 +1573,11 @@ FString FMaterialDerivativeAutogen::GenerateUsedFunctions(FHLSLMaterialTranslato
 					Ret += TEXT("\t") + BaseName + TEXT(" Ret;") LINE_TERMINATOR;
 					if (IsLWCType(DerivType))
 					{
+						Translator.AddLWCFuncUsage(ELWCFunctionKind::Divide);
 						Ret += TEXT("\tRet.Value = LWCDivide(A.Value, B.Value);") LINE_TERMINATOR;
+						Translator.AddLWCFuncUsage(ELWCFunctionKind::Other);
+						Translator.AddLWCFuncUsage(ELWCFunctionKind::Demote, 2);
+						Translator.AddLWCFuncUsage(ELWCFunctionKind::MultiplyVectorVector, 3);
 						Ret += TEXT("\t") + FieldName + TEXT(" Denom = LWCRcp(LWCMultiply(B.Value, B.Value));") LINE_TERMINATOR;
 						Ret += TEXT("\t") + FieldName + TEXT(" dFdA =  LWCToFloat(LWCMultiply(B.Value, Denom));") LINE_TERMINATOR;
 						Ret += TEXT("\t") + FieldName + TEXT(" dFdB = -LWCToFloat(LWCMultiply(A.Value, Denom));") LINE_TERMINATOR;
@@ -1564,9 +1603,12 @@ FString FMaterialDerivativeAutogen::GenerateUsedFunctions(FHLSLMaterialTranslato
 						Ret += BaseName + TEXT(" DivDeriv") + Suffix + TEXT("(") + BaseName + TEXT(" A, ") + NonLWCBaseName + TEXT(" B)") LINE_TERMINATOR;
 						Ret += TEXT("{") LINE_TERMINATOR;
 						Ret += TEXT("\t") + BaseName + TEXT(" Ret;") LINE_TERMINATOR;
+						Translator.AddLWCFuncUsage(ELWCFunctionKind::Divide);
 						Ret += TEXT("\tRet.Value = LWCDivide(A.Value, B.Value);") LINE_TERMINATOR;
 						Ret += TEXT("\t") + FieldName + TEXT(" Denom = rcp(B.Value * B.Value);") LINE_TERMINATOR;
 						Ret += TEXT("\t") + FieldName + TEXT(" dFdA =  B.Value * Denom;") LINE_TERMINATOR;
+						Translator.AddLWCFuncUsage(ELWCFunctionKind::Demote);
+						Translator.AddLWCFuncUsage(ELWCFunctionKind::MultiplyVectorVector);
 						Ret += TEXT("\t") + FieldName + TEXT(" dFdB = -LWCToFloat(LWCMultiply(A.Value, Denom));") LINE_TERMINATOR;
 						Ret += TEXT("\tRet.Ddx = dFdA * A.Ddx + dFdB * B.Ddx;") LINE_TERMINATOR;
 						Ret += TEXT("\tRet.Ddy = dFdA * A.Ddy + dFdB * B.Ddy;") LINE_TERMINATOR;
@@ -1583,6 +1625,7 @@ FString FMaterialDerivativeAutogen::GenerateUsedFunctions(FHLSLMaterialTranslato
 					Ret += TEXT("\t") + NonLWCBaseName + TEXT(" Ret;") LINE_TERMINATOR;
 					if (IsLWCType(DerivType))
 					{
+						Translator.AddLWCFuncUsage(ELWCFunctionKind::Other);
 						Ret += TEXT("\tRet.Value = LWCFmod(A.Value, B.Value);") LINE_TERMINATOR;
 					}
 					else
@@ -1601,6 +1644,7 @@ FString FMaterialDerivativeAutogen::GenerateUsedFunctions(FHLSLMaterialTranslato
 					Ret += TEXT("\t") + BaseName + TEXT(" Ret;") LINE_TERMINATOR;
 					if (IsLWCType(DerivType))
 					{
+						Translator.AddLWCFuncUsage(ELWCFunctionKind::Other);
 						Ret += TEXT("\t") + BoolName + TEXT(" Cmp = LWCLess(A.Value, B.Value);") LINE_TERMINATOR;
 						Ret += TEXT("\tRet.Value = LWCSelect(Cmp, A.Value, B.Value);") LINE_TERMINATOR;
 					}
@@ -1621,6 +1665,7 @@ FString FMaterialDerivativeAutogen::GenerateUsedFunctions(FHLSLMaterialTranslato
 					Ret += TEXT("\t") + BaseName + TEXT(" Ret;") LINE_TERMINATOR;
 					if (IsLWCType(DerivType))
 					{
+						Translator.AddLWCFuncUsage(ELWCFunctionKind::Other);
 						Ret += TEXT("\t") + BoolName + TEXT(" Cmp = LWCGreater(A.Value, B.Value);") LINE_TERMINATOR;
 						Ret += TEXT("\tRet.Value = LWCSelect(Cmp, A.Value, B.Value);") LINE_TERMINATOR;
 					}
@@ -1764,6 +1809,7 @@ FString FMaterialDerivativeAutogen::GenerateUsedFunctions(FHLSLMaterialTranslato
 					Ret += TEXT("\t") + BaseName + TEXT(" Ret;") LINE_TERMINATOR;
 					if (IsLWCType(DerivType))
 					{
+						Translator.AddLWCFuncUsage(ELWCFunctionKind::Other);
 						Ret += TEXT("\tRet.Value = LWCAbs(A.Value);") LINE_TERMINATOR;
 						Ret += TEXT("\t") + NonLWCFieldName + TEXT(" dFdA = select(LWCGreaterEqual(A.Value, 0.0f), One, -One);") LINE_TERMINATOR;
 					}
@@ -1785,6 +1831,7 @@ FString FMaterialDerivativeAutogen::GenerateUsedFunctions(FHLSLMaterialTranslato
 					Ret += TEXT("\t") + NonLWCBaseName + TEXT(" Ret;") LINE_TERMINATOR;
 					if (IsLWCType(DerivType))
 					{
+						Translator.AddLWCFuncUsage(ELWCFunctionKind::Other);
 						Ret += TEXT("\tRet.Value = LWCSin(A.Value);") LINE_TERMINATOR;
 						Ret += TEXT("\t") + FieldName + TEXT(" dFdA = LWCCos(A.Value);");
 					}
@@ -1805,6 +1852,7 @@ FString FMaterialDerivativeAutogen::GenerateUsedFunctions(FHLSLMaterialTranslato
 					Ret += TEXT("\t") + NonLWCBaseName + TEXT(" Ret;") LINE_TERMINATOR;
 					if (IsLWCType(DerivType))
 					{
+						Translator.AddLWCFuncUsage(ELWCFunctionKind::Other);
 						Ret += TEXT("\tRet.Value = LWCCos(A.Value);") LINE_TERMINATOR;
 						Ret += TEXT("\t") + FieldName + TEXT(" dFdA = -LWCSin(A.Value);");
 					}
@@ -1825,6 +1873,7 @@ FString FMaterialDerivativeAutogen::GenerateUsedFunctions(FHLSLMaterialTranslato
 					Ret += TEXT("\t") + NonLWCBaseName + TEXT(" Ret;") LINE_TERMINATOR;
 					if (IsLWCType(DerivType))
 					{
+						Translator.AddLWCFuncUsage(ELWCFunctionKind::Other);
 						Ret += TEXT("\tRet.Value = LWCTan(A.Value);") LINE_TERMINATOR;
 						Ret += TEXT("\t") + FieldName + TEXT(" dFdA = rcp(Pow2(LWCCos(A.Value)));");
 					}
@@ -1917,6 +1966,7 @@ FString FMaterialDerivativeAutogen::GenerateUsedFunctions(FHLSLMaterialTranslato
 					Ret += TEXT("\t") + NonLWCBaseName + TEXT(" Ret;") LINE_TERMINATOR;
 					if (IsLWCType(DerivType))
 					{
+						Translator.AddLWCFuncUsage(ELWCFunctionKind::Other);
 						Ret += TEXT("\tRet.Value = LWCSqrt(A.Value);") LINE_TERMINATOR;
 						Ret += TEXT("\t") + NonLWCFieldName + TEXT(" dFdA = 0.5f * LWCRsqrt(LWCMax(A.Value, 0.00001f));") LINE_TERMINATOR;
 					}
@@ -1938,6 +1988,7 @@ FString FMaterialDerivativeAutogen::GenerateUsedFunctions(FHLSLMaterialTranslato
 					Ret += TEXT("\t") + NonLWCBaseName + TEXT(" Ret;") LINE_TERMINATOR;
 					if (IsLWCType(DerivType))
 					{
+						Translator.AddLWCFuncUsage(ELWCFunctionKind::Other);
 						Ret += TEXT("\tRet.Value = LWCRcp(A.Value);") LINE_TERMINATOR;
 					}
 					else
@@ -1957,6 +2008,7 @@ FString FMaterialDerivativeAutogen::GenerateUsedFunctions(FHLSLMaterialTranslato
 					Ret += TEXT("\t") + NonLWCBaseName + TEXT(" Ret;") LINE_TERMINATOR;
 					if (IsLWCType(DerivType))
 					{
+						Translator.AddLWCFuncUsage(ELWCFunctionKind::Other);
 						Ret += TEXT("\tRet.Value = LWCRsqrt(A.Value);") LINE_TERMINATOR;
 						Ret += TEXT("\t") + NonLWCFieldName + TEXT(" dFdA = -0.5f * LWCRsqrt(A.Value) * LWCRcp(A.Value);");
 					}
@@ -1979,6 +2031,7 @@ FString FMaterialDerivativeAutogen::GenerateUsedFunctions(FHLSLMaterialTranslato
 					Ret += TEXT("\t") + NonLWCBaseName + TEXT(" Ret;") LINE_TERMINATOR;
 					if (IsLWCType(DerivType))
 					{
+						Translator.AddLWCFuncUsage(ELWCFunctionKind::Other);
 						Ret += TEXT("\tRet.Value = LWCSaturate(A.Value);") LINE_TERMINATOR;
 						Ret += TEXT("\t") + BoolName + TEXT(" InRange = LWCEquals(Ret.Value, A.Value);") LINE_TERMINATOR;
 					}
@@ -1999,6 +2052,7 @@ FString FMaterialDerivativeAutogen::GenerateUsedFunctions(FHLSLMaterialTranslato
 					Ret += TEXT("\t") + NonLWCBaseName + TEXT(" Ret;") LINE_TERMINATOR;
 					if (IsLWCType(DerivType))
 					{
+						Translator.AddLWCFuncUsage(ELWCFunctionKind::Other);
 						Ret += TEXT("\tRet.Value = LWCFrac(A.Value);") LINE_TERMINATOR;
 					}
 					else
@@ -2052,6 +2106,7 @@ FString FMaterialDerivativeAutogen::GenerateUsedFunctions(FHLSLMaterialTranslato
 					Ret += TEXT("\t") + ScalarName + TEXT(" Ret;") LINE_TERMINATOR;
 					if (IsLWCType(DerivType))
 					{
+						Translator.AddLWCFuncUsage(ELWCFunctionKind::Other);
 						Ret += TEXT("\tRet.Value = LWCLength(A.Value);") LINE_TERMINATOR;
 					}
 					else
@@ -2071,6 +2126,7 @@ FString FMaterialDerivativeAutogen::GenerateUsedFunctions(FHLSLMaterialTranslato
 					Ret += TEXT("\tFloatDeriv Ret;") LINE_TERMINATOR;
 					if (IsLWCType(DerivType))
 					{
+						Translator.AddLWCFuncUsage(ELWCFunctionKind::Other);
 						Ret += TEXT("\tRet.Value = LWCRcpLength(A.Value);") LINE_TERMINATOR;
 					}
 					else
@@ -2127,7 +2183,9 @@ FString FMaterialDerivativeAutogen::GenerateUsedFunctions(FHLSLMaterialTranslato
 			Ret += TEXT("\t") + BaseName + TEXT(" Ret;") LINE_TERMINATOR;
 			if (IsLWCType(DerivType))
 			{
+				Translator.AddLWCFuncUsage(ELWCFunctionKind::Other);
 				Ret += TEXT("\tRet.Value = LWCLerp(A.Value, B.Value, S.Value);") LINE_TERMINATOR;
+				Translator.AddLWCFuncUsage(ELWCFunctionKind::Subtract, 2);
 				Ret += TEXT("\tRet.Ddx = lerp(A.Ddx, B.Ddx, S.Value) + S.Ddx * LWCToFloat(LWCSubtract(B.Value, A.Value));") LINE_TERMINATOR;
 				Ret += TEXT("\tRet.Ddy = lerp(A.Ddy, B.Ddy, S.Value) + S.Ddy * LWCToFloat(LWCSubtract(B.Value, A.Value));") LINE_TERMINATOR;
 			}
diff --git a/Engine/Source/Runtime/Engine/Private/Materials/HLSLMaterialTranslator.cpp b/Engine/Source/Runtime/Engine/Private/Materials/HLSLMaterialTranslator.cpp
--- a/Engine/Source/Runtime/Engine/Private/Materials/HLSLMaterialTranslator.cpp
+++ b/Engine/Source/Runtime/Engine/Private/Materials/HLSLMaterialTranslator.cpp
@@ -3343,6 +3343,17 @@ void FHLSLMaterialTranslator::AddEstimatedTextureSample(const uint32 Count)
 	}
 }
 
+void FHLSLMaterialTranslator::AddLWCFuncUsage(ELWCFunctionKind Kind, const uint32 Count)
+{
+	if (IsCurrentlyCompilingForPreviousFrame())
+	{
+		// Ignore non-actionable cases
+		return;
+	}
+
+	MaterialCompilationOutput.EstimatedLWCFuncUsages[(int)Kind] += Count;
+}
+
 /** Creates a unique symbol name and adds it to the symbol list. */
 FString FHLSLMaterialTranslator::CreateSymbolName(const TCHAR* SymbolNameHint)
 {
@@ -3792,6 +3803,7 @@ int32 FHLSLMaterialTranslator::AccessUniformExpression(int32 Index)
 		const uint32 UniformOffset = CodeChunk.UniformExpression->UniformOffset;
 		if (bIsLWC)
 		{
+			AddLWCFuncUsage(ELWCFunctionKind::Constructor);
 			if (NumComponents == 1)
 			{
 				FormattedCode.Append(TEXT("MakeLWCScalar("));
@@ -4611,11 +4623,13 @@ FString FHLSLMaterialTranslator::CastValue(const FString& Code, EMaterialValueTy
 			if (bIsLWC)
 			{
 				// float->LWC
+				AddLWCFuncUsage(ELWCFunctionKind::Promote);
 				return FString::Printf(TEXT("LWCPromote(%s)"), *CastValue(Code, SourceType, MakeNonLWCType(DestType), Flags));
 			}
 			else
 			{
 				//LWC->float
+				AddLWCFuncUsage(ELWCFunctionKind::Demote);
 				return CastValue(FString::Printf(TEXT("LWCToFloat(%s)"), *Code), MakeNonLWCType(SourceType), DestType, Flags);
 			}
 		}
@@ -4625,6 +4639,7 @@ FString FHLSLMaterialTranslator::CastValue(const FString& Code, EMaterialValueTy
 		bool bNeedClosingParen = false;
 		if (bIsLWC)
 		{
+			AddLWCFuncUsage(ELWCFunctionKind::Constructor);
 			Result = TEXT("MakeLWCVector(");
 			bNeedClosingParen = true;
 		}
@@ -4676,6 +4691,7 @@ FString FHLSLMaterialTranslator::CastValue(const FString& Code, EMaterialValueTy
 					if (!bReplicateScalar && ComponentIndex >= NumSourceComponents)
 					{
 						check(bAllowAppendZeroes);
+						AddLWCFuncUsage(ELWCFunctionKind::Promote);
 						Result += TEXT("LWCPromote(0.0f)");
 					}
 					else
@@ -5257,6 +5273,7 @@ int32 FHLSLMaterialTranslator::Floor(int32 X)
 		const EMaterialValueType ValueType = GetParameterType(X);
 		if (IsLWCType(ValueType))
 		{
+			AddLWCFuncUsage(ELWCFunctionKind::Other);
 			return AddCodeChunkZeroDeriv(ValueType, TEXT("LWCFloor(%s)"), *GetParameterCode(X));
 		}
 		else
@@ -5282,6 +5299,7 @@ int32 FHLSLMaterialTranslator::Ceil(int32 X)
 		const EMaterialValueType ValueType = GetParameterType(X);
 		if (IsLWCType(ValueType))
 		{
+			AddLWCFuncUsage(ELWCFunctionKind::Other);
 			return AddCodeChunkZeroDeriv(ValueType, TEXT("LWCCeil(%s)"), *GetParameterCode(X));
 		}
 		else
@@ -5307,6 +5325,7 @@ int32 FHLSLMaterialTranslator::Round(int32 X)
 		const EMaterialValueType ValueType = GetParameterType(X);
 		if (IsLWCType(ValueType))
 		{
+			AddLWCFuncUsage(ELWCFunctionKind::Other);
 			return AddCodeChunkZeroDeriv(ValueType, TEXT("LWCRound(%s)"), *GetParameterCode(X));
 		}
 		else
@@ -5332,6 +5351,7 @@ int32 FHLSLMaterialTranslator::Truncate(int32 X)
 		const EMaterialValueType ValueType = GetParameterType(X);
 		if (IsLWCType(ValueType))
 		{
+			AddLWCFuncUsage(ELWCFunctionKind::Other);
 			return AddCodeChunkZeroDeriv(ValueType, TEXT("LWCTrunc(%s)"), *GetParameterCode(X));
 		}
 		else
@@ -5357,6 +5377,7 @@ int32 FHLSLMaterialTranslator::Sign(int32 X)
 		const EMaterialValueType ValueType = GetParameterType(X);
 		if (IsLWCType(ValueType))
 		{
+			AddLWCFuncUsage(ELWCFunctionKind::Other);
 			return AddCodeChunkZeroDeriv(MakeNonLWCType(ValueType), TEXT("LWCSign(%s)"), *GetParameterCode(X));
 		}
 		else
@@ -5388,6 +5409,7 @@ int32 FHLSLMaterialTranslator::Frac(int32 X)
 			const EMaterialValueType Type = GetParameterType(X);
 			if (IsLWCType(Type))
 			{
+				AddLWCFuncUsage(ELWCFunctionKind::Other);
 				return AddCodeChunk(MakeNonLWCType(Type), TEXT("LWCFrac(%s)"), *GetParameterCode(X));
 			}
 			else
@@ -5636,6 +5658,7 @@ int32 FHLSLMaterialTranslator::ParticlePosition()
 		return NonVertexOrPixelShaderExpressionError();
 	}
 	bNeedsParticlePosition = true;
+	AddLWCFuncUsage(ELWCFunctionKind::Subtract);
 	const int32 Result = AddInlinedCodeChunkZeroDeriv(MCT_LWCVector3,TEXT("LWCSubtract(Parameters.Particle.TranslatedWorldPositionAndSize.xyz, ResolvedView.PreViewTranslation)"));
 	return CastToNonLWCIfDisabled(Result);
 }
@@ -5887,6 +5910,7 @@ int32 FHLSLMaterialTranslator::ActorWorldPosition()
 		// Decal VS doesn't have material code so FMaterialVertexParameters
 		// and primitve uniform buffer are guaranteed to exist if ActorPosition
 		// material node is used in VS
+		AddLWCFuncUsage(ELWCFunctionKind::MultiplyVectorMatrix, 2);
 		Result = AddInlinedCodeChunkZeroDeriv(
 			MCT_LWCVector3,
 			TEXT("LWCMultiply(LWCMultiply(GetActorWorldPosition(Parameters), GetWorldToInstance(Parameters)), Parameters.PrevFrameLocalToWorld)"));
@@ -6669,6 +6693,7 @@ int32 FHLSLMaterialTranslator::TextureSample(
 			default: checkNoEntry(); return TEXT("");
 			}
 		};
+		AddLWCFuncUsage(ELWCFunctionKind::Other, 1);
 		const uint32 NumComponents = GetNumComponents(UVsType);
 		switch (NumComponents)
 		{
@@ -6708,6 +6733,7 @@ int32 FHLSLMaterialTranslator::TextureSample(
 				MipScaleIndex = AddCodeChunkZeroDeriv(UVsType, TEXT("exp2(%s)"), *CoerceParameter(MipValue0Index, MCT_Float1));
 			}
 
+			AddLWCFuncUsage(ELWCFunctionKind::Other, 2);
 			MipValue0Index = AddCodeChunkZeroDeriv(UVsType, TEXT("LWCDdx(%s)"), *GetParameterCode(CoordinateIndex));
 			MipValue1Index = AddCodeChunkZeroDeriv(UVsType, TEXT("LWCDdy(%s)"), *GetParameterCode(CoordinateIndex));
 			if (MipScaleIndex != INDEX_NONE)
@@ -8299,6 +8325,7 @@ int32 FHLSLMaterialTranslator::Add(int32 A, int32 B)
 		{
 			if (ResultType & MCT_LWCType)
 			{
+				AddLWCFuncUsage(ELWCFunctionKind::Add);
 				return AddCodeChunk(ResultType, TEXT("LWCAdd(%s, %s)"), *GetParameterCode(A), *GetParameterCode(B));
 			}
 			else
@@ -8335,6 +8362,7 @@ int32 FHLSLMaterialTranslator::Sub(int32 A, int32 B)
 		{
 			if (ResultType & MCT_LWCType)
 			{
+				AddLWCFuncUsage(ELWCFunctionKind::Subtract);
 				ResultCode = AddCodeChunk(ResultType, TEXT("LWCSubtract(%s, %s)"), *GetParameterCode(A), *GetParameterCode(B));
 			}
 			else
@@ -8839,6 +8867,7 @@ int32 FHLSLMaterialTranslator::Step(int32 Y, int32 X)
 
 	if (IsLWCType(ResultType))
 	{
+		AddLWCFuncUsage(ELWCFunctionKind::Other);
 		return AddCodeChunkZeroDeriv(MakeNonLWCType(ResultType), TEXT("LWCStep(%s,%s)"), *CoerceParameter(Y, ResultType), *CoerceParameter(X, ResultType));
 	}
 	else
@@ -8959,6 +8988,7 @@ int32 FHLSLMaterialTranslator::SmoothStep(int32 X, int32 Y, int32 A)
 
 	if (IsLWCType(ResultType))
 	{
+		AddLWCFuncUsage(ELWCFunctionKind::Other);
 		return AddCodeChunk(MakeNonLWCType(ResultType), TEXT("LWCSmoothStep(%s,%s,%s)"), *CoerceParameter(X, ResultType), *CoerceParameter(Y, ResultType), *CoerceParameter(A, ResultType));
 	}
 	else
@@ -9504,6 +9534,8 @@ int32 FHLSLMaterialTranslator::AppendVector(int32 A,int32 B)
 		FString FiniteCode;
 		if (bIsLWC)
 		{
+			AddLWCFuncUsage(ELWCFunctionKind::Promote, 2);
+			AddLWCFuncUsage(ELWCFunctionKind::Constructor);
 			FiniteCode = FString::Printf(TEXT("MakeLWCVector(LWCPromote(%s),LWCPromote(%s))"), *GetParameterCode(A), *GetParameterCode(B));
 		}
 		else
@@ -9598,6 +9630,7 @@ int32 FHLSLMaterialTranslator::TransformBase(EMaterialCommonBasis SourceCoordBas
 	}
 
 	const EMaterialValueType SourceType = GetParameterType(A);
+	const bool bIsPositionTranform = AWComponent != 0;
 		
 	{ // validation
 		if (ShaderFrequency != SF_Pixel && ShaderFrequency != SF_Compute && ShaderFrequency != SF_Vertex)
@@ -9658,6 +9691,7 @@ int32 FHLSLMaterialTranslator::TransformBase(EMaterialCommonBasis SourceCoordBas
 		{
 			if (DestCoordBasis == MCB_World)
 			{
+				if (bIsPositionTranform) { AddLWCFuncUsage(ELWCFunctionKind::MultiplyVectorMatrix); }
 				CodeStr = TEXT("TransformLocal<TO><PREV>World(Parameters, <A>)");
 				CodeDerivStr = TEXT("TransformLocal<TO><PREV>World(Parameters, <A>)");
 			}
@@ -9670,6 +9704,7 @@ int32 FHLSLMaterialTranslator::TransformBase(EMaterialCommonBasis SourceCoordBas
 			{
 				if (AWComponent)
 				{
+					AddLWCFuncUsage(ELWCFunctionKind::Subtract);
 					CodeStr = TEXT("LWCSubtract(<A>, ResolvedView.<PREV>PreViewTranslation)");
 				}
 				else
@@ -9710,6 +9745,7 @@ int32 FHLSLMaterialTranslator::TransformBase(EMaterialCommonBasis SourceCoordBas
 				}
 
 				// TODO: inconsistent with TransformLocal<TO>World with instancing
+				AddLWCFuncUsage(ELWCFunctionKind::MultiplyVectorMatrix);
 				CodeStr = LWCMultiplyMatrix(TEXT("<A>"), TEXT("GetPrimitiveData(Parameters).<PREVIOUS>WorldToLocal"), AWComponent);
 				CodeDerivStr = LWCMultiplyMatrix(TEXT("<A>"), TEXT("GetPrimitiveData(Parameters).<PREVIOUS>WorldToLocal"), 0);
 			}
@@ -9717,6 +9753,7 @@ int32 FHLSLMaterialTranslator::TransformBase(EMaterialCommonBasis SourceCoordBas
 			{
 				if (AWComponent)
 				{
+					AddLWCFuncUsage(ELWCFunctionKind::Add);
 					CodeStr = TEXT("LWCToFloat(LWCAdd(<A>, ResolvedView.<PREV>PreViewTranslation))");
 				}
 				else
@@ -9727,12 +9764,14 @@ int32 FHLSLMaterialTranslator::TransformBase(EMaterialCommonBasis SourceCoordBas
 			}
 			else if (DestCoordBasis == MCB_MeshParticle)
 			{
+				AddLWCFuncUsage(ELWCFunctionKind::MultiplyVectorMatrix);
 				CodeStr = LWCMultiplyMatrix(TEXT("<A>"), TEXT("Parameters.Particle.WorldToParticle"), AWComponent);
 				CodeDerivStr = LWCMultiplyMatrix(TEXT("<A>"), TEXT("Parameters.Particle.WorldToParticle"), 0);
 				bUsesParticleWorldToLocal = true;
 			}
 			else if (DestCoordBasis == MCB_Instance)
 			{
+				AddLWCFuncUsage(ELWCFunctionKind::MultiplyVectorMatrix);
 				CodeStr = LWCMultiplyMatrix(TEXT("<A>"), TEXT("GetWorldToInstance(Parameters)"), AWComponent);
 				CodeDerivStr = LWCMultiplyMatrix(TEXT("<A>"), TEXT("GetWorldToInstance(Parameters)"), 0);
 				bUsesInstanceWorldToLocalPS |= ShaderFrequency == SF_Pixel;
@@ -9768,6 +9807,7 @@ int32 FHLSLMaterialTranslator::TransformBase(EMaterialCommonBasis SourceCoordBas
 		{
 			if (DestCoordBasis == MCB_World)
 			{
+				AddLWCFuncUsage(ELWCFunctionKind::MultiplyVectorMatrix);
 				CodeStr = LWCMultiplyMatrix(TEXT("<A>"), TEXT("Parameters.Particle.ParticleToWorld"), AWComponent);
 				CodeDerivStr = LWCMultiplyMatrix(TEXT("<A>"), TEXT("Parameters.Particle.ParticleToWorld"), 0);
 				bUsesParticleLocalToWorld = true;
@@ -9779,6 +9819,7 @@ int32 FHLSLMaterialTranslator::TransformBase(EMaterialCommonBasis SourceCoordBas
 		{
 			if (DestCoordBasis == MCB_World)
 			{
+				AddLWCFuncUsage(ELWCFunctionKind::MultiplyVectorMatrix);
 				CodeStr = LWCMultiplyMatrix(TEXT("<A>"), TEXT("Get<PREV>InstanceToWorld(Parameters)"), AWComponent);
 				CodeDerivStr = LWCMultiplyMatrix(TEXT("<A>"), TEXT("Get<PREV>InstanceToWorld(Parameters)"), 0);
 				bUsesInstanceLocalToWorldPS |= ShaderFrequency == SF_Pixel;
@@ -9804,7 +9845,7 @@ int32 FHLSLMaterialTranslator::TransformBase(EMaterialCommonBasis SourceCoordBas
 		return TransformBase(IntermediaryBasis, DestCoordBasis, IntermediaryA, AWComponent);
 	}
 		
-	if (AWComponent != 0)
+	if (bIsPositionTranform)
 	{
 		CodeStr.ReplaceInline(TEXT("<TO>"),TEXT("PositionTo"));
 	}
@@ -10004,6 +10045,7 @@ int32 FHLSLMaterialTranslator::GenericSwitch(const TCHAR* SwitchExpressionText,
 	
 	if (IsLWCType(ResultType))
 	{
+		AddLWCFuncUsage(ELWCFunctionKind::Other);
 		return AddCodeChunk(ResultType, TEXT("LWCSelect(%s, %s, %s)"), SwitchExpressionText, *IfTrueCode, *IfFalseCode);
 	}
 	else
@@ -10157,6 +10199,10 @@ int32 FHLSLMaterialTranslator::Derivative(int32 A, const TCHAR* Component)
 	const EMaterialValueType ResultType = MakeNonLWCType(ADerivInfo.Type);
 	const TCHAR* FunctionName = IsLWCType(ADerivInfo.Type) ? TEXT("LWCDd") : TEXT("dd");
     
+	if (IsLWCType(ADerivInfo.Type))
+	{
+		AddLWCFuncUsage(ELWCFunctionKind::Other);
+	}
     // Spirv expects DDX/Y to have 32 bit width, so need to ensure we use float and not half
     const EMaterialValueType ParameterType = GetParameterType(A);
     const TCHAR* CastType = TEXT("");
@@ -13078,6 +13124,7 @@ int32 FHLSLMaterialTranslator::CustomExpression( class UMaterialExpressionCustom
 			{
 				// LWC types get two values, first the value converted to float, then the raw LWC value
 				// This way legacy custom expressions can continue to operate on regular float values
+				AddLWCFuncUsage(ELWCFunctionKind::Demote);
 				CodeChunk += TEXT("LWCToFloat(");
 				CodeChunk += *ParamCode;
 				CodeChunk += TEXT("),");
@@ -13170,6 +13217,7 @@ int32 FHLSLMaterialTranslator::CustomOutput(class UMaterialExpressionCustomOutpu
 	if (IsLWCType(OutputType))
 	{
 		// Add a wrapper with no suffix to return a non-LWC type
+		AddLWCFuncUsage(ELWCFunctionKind::Demote);
 		ImplementationCodeFinite += FString::Printf(TEXT("%s %s(%s FMaterial%sParameters Parameters) { return LWCToFloat(%s_LWC(Parameters)); }\r\n"),
 			HLSLTypeString(MakeNonLWCType(OutputType)),
 			*FunctionNameBase,
@@ -13180,6 +13228,7 @@ int32 FHLSLMaterialTranslator::CustomOutput(class UMaterialExpressionCustomOutpu
 	else
 	{
 		// Add a wrapper with LWC suffix to return a LWC type
+		AddLWCFuncUsage(ELWCFunctionKind::Promote);
 		ImplementationCodeFinite += FString::Printf(TEXT("%s %s_LWC(%s FMaterial%sParameters Parameters) { return LWCPromote(%s(Parameters)); }\r\n"),
 			HLSLTypeString(MakeLWCType(OutputType)),
 			*FunctionNameBase,
diff --git a/Engine/Source/Runtime/Engine/Private/Materials/HLSLMaterialTranslator.h b/Engine/Source/Runtime/Engine/Private/Materials/HLSLMaterialTranslator.h
--- a/Engine/Source/Runtime/Engine/Private/Materials/HLSLMaterialTranslator.h
+++ b/Engine/Source/Runtime/Engine/Private/Materials/HLSLMaterialTranslator.h
@@ -626,6 +626,7 @@ protected:
 	int32 NonVertexOrPixelShaderExpressionError();
 
 	void AddEstimatedTextureSample(const uint32 Count = 1);
+	void AddLWCFuncUsage(ELWCFunctionKind Kind, const uint32 Count = 1);
 
 	/** Creates a unique symbol name and adds it to the symbol list. */
 	FString CreateSymbolName(const TCHAR* SymbolNameHint);
diff --git a/Engine/Source/Runtime/Engine/Private/Materials/MaterialShared.cpp b/Engine/Source/Runtime/Engine/Private/Materials/MaterialShared.cpp
--- a/Engine/Source/Runtime/Engine/Private/Materials/MaterialShared.cpp
+++ b/Engine/Source/Runtime/Engine/Private/Materials/MaterialShared.cpp
@@ -3919,6 +3919,16 @@ uint32 FMaterialResource::GetEstimatedNumVirtualTextureLookups() const
 	}
 	return 0;
 }
+
+TStaticArray<uint16, (int)ELWCFunctionKind::Max> FMaterialResource::GetEstimatedLWCFuncUsages() const
+{
+	if (const FMaterialShaderMap* ShaderMap = GetGameThreadShaderMap())
+	{
+		return ShaderMap->GetEstimatedLWCFuncUsages();
+	}
+	return {};
+}
+
 #endif // WITH_EDITOR
 
 uint32 FMaterialResource::GetNumVirtualTextureStacks() const
diff --git a/Engine/Source/Runtime/Engine/Public/MaterialShared.h b/Engine/Source/Runtime/Engine/Public/MaterialShared.h
--- a/Engine/Source/Runtime/Engine/Public/MaterialShared.h
+++ b/Engine/Source/Runtime/Engine/Public/MaterialShared.h
@@ -50,6 +50,8 @@
 
 #include <atomic>
 
+#include "MaterialShared.generated.h"
+
 struct FExpressionInput;
 struct FExtraShaderCompilerSettings;
 class FMaterial;
@@ -275,6 +277,22 @@ inline EMaterialValueType MakeLWCType(EMaterialValueType Type)
 	}
 }
 
+UENUM()
+enum class ELWCFunctionKind
+{
+	Constructor,
+	Promote,
+	Demote,
+	Add,
+	Subtract,
+	Divide,
+	MultiplyVectorVector,
+	MultiplyVectorMatrix,
+	MultiplyMatrixMatrix,
+	Other,
+	Max
+};
+
 /**
  * The common bases of material
  */
@@ -692,7 +710,11 @@ public:
 		bHasRuntimeVirtualTextureOutputNode(false),
 		bUsesAnisotropy(false),
 		bUsesDisplacement(false)
-	{}
+	{
+#if WITH_EDITOR
+		FMemory::Memzero(EstimatedLWCFuncUsages);
+#endif
+	}
 
 	ENGINE_API bool IsSceneTextureUsed(ESceneTextureId TexId) const { return (UsedSceneTextures & (1 << TexId)) != 0; }
 	ENGINE_API void SetIsSceneTextureUsed(ESceneTextureId TexId) { UsedSceneTextures |= (1 << TexId); }
@@ -747,6 +769,9 @@ public:
 	LAYOUT_FIELD_EDITORONLY(uint16, EstimatedNumTextureSamplesVS);
 	LAYOUT_FIELD_EDITORONLY(uint16, EstimatedNumTextureSamplesPS);
 
+	/** Estimate of the number of times each LWC operator occurs in the material shader code */
+	LAYOUT_ARRAY_EDITORONLY(uint16, EstimatedLWCFuncUsages, (int)ELWCFunctionKind::Max)
+	
 	/** Number of virtual texture lookups performed, excludes direct invocation in shaders (for example VT lightmaps) */
 	LAYOUT_FIELD_EDITORONLY(uint16, EstimatedNumVirtualTextureLookups);
 
@@ -1378,6 +1403,12 @@ public:
 	uint32 GetNumUsedCustomInterpolatorScalars() const { return GetContent()->MaterialCompilationOutput.NumUsedCustomInterpolatorScalars; }
 	void GetEstimatedNumTextureSamples(uint32& VSSamples, uint32& PSSamples) const { VSSamples = GetContent()->MaterialCompilationOutput.EstimatedNumTextureSamplesVS; PSSamples = GetContent()->MaterialCompilationOutput.EstimatedNumTextureSamplesPS; }
 	uint32 GetEstimatedNumVirtualTextureLookups() const { return GetContent()->MaterialCompilationOutput.EstimatedNumVirtualTextureLookups; }
+	TStaticArray<uint16, (int)ELWCFunctionKind::Max> GetEstimatedLWCFuncUsages() const
+	{
+		TStaticArray<uint16, (int)ELWCFunctionKind::Max> Result;
+		CopyAssignItems(Result.GetData(), GetContent()->MaterialCompilationOutput.EstimatedLWCFuncUsages, (int)ELWCFunctionKind::Max);
+		return Result;
+	}
 #endif
 	uint32 GetNumVirtualTextureStacks() const { return GetContent()->MaterialCompilationOutput.UniformExpressionSet.VTStacks.Num(); }
 	uint8 GetRuntimeVirtualTextureOutputAttributeMask() const { return GetContent()->MaterialCompilationOutput.RuntimeVirtualTextureOutputAttributeMask; }
@@ -2546,6 +2577,7 @@ public:
 	ENGINE_API void GetUserInterpolatorUsage(uint32& NumUsedUVScalars, uint32& NumUsedCustomInterpolatorScalars) const;
 	ENGINE_API void GetEstimatedNumTextureSamples(uint32& VSSamples, uint32& PSSamples) const;
 	ENGINE_API uint32 GetEstimatedNumVirtualTextureLookups() const;
+	ENGINE_API TStaticArray<uint16, (int)ELWCFunctionKind::Max> GetEstimatedLWCFuncUsages() const;
 #endif
 	ENGINE_API uint32 GetNumVirtualTextureStacks() const;
 
