Fix when multiple viewports are open, only the viewport in focus has valid velocity on skeletal meshes.
This is fixed by switching to use GFrameCounter which is incremented once per engine tick, so multi views of the same frame share the same value.

diff --git a/Engine/Shaders/Private/LocalVertexFactory.ush b/Engine/Shaders/Private/LocalVertexFactory.ush
--- a/Engine/Shaders/Private/LocalVertexFactory.ush
+++ b/Engine/Shaders/Private/LocalVertexFactory.ush
@@ -1203,7 +1203,7 @@ float4 VertexFactoryGetPreviousWorldPosition(FVertexFactoryInput Input, FVertexF
 	if (IsGPUSkinPassThrough())
 	{
 		// Duplicates logic in FGPUSkinVertexFactoryShaderParameters::GetElementShaderBindings().
-		if (ResolvedView.FrameNumber == LocalVFLooseParameters.FrameNumber && ResolvedView.WorldIsPaused == 0)
+		if (ResolvedView.FrameCounter == LocalVFLooseParameters.FrameNumber && ResolvedView.WorldIsPaused == 0)
 		{
 			uint Offset = Input.VertexId * 3;
 			float3 PreviousPos;
diff --git a/Engine/Source/Runtime/Engine/Private/GPUSkinVertexFactory.cpp b/Engine/Source/Runtime/Engine/Private/GPUSkinVertexFactory.cpp
--- a/Engine/Source/Runtime/Engine/Private/GPUSkinVertexFactory.cpp
+++ b/Engine/Source/Runtime/Engine/Private/GPUSkinVertexFactory.cpp
@@ -891,7 +891,7 @@ public:
 
 				// Bone data is updated whenever animation triggers a dynamic update, animation can skip frames hence the frequency is not necessary every frame.
 				// So check if bone data is updated this frame, if not then the previous frame data is stale and not suitable for motion blur.
-				bool bBoneDataUpdatedThisFrame = View->Family->FrameNumber == ShaderData.UpdatedFrameNumber;
+				bool bBoneDataUpdatedThisFrame = View->Family->FrameCounter == ShaderData.UpdatedFrameNumber;
 				// If world is paused, use current frame bone matrices, so velocity is canceled and skeletal mesh isn't blurred from motion.
 				bool bPrevious = !View->Family->bWorldIsPaused && bBoneDataUpdatedThisFrame;
 				FRHIShaderResourceView* PreviousData = ShaderData.GetBoneBufferForReading(bPrevious).VertexBufferSRV;
@@ -926,7 +926,7 @@ public:
 		if (!bIsMobile)
 		{
 			const auto* GPUSkinVertexFactory = (const FGPUBaseSkinVertexFactory*)VertexFactory;
-			bool bMorphUpdatedThisFrame = (View->Family->FrameNumber == GPUSkinVertexFactory->GetMorphVertexBufferUpdatedFrameNumber());
+			bool bMorphUpdatedThisFrame = (View->Family->FrameCounter == GPUSkinVertexFactory->GetMorphVertexBufferUpdatedFrameNumber());
 			bool bPrevious = !View->Family->bWorldIsPaused && bMorphUpdatedThisFrame;
 			const FMorphVertexBuffer* MorphVertexBuffer = GPUSkinVertexFactory->GetMorphVertexBuffer(bPrevious);
 			ShaderBindings.Add(PreviousMorphBufferParameter, MorphVertexBuffer ? MorphVertexBuffer->GetSRV() : GNullVertexBuffer.VertexBufferSRV.GetReference());
diff --git a/Engine/Source/Runtime/Engine/Private/SceneView.cpp b/Engine/Source/Runtime/Engine/Private/SceneView.cpp
--- a/Engine/Source/Runtime/Engine/Private/SceneView.cpp
+++ b/Engine/Source/Runtime/Engine/Private/SceneView.cpp
@@ -2749,8 +2749,10 @@ void FSceneView::SetupCommonViewUniformBufferParameters(
 #endif
 
 	ViewUniformShaderParameters.Random = FMath::Rand();
+	// FrameNumber corresponds to how many times FRendererModule::BeginRenderingViewFamilies has been called, so multi views of the same frame have incremental values.
 	ViewUniformShaderParameters.FrameNumber = Family->FrameNumber;
-
+	// FrameCounter is incremented once per engine tick, so multi views of the same frame have the same value.
+	ViewUniformShaderParameters.FrameCounter = Family->FrameCounter;
 	ViewUniformShaderParameters.WorldIsPaused = Family->bWorldIsPaused;
 	ViewUniformShaderParameters.CameraCut = bCameraCut ? 1 : 0;
 
diff --git a/Engine/Source/Runtime/Engine/Private/SkeletalRenderCPUSkin.cpp b/Engine/Source/Runtime/Engine/Private/SkeletalRenderCPUSkin.cpp
--- a/Engine/Source/Runtime/Engine/Private/SkeletalRenderCPUSkin.cpp
+++ b/Engine/Source/Runtime/Engine/Private/SkeletalRenderCPUSkin.cpp
@@ -143,14 +143,11 @@ void FSkeletalMeshObjectCPUSkin::Update(
 		// this data is only deleted when another update is sent
 		FDynamicSkelMeshObjectDataCPUSkin* NewDynamicData = new FDynamicSkelMeshObjectDataCPUSkin(InMeshComponent,SkeletalMeshRenderData,LODIndex,InActiveMorphTargets, InMorphTargetWeights);
 
-		// We prepare the next frame but still have the value from the last one
-		uint32 FrameNumberToPrepare = GFrameNumber + 1;
+		uint64 FrameNumberToPrepare = GFrameCounter;
 		uint32 RevisionNumber = 0;
 
 		if (InMeshComponent->SceneProxy)
 		{
-			// We allow caching of per-frame, per-scene data
-			FrameNumberToPrepare = InMeshComponent->SceneProxy->GetScene().GetFrameNumber() + 1;
 			RevisionNumber = InMeshComponent->GetBoneTransformRevisionNumber();
 		}
 
@@ -189,7 +186,7 @@ void FSkeletalMeshObjectCPUSkin::UpdateSkinWeightBuffer(USkinnedMeshComponent* I
 	}
 }
 
-void FSkeletalMeshObjectCPUSkin::UpdateDynamicData_RenderThread(FRHICommandListImmediate& RHICmdList, FDynamicSkelMeshObjectDataCPUSkin* InDynamicData, uint32 FrameNumberToPrepare, uint32 RevisionNumber)
+void FSkeletalMeshObjectCPUSkin::UpdateDynamicData_RenderThread(FRHICommandListImmediate& RHICmdList, FDynamicSkelMeshObjectDataCPUSkin* InDynamicData, uint64 FrameNumberToPrepare, uint32 RevisionNumber)
 {
 	// we should be done with the old data at this point
 	delete DynamicData;
diff --git a/Engine/Source/Runtime/Engine/Private/SkeletalRenderCPUSkin.h b/Engine/Source/Runtime/Engine/Private/SkeletalRenderCPUSkin.h
--- a/Engine/Source/Runtime/Engine/Private/SkeletalRenderCPUSkin.h
+++ b/Engine/Source/Runtime/Engine/Private/SkeletalRenderCPUSkin.h
@@ -99,7 +99,7 @@ public:
 	ENGINE_API virtual void InitResources(USkinnedMeshComponent* InMeshComponent) override;
 	ENGINE_API virtual void ReleaseResources() override;
 	ENGINE_API virtual void Update(int32 LODIndex,USkinnedMeshComponent* InMeshComponent,const FMorphTargetWeightMap& InActiveMorphTargets, const TArray<float>& MorphTargetsWeights, EPreviousBoneTransformUpdateMode PreviousBoneTransformUpdateMode, const FExternalMorphWeightData& InExternalMorphWeightData) override;
-	ENGINE_API void UpdateDynamicData_RenderThread(FRHICommandListImmediate& RHICmdList, FDynamicSkelMeshObjectDataCPUSkin* InDynamicData, uint32 FrameNumberToPrepare, uint32 RevisionNumber);
+	ENGINE_API void UpdateDynamicData_RenderThread(FRHICommandListImmediate& RHICmdList, FDynamicSkelMeshObjectDataCPUSkin* InDynamicData, uint64 FrameNumberToPrepare, uint32 RevisionNumber);
 	ENGINE_API virtual void EnableOverlayRendering(bool bEnabled, const TArray<int32>* InBonesOfInterest, const TArray<UMorphTarget*>* InMorphTargetOfInterest) override;
 	ENGINE_API virtual void CacheVertices(int32 LODIndex, bool bForce) const override;
 	virtual bool IsCPUSkinned() const override { return true; }
diff --git a/Engine/Source/Runtime/Engine/Private/SkeletalRenderGPUSkin.cpp b/Engine/Source/Runtime/Engine/Private/SkeletalRenderGPUSkin.cpp
--- a/Engine/Source/Runtime/Engine/Private/SkeletalRenderGPUSkin.cpp
+++ b/Engine/Source/Runtime/Engine/Private/SkeletalRenderGPUSkin.cpp
@@ -405,8 +405,7 @@ void FSkeletalMeshObjectGPUSkin::Update(
 		PreviousBoneTransformUpdateMode,
 		InExternalMorphWeightData);
 
-	// We prepare the next frame but still have the value from the last one
-	uint32 FrameNumberToPrepare = GFrameNumber + 1;
+	uint64 FrameNumberToPrepare = GFrameCounter;
 	uint32 RevisionNumber = 0;
 
 	FGPUSkinCache* GPUSkinCache = nullptr;
@@ -415,7 +414,6 @@ void FSkeletalMeshObjectGPUSkin::Update(
 	{
 		// We allow caching of per-frame, per-scene data
 		Scene = &(InMeshComponent->SceneProxy->GetScene());
-		FrameNumberToPrepare = Scene->GetFrameNumber() + 1;
 		GPUSkinCache = Scene->GetGPUSkinCache();
 		RevisionNumber = InMeshComponent->GetBoneTransformRevisionNumber();
 	}
@@ -484,7 +482,7 @@ FORCEINLINE bool IsDeferredSkeletalDynamicDataUpdateEnabled()
 	return CVarDeferSkeletalDynamicDataUpdateUntilGDME.GetValueOnRenderThread() > 0;
 }
 
-void FSkeletalMeshObjectGPUSkin::UpdateDynamicData_RenderThread(FGPUSkinCache* GPUSkinCache, FRHICommandListImmediate& RHICmdList, FDynamicSkelMeshObjectDataGPUSkin* InDynamicData, FSceneInterface* Scene, uint32 FrameNumberToPrepare, uint32 RevisionNumber)
+void FSkeletalMeshObjectGPUSkin::UpdateDynamicData_RenderThread(FGPUSkinCache* GPUSkinCache, FRHICommandListImmediate& RHICmdList, FDynamicSkelMeshObjectDataGPUSkin* InDynamicData, FSceneInterface* Scene, uint64 FrameNumberToPrepare, uint32 RevisionNumber)
 {
 	TRACE_CPUPROFILER_EVENT_SCOPE(GPUSkin::UpdateDynamicData_RT);
 
diff --git a/Engine/Source/Runtime/Engine/Private/SkeletalRenderGPUSkin.h b/Engine/Source/Runtime/Engine/Private/SkeletalRenderGPUSkin.h
--- a/Engine/Source/Runtime/Engine/Private/SkeletalRenderGPUSkin.h
+++ b/Engine/Source/Runtime/Engine/Private/SkeletalRenderGPUSkin.h
@@ -352,7 +352,7 @@ public:
 	virtual void InitResources(USkinnedMeshComponent* InMeshComponent) override;
 	virtual void ReleaseResources() override;
 	virtual void Update(int32 LODIndex,USkinnedMeshComponent* InMeshComponent,const FMorphTargetWeightMap& InActiveMorphTargets, const TArray<float>& InMorphTargetWeights, EPreviousBoneTransformUpdateMode PreviousBoneTransformUpdateMode, const FExternalMorphWeightData& InExternalMorphWeightData) override;
-	void UpdateDynamicData_RenderThread(FGPUSkinCache* GPUSkinCache, FRHICommandListImmediate& RHICmdList, FDynamicSkelMeshObjectDataGPUSkin* InDynamicData, FSceneInterface* Scene, uint32 FrameNumberToPrepare, uint32 RevisionNumber);
+	void UpdateDynamicData_RenderThread(FGPUSkinCache* GPUSkinCache, FRHICommandListImmediate& RHICmdList, FDynamicSkelMeshObjectDataGPUSkin* InDynamicData, FSceneInterface* Scene, uint64 FrameNumberToPrepare, uint32 RevisionNumber);
 	virtual void PreGDMECallback(FGPUSkinCache* GPUSkinCache, uint32 FrameNumber) override;
 	virtual const FVertexFactory* GetSkinVertexFactory(const FSceneView* View, int32 LODIndex,int32 ChunkIdx, ESkinVertexFactoryMode VFMode = ESkinVertexFactoryMode::Default) const override;
 	virtual const FSkinBatchVertexFactoryUserData* GetVertexFactoryUserData(const int32 LODIndex, int32 ChunkIdx, ESkinVertexFactoryMode VFMode) const override;
diff --git a/Engine/Source/Runtime/Engine/Public/GPUSkinVertexFactory.h b/Engine/Source/Runtime/Engine/Public/GPUSkinVertexFactory.h
--- a/Engine/Source/Runtime/Engine/Public/GPUSkinVertexFactory.h
+++ b/Engine/Source/Runtime/Engine/Public/GPUSkinVertexFactory.h
@@ -286,7 +286,7 @@ public:
 		int32 InputWeightIndexSize = 0;
 		FShaderResourceViewRHIRef InputWeightStream;
 		// Frame number of the bone data that is last updated
-		uint32 UpdatedFrameNumber = 0;
+		uint64 UpdatedFrameNumber = 0;
 
 	private:
 		// double buffered bone positions+orientations to support normal rendering and velocity (new-old position) rendering
diff --git a/Engine/Source/Runtime/Engine/Public/SceneView.h b/Engine/Source/Runtime/Engine/Public/SceneView.h
--- a/Engine/Source/Runtime/Engine/Public/SceneView.h
+++ b/Engine/Source/Runtime/Engine/Public/SceneView.h
@@ -758,6 +758,7 @@ enum ETranslucencyVolumeCascade
 	VIEW_UNIFORM_BUFFER_MEMBER(float, MaterialTextureDerivativeMultiply) \
 	VIEW_UNIFORM_BUFFER_MEMBER(uint32, Random) \
 	VIEW_UNIFORM_BUFFER_MEMBER(uint32, FrameNumber) \
+	VIEW_UNIFORM_BUFFER_MEMBER(uint32, FrameCounter) \
 	VIEW_UNIFORM_BUFFER_MEMBER(uint32, StateFrameIndexMod8) \
 	VIEW_UNIFORM_BUFFER_MEMBER(uint32, StateFrameIndex) \
 	VIEW_UNIFORM_BUFFER_MEMBER(uint32, DebugViewModeMask) \
@@ -1964,6 +1965,9 @@ public:
 	/** Copy from main thread GFrameNumber to be accessible on render thread side. UINT_MAX before CreateSceneRenderer() or BeginRenderingViewFamily() was called */
 	uint32 FrameNumber;
 
+	/** Copy from main thread GFrameCounter to be accessible on render thread side. GFrameCounter is incremented once per engine tick, so multi views of the same frame have the same value. */
+	uint64 FrameCounter = 0;
+
 	/** Indicates this view family is an additional one. */
 	bool bAdditionalViewFamily;
 
diff --git a/Engine/Source/Runtime/Renderer/Private/SceneRendering.cpp b/Engine/Source/Runtime/Renderer/Private/SceneRendering.cpp
--- a/Engine/Source/Runtime/Renderer/Private/SceneRendering.cpp
+++ b/Engine/Source/Runtime/Renderer/Private/SceneRendering.cpp
@@ -4713,6 +4713,8 @@ void FRendererModule::BeginRenderingViewFamilies(FCanvas* Canvas, TArrayView<FSc
 
 	for (FSceneViewFamily* ViewFamily : ViewFamilies)
 	{
+		ViewFamily->FrameCounter = GFrameCounter;
+
 #if !(UE_BUILD_SHIPPING || UE_BUILD_TEST)
 		{
 			extern TSharedRef<ISceneViewExtension, ESPMode::ThreadSafe> GetRendererViewExtension();
