Fixed calls to FMemory::Memswap with compile-time sizes.

diff --git a/Engine/Plugins/Experimental/Water/Source/Editor/Private/WaterLandscapeBrush.cpp b/Engine/Plugins/Experimental/Water/Source/Editor/Private/WaterLandscapeBrush.cpp
--- a/Engine/Plugins/Experimental/Water/Source/Editor/Private/WaterLandscapeBrush.cpp
+++ b/Engine/Plugins/Experimental/Water/Source/Editor/Private/WaterLandscapeBrush.cpp
@@ -134,8 +134,8 @@ void AWaterLandscapeBrush::UpdateActors(bool bInTriggerEvents)
 	ClearActors();
 
 	// Backup Cache
-	TMap<TWeakObjectPtr<AActor>, UObject*> PreviousCache;
-	FMemory::Memswap(&Cache, &PreviousCache, sizeof(Cache));
+	TMap<TWeakObjectPtr<AActor>, TObjectPtr<UObject>> PreviousCache;
+	Swap(Cache, PreviousCache);
 
 	if (UWorld* World = GetWorld())
 	{
@@ -144,7 +144,7 @@ void AWaterLandscapeBrush::UpdateActors(bool bInTriggerEvents)
 			AActor* Actor = *It;
 			if (IWaterBrushActorInterface* WaterBrushActor = Cast<IWaterBrushActorInterface>(Actor))
 			{
-				UObject* const* FoundCache = PreviousCache.Find(TWeakObjectPtr<AActor>(Actor));
+				const TObjectPtr<UObject>* FoundCache = PreviousCache.Find(TWeakObjectPtr<AActor>(Actor));
 				const bool bTriggerEvent = false;
 				const bool bModify = false;
 				AddActorInternal(Actor, World, FoundCache != nullptr ? *FoundCache : nullptr, bTriggerEvent, bModify);
diff --git a/Engine/Plugins/FX/Niagara/Source/Niagara/Private/NiagaraGPUInstanceCountManager.cpp b/Engine/Plugins/FX/Niagara/Source/Niagara/Private/NiagaraGPUInstanceCountManager.cpp
--- a/Engine/Plugins/FX/Niagara/Source/Niagara/Private/NiagaraGPUInstanceCountManager.cpp
+++ b/Engine/Plugins/FX/Niagara/Source/Niagara/Private/NiagaraGPUInstanceCountManager.cpp
@@ -298,7 +298,7 @@ void FNiagaraGPUInstanceCountManager::ResizeBuffers(FRHICommandListImmediate& RH
 
 			// Swap the buffers
 			AllocatedInstanceCounts = RecommendedInstanceCounts;
-			FMemory::Memswap(&NextCountBuffer, &CountBuffer, sizeof(NextCountBuffer));
+			Swap(NextCountBuffer, CountBuffer);
 			//UE_LOG(LogNiagara, Log, TEXT("FNiagaraGPUInstanceCountManager::ResizeBuffers Resize AllocatedInstanceCounts: %d ReservedInstanceCounts: %d"), AllocatedInstanceCounts, ReservedInstanceCounts);
 		}
 		// If we need to shrink the buffer size because use way to much buffer size.
diff --git a/Engine/Source/Editor/HierarchicalLODOutliner/Private/TreeItemID.h b/Engine/Source/Editor/HierarchicalLODOutliner/Private/TreeItemID.h
--- a/Engine/Source/Editor/HierarchicalLODOutliner/Private/TreeItemID.h
+++ b/Engine/Source/Editor/HierarchicalLODOutliner/Private/TreeItemID.h
@@ -38,7 +38,7 @@ namespace HLODOutliner
 
 		FTreeItemID& operator=(FTreeItemID&& Other)
 		{
-			FMemory::Memswap(this, &Other, sizeof(FTreeItemID));
+			Swap(*this, Other);
 			return *this;
 		}
 
diff --git a/Engine/Source/Editor/SceneOutliner/Public/SceneOutlinerStandaloneTypes.h b/Engine/Source/Editor/SceneOutliner/Public/SceneOutlinerStandaloneTypes.h
--- a/Engine/Source/Editor/SceneOutliner/Public/SceneOutlinerStandaloneTypes.h
+++ b/Engine/Source/Editor/SceneOutliner/Public/SceneOutlinerStandaloneTypes.h
@@ -106,7 +106,7 @@ public:
 	}
 	FSceneOutlinerTreeItemID& operator=(FSceneOutlinerTreeItemID&& Other)
 	{
-		FMemory::Memswap(this, &Other, sizeof(FSceneOutlinerTreeItemID));
+		Swap(*this, Other);
 		return *this;
 	}
 
diff --git a/Engine/Source/Editor/WorldBrowser/Private/WorldBrowserPrivateTypes.h b/Engine/Source/Editor/WorldBrowser/Private/WorldBrowserPrivateTypes.h
--- a/Engine/Source/Editor/WorldBrowser/Private/WorldBrowserPrivateTypes.h
+++ b/Engine/Source/Editor/WorldBrowser/Private/WorldBrowserPrivateTypes.h
@@ -81,7 +81,7 @@ namespace WorldHierarchy
 		}
 		FWorldTreeItemID& operator=(FWorldTreeItemID&& Other)
 		{
-			FMemory::Memswap(this, &Other, sizeof(FWorldTreeItemID));
+			Swap(*this, Other);
 			return *this;
 		}
 
diff --git a/Engine/Source/Runtime/Core/Public/Containers/Array.h b/Engine/Source/Runtime/Core/Public/Containers/Array.h
--- a/Engine/Source/Runtime/Core/Public/Containers/Array.h
+++ b/Engine/Source/Runtime/Core/Public/Containers/Array.h
@@ -2721,10 +2721,9 @@ public:
 	 */
 	FORCEINLINE void SwapMemory(SizeType FirstIndexToSwap, SizeType SecondIndexToSwap)
 	{
-		FMemory::Memswap(
-			(uint8*)AllocatorInstance.GetAllocation() + (sizeof(ElementType)*FirstIndexToSwap),
-			(uint8*)AllocatorInstance.GetAllocation() + (sizeof(ElementType)*SecondIndexToSwap),
-			sizeof(ElementType)
+		::Swap(
+			*(ElementType*)((uint8*)AllocatorInstance.GetAllocation() + (sizeof(ElementType)*FirstIndexToSwap)),
+			*(ElementType*)((uint8*)AllocatorInstance.GetAllocation() + (sizeof(ElementType)*SecondIndexToSwap))
 		);
 	}
 
diff --git a/Engine/Source/Runtime/Core/Public/Templates/SharedPointer.h b/Engine/Source/Runtime/Core/Public/Templates/SharedPointer.h
--- a/Engine/Source/Runtime/Core/Public/Templates/SharedPointer.h
+++ b/Engine/Source/Runtime/Core/Public/Templates/SharedPointer.h
@@ -392,7 +392,7 @@ public:
 
 	FORCEINLINE TSharedRef& operator=( TSharedRef&& InSharedRef )
 	{
-		FMemory::Memswap(this, &InSharedRef, sizeof(TSharedRef));
+		Swap(*this, InSharedRef);
 		return *this;
 	}
 
diff --git a/Engine/Source/Runtime/Engine/Private/Components/MeshComponent.cpp b/Engine/Source/Runtime/Engine/Private/Components/MeshComponent.cpp
--- a/Engine/Source/Runtime/Engine/Private/Components/MeshComponent.cpp
+++ b/Engine/Source/Runtime/Engine/Private/Components/MeshComponent.cpp
@@ -543,7 +543,7 @@ void UMeshComponent::LogMaterialsAndTextures(FOutputDevice& Ar, int32 Indent) co
 
 	// Backup the material overrides so we can access the mesh original materials.
 	TArray<class UMaterialInterface*> OverrideMaterialsBackup;
-	FMemory::Memswap(&OverrideMaterialsBackup, &const_cast<UMeshComponent*>(this)->OverrideMaterials, sizeof(OverrideMaterialsBackup));
+	Swap(OverrideMaterialsBackup, const_cast<UMeshComponent*>(this)->OverrideMaterials);
 
 	TArray<UMaterialInterface*> MaterialInterfaces = GetMaterials();
 	for (int32 MaterialIndex = 0; MaterialIndex < MaterialInterfaces.Num(); ++MaterialIndex)
@@ -561,7 +561,7 @@ void UMeshComponent::LogMaterialsAndTextures(FOutputDevice& Ar, int32 Indent) co
 	}
 
 	// Restore the overrides.
-	FMemory::Memswap(&OverrideMaterialsBackup, &const_cast<UMeshComponent*>(this)->OverrideMaterials, sizeof(OverrideMaterialsBackup));
+	Swap(OverrideMaterialsBackup, const_cast<UMeshComponent*>(this)->OverrideMaterials);
 }
 
 #endif
diff --git a/Engine/Source/Runtime/Engine/Private/InstancedStaticMesh.cpp b/Engine/Source/Runtime/Engine/Private/InstancedStaticMesh.cpp
--- a/Engine/Source/Runtime/Engine/Private/InstancedStaticMesh.cpp
+++ b/Engine/Source/Runtime/Engine/Private/InstancedStaticMesh.cpp
@@ -533,7 +533,7 @@ void FStaticMeshInstanceBuffer::InitFromPreallocatedData(FStaticMeshInstanceData
 	QUICK_SCOPE_CYCLE_COUNTER(STAT_FStaticMeshInstanceBuffer_InitFromPreallocatedData);
 
 	InstanceData = MakeShared<FStaticMeshInstanceData, ESPMode::ThreadSafe>();
-	FMemory::Memswap(&Other, InstanceData.Get(), sizeof(FStaticMeshInstanceData));
+	Swap(Other, *InstanceData.Get());
 	InstanceData->SetAllowCPUAccess(RequireCPUAccess);
 }
 
@@ -543,7 +543,7 @@ void FStaticMeshInstanceBuffer::UpdateFromCommandBuffer_Concurrent(FInstanceUpda
 	
 	FStaticMeshInstanceBuffer* InstanceBuffer = this; 
 	FInstanceUpdateCmdBuffer* NewCmdBuffer = new FInstanceUpdateCmdBuffer();
-	FMemory::Memswap(&CmdBuffer, NewCmdBuffer, sizeof(FInstanceUpdateCmdBuffer));
+	Swap(CmdBuffer, *NewCmdBuffer);
 	
 	// leave NumEdits unchanged in commandbuffer
 	CmdBuffer.NumEdits = NewCmdBuffer->NumEdits; 
@@ -1295,7 +1295,7 @@ void FPerInstanceRenderData::UpdateFromPreallocatedData(FStaticMeshInstanceData&
 	InOther.SetAllowCPUAccess(InstanceBuffer.RequireCPUAccess);
 
 	InstanceBuffer_GameThread = MakeShared<FStaticMeshInstanceData, ESPMode::ThreadSafe>();
-	FMemory::Memswap(&InOther, InstanceBuffer_GameThread.Get(), sizeof(FStaticMeshInstanceData));
+	Swap(InOther, *InstanceBuffer_GameThread.Get());
 
 	typedef TSharedPtr<FStaticMeshInstanceData, ESPMode::ThreadSafe> FStaticMeshInstanceDataPtr;
 
diff --git a/Engine/Source/Runtime/Engine/Private/MapBuildData.cpp b/Engine/Source/Runtime/Engine/Private/MapBuildData.cpp
--- a/Engine/Source/Runtime/Engine/Private/MapBuildData.cpp
+++ b/Engine/Source/Runtime/Engine/Private/MapBuildData.cpp
@@ -898,8 +898,8 @@ void UMapBuildDataRegistry::InvalidateSurfaceLightmaps(UWorld* World, bool bRecr
 		{
 			TMap<FGuid, FMeshMapBuildData> PrevMeshData;
 			TMap<FGuid, FLightComponentMapBuildData> PrevLightData;
-			FMemory::Memswap(&MeshBuildData, &PrevMeshData, sizeof(MeshBuildData));
-			FMemory::Memswap(&LightBuildData, &PrevLightData, sizeof(LightBuildData));
+			Swap(MeshBuildData, PrevMeshData);
+			Swap(LightBuildData, PrevLightData);
 
 			for (const FGuid& Guid : *ResourcesToKeep)
 			{
@@ -945,7 +945,7 @@ void UMapBuildDataRegistry::InvalidateReflectionCaptures(const TSet<FGuid>* Reso
 		// Warning: caller is responsible for unregistering any components potentially referencing this UMapBuildDataRegistry before we change its contents!
 
 		TMap<FGuid, FReflectionCaptureMapBuildData> PrevReflectionCapturedData;
-		FMemory::Memswap(&ReflectionCaptureBuildData , &PrevReflectionCapturedData, sizeof(ReflectionCaptureBuildData));
+		Swap(ReflectionCaptureBuildData, PrevReflectionCapturedData);
 
 		for (TMap<FGuid, FReflectionCaptureMapBuildData>::TIterator It(PrevReflectionCapturedData); It; ++It)
 		{
@@ -1110,8 +1110,8 @@ void UMapBuildDataRegistry::EmptyLevelData(const TSet<FGuid>* ResourcesToKeep)
 {
 	TMap<FGuid, FPrecomputedLightVolumeData*> PrevPrecomputedLightVolumeData;
 	TMap<FGuid, FPrecomputedVolumetricLightmapData*> PrevPrecomputedVolumetricLightmapData;
-	FMemory::Memswap(&LevelPrecomputedLightVolumeBuildData , &PrevPrecomputedLightVolumeData, sizeof(LevelPrecomputedLightVolumeBuildData));
-	FMemory::Memswap(&LevelPrecomputedVolumetricLightmapBuildData , &PrevPrecomputedVolumetricLightmapData, sizeof(LevelPrecomputedVolumetricLightmapBuildData));
+	Swap(LevelPrecomputedLightVolumeBuildData, PrevPrecomputedLightVolumeData);
+	Swap(LevelPrecomputedVolumetricLightmapBuildData, PrevPrecomputedVolumetricLightmapData);
 
 	for (TMap<FGuid, FPrecomputedLightVolumeData*>::TIterator It(PrevPrecomputedLightVolumeData); It; ++It)
 	{
diff --git a/Engine/Source/Runtime/Engine/Private/SkeletalMeshRenderData.cpp b/Engine/Source/Runtime/Engine/Private/SkeletalMeshRenderData.cpp
--- a/Engine/Source/Runtime/Engine/Private/SkeletalMeshRenderData.cpp
+++ b/Engine/Source/Runtime/Engine/Private/SkeletalMeshRenderData.cpp
@@ -456,12 +456,12 @@ void FSkeletalMeshRenderData::SyncUVChannelData(const TArray<FSkeletalMaterial>&
 	{
 		ENQUEUE_RENDER_COMMAND(SyncUVChannelData)([this, UpdateData = MoveTemp(UpdateData)](FRHICommandListImmediate& RHICmdList)
 		{
-			FMemory::Memswap(&UVChannelDataPerMaterial, UpdateData.Get(), sizeof(UVChannelDataPerMaterial));
+			Swap(UVChannelDataPerMaterial, *UpdateData.Get());
 		});
 	}
 	else
 	{
-		FMemory::Memswap(&UVChannelDataPerMaterial, UpdateData.Get(), sizeof(UVChannelDataPerMaterial));
+		Swap(UVChannelDataPerMaterial, *UpdateData.Get());
 	}
 }
 
diff --git a/Engine/Source/Runtime/Engine/Private/StaticMesh.cpp b/Engine/Source/Runtime/Engine/Private/StaticMesh.cpp
--- a/Engine/Source/Runtime/Engine/Private/StaticMesh.cpp
+++ b/Engine/Source/Runtime/Engine/Private/StaticMesh.cpp
@@ -2132,12 +2132,12 @@ void FStaticMeshRenderData::SyncUVChannelData(const TArray<FStaticMaterial>& Obj
 	{
 		ENQUEUE_RENDER_COMMAND(SyncUVChannelData)([this, UpdateData = MoveTemp(UpdateData)](FRHICommandListImmediate& RHICmdList)
 		{
-			FMemory::Memswap(&UVChannelDataPerMaterial, UpdateData.Get(), sizeof(UVChannelDataPerMaterial));
+			Swap(UVChannelDataPerMaterial, *UpdateData.Get());
 		});
 	}
 	else
 	{
-		FMemory::Memswap(&UVChannelDataPerMaterial, UpdateData.Get(), sizeof(UVChannelDataPerMaterial));
+		Swap(UVChannelDataPerMaterial, *UpdateData.Get());
 	}
 }
 
diff --git a/Engine/Source/Runtime/Engine/Private/Streaming/AsyncTextureStreaming.cpp b/Engine/Source/Runtime/Engine/Private/Streaming/AsyncTextureStreaming.cpp
--- a/Engine/Source/Runtime/Engine/Private/Streaming/AsyncTextureStreaming.cpp
+++ b/Engine/Source/Runtime/Engine/Private/Streaming/AsyncTextureStreaming.cpp
@@ -329,8 +329,8 @@ void FRenderAssetStreamingMipCalcTask::ApplyPakStateChanges_Async()
 	// Acquire the pending file state changes from the streaming manager.
 	{
 		FScopeLock Lock(&StreamingManager.MountedStateDirtyFilesCS);
-		FMemory::Memswap(&MountedStateDirtyFiles, &StreamingManager.MountedStateDirtyFiles, sizeof(FIoFilenameHashSet));
-		FMemory::Memswap(&bRecacheAllFiles, &StreamingManager.bRecacheAllFiles, sizeof(bool));
+		Swap(MountedStateDirtyFiles, StreamingManager.MountedStateDirtyFiles);
+		Swap(bRecacheAllFiles, StreamingManager.bRecacheAllFiles);
 	}
 
 	if (bRecacheAllFiles || MountedStateDirtyFiles.Num())
diff --git a/Engine/Source/Runtime/Engine/Private/Streaming/TextureInstanceView.cpp b/Engine/Source/Runtime/Engine/Private/Streaming/TextureInstanceView.cpp
--- a/Engine/Source/Runtime/Engine/Private/Streaming/TextureInstanceView.cpp
+++ b/Engine/Source/Runtime/Engine/Private/Streaming/TextureInstanceView.cpp
@@ -257,10 +257,10 @@ void FRenderAssetInstanceView::SwapData(FRenderAssetInstanceView* Lfs, FRenderAs
 	check(Lfs->CompiledRenderAssetMap.Num() == 0 && Rhs->CompiledRenderAssetMap.Num() == 0);
 	check(!Lfs->CompiledNumForcedLODCompMap.Num() && !Rhs->CompiledNumForcedLODCompMap.Num());
 
-	FMemory::Memswap(&Lfs->Bounds4 , &Rhs->Bounds4, sizeof(Lfs->Bounds4));
-	FMemory::Memswap(&Lfs->Elements , &Rhs->Elements, sizeof(Lfs->Elements));
-	FMemory::Memswap(&Lfs->RenderAssetMap, &Rhs->RenderAssetMap, sizeof(Lfs->RenderAssetMap));
-	FMemory::Memswap(&Lfs->MaxTexelFactor , &Rhs->MaxTexelFactor, sizeof(Lfs->MaxTexelFactor));
+	Swap(Lfs->Bounds4 , Rhs->Bounds4);
+	Swap(Lfs->Elements , Rhs->Elements);
+	Swap(Lfs->RenderAssetMap, Rhs->RenderAssetMap);
+	Swap(Lfs->MaxTexelFactor , Rhs->MaxTexelFactor);
 }
 
 void FRenderAssetInstanceView::OnVerifyElementIdxFailed(int32 Idx, bool bInRange, int32 IterationCount, TMap<const UPrimitiveComponent*, int32>* ComponentMapPtr, TArray<int32>* FreeIndicesPtr) const
diff --git a/Engine/Source/Runtime/Renderer/Private/MeshDrawCommands.cpp b/Engine/Source/Runtime/Renderer/Private/MeshDrawCommands.cpp
--- a/Engine/Source/Runtime/Renderer/Private/MeshDrawCommands.cpp
+++ b/Engine/Source/Runtime/Renderer/Private/MeshDrawCommands.cpp
@@ -554,7 +554,7 @@ static void BuildMeshDrawCommandPrimitiveIdBuffer(
 		NewPassVisibleMeshDrawCommandsNum = TempVisibleMeshDrawCommands.Num();
 
 		// Replace VisibleMeshDrawCommands
-		FMemory::Memswap(&VisibleMeshDrawCommands, &TempVisibleMeshDrawCommands, sizeof(TempVisibleMeshDrawCommands));
+		Swap(VisibleMeshDrawCommands, TempVisibleMeshDrawCommands);
 		TempVisibleMeshDrawCommands.Reset();
 	}
 	else
@@ -823,7 +823,7 @@ void ApplyViewOverridesToMeshDrawCommands(
 			}
 
 			// Replace VisibleMeshDrawCommands
-			FMemory::Memswap(&VisibleMeshDrawCommands, &TempVisibleMeshDrawCommands, sizeof(TempVisibleMeshDrawCommands));
+			Swap(VisibleMeshDrawCommands, TempVisibleMeshDrawCommands);
 			TempVisibleMeshDrawCommands.Reset();
 		}
 	}
@@ -1258,12 +1258,12 @@ void FParallelMeshDrawCommandPass::DispatchPassSetup(
 		case EMeshPass::TranslucencyAll: TaskContext.TranslucencyPass				= ETranslucencyPass::TPT_AllTranslucency; break;
 	}
 
-	FMemory::Memswap(&TaskContext.MeshDrawCommands, &InOutMeshDrawCommands, sizeof(InOutMeshDrawCommands));
-	FMemory::Memswap(&TaskContext.DynamicMeshCommandBuildRequests, &InOutDynamicMeshCommandBuildRequests, sizeof(InOutDynamicMeshCommandBuildRequests));
+	Swap(TaskContext.MeshDrawCommands, InOutMeshDrawCommands);
+	Swap(TaskContext.DynamicMeshCommandBuildRequests, InOutDynamicMeshCommandBuildRequests);
 
 	if (TaskContext.ShadingPath == EShadingPath::Mobile && TaskContext.PassType == EMeshPass::BasePass)
 	{
-		FMemory::Memswap(&TaskContext.MobileBasePassCSMMeshDrawCommands, InOutMobileBasePassCSMMeshDrawCommands, sizeof(*InOutMobileBasePassCSMMeshDrawCommands));
+		Swap(TaskContext.MobileBasePassCSMMeshDrawCommands, *InOutMobileBasePassCSMMeshDrawCommands);
 	}
 	else
 	{
diff --git a/Engine/Source/Runtime/Renderer/Private/MeshPassProcessor.cpp b/Engine/Source/Runtime/Renderer/Private/MeshPassProcessor.cpp
--- a/Engine/Source/Runtime/Renderer/Private/MeshPassProcessor.cpp
+++ b/Engine/Source/Runtime/Renderer/Private/MeshPassProcessor.cpp
@@ -1758,7 +1758,7 @@ void ApplyViewOverridesToMeshDrawCommands(const FSceneView& View, FMeshCommandOn
 		}
 
 		// Replace VisibleMeshDrawCommands
-		FMemory::Memswap(&VisibleMeshDrawCommands, &ViewOverriddenMeshCommands, sizeof(ViewOverriddenMeshCommands));
+		Swap(VisibleMeshDrawCommands, ViewOverriddenMeshCommands);
 	}
 }
 
diff --git a/Engine/Source/Runtime/Renderer/Private/SceneOcclusion.cpp b/Engine/Source/Runtime/Renderer/Private/SceneOcclusion.cpp
--- a/Engine/Source/Runtime/Renderer/Private/SceneOcclusion.cpp
+++ b/Engine/Source/Runtime/Renderer/Private/SceneOcclusion.cpp
@@ -416,7 +416,7 @@ void FFrameBasedOcclusionQueryPool::AdvanceFrame(uint32 InOcclusionFrameCounter,
 			if (Frame.OcclusionFrameCounter > NewFrame.OcclusionFrameCounter)
 			{
 				Frame.Queries.Append(MoveTemp(NewFrame.Queries));
-				FMemory::Memswap(&Frame, &NewFrame, sizeof(FFrameOcclusionQueries));
+				Swap(Frame, NewFrame);
 			}
 			else
 			{
@@ -424,7 +424,7 @@ void FFrameBasedOcclusionQueryPool::AdvanceFrame(uint32 InOcclusionFrameCounter,
 			}
 		}
 
-		FMemory::Memswap(FrameQueries, TmpFrameQueries, sizeof(FrameQueries));
+		Swap(FrameQueries, TmpFrameQueries);
 		NumBufferedFrames = InNumBufferedFrames;
 	}
 	
diff --git a/Engine/Source/Runtime/SoundFieldRendering/Public/SphericalHarmonicCalculator.h b/Engine/Source/Runtime/SoundFieldRendering/Public/SphericalHarmonicCalculator.h
--- a/Engine/Source/Runtime/SoundFieldRendering/Public/SphericalHarmonicCalculator.h
+++ b/Engine/Source/Runtime/SoundFieldRendering/Public/SphericalHarmonicCalculator.h
@@ -37,6 +37,6 @@ public:
 	{
 		InOutVector.X = -(InOutVector.X - HALF_PI);
 		InOutVector.Y *= -1.0f;
-		FMemory::Memswap(&InOutVector.X,&InOutVector.Y, sizeof(InOutVector.X));
+		Swap(InOutVector.X,InOutVector.Y);
 	}
 };
diff --git a/Engine/Source/Runtime/VulkanRHI/Private/VulkanMemory.cpp b/Engine/Source/Runtime/VulkanRHI/Private/VulkanMemory.cpp
--- a/Engine/Source/Runtime/VulkanRHI/Private/VulkanMemory.cpp
+++ b/Engine/Source/Runtime/VulkanRHI/Private/VulkanMemory.cpp
@@ -3799,7 +3799,7 @@ namespace VulkanRHI
 	}
 	void FVulkanAllocation::Swap(FVulkanAllocation& Other)
 	{
-		FMemory::Memswap(this, &Other, sizeof(*this));
+		::Swap(*this, Other);
 	}
 
 	void FVulkanAllocation::Reference(const FVulkanAllocation& Other)
