Fix non-unity build errors

diff --git a/Engine/Source/Runtime/Engine/Public/MeshBatch.h b/Engine/Source/Runtime/Engine/Public/MeshBatch.h
--- a/Engine/Source/Runtime/Engine/Public/MeshBatch.h
+++ b/Engine/Source/Runtime/Engine/Public/MeshBatch.h
@@ -36,6 +36,7 @@ struct FInstanceDynamicData;
 struct FInstanceSceneData;
 struct FMaterialShaderParameters;
 struct FRenderBounds;
+template<typename TBufferStruct> class TUniformBuffer;
 
 enum EPrimitiveIdMode
 {
