Fix GenerateObjectsClusters API to pass in object array by const reference.

diff --git a/Engine/Source/Runtime/Engine/Private/ReferenceCluster.cpp b/Engine/Source/Runtime/Engine/Private/ReferenceCluster.cpp
--- a/Engine/Source/Runtime/Engine/Private/ReferenceCluster.cpp
+++ b/Engine/Source/Runtime/Engine/Private/ReferenceCluster.cpp
@@ -3,7 +3,7 @@
 #include "ReferenceCluster.h"
 #include "Misc/Guid.h"
 
-TArray<TArray<FGuid>> GenerateObjectsClusters(TArray<TPair<FGuid, TArray<FGuid>>> InObjects)
+TArray<TArray<FGuid>> GenerateObjectsClusters(const TArray<TPair<FGuid, TArray<FGuid>>>& InObjects)
 {
 	TMap<FGuid, FGuid> ObjectToClusters;
 	TMap<FGuid, TSet<FGuid>> Clusters;
diff --git a/Engine/Source/Runtime/Engine/Public/ReferenceCluster.h b/Engine/Source/Runtime/Engine/Public/ReferenceCluster.h
--- a/Engine/Source/Runtime/Engine/Public/ReferenceCluster.h
+++ b/Engine/Source/Runtime/Engine/Public/ReferenceCluster.h
@@ -14,4 +14,4 @@
  *	@return	Array of reference clusters (e.g. group of objects referencing each other). For example, an
  *			input of n objects without any references will generate n clusters of a single object each.
  */
-ENGINE_API TArray<TArray<FGuid>> GenerateObjectsClusters(TArray<TPair<FGuid, TArray<FGuid>>> InObjects);
\ No newline at end of file
+ENGINE_API TArray<TArray<FGuid>> GenerateObjectsClusters(const TArray<TPair<FGuid, TArray<FGuid>>>& InObjects);
\ No newline at end of file
