Make abstract interface IPhysicsReplication, so that we can introduce new replication managers without reliance on the implementation of FPhysicsReplication

diff --git a/Engine/Source/Runtime/Engine/Private/ActorReplication.cpp b/Engine/Source/Runtime/Engine/Private/ActorReplication.cpp
--- a/Engine/Source/Runtime/Engine/Private/ActorReplication.cpp
+++ b/Engine/Source/Runtime/Engine/Private/ActorReplication.cpp
@@ -303,7 +303,7 @@ void AActor::SyncReplicatedPhysicsSimulation()
 				{
 					if (FPhysScene* PhysScene = World->GetPhysicsScene())
 					{
-						if (FPhysicsReplication* PhysicsReplication = PhysScene->GetPhysicsReplication())
+						if (IPhysicsReplication* PhysicsReplication = PhysScene->GetPhysicsReplication())
 						{
 							PhysicsReplication->RemoveReplicatedTarget(RootPrimComp);
 						}
diff --git a/Engine/Source/Runtime/Engine/Private/PhysicsEngine/AsyncPhysicsInputComponent.cpp b/Engine/Source/Runtime/Engine/Private/PhysicsEngine/AsyncPhysicsInputComponent.cpp
--- a/Engine/Source/Runtime/Engine/Private/PhysicsEngine/AsyncPhysicsInputComponent.cpp
+++ b/Engine/Source/Runtime/Engine/Private/PhysicsEngine/AsyncPhysicsInputComponent.cpp
@@ -123,7 +123,7 @@ struct FAsyncPhysicsInputRewindCallback : public Chaos::IRewindCallback
 		{
 			if (FPhysScene_Chaos* Scene = static_cast<FPhysScene_Chaos*>(World->GetPhysicsScene()))
 			{
-				if(const FPhysicsReplication* PhysicsReplication  = Scene->GetPhysicsReplication())
+				if (const IPhysicsReplication* PhysicsReplication  = Scene->GetPhysicsReplication())
 				{
 					CachedClientFrame = Scene->GetPhysicsReplication()->GetResimFrame();
 				}
diff --git a/Engine/Source/Runtime/Engine/Private/PhysicsEngine/Experimental/PhysScene_Chaos.cpp b/Engine/Source/Runtime/Engine/Private/PhysicsEngine/Experimental/PhysScene_Chaos.cpp
--- a/Engine/Source/Runtime/Engine/Private/PhysicsEngine/Experimental/PhysScene_Chaos.cpp
+++ b/Engine/Source/Runtime/Engine/Private/PhysicsEngine/Experimental/PhysScene_Chaos.cpp
@@ -516,7 +516,10 @@ FPhysScene_Chaos::FPhysScene_Chaos(AActor* InSolverActor
 #endif
 
 	// Create replication manager
-	PhysicsReplication = PhysicsReplicationFactory.IsValid() ? PhysicsReplicationFactory->Create(this) : new FPhysicsReplication(this);
+	PhysicsReplication
+		= PhysicsReplicationFactory.IsValid()
+		? PhysicsReplicationFactory->CreatePhysicsReplication(this)
+		: MakeUnique<FPhysicsReplication>(this);
 
 	FPhysicsDelegates::OnPhysSceneInit.Broadcast(this);
 
@@ -532,18 +535,12 @@ FPhysScene_Chaos::~FPhysScene_Chaos()
 
 	// Must ensure deferred components do not hold onto scene pointer.
 	ProcessDeferredCreatePhysicsState();
+	
+	// Make sure physics replication is cleared before we're fully destructed
+	PhysicsReplication.Reset();
 
 	FPhysicsDelegates::OnPhysSceneTerm.Broadcast(this);
 
-	if (IPhysicsReplicationFactory* RawReplicationFactory = FPhysScene_Chaos::PhysicsReplicationFactory.Get())
-	{
-		RawReplicationFactory->Destroy(PhysicsReplication);
-	}
-	else if(PhysicsReplication)
-	{
-		delete PhysicsReplication;
-	}
-
 #if CHAOS_WITH_PAUSABLE_SOLVER
 	if (SyncCaller)
 	{
@@ -771,14 +768,14 @@ void FPhysScene_Chaos::RemoveObject(Chaos::FClusterUnionPhysicsProxy* InObject)
 	RemoveFromComponentMaps(InObject);
 }
 
-FPhysicsReplication* FPhysScene_Chaos::GetPhysicsReplication()
+IPhysicsReplication* FPhysScene_Chaos::GetPhysicsReplication()
 {
-	return PhysicsReplication;
+	return PhysicsReplication.Get();
 }
 
-void FPhysScene_Chaos::SetPhysicsReplication(FPhysicsReplication* InPhysicsReplication)
+void FPhysScene_Chaos::SetPhysicsReplication(IPhysicsReplication* InPhysicsReplication)
 {
-	PhysicsReplication = InPhysicsReplication;
+	PhysicsReplication = TUniquePtr<IPhysicsReplication>(InPhysicsReplication);
 }
 
 void FPhysScene_Chaos::AddReferencedObjects(FReferenceCollector& Collector)
@@ -2107,9 +2104,6 @@ void FPhysScene_Chaos::OnSyncBodies(Chaos::FPhysicsSolverBase* Solver)
 				{
 					const FVector MoveBy = NewTransform.GetLocation() - ParentComponent->GetComponentTransform().GetLocation();
 					PendingTransforms.Add(FPhysScenePendingComponentTransform_Chaos(ParentComponent, MoveBy, NewTransform.GetRotation(), DirtyParticle->GetWakeEvent()));
-
-					// Any changes in the linear/angular velocity of the particle is meaningless if the X/R doesn't change too.
-					ParentComponent->SyncVelocitiesFromPhysics(DirtyParticle->V(), DirtyParticle->W());
 				}
 
 				Interface->AddToSpatialAcceleration({ &Handle, 1 }, Outer->GetSpacialAcceleration());
diff --git a/Engine/Source/Runtime/Engine/Private/PhysicsEngine/PhysicsInterfaceUtils.cpp b/Engine/Source/Runtime/Engine/Private/PhysicsEngine/PhysicsInterfaceUtils.cpp
--- a/Engine/Source/Runtime/Engine/Private/PhysicsEngine/PhysicsInterfaceUtils.cpp
+++ b/Engine/Source/Runtime/Engine/Private/PhysicsEngine/PhysicsInterfaceUtils.cpp
@@ -103,3 +103,16 @@ FCollisionFilterData CreateQueryFilterData(const uint8 MyChannel, const bool bTr
 		return CreateTraceQueryFilterData(MyChannel, bTraceComplex, InCollisionResponseContainer, QueryParam);
 	}
 }
+
+//
+// NOTE: Once the old Create/Destroy methods are deprecated, remove the default implementations
+//
+
+#include "PhysicsReplication.h"
+TUniquePtr<IPhysicsReplication> IPhysicsReplicationFactory::CreatePhysicsReplication(FPhysScene* OwningPhysScene)
+{
+PRAGMA_DISABLE_DEPRECATION_WARNINGS
+	UE_LOG(LogPhysics, Warning, TEXT("\nIPhysicsReplicationFactory::Create has been deprecated in favor of IPhysicsReplicationFactory::CreatePhysicsReplication. Use the new CreatePhysicsReplication method instead moving forward, as the Create function will be removed in future versions. Please update your code to the new API before upgrading to the next release, otherwise your project will no longer compile."));
+	return TUniquePtr<IPhysicsReplication>(Create(OwningPhysScene));
+PRAGMA_ENABLE_DEPRECATION_WARNINGS
+}
diff --git a/Engine/Source/Runtime/Engine/Private/PrimitiveComponentPhysics.cpp b/Engine/Source/Runtime/Engine/Private/PrimitiveComponentPhysics.cpp
--- a/Engine/Source/Runtime/Engine/Private/PrimitiveComponentPhysics.cpp
+++ b/Engine/Source/Runtime/Engine/Private/PrimitiveComponentPhysics.cpp
@@ -44,7 +44,7 @@ void UPrimitiveComponent::SetRigidBodyReplicatedTarget(FRigidBodyState& UpdatedS
 	{
 		if (FPhysScene* PhysScene = World->GetPhysicsScene())
 		{
-			if (FPhysicsReplication* PhysicsReplication = PhysScene->GetPhysicsReplication())
+			if (IPhysicsReplication* PhysicsReplication = PhysScene->GetPhysicsReplication())
 			{
 				// If we are not allowed to replicate physics objects,
 				// don't set replicated target unless we have a BodyInstance.
diff --git a/Engine/Source/Runtime/Engine/Public/Physics/Experimental/PhysScene_Chaos.h b/Engine/Source/Runtime/Engine/Public/Physics/Experimental/PhysScene_Chaos.h
--- a/Engine/Source/Runtime/Engine/Public/Physics/Experimental/PhysScene_Chaos.h
+++ b/Engine/Source/Runtime/Engine/Public/Physics/Experimental/PhysScene_Chaos.h
@@ -26,7 +26,7 @@
 class UPrimitiveComponent;
 
 class AdvanceOneTimeStepTask;
-class FPhysicsReplication;
+class IPhysicsReplication;
 class FPhysInterface_Chaos;
 class FChaosSolversModule;
 struct FForceFieldProxy;
@@ -153,8 +153,10 @@ public:
 	void RemoveObject(FGeometryCollectionPhysicsProxy* InObject);
 	void RemoveObject(Chaos::FClusterUnionPhysicsProxy* InObject);
 
-	FPhysicsReplication* GetPhysicsReplication();
-	void SetPhysicsReplication(FPhysicsReplication* InPhysicsReplication);
+	IPhysicsReplication* GetPhysicsReplication();
+
+	UE_DEPRECATED(5.3, "Can no longer direclty set physics replication at runtime. For now, specify a PhysicsReplication factory instead. This function will take ownership of the IPhysicsReplication's lifetime.")
+	void SetPhysicsReplication(IPhysicsReplication* InPhysicsReplication);
 
 	virtual void AddReferencedObjects(FReferenceCollector& Collector) override;
 	/** Given a solver object, returns its associated component. */
@@ -294,7 +296,7 @@ private:
 	void HandleCrumblingEvents(const Chaos::FCrumblingEventData& Event);
 
 	/** Replication manager that updates physics bodies towards replicated physics state */
-	FPhysicsReplication* PhysicsReplication;
+	TUniquePtr<IPhysicsReplication> PhysicsReplication;
 
 #if CHAOS_WITH_PAUSABLE_SOLVER
 	/** Callback that checks the status of the world settings for this scene before pausing/unpausing its solver. */
diff --git a/Engine/Source/Runtime/Engine/Public/Physics/PhysicsInterfaceUtils.h b/Engine/Source/Runtime/Engine/Public/Physics/PhysicsInterfaceUtils.h
--- a/Engine/Source/Runtime/Engine/Public/Physics/PhysicsInterfaceUtils.h
+++ b/Engine/Source/Runtime/Engine/Public/Physics/PhysicsInterfaceUtils.h
@@ -6,6 +6,8 @@
 #include "Containers/Union.h"
 #include "Physics/PhysicsInterfaceTypes.h"
 #include "PhysicsInterfaceUtilsCore.h"
+#include "PhysicsReplicationInterface.h"
+#include "Misc/CoreMiscDefines.h"
 
 class FPhysScene_PhysX;
 struct FConstraintInstance;
@@ -31,9 +33,15 @@ struct FConstraintBrokenDelegateData
 class FPhysicsReplication;
 
 /** Interface for the creation of customized physics replication.*/
-class IPhysicsReplicationFactory
+class ENGINE_API IPhysicsReplicationFactory
 {
 public:
-	virtual FPhysicsReplication* Create(FPhysScene* OwningPhysScene) = 0;
-	virtual void Destroy(FPhysicsReplication* PhysicsReplication) = 0;
+
+	// NOTE: Once the old Create/Destroy methods are deprecated, remove the default implementation!
+	virtual TUniquePtr<IPhysicsReplication> CreatePhysicsReplication(FPhysScene* OwningPhysScene);
+
+	UE_DEPRECATED(5.3, "Use CreatePhysicsReplication (which gives a replication interface) instead")
+	virtual FPhysicsReplication* Create(FPhysScene* OwningPhysScene) { return nullptr; }
+	UE_DEPRECATED(5.3, "No longer used - PhysScene owns the physics replication ptr once it is created. If you need a custom deleter, override CreatePhysicsReplication and specify the deleter in TUniquePtr.")
+	virtual void Destroy(FPhysicsReplication* PhysicsReplication) { }
 };
diff --git a/Engine/Source/Runtime/Engine/Public/PhysicsReplication.h b/Engine/Source/Runtime/Engine/Public/PhysicsReplication.h
--- a/Engine/Source/Runtime/Engine/Public/PhysicsReplication.h
+++ b/Engine/Source/Runtime/Engine/Public/PhysicsReplication.h
@@ -9,6 +9,7 @@
 
 #include "Engine/EngineTypes.h"
 #include "Engine/ReplicatedState.h"
+#include "PhysicsReplicationInterface.h"
 #include "Physics/PhysicsInterfaceDeclares.h"
 #include "PhysicsProxy/SingleParticlePhysicsProxyFwd.h"
 #include "Chaos/Particles.h"
@@ -85,7 +86,7 @@ class UPrimitiveComponent;
 class FPhysicsReplicationAsyncCallback;
 struct FAsyncPhysicsRepCallbackData;
 
-class ENGINE_API FPhysicsReplication
+class ENGINE_API FPhysicsReplication : public IPhysicsReplication
 {
 public:
 	FPhysicsReplication(FPhysScene* PhysScene);
@@ -95,21 +96,21 @@ public:
 	static bool ShouldSkipPhysicsReplication();
 
 	/** Tick and update all body states according to replicated targets */
-	void Tick(float DeltaSeconds);
+	virtual void Tick(float DeltaSeconds) override;
 
 	/** Sets the latest replicated target for a body instance */
 	UE_DEPRECATED(5.1, "SetReplicatedTarget now takes the ServerFrame.  Please update calls and overloads.")
 	virtual void SetReplicatedTarget(UPrimitiveComponent* Component, FName BoneName, const FRigidBodyState& ReplicatedTarget) { SetReplicatedTarget(Component, BoneName, ReplicatedTarget, 0); }
-	virtual void SetReplicatedTarget(UPrimitiveComponent* Component, FName BoneName, const FRigidBodyState& ReplicatedTarget, int32 ServerFrame);
+	virtual void SetReplicatedTarget(UPrimitiveComponent* Component, FName BoneName, const FRigidBodyState& ReplicatedTarget, int32 ServerFrame) override;
 
 	/** Remove the replicated target*/
-	virtual void RemoveReplicatedTarget(UPrimitiveComponent* Component);
+	virtual void RemoveReplicatedTarget(UPrimitiveComponent* Component) override;
 
 	/** Get the resim frame (min server frame from the targets) */
-	int32 GetResimFrame() const;
+	virtual int32 GetResimFrame() const override;
 
 	/** Set the resim frame for replication */
-	void SetResimFrame(const int32 InResimFrame);
+	virtual void SetResimFrame(const int32 InResimFrame) override;
 
 protected:
 
diff --git a/Engine/Source/Runtime/Engine/Public/PhysicsReplicationInterface.h b/Engine/Source/Runtime/Engine/Public/PhysicsReplicationInterface.h
new file mode 100644
--- /dev/null
+++ b/Engine/Source/Runtime/Engine/Public/PhysicsReplicationInterface.h
@@ -0,0 +1,24 @@
+// Copyright Epic Games, Inc. All Rights Reserved.
+
+#pragma once
+
+#include "CoreMinimal.h"
+
+class UPrimitiveComponent;
+struct FRigidBodyState;
+
+class ENGINE_API IPhysicsReplication
+{
+public:
+	virtual ~IPhysicsReplication() { }
+
+	virtual void Tick(float DeltaSeconds) { }
+
+	virtual void SetReplicatedTarget(UPrimitiveComponent* Component, FName BoneName, const FRigidBodyState& ReplicatedTarget, int32 ServerFrame) = 0;
+
+	virtual void RemoveReplicatedTarget(UPrimitiveComponent* Component) = 0;
+
+	virtual void SetResimFrame(const int32 InResimFrame) { }
+
+	virtual int32 GetResimFrame() const { return INDEX_NONE; }
+};
