Instance static mesh components API improvements:
- Add array view version of batching transform updates
- Add option for passing sorted array to removal function

diff --git a/Engine/Source/Runtime/Engine/Classes/Components/HierarchicalInstancedStaticMeshComponent.h b/Engine/Source/Runtime/Engine/Classes/Components/HierarchicalInstancedStaticMeshComponent.h
--- a/Engine/Source/Runtime/Engine/Classes/Components/HierarchicalInstancedStaticMeshComponent.h
+++ b/Engine/Source/Runtime/Engine/Classes/Components/HierarchicalInstancedStaticMeshComponent.h
@@ -236,10 +236,13 @@ public:
 	virtual TArray<int32> AddInstances(const TArray<FTransform>& InstanceTransforms, bool bShouldReturnIndices, bool bWorldSpace = false) override;
 	virtual bool RemoveInstance(int32 InstanceIndex) override;
 	virtual bool RemoveInstances(const TArray<int32>& InstancesToRemove) override;
+	virtual bool RemoveInstances(const TArray<int32>& InstancesToRemove, bool bInstanceArrayAlreadySortedInReverseOrder) override;
 	virtual bool UpdateInstanceTransform(int32 InstanceIndex, const FTransform& NewInstanceTransform, bool bWorldSpace, bool bMarkRenderStateDirty = false, bool bTeleport = false) override;
 	virtual bool SetCustomDataValue(int32 InstanceIndex, int32 CustomDataIndex, float CustomDataValue, bool bMarkRenderStateDirty = false) override;
 	virtual bool SetCustomData(int32 InstanceIndex, TArrayView<const float> InCustomData, bool bMarkRenderStateDirty = false) override;
 	virtual bool BatchUpdateInstancesTransforms(int32 StartInstanceIndex, const TArray<FTransform>& NewInstancesTransforms, bool bWorldSpace=false, bool bMarkRenderStateDirty=false, bool bTeleport=false) override;
+	virtual bool BatchUpdateInstancesTransforms(int32 StartInstanceIndex, TArrayView<const FTransform> NewInstancesTransforms, bool bWorldSpace, bool bMarkRenderStateDirty, bool bTeleport) override;
+
 	virtual bool BatchUpdateInstancesTransform(int32 StartInstanceIndex, int32 NumInstances, const FTransform& NewInstancesTransform, bool bWorldSpace=false, bool bMarkRenderStateDirty=false, bool bTeleport=false) override;
 	virtual bool BatchUpdateInstancesData(int32 StartInstanceIndex, int32 NumInstances, FInstancedStaticMeshInstanceData* StartInstanceData, bool bMarkRenderStateDirty = false, bool bTeleport = false) override;
 
@@ -327,6 +330,10 @@ protected:
 	FGraphEventArray BuildTreeAsyncTasks;
 	EHISMViewRelevanceType ViewRelevanceType = EHISMViewRelevanceType::HISM;
 
+private:
+	/** hidden Implementation of BatchUpdateInstancesTransforms - it is shared by the TArray and TArrayView version of public API */
+	bool BatchUpdateInstancesTransformsInternal(int32 StartInstanceIndex, TArrayView<const FTransform> NewInstancesTransforms, bool bWorldSpace, bool bMarkRenderStateDirty, bool bTeleport);
+
 protected:
 	friend FStaticLightingTextureMapping_InstancedStaticMesh;
 	friend FInstancedLightMap2D;
diff --git a/Engine/Source/Runtime/Engine/Classes/Components/InstancedStaticMeshComponent.h b/Engine/Source/Runtime/Engine/Classes/Components/InstancedStaticMeshComponent.h
--- a/Engine/Source/Runtime/Engine/Classes/Components/InstancedStaticMeshComponent.h
+++ b/Engine/Source/Runtime/Engine/Classes/Components/InstancedStaticMeshComponent.h
@@ -217,6 +217,9 @@ class ENGINE_API UInstancedStaticMeshComponent : public UStaticMeshComponent, pu
 	UFUNCTION(BlueprintCallable, Category = "Components|InstancedStaticMesh")
 	virtual bool BatchUpdateInstancesTransforms(int32 StartInstanceIndex, const TArray<FTransform>& NewInstancesTransforms, bool bWorldSpace=false, bool bMarkRenderStateDirty=false, bool bTeleport=false);
 
+	/** this is array view version of the UFUNCTION, blueprints do not support ArrayViews at the time of adding this one  */
+	virtual bool BatchUpdateInstancesTransforms(int32 StartInstanceIndex, TArrayView<const FTransform> NewInstancesTransforms, bool bWorldSpace = false, bool bMarkRenderStateDirty = false, bool bTeleport = false);
+
 	// TODO: KevinO cleanup
 	/**
 	* Update the transform for an array of instances. Overloaded version which takes an array of NewPreviousFrameTransforms.
@@ -264,6 +267,12 @@ class ENGINE_API UInstancedStaticMeshComponent : public UStaticMeshComponent, pu
 	UFUNCTION(BlueprintCallable, Category = "Components|InstancedStaticMesh")
 	virtual bool RemoveInstances(const TArray<int32>& InstancesToRemove);
 
+	/** Remove the instances specified. Returns True on success.
+	* @param InstancesToRemove array of the instance indices to remove ( can be sorted @see bInstanceArrayAlreadySortedInReverseOrder )
+	* @param bInstanceArrayAlreadySortedInReverseOrder true is the array is already sorted in reverse order
+	*/
+	virtual bool RemoveInstances(const TArray<int32>& InstancesToRemove, bool bInstanceArrayAlreadySortedInReverseOrder);
+
 	/** Clear all instances being rendered by this component. */
 	UFUNCTION(BlueprintCallable, Category="Components|InstancedStaticMesh")
 	virtual void ClearInstances();
@@ -465,6 +474,9 @@ private:
 	/** Update instance body with a new transform */
 	void UpdateInstanceBodyTransform(int32 InstanceIndex, const FTransform& WorldSpaceInstanceTransform, bool bTeleport);
 
+	/** hidden Implementation of BatchUpdateInstancesTransforms - it is shared by the TArray and TArrayView version of public API */
+	bool BatchUpdateInstancesTransformsInternal(int32 StartInstanceIndex, TArrayView<const FTransform> NewInstancesTransforms, bool bWorldSpace, bool bMarkRenderStateDirty, bool bTeleport);
+
 protected:
 	/** Creates body instances for all instances owned by this component. */
 	void CreateAllInstanceBodies();
diff --git a/Engine/Source/Runtime/Engine/Private/HierarchicalInstancedStaticMesh.cpp b/Engine/Source/Runtime/Engine/Private/HierarchicalInstancedStaticMesh.cpp
--- a/Engine/Source/Runtime/Engine/Private/HierarchicalInstancedStaticMesh.cpp
+++ b/Engine/Source/Runtime/Engine/Private/HierarchicalInstancedStaticMesh.cpp
@@ -2318,6 +2318,11 @@ void UHierarchicalInstancedStaticMeshComponent::RemoveInstancesInternal(const in
 }
 
 bool UHierarchicalInstancedStaticMeshComponent::RemoveInstances(const TArray<int32>& InstancesToRemove)
+{
+	return RemoveInstances(InstancesToRemove, false /*bInstanceArrayAlreadySortedInReverseOrder*/);
+}
+
+bool UHierarchicalInstancedStaticMeshComponent::RemoveInstances(const TArray<int32>& InstancesToRemove, bool bInstanceArrayAlreadySortedInReverseOrder)
 {
 	LLM_SCOPE(ELLMTag::StaticMesh);
 
@@ -2328,26 +2333,40 @@ bool UHierarchicalInstancedStaticMeshComponent::RemoveInstances(const TArray<int
 
 	SCOPE_CYCLE_COUNTER(STAT_HISMCRemoveInstance);
 
-	TArray<int32> SortedInstancesToRemove = InstancesToRemove;
+	auto RemoveInstanceFromSortedArray = [this](const TArray<int32>& SortedInstancesToRemove) -> bool
+	{
+		if (!PerInstanceSMData.IsValidIndex(SortedInstancesToRemove[0]) || !PerInstanceSMData.IsValidIndex(SortedInstancesToRemove.Last()))
+		{
+			return false;
+		}
 
-	// Sort so RemoveAtSwaps don't alter the indices of items still to remove
-	SortedInstancesToRemove.Sort(TGreater<int32>());
+		RemoveInstancesInternal(SortedInstancesToRemove.GetData(), SortedInstancesToRemove.Num());
 
-	if (!PerInstanceSMData.IsValidIndex(SortedInstancesToRemove[0]) || !PerInstanceSMData.IsValidIndex(SortedInstancesToRemove.Last()))
-	{
-		return false;
-	}
+		if (bAutoRebuildTreeOnInstanceChanges)
+		{
+			BuildTreeIfOutdated(/*Async*/true, /*ForceUpdate*/false);
+		}
 
-	RemoveInstancesInternal(SortedInstancesToRemove.GetData(), SortedInstancesToRemove.Num());
+		MarkRenderStateDirty();
 
-	if (bAutoRebuildTreeOnInstanceChanges)
+		return true;
+	};
+
+	bool bSuccess = false;
+	if (bInstanceArrayAlreadySortedInReverseOrder)
 	{
-		BuildTreeIfOutdated(/*Async*/true, /*ForceUpdate*/false);
+		bSuccess = RemoveInstanceFromSortedArray(InstancesToRemove);
 	}
+	else
+	{
+		TArray<int32> SortedInstancesToRemove = InstancesToRemove;
 
-	MarkRenderStateDirty();
+		// Sort so RemoveAtSwaps don't alter the indices of items still to remove
+		SortedInstancesToRemove.Sort(TGreater<int32>());
+		bSuccess = RemoveInstanceFromSortedArray(SortedInstancesToRemove);
+	}
 
-	return true;
+	return bSuccess;
 }
 
 bool UHierarchicalInstancedStaticMeshComponent::RemoveInstance(int32 InstanceIndex)
@@ -2497,6 +2516,16 @@ bool UHierarchicalInstancedStaticMeshComponent::SetCustomData(int32 InstanceInde
 }
 
 bool UHierarchicalInstancedStaticMeshComponent::BatchUpdateInstancesTransforms(int32 StartInstanceIndex, const TArray<FTransform>& NewInstancesTransforms, bool bWorldSpace, bool bMarkRenderStateDirty, bool bTeleport)
+{
+	return BatchUpdateInstancesTransformsInternal(StartInstanceIndex, MakeArrayView(NewInstancesTransforms), bWorldSpace, bMarkRenderStateDirty, bTeleport);
+}
+
+bool UHierarchicalInstancedStaticMeshComponent::BatchUpdateInstancesTransforms(int32 StartInstanceIndex, TArrayView<const FTransform> NewInstancesTransforms, bool bWorldSpace, bool bMarkRenderStateDirty, bool bTeleport)
+{
+	return BatchUpdateInstancesTransformsInternal(StartInstanceIndex, NewInstancesTransforms, bWorldSpace, bMarkRenderStateDirty, bTeleport);
+}
+
+bool UHierarchicalInstancedStaticMeshComponent::BatchUpdateInstancesTransformsInternal(int32 StartInstanceIndex, TArrayView<const FTransform> NewInstancesTransforms, bool bWorldSpace, bool bMarkRenderStateDirty, bool bTeleport)
 {
 	bool BatchResult = true;
 
diff --git a/Engine/Source/Runtime/Engine/Private/InstancedStaticMesh.cpp b/Engine/Source/Runtime/Engine/Private/InstancedStaticMesh.cpp
--- a/Engine/Source/Runtime/Engine/Private/InstancedStaticMesh.cpp
+++ b/Engine/Source/Runtime/Engine/Private/InstancedStaticMesh.cpp
@@ -3790,21 +3790,39 @@ bool UInstancedStaticMeshComponent::RemoveInstance(int32 InstanceIndex)
 
 bool UInstancedStaticMeshComponent::RemoveInstances(const TArray<int32>& InstancesToRemove)
 {
-	// Sort so Remove doesn't alter the indices of items still to remove
-	TArray<int32> SortedInstancesToRemove = InstancesToRemove;
-	SortedInstancesToRemove.Sort(TGreater<int32>());
+	return RemoveInstances(InstancesToRemove, false /*bInstanceArrayAlreadySortedInReverseOrder*/);
+}
 
-	if (!PerInstanceSMData.IsValidIndex(SortedInstancesToRemove[0]) || !PerInstanceSMData.IsValidIndex(SortedInstancesToRemove.Last()))
+bool UInstancedStaticMeshComponent::RemoveInstances(const TArray<int32>& InstancesToRemove, bool bInstanceArrayAlreadySortedInReverseOrder)
+{
+	auto RemoveInstancedWithSortedArray = [this](const TArray<int32>& SortedInstancesToRemove) -> bool
 	{
-		return false;
-	}
+		if (!PerInstanceSMData.IsValidIndex(SortedInstancesToRemove[0]) || !PerInstanceSMData.IsValidIndex(SortedInstancesToRemove.Last()))
+		{
+			return false;
+		}
+
+		for (const int32 InstanceIndex : SortedInstancesToRemove)
+		{
+			RemoveInstanceInternal(InstanceIndex, false);
+		}
+		return true;
+	};
 
-	for (const int32 InstanceIndex : SortedInstancesToRemove)
+	bool bSuccess = false;
+	if (bInstanceArrayAlreadySortedInReverseOrder)
+	{
+		bSuccess = RemoveInstancedWithSortedArray(InstancesToRemove);
+	}
+	else
 	{
-		RemoveInstanceInternal(InstanceIndex, false);
+		// Sort so Remove doesn't alter the indices of items still to remove
+		TArray<int32> SortedInstancesToRemove = InstancesToRemove;
+		SortedInstancesToRemove.Sort(TGreater<int32>());
+		bSuccess = RemoveInstancedWithSortedArray(SortedInstancesToRemove);
 	}
 
-	return true;
+	return bSuccess;
 }
 
 bool UInstancedStaticMeshComponent::GetInstanceTransform(int32 InstanceIndex, FTransform& OutInstanceTransform, bool bWorldSpace) const
@@ -4224,6 +4242,16 @@ bool UInstancedStaticMeshComponent::UpdateInstances(
 }
 
 bool UInstancedStaticMeshComponent::BatchUpdateInstancesTransforms(int32 StartInstanceIndex, const TArray<FTransform>& NewInstancesTransforms, bool bWorldSpace, bool bMarkRenderStateDirty, bool bTeleport)
+{
+	return BatchUpdateInstancesTransformsInternal(StartInstanceIndex, MakeArrayView(NewInstancesTransforms), bWorldSpace, bMarkRenderStateDirty, bTeleport);
+}
+
+bool UInstancedStaticMeshComponent::BatchUpdateInstancesTransforms(int32 StartInstanceIndex, TArrayView<const FTransform> NewInstancesTransforms, bool bWorldSpace, bool bMarkRenderStateDirty, bool bTeleport)
+{
+	return BatchUpdateInstancesTransformsInternal(StartInstanceIndex, NewInstancesTransforms, bWorldSpace, bMarkRenderStateDirty, bTeleport);
+}
+
+bool UInstancedStaticMeshComponent::BatchUpdateInstancesTransformsInternal(int32 StartInstanceIndex, TArrayView<const FTransform> NewInstancesTransforms, bool bWorldSpace, bool bMarkRenderStateDirty, bool bTeleport)
 {
 	CSV_SCOPED_TIMING_STAT_EXCLUSIVE(UInstancedStaticMeshComponent_BatchUpdateInstancesTransforms);
 	TRACE_CPUPROFILER_EVENT_SCOPE_STR("UInstancedStaticMeshComponent::BatchUpdateInstancesTransforms");
