ReplicatedSubObject
* Fix refcount randomly breaking due to accessing invalid memory from a TSet iterator after having removed the iterator

diff --git a/Engine/Source/Runtime/Engine/Private/DataChannel.cpp b/Engine/Source/Runtime/Engine/Private/DataChannel.cpp
--- a/Engine/Source/Runtime/Engine/Private/DataChannel.cpp
+++ b/Engine/Source/Runtime/Engine/Private/DataChannel.cpp
@@ -3226,7 +3226,10 @@ void UActorChannel::ProcessBunch( FInBunch & Bunch )
 	}
 
 #if UE_REPLICATED_OBJECT_REFCOUNTING
-	Connection->Driver->GetNetworkObjectList().RemoveMultipleSubObjectChannelReference(Actor, ReferencesToRemove, this);
+	if (ReferencesToRemove.Num() > 0 )
+	{
+		Connection->Driver->GetNetworkObjectList().RemoveMultipleSubObjectChannelReference(Actor, ReferencesToRemove, this);
+	}
 #endif
 
 	// After all properties have been initialized, call PostNetInit. This should call BeginPlay() so initialization can be done with proper starting values.
diff --git a/Engine/Source/Runtime/Engine/Private/NetConnection.cpp b/Engine/Source/Runtime/Engine/Private/NetConnection.cpp
--- a/Engine/Source/Runtime/Engine/Private/NetConnection.cpp
+++ b/Engine/Source/Runtime/Engine/Private/NetConnection.cpp
@@ -5244,7 +5244,10 @@ void UNetConnection::CleanupDormantReplicatorsForActor(AActor* Actor)
 
 			DormantReplicatorSet.ForEachDormantReplicatorOfActor(Actor, ExecuteFunction);
 
-			Driver->GetNetworkObjectList().RemoveMultipleSubObjectChannelReference(Actor, RemovedObjects, this);
+			if (RemovedObjects.Num() > 0)
+			{
+				Driver->GetNetworkObjectList().RemoveMultipleSubObjectChannelReference(Actor, RemovedObjects, this);
+			}
 		}
 #endif
 
