LocalHeightFogVolume - fix detail panel authoring issue where new values where not always accounted for.

diff --git a/Engine/Source/Runtime/Engine/Private/Components/LocalHeightFogComponent.cpp b/Engine/Source/Runtime/Engine/Private/Components/LocalHeightFogComponent.cpp
--- a/Engine/Source/Runtime/Engine/Private/Components/LocalHeightFogComponent.cpp
+++ b/Engine/Source/Runtime/Engine/Private/Components/LocalHeightFogComponent.cpp
@@ -96,7 +96,20 @@ void ULocalHeightFogComponent::DestroyRenderState_Concurrent()
 void ULocalHeightFogComponent::PostEditChangeProperty(FPropertyChangedEvent& PropertyChangedEvent)
 {
 	Super::PostEditChangeProperty(PropertyChangedEvent);
-	SendRenderTransformCommand();
+
+	FProperty* Property = PropertyChangedEvent.Property;
+	const FName PropertyName = Property ? Property->GetFName() : NAME_None;
+
+	if (PropertyName == TEXT("RelativeLocation") ||
+		PropertyName == TEXT("RelativeRotation") ||
+		PropertyName == TEXT("RelativeScale3D"))
+	{
+		SendRenderTransformCommand();
+	}
+	else
+	{
+		MarkRenderStateDirty();
+	}
 }
 
 #endif // WITH_EDITOR
