Cut time spent in ValidateTextureOverrides by more than half by validating only unique materials

#rnx

diff --git a/Engine/Source/Runtime/Engine/Private/World.cpp b/Engine/Source/Runtime/Engine/Private/World.cpp
--- a/Engine/Source/Runtime/Engine/Private/World.cpp
+++ b/Engine/Source/Runtime/Engine/Private/World.cpp
@@ -1124,6 +1124,7 @@ void UWorld::PostDuplicate(bool bDuplicateForPIE)
 				{
 					TRACE_CPUPROFILER_EVENT_SCOPE(ValidateTextureOverridesForPIE);
 
+					TSet<UMaterialInterface*> ProcessedMaterials;
 					for (TObjectIterator<UPrimitiveComponent> It; It; ++It)
 					{
 						UPrimitiveComponent* Component = *It;
@@ -1134,9 +1135,14 @@ void UWorld::PostDuplicate(bool bDuplicateForPIE)
 							Component->GetUsedMaterials(Materials);
 							for (UMaterialInterface* Material : Materials)
 							{
-								if (UMaterialInstance* MaterialInstance = Cast<UMaterialInstance>(Material))
+								bool bIsAlreadyInSet = false;
+								ProcessedMaterials.FindOrAdd(Material, &bIsAlreadyInSet);
+								if (!bIsAlreadyInSet)
 								{
-									MaterialInstance->ValidateTextureOverrides(FeatureLevel);
+									if (UMaterialInstance* MaterialInstance = Cast<UMaterialInstance>(Material))
+									{
+										MaterialInstance->ValidateTextureOverrides(FeatureLevel);
+									}
 								}
 							}
 						}
