Moves ITemporalUpscaler to UE::Rendering::Private namespace to not commit on backward compatibility between minor versions of the engine

diff --git a/Engine/Source/Runtime/Engine/Public/SceneView.h b/Engine/Source/Runtime/Engine/Public/SceneView.h
--- a/Engine/Source/Runtime/Engine/Public/SceneView.h
+++ b/Engine/Source/Runtime/Engine/Public/SceneView.h
@@ -41,8 +41,14 @@ class FSceneViewFamily;
 class FVolumetricFogViewResources;
 class ISceneRenderer;
 class ISpatialUpscaler;
+
+namespace UE::Renderer::Private
+{
+
 class ITemporalUpscaler;
 
+}
+
 class FRenderTarget;
 
 
@@ -2178,14 +2184,14 @@ public:
 		return ViewExtensionData;
 	}
 
-	FORCEINLINE void SetTemporalUpscalerInterface(ITemporalUpscaler* InTemporalUpscalerInterface)
+	FORCEINLINE void SetTemporalUpscalerInterface(UE::Renderer::Private::ITemporalUpscaler* InTemporalUpscalerInterface)
 	{
 		check(InTemporalUpscalerInterface);
 		checkf(TemporalUpscalerInterface == nullptr, TEXT("View family already had a temporal upscaler assigned."));
 		TemporalUpscalerInterface = InTemporalUpscalerInterface;
 	}
 
-	FORCEINLINE const ITemporalUpscaler* GetTemporalUpscalerInterface() const
+	FORCEINLINE const UE::Renderer::Private::ITemporalUpscaler* GetTemporalUpscalerInterface() const
 	{
 		return TemporalUpscalerInterface;
 	}
@@ -2228,7 +2234,7 @@ private:
 	ISceneViewFamilyScreenPercentage* ScreenPercentageInterface;
 
 	/** Renderer private interfaces, automatically have same lifetime as FSceneViewFamily. */
-	ITemporalUpscaler* TemporalUpscalerInterface;
+	UE::Renderer::Private::ITemporalUpscaler* TemporalUpscalerInterface;
 	ISpatialUpscaler* PrimarySpatialUpscalerInterface;
 	ISpatialUpscaler* SecondarySpatialUpscalerInterface;
 
diff --git a/Engine/Source/Runtime/Renderer/Private/PostProcess/PostProcessing.cpp b/Engine/Source/Runtime/Renderer/Private/PostProcess/PostProcessing.cpp
--- a/Engine/Source/Runtime/Renderer/Private/PostProcess/PostProcessing.cpp
+++ b/Engine/Source/Runtime/Renderer/Private/PostProcess/PostProcessing.cpp
@@ -201,6 +201,8 @@ FDefaultTemporalUpscaler::FOutputs AddThirdPartyTemporalUpscalerPasses(
 	const FViewInfo& View,
 	const FDefaultTemporalUpscaler::FInputs& Inputs)
 {
+	using UE::Renderer::Private::ITemporalUpscaler;
+
 	const ITemporalUpscaler* UpscalerToUse = View.Family->GetTemporalUpscalerInterface();
 	check(UpscalerToUse);
 
diff --git a/Engine/Source/Runtime/Renderer/Private/PostProcess/TemporalAA.h b/Engine/Source/Runtime/Renderer/Private/PostProcess/TemporalAA.h
--- a/Engine/Source/Runtime/Renderer/Private/PostProcess/TemporalAA.h
+++ b/Engine/Source/Runtime/Renderer/Private/PostProcess/TemporalAA.h
@@ -229,7 +229,7 @@ struct FVisualizeTemporalUpscalerInputs
 
 	// Temporal upscaler used and its inputs and outputs.
 	EMainTAAPassConfig TAAConfig = EMainTAAPassConfig::Disabled;
-	const ITemporalUpscaler* UpscalerUsed = nullptr;
+	const UE::Renderer::Private::ITemporalUpscaler* UpscalerUsed = nullptr;
 	FDefaultTemporalUpscaler::FInputs Inputs;
 	FDefaultTemporalUpscaler::FOutputs Outputs;
 };
diff --git a/Engine/Source/Runtime/Renderer/Private/SceneRendering.h b/Engine/Source/Runtime/Renderer/Private/SceneRendering.h
--- a/Engine/Source/Runtime/Renderer/Private/SceneRendering.h
+++ b/Engine/Source/Runtime/Renderer/Private/SceneRendering.h
@@ -1134,7 +1134,7 @@ struct FPreviousViewInfo
 	FTSRHistory TSRHistory;
 
 	// Custom Temporal AA result of last frame, used by plugins
-	TRefCountPtr<ITemporalUpscaler::IHistory> ThirdPartyTemporalUpscalerHistory;
+	TRefCountPtr<UE::Renderer::Private::ITemporalUpscaler::IHistory> ThirdPartyTemporalUpscalerHistory;
 
 	// Half resolution version temporal AA result of last frame
 	TRefCountPtr<IPooledRenderTarget> HalfResTemporalAAHistory;
diff --git a/Engine/Source/Runtime/Renderer/Public/TemporalUpscaler.h b/Engine/Source/Runtime/Renderer/Public/TemporalUpscaler.h
--- a/Engine/Source/Runtime/Renderer/Public/TemporalUpscaler.h
+++ b/Engine/Source/Runtime/Renderer/Public/TemporalUpscaler.h
@@ -4,7 +4,10 @@
 
 #include "ScreenPass.h"
 
-/** Interface for the main temporal upscaling algorithm. */
+namespace UE::Renderer::Private
+{
+
+/** Interface for implementing third party temporal upscaler. */
 class RENDERER_API ITemporalUpscaler : public ISceneViewFamilyExtention
 {
 public:
@@ -79,3 +82,5 @@ public:
 
 	virtual ITemporalUpscaler* Fork_GameThread(const class FSceneViewFamily& ViewFamily) const = 0;
 };
+
+} // namespace UE::Renderer::Private
