Reduce texture compiler cost when only PIE textures should be built

#rnx

diff --git a/Engine/Source/Runtime/Engine/Private/TextureCompiler.cpp b/Engine/Source/Runtime/Engine/Private/TextureCompiler.cpp
--- a/Engine/Source/Runtime/Engine/Private/TextureCompiler.cpp
+++ b/Engine/Source/Runtime/Engine/Private/TextureCompiler.cpp
@@ -698,11 +698,22 @@ void FTextureCompilingManager::FinishCompilationsForGame()
 }
 
 void FTextureCompilingManager::ProcessAsyncTasks(bool bLimitExecutionTime)
+{
+	AssetCompilation::FProcessAsyncTaskParams Params;
+	Params.bLimitExecutionTime = bLimitExecutionTime;
+	Params.bPlayInEditorAssetsOnly = false;
+	ProcessAsyncTasks(Params);
+}
+
+void FTextureCompilingManager::ProcessAsyncTasks(const AssetCompilation::FProcessAsyncTaskParams& Params)
 {
 	FObjectCacheContextScope ObjectCacheScope;
 	FinishCompilationsForGame();
 
-	ProcessTextures(bLimitExecutionTime);
+	if (!Params.bPlayInEditorAssetsOnly)
+	{
+		ProcessTextures(Params.bLimitExecutionTime);
+	}
 
 	UpdateCompilationNotification();
 }
diff --git a/Engine/Source/Runtime/Engine/Public/TextureCompiler.h b/Engine/Source/Runtime/Engine/Public/TextureCompiler.h
--- a/Engine/Source/Runtime/Engine/Public/TextureCompiler.h
+++ b/Engine/Source/Runtime/Engine/Public/TextureCompiler.h
@@ -91,7 +91,8 @@ private:
 	TArrayView<FName> GetDependentTypeNames() const override;
 	int32 GetNumRemainingAssets() const override;
 	void ProcessAsyncTasks(bool bLimitExecutionTime = false) override;
-	
+	void ProcessAsyncTasks(const AssetCompilation::FProcessAsyncTaskParams& Params) override;
+
 	void FinishCompilationsForGame();
 	void ProcessTextures(bool bLimitExecutionTime, int32 MaximumPriority = -1);
 	void UpdateCompilationNotification();
