WorldPartition: change defaults for some wp.Runtime.RuntimeSpatialHash* cvars since the current defaults have known issues.

diff --git a/Engine/Source/Runtime/Core/Public/UObject/FortniteMainBranchObjectVersion.h b/Engine/Source/Runtime/Core/Public/UObject/FortniteMainBranchObjectVersion.h
--- a/Engine/Source/Runtime/Core/Public/UObject/FortniteMainBranchObjectVersion.h
+++ b/Engine/Source/Runtime/Core/Public/UObject/FortniteMainBranchObjectVersion.h
@@ -325,6 +325,9 @@ struct CORE_API FFortniteMainBranchObjectVersion
 		// Add FPackedLevelActorDesc for APackedLevelActor and support for APackedLevelActor Filters
 		PackedLevelActorDesc,
 
+		// Add customizable values for several UWorldPartitionRuntimeSpatialHash cvars
+		WorldPartitionRuntimeSpatialHashCVarOverrides,
+
 		// -----<new versions can be added above this line>-------------------------------------------------
 		VersionPlusOne,
 		LatestVersion = VersionPlusOne - 1
diff --git a/Engine/Source/Runtime/Engine/Private/WorldPartition/RuntimeSpatialHash/RuntimeSpatialHashGridHelper.cpp b/Engine/Source/Runtime/Engine/Private/WorldPartition/RuntimeSpatialHash/RuntimeSpatialHashGridHelper.cpp
--- a/Engine/Source/Runtime/Engine/Private/WorldPartition/RuntimeSpatialHash/RuntimeSpatialHashGridHelper.cpp
+++ b/Engine/Source/Runtime/Engine/Private/WorldPartition/RuntimeSpatialHash/RuntimeSpatialHashGridHelper.cpp
@@ -31,6 +31,11 @@ static FAutoConsoleVariableRef CVarRuntimeSpatialHashPlacePartitionActorsUsingLo
 	GRuntimeSpatialHashPlacePartitionActorsUsingLocation,
 	TEXT("Set RuntimeSpatialHashPlacePartitionActorsUsingLocation to true to place partitioned actors into their corresponding cell using their location instead of their bounding box."));
 
+bool GRuntimeSpatialHashUseAlignedGridLevelsEffective = true;
+bool GRuntimeSpatialHashSnapNonAlignedGridLevelsToLowerLevelsEffective = true;
+bool GRuntimeSpatialHashPlaceSmallActorsUsingLocationEffective = false;
+bool GRuntimeSpatialHashPlacePartitionActorsUsingLocationEffective = true;
+
 FSquare2DGridHelper::FSquare2DGridHelper(const FBox& InWorldBounds, const FVector& InOrigin, int64 InCellSize)
 	: WorldBounds(InWorldBounds)
 	, Origin(InOrigin)
@@ -66,7 +71,7 @@ FSquare2DGridHelper::FSquare2DGridHelper(const FBox& InWorldBounds, const FVecto
 	{
 		int64 LevelGridSize = CurrentGridSize;
 
-		if (!GRuntimeSpatialHashUseAlignedGridLevels)
+		if (!GRuntimeSpatialHashUseAlignedGridLevelsEffective)
 		{
 			// Except for top level, adding 1 to CurrentGridSize (which is always a power of 2) breaks the pattern of perfectly aligned cell edges between grid level cells.
 			// This will prevent weird artefact during actor promotion when an actor is placed using its bounds and which overlaps multiple cells.
@@ -165,12 +170,12 @@ FSquare2DGridHelper GetPartitionedActors(const FBox& WorldBounds, const FSpatial
 	auto ShouldActorUseLocationPlacement = [CellArea, CellSize, GridLevelCount](const IStreamingGenerationContext::FActorSetInstance* InActorSetInstance, const FBox2D& InActorSetInstanceBounds, int32& OutGridLevel)
 	{
 		OutGridLevel = 0;
-		if (GRuntimeSpatialHashPlaceSmallActorsUsingLocation && (InActorSetInstanceBounds.GetArea() <= CellArea))
+		if (GRuntimeSpatialHashPlaceSmallActorsUsingLocationEffective && (InActorSetInstanceBounds.GetArea() <= CellArea))
 		{
 			return true;
 		}
 
-		if (GRuntimeSpatialHashPlacePartitionActorsUsingLocation)
+		if (GRuntimeSpatialHashPlacePartitionActorsUsingLocationEffective)
 		{
 			bool bUseLocation = true;
 			for (const FGuid& ActorGuid : InActorSetInstance->ActorSet->Actors)
diff --git a/Engine/Source/Runtime/Engine/Private/WorldPartition/WorldPartitionRuntimeSpatialHash.cpp b/Engine/Source/Runtime/Engine/Private/WorldPartition/WorldPartitionRuntimeSpatialHash.cpp
--- a/Engine/Source/Runtime/Engine/Private/WorldPartition/WorldPartitionRuntimeSpatialHash.cpp
+++ b/Engine/Source/Runtime/Engine/Private/WorldPartition/WorldPartitionRuntimeSpatialHash.cpp
@@ -34,10 +34,16 @@
 #include "Misc/HierarchicalLogArchive.h"
 #include "UObject/UnrealType.h"
 #include "UObject/ObjectSaveContext.h"
+#include "UObject/FortniteMainBranchObjectVersion.h"
 #include "Misc/HashBuilder.h"
 
 #include UE_INLINE_GENERATED_CPP_BY_NAME(WorldPartitionRuntimeSpatialHash)
 
+extern ENGINE_API bool GRuntimeSpatialHashUseAlignedGridLevels;
+extern ENGINE_API bool GRuntimeSpatialHashSnapNonAlignedGridLevelsToLowerLevels;
+extern ENGINE_API bool GRuntimeSpatialHashPlaceSmallActorsUsingLocation;
+extern ENGINE_API bool GRuntimeSpatialHashPlacePartitionActorsUsingLocation;
+
 #if WITH_EDITOR
 extern UNREALED_API class UEditorEngine* GEditor;
 #else
@@ -272,7 +278,7 @@ void FSpatialHashStreamingGrid::GetCells(const TArray<FWorldPartitionStreamingSo
 								{
 									check(Source.TargetState == EStreamingSourceTargetState::Activated);
 									OutActivateCells.AddCell(Cell, Source, Shape);
-									bAddedActivatedCell = !GRuntimeSpatialHashUseAlignedGridLevels && GRuntimeSpatialHashSnapNonAlignedGridLevelsToLowerLevels;
+									bAddedActivatedCell = !GRuntimeSpatialHashUseAlignedGridLevelsEffective && GRuntimeSpatialHashSnapNonAlignedGridLevelsToLowerLevelsEffective;
 								}
 							}
 							else if (Cell->HasAnyDataLayerInEffectiveRuntimeState(EDataLayerRuntimeState::Loaded))
@@ -292,7 +298,7 @@ void FSpatialHashStreamingGrid::GetCells(const TArray<FWorldPartitionStreamingSo
 
 	GetNonSpatiallyLoadedCells(OutActivateCells.GetCells(), OutLoadCells.GetCells());
 
-	if (!GRuntimeSpatialHashUseAlignedGridLevels && GRuntimeSpatialHashSnapNonAlignedGridLevelsToLowerLevels)
+	if (!GRuntimeSpatialHashUseAlignedGridLevelsEffective && GRuntimeSpatialHashSnapNonAlignedGridLevelsToLowerLevelsEffective)
 	{
 		auto FindIntersectingParents = [&Helper, this](const FIntersectingCells& InAllCells, const FIntersectingCells& InTestCells, FIntersectingCells& OutIntersectingCells)
 		{
@@ -996,6 +1002,10 @@ ASpatialHashRuntimeGridInfo::ASpatialHashRuntimeGridInfo(const FObjectInitialize
 UWorldPartitionRuntimeSpatialHash::UWorldPartitionRuntimeSpatialHash(const FObjectInitializer& ObjectInitializer)
 	: Super(ObjectInitializer)
 #if WITH_EDITORONLY_DATA
+	, UseAlignedGridLevels(EWorldPartitionCVarProjectDefaultOverride::Disabled)
+	, SnapNonAlignedGridLevelsToLowerLevels(EWorldPartitionCVarProjectDefaultOverride::Disabled)
+	, PlaceSmallActorsUsingLocation(EWorldPartitionCVarProjectDefaultOverride::Enabled)
+	, PlacePartitionActorsUsingLocation(EWorldPartitionCVarProjectDefaultOverride::Enabled)
 	, bPreviewGrids(false)
 	, PreviewGridLevel(0)
 #endif
@@ -1060,7 +1070,38 @@ URuntimeHashExternalStreamingObjectBase* UWorldPartitionRuntimeSpatialHash::Stor
 
 	return StreamingObject;
 }
+#endif
+
+void UWorldPartitionRuntimeSpatialHash::ApplyCVars()
+{
+#if WITH_EDITORONLY_DATA
+	GRuntimeSpatialHashUseAlignedGridLevelsEffective = (UseAlignedGridLevels == EWorldPartitionCVarProjectDefaultOverride::ProjectDefault) ? GRuntimeSpatialHashUseAlignedGridLevels : (UseAlignedGridLevels == EWorldPartitionCVarProjectDefaultOverride::Enabled);
+	GRuntimeSpatialHashSnapNonAlignedGridLevelsToLowerLevelsEffective = (SnapNonAlignedGridLevelsToLowerLevels == EWorldPartitionCVarProjectDefaultOverride::ProjectDefault) ? GRuntimeSpatialHashSnapNonAlignedGridLevelsToLowerLevels : (SnapNonAlignedGridLevelsToLowerLevels == EWorldPartitionCVarProjectDefaultOverride::Enabled);
+	GRuntimeSpatialHashPlaceSmallActorsUsingLocationEffective = (PlaceSmallActorsUsingLocation == EWorldPartitionCVarProjectDefaultOverride::ProjectDefault) ? GRuntimeSpatialHashPlaceSmallActorsUsingLocation : (PlaceSmallActorsUsingLocation == EWorldPartitionCVarProjectDefaultOverride::Enabled);
+	GRuntimeSpatialHashPlacePartitionActorsUsingLocationEffective = (PlacePartitionActorsUsingLocation == EWorldPartitionCVarProjectDefaultOverride::ProjectDefault) ? GRuntimeSpatialHashPlacePartitionActorsUsingLocation : (PlacePartitionActorsUsingLocation == EWorldPartitionCVarProjectDefaultOverride::Enabled);
+#endif
+}
+
+void UWorldPartitionRuntimeSpatialHash::Serialize(FArchive& Ar)
+{
+	Ar.UsingCustomVersion(FFortniteMainBranchObjectVersion::GUID);
+
+	Super::Serialize(Ar);
 
+#if WITH_EDITORONLY_DATA
+	if (Ar.CustomVer(FFortniteMainBranchObjectVersion::GUID) < FFortniteMainBranchObjectVersion::WorldPartitionRuntimeSpatialHashCVarOverrides)
+	{
+		UseAlignedGridLevels = EWorldPartitionCVarProjectDefaultOverride::ProjectDefault;
+		SnapNonAlignedGridLevelsToLowerLevels = EWorldPartitionCVarProjectDefaultOverride::ProjectDefault;
+		PlaceSmallActorsUsingLocation = EWorldPartitionCVarProjectDefaultOverride::ProjectDefault;
+		PlacePartitionActorsUsingLocation = EWorldPartitionCVarProjectDefaultOverride::ProjectDefault;
+	}
+#endif
+
+	ApplyCVars();
+}
+
+#if WITH_EDITOR
 void UWorldPartitionRuntimeSpatialHash::PostEditChangeProperty(FPropertyChangedEvent& PropertyChangedEvent)
 {
 	const FName PropertyName = PropertyChangedEvent.GetPropertyName();
@@ -1091,6 +1132,8 @@ void UWorldPartitionRuntimeSpatialHash::PostEditChangeProperty(FPropertyChangedE
 		}
 	}
 
+	ApplyCVars();
+
 	Super::PostEditChangeProperty(PropertyChangedEvent);
 }
 
diff --git a/Engine/Source/Runtime/Engine/Private/WorldPartition/WorldPartitionRuntimeSpatialHashGridPreviewer.cpp b/Engine/Source/Runtime/Engine/Private/WorldPartition/WorldPartitionRuntimeSpatialHashGridPreviewer.cpp
--- a/Engine/Source/Runtime/Engine/Private/WorldPartition/WorldPartitionRuntimeSpatialHashGridPreviewer.cpp
+++ b/Engine/Source/Runtime/Engine/Private/WorldPartition/WorldPartitionRuntimeSpatialHashGridPreviewer.cpp
@@ -83,7 +83,7 @@ void FWorldPartitionRuntimeSpatialHashGridPreviewer::Draw(UWorld* World, const T
 						CachedParameters.GridColor = Grid.DebugColor;
 					}
 
-					FVector GridOffset = FVector(Grid.Origin, 0) + (GRuntimeSpatialHashUseAlignedGridLevels ? FVector(0.5 * PreviewCellSize) : FVector::ZeroVector);
+					FVector GridOffset = FVector(Grid.Origin, 0) + (GRuntimeSpatialHashUseAlignedGridLevelsEffective ? FVector(0.5 * PreviewCellSize) : FVector::ZeroVector);
 					if (CachedParameters.GridOffset != GridOffset)
 					{
 						MID->SetVectorParameterValue(*FString::Printf(TEXT("Grid%d_Offset"), i), GridOffset);
diff --git a/Engine/Source/Runtime/Engine/Public/WorldPartition/RuntimeSpatialHash/RuntimeSpatialHashGridHelper.h b/Engine/Source/Runtime/Engine/Public/WorldPartition/RuntimeSpatialHash/RuntimeSpatialHashGridHelper.h
--- a/Engine/Source/Runtime/Engine/Public/WorldPartition/RuntimeSpatialHash/RuntimeSpatialHashGridHelper.h
+++ b/Engine/Source/Runtime/Engine/Public/WorldPartition/RuntimeSpatialHash/RuntimeSpatialHashGridHelper.h
@@ -11,8 +11,10 @@
 #include "WorldPartition/DataLayer/DataLayersID.h"
 #include "WorldPartition/DataLayer/DataLayerInstance.h"
 
-extern ENGINE_API bool GRuntimeSpatialHashUseAlignedGridLevels;
-extern ENGINE_API bool GRuntimeSpatialHashSnapNonAlignedGridLevelsToLowerLevels;
+extern ENGINE_API bool GRuntimeSpatialHashUseAlignedGridLevelsEffective;
+extern ENGINE_API bool GRuntimeSpatialHashSnapNonAlignedGridLevelsToLowerLevelsEffective;
+extern ENGINE_API bool GRuntimeSpatialHashPlaceSmallActorsUsingLocationEffective;
+extern ENGINE_API bool GRuntimeSpatialHashPlacePartitionActorsUsingLocationEffective;
 
 /**
   * Square 2D grid helper
diff --git a/Engine/Source/Runtime/Engine/Public/WorldPartition/WorldPartitionRuntimeSpatialHash.h b/Engine/Source/Runtime/Engine/Public/WorldPartition/WorldPartitionRuntimeSpatialHash.h
--- a/Engine/Source/Runtime/Engine/Public/WorldPartition/WorldPartitionRuntimeSpatialHash.h
+++ b/Engine/Source/Runtime/Engine/Public/WorldPartition/WorldPartitionRuntimeSpatialHash.h
@@ -19,6 +19,14 @@ typedef UE::Math::TIntVector3<int64> FGridCellCoord;
 typedef UE::Math::TIntVector2<int64> FGridCellCoord2;
 class FWorldPartitionDraw2DContext;
 
+UENUM()
+enum class EWorldPartitionCVarProjectDefaultOverride : uint8
+{
+	ProjectDefault,
+	Disabled,
+	Enabled,
+};
+
 USTRUCT()
 struct FSpatialHashStreamingGridLayerCell
 {
@@ -258,6 +266,10 @@ public:
 	PRAGMA_ENABLE_DEPRECATION_WARNINGS
 	virtual void PreSave(FObjectPreSaveContext ObjectSaveContext) override;
 
+	virtual void Serialize(FArchive& Ar) override;
+
+	void ApplyCVars();
+
 #if WITH_EDITOR
 	virtual void PostEditChangeProperty(FPropertyChangedEvent& PropertyChangedEvent) override;
 
@@ -315,6 +327,22 @@ private:
 	UPROPERTY(EditAnywhere, Config, Category = RuntimeSettings)
 	TArray<FSpatialHashRuntimeGrid> Grids;
 
+	/** Disable to help break the pattern caused by world partition promotion of actors to upper grid levels that are always aligned on child levels. */
+	UPROPERTY(EditAnywhere, Config, AdvancedDisplay, Category = RuntimeSettings)
+	EWorldPartitionCVarProjectDefaultOverride UseAlignedGridLevels;
+
+	/** Disable to avoid snapping higher levels cells to child cells. Only used when bUseAlignedGridLevels is false. */
+	UPROPERTY(EditAnywhere, Config, AdvancedDisplay, Category = RuntimeSettings)
+	EWorldPartitionCVarProjectDefaultOverride SnapNonAlignedGridLevelsToLowerLevels;
+
+	/** Enable to place actors smaller than a cell size into their corresponding cell using their location instead of their bounding box. */
+	UPROPERTY(EditAnywhere, Config, AdvancedDisplay, Category = RuntimeSettings)
+	EWorldPartitionCVarProjectDefaultOverride PlaceSmallActorsUsingLocation;
+
+	/** Enable to place partitioned actors into their corresponding cell using their location instead of their bounding box. */
+	UPROPERTY(EditAnywhere, Config, AdvancedDisplay, Category = RuntimeSettings)
+	EWorldPartitionCVarProjectDefaultOverride PlacePartitionActorsUsingLocation;
+
 	/** Whether to preview runtime grids. */
 	UPROPERTY(Transient)
 	bool bPreviewGrids;
