Adding a CVar for a recent workaround that flushes the render thread when resetting LODInfo on a skinned mesh component, as this incurs a render thread flush, stalling the game thread.

diff --git a/Engine/Source/Runtime/Engine/Private/Components/SkinnedMeshComponent.cpp b/Engine/Source/Runtime/Engine/Private/Components/SkinnedMeshComponent.cpp
--- a/Engine/Source/Runtime/Engine/Private/Components/SkinnedMeshComponent.cpp
+++ b/Engine/Source/Runtime/Engine/Private/Components/SkinnedMeshComponent.cpp
@@ -75,6 +75,14 @@ FAutoConsoleVariableRef CVarUpdateBoundsNotifyStreamingRadiusChangeRatio(
 	ECVF_Default
 );
 
+static bool GReleasePreviousLODInfoOnInitialization = true;
+static FAutoConsoleVariableRef CVarReleasePreviousLODInfoOnInitialization(
+	TEXT("r.SkinnedMesh.ReleasePreviousLODInfoOnInitialization"),
+	GReleasePreviousLODInfoOnInitialization,
+	TEXT("Whether to flush the render thread (incurring a game thread stall) and clean existing LOD info when re-initializating."),
+	ECVF_Default
+);
+
 namespace FAnimUpdateRateManager
 {
 	static float TargetFrameTimeForUpdateRate = 1.f / 30.f; //Target frame rate for lookahead URO
@@ -1215,7 +1223,7 @@ void USkinnedMeshComponent::InitLODInfos()
 		if (GetSkinnedAsset()->GetLODNum() != LODInfo.Num())
 		{
 			// Perform cleanup if LOD infos have been initialized before 
-			if (!LODInfo.IsEmpty())
+			if (!LODInfo.IsEmpty() && GReleasePreviousLODInfoOnInitialization)
 			{
 				// Batch release all resources of LOD infos we're about to destruct.
 				// This is relevant when overrides have been set but LODInfos are
