Cook fix follow-up, also reset data/hash if the compression failed to generate valid compressed data

diff --git a/Engine/Source/Runtime/Engine/Private/Animation/AnimSequence.cpp b/Engine/Source/Runtime/Engine/Private/Animation/AnimSequence.cpp
--- a/Engine/Source/Runtime/Engine/Private/Animation/AnimSequence.cpp
+++ b/Engine/Source/Runtime/Engine/Private/Animation/AnimSequence.cpp
@@ -5310,7 +5310,7 @@ void UAnimSequence::FinishAsyncTasks()
 				return DataByPlatformKeyHash.FindChecked(Hash).Get();
 			}();
 
-			if (It->Value.Get()->WasCancelled())
+			auto ResetData = [this, CompressedAnimData, It]()
 			{
 				CompressedAnimData->Reset();
 				DataByPlatformKeyHash.Remove(It->Key);
@@ -5320,6 +5320,11 @@ void UAnimSequence::FinishAsyncTasks()
 				{
 					DataKeyHash = FIoHash::Zero;
 				}
+			};
+
+			if (It->Value.Get()->WasCancelled())
+			{
+				ResetData();
 			}
 			else
 			{
@@ -5348,7 +5353,8 @@ void UAnimSequence::FinishAsyncTasks()
 				else
 				{
 					// Failed to compress
-					UE_LOG(LogAnimation, Warning, TEXT("Failed to finish async Animation Compression task for %s, as the generated data is not valid."), *GetName());
+					UE_LOG(LogAnimation, Warning, TEXT("Failed to finish async Animation Compression task for %s, as the generated data is not valid."), *GetName());					
+					ResetData();
 				}
 			}
 
