SyncBodies optimization : completely remove allocation from the () operator for geometry collection
directly use chaos direct APi to avoid overhead

diff --git a/Engine/Source/Runtime/Engine/Private/PhysicsEngine/Experimental/PhysScene_Chaos.cpp b/Engine/Source/Runtime/Engine/Private/PhysicsEngine/Experimental/PhysScene_Chaos.cpp
--- a/Engine/Source/Runtime/Engine/Private/PhysicsEngine/Experimental/PhysScene_Chaos.cpp
+++ b/Engine/Source/Runtime/Engine/Private/PhysicsEngine/Experimental/PhysScene_Chaos.cpp
@@ -1958,6 +1958,81 @@ DECLARE_CYCLE_STAT(TEXT("FPhysScene_Chaos::OnSyncBodies-FJointConstraintPhysicsP
 DECLARE_CYCLE_STAT(TEXT("FPhysScene_Chaos::OnSyncBodies-FGeometryCollectionPhysicsProxy"), STAT_SyncBodiesGeometryCollectionPhysicsProxy, STATGROUP_Chaos);
 DECLARE_CYCLE_STAT(TEXT("FPhysScene_Chaos::OnSyncBodies-FClusterUnionPhysicsProxy"), STAT_SyncBodiesClusterUnionPhysicsProxy, STATGROUP_Chaos);
 
+static void UpdateAccelerationStructureFromGeometryCollectionProxy(FGeometryCollectionPhysicsProxy& Proxy, Chaos::ISpatialAcceleration<Chaos::FAccelerationStructureHandle, Chaos::FReal, 3>& SpatialAcceleration)
+{
+	SCOPE_CYCLE_COUNTER(STAT_SyncBodiesGeometryCollectionPhysicsProxy);
+
+	const bool bIsParentProxyNull = (Proxy.GetParentProxy() == nullptr);
+
+	auto AreAllShapesQueryEnabled = [](Chaos::FPBDRigidParticle& Particle) -> bool
+	{
+		for (const TUniquePtr<Chaos::FPerShapeData>& ShapeData : Particle.ShapesArray())
+		{
+			if (!ShapeData->GetCollisionData().bQueryCollision)
+			{
+				return false;
+			}
+		}
+		return true;
+	};
+
+	auto GetParticleWorldBounds = [](Chaos::FPBDRigidParticle& Particle) -> Chaos::FAABB3
+	{
+		const Chaos::FRigidTransform3 ParticleWorldTransform(Particle.X(), Particle.R());
+		Chaos::FAABB3 WorldBounds;
+		if (const Chaos::FImplicitObject* Geometry = Particle.Geometry().Get(); Geometry && Geometry->HasBoundingBox())
+		{
+			WorldBounds = Geometry->CalculateTransformedBounds(ParticleWorldTransform);
+		}
+		return WorldBounds;
+	};
+
+	const TManagedArray<TUniquePtr<Chaos::FPBDRigidParticle>>& GTParticles = Proxy.GetExternalParticles();
+	if (bIsParentProxyNull)
+	{
+		for (const TUniquePtr<Chaos::FPBDRigidParticle>& GTParticle : GTParticles)
+		{
+			if (GTParticle)
+			{
+
+				const Chaos::FSpatialAccelerationIdx SpatialIndex = GTParticle->SpatialIdx();
+
+				// It's possible to be an enabled particle and not qualify for query collisions if the GC particle has been replication abandoned. 
+				// Furthermore, it's possible for a particle to be enabled on the game thread but not actually enabled on the physics thread. A particle
+				// with an internal cluster parent (e.g. a cluster union) could get into this state where the geometry collection physics proxy will not
+				// transfer the disabled state from the PT to the GT because the particle has an internal cluster for a parent. We can detect this case by making
+				// sure that the proxy of the GC does not have a parent proxy (only happens with a cluster union currently).
+				if (!GTParticle->Disabled() && AreAllShapesQueryEnabled(*GTParticle))
+				{
+					const Chaos::FAccelerationStructureHandle AccelerationHandleForUpdate(GTParticle.Get());
+					const Chaos::FAABB3 WorldBounds = GetParticleWorldBounds(*GTParticle);
+					SpatialAcceleration.UpdateElementIn(AccelerationHandleForUpdate, WorldBounds, true /* bHasBounds */, SpatialIndex);
+				}
+				else
+				{
+					// Perf optimization : We only use this FAccelerationStructureHandle for remove from the structure so precomputing the prefiltering data is not necessary 
+					const Chaos::FAccelerationStructureHandle AccelerationHandleForRemove(GTParticle.Get(), false /*bUsePrefiltering*/);
+					SpatialAcceleration.RemoveElementFrom(AccelerationHandleForRemove, SpatialIndex);
+				}
+			}
+		}
+	}
+	else
+	{
+		// if we have a parent proxy ( like attached to a cluster union) then none of the handles should be in the acceleration structure 
+		// todo : right now we are sending all the particles, but we should be able to get away with active + direct children that could reduce the overall cost 
+		for (const TUniquePtr<Chaos::FPBDRigidParticle>& GTParticle : GTParticles)
+		{
+			if (GTParticle)
+			{
+				// Perf optimization : We only use this FAccelerationStructureHandle for remove from the structure so precomputing the prefiltering data is not necessary 
+				const Chaos::FAccelerationStructureHandle AccelerationHandleForRemove(GTParticle.Get(), false /*bUsePrefiltering*/);
+				SpatialAcceleration.RemoveElementFrom(AccelerationHandleForRemove, GTParticle->SpatialIdx());
+			}
+		}
+	}
+}
+
 void FPhysScene_Chaos::OnSyncBodies(Chaos::FPhysicsSolverBase* Solver)
 {
 	using namespace Chaos;
@@ -2044,47 +2119,9 @@ void FPhysScene_Chaos::OnSyncBodies(Chaos::FPhysicsSolverBase* Solver)
 
 		void operator()(FGeometryCollectionPhysicsProxy* Proxy)
 		{
-			SCOPE_CYCLE_COUNTER(STAT_SyncBodiesGeometryCollectionPhysicsProxy);
-			// Don't pass in anything here so we don't end up locking anything because we can assume the scene is already locked.
-			FLockedWritePhysicsObjectExternalInterface Interface = FPhysicsObjectExternalInterface::LockWrite({});
-
-			const bool bIsParentProxyNull = (Proxy->GetParentProxy() == nullptr);
-
-			if (bIsParentProxyNull)
-			{
-				auto ShouldGCParticleBeInAccelerationStructure = [&Interface](Chaos::FPhysicsObjectHandle Handle)
-				{
-					// It's possible to be an enabled particle and not qualify for query collisions if the GC particle has been replication abandoned. 
-					// Furthermore, it's possible for a particle to be enabled on the game thread but not actually enabled on the physics thread. A particle
-					// with an internal cluster parent (e.g. a cluster union) could get into this state where the geometry collection physics proxy will not
-					// transfer the disabled state from the PT to the GT because the particle has an internal cluster for a parent. We can detect this case by making
-					// sure that the proxy of the GC does not have a parent proxy (only happens with a cluster union currently).
-					return !Interface->AreAllDisabled({ &Handle, 1 })
-						&& Interface->AreAllShapesQueryEnabled({ &Handle, 1 });
-				};
-
-				TArray<Chaos::FPhysicsObjectHandle> ActiveHandles = Proxy->GetAllPhysicsObjects().FilterByPredicate(
-					[&ShouldGCParticleBeInAccelerationStructure](Chaos::FPhysicsObjectHandle Handle)
-					{
-						return ShouldGCParticleBeInAccelerationStructure(Handle);
-					}
-				);
-
-				TArray<Chaos::FPhysicsObjectHandle> DisabledHandles = Proxy->GetAllPhysicsObjects().FilterByPredicate(
-					[&ShouldGCParticleBeInAccelerationStructure](Chaos::FPhysicsObjectHandle Handle)
-					{
-						return !ShouldGCParticleBeInAccelerationStructure(Handle);
-					}
-				);
-
-				Interface->AddToSpatialAcceleration(ActiveHandles, Outer->GetSpacialAcceleration());
-				Interface->RemoveFromSpatialAcceleration(DisabledHandles, Outer->GetSpacialAcceleration());
-			}
-			else
+			if (Proxy && Outer->GetSpacialAcceleration())
 			{
-				// if we have aparent proxy ( like attached to a cluster union) then none of the handles should be in the acceleration structure 
-				// todo : right now we are sending all the particles, but we should be able to get away with active + direct children that could reduce the overall cost 
-				Interface->RemoveFromSpatialAcceleration(Proxy->GetAllPhysicsObjects(), Outer->GetSpacialAcceleration());
+				UpdateAccelerationStructureFromGeometryCollectionProxy(*Proxy, *Outer->GetSpacialAcceleration());
 			}
 		}
 
diff --git a/Engine/Source/Runtime/Experimental/Chaos/Public/Chaos/ParticleHandle.h b/Engine/Source/Runtime/Experimental/Chaos/Public/Chaos/ParticleHandle.h
--- a/Engine/Source/Runtime/Experimental/Chaos/Public/Chaos/ParticleHandle.h
+++ b/Engine/Source/Runtime/Experimental/Chaos/Public/Chaos/ParticleHandle.h
@@ -196,7 +196,9 @@ public:
 	static constexpr bool bHasPayloadOnInternalThread = true;
 
 	FAccelerationStructureHandle(FGeometryParticleHandle* InHandle);
-	FAccelerationStructureHandle(FGeometryParticle* InGeometryParticle = nullptr);
+
+	// @param bUsePrefiltering make sur the prefiltering data is computed. setting it to false may be useful when building a payload for a remove operation
+	FAccelerationStructureHandle(FGeometryParticle* InGeometryParticle = nullptr, bool bUsePrefiltering = true);
 
 	template <bool bPersistent>
 	FAccelerationStructureHandle(TGeometryParticleHandleImp<FReal, 3, bPersistent>& InHandle);
@@ -3359,7 +3361,7 @@ FORCEINLINE_DEBUGGABLE FAccelerationStructureHandle::FAccelerationStructureHandl
 	}
 }
 
-FORCEINLINE_DEBUGGABLE FAccelerationStructureHandle::FAccelerationStructureHandle(FGeometryParticle* InGeometryParticle)
+FORCEINLINE_DEBUGGABLE FAccelerationStructureHandle::FAccelerationStructureHandle(FGeometryParticle* InGeometryParticle, bool bUsePrefiltering)
 	: ExternalGeometryParticle(InGeometryParticle)
 	, GeometryParticleHandle(InGeometryParticle ? InGeometryParticle->Handle() : nullptr)
 	, CachedUniqueIdx(InGeometryParticle ? InGeometryParticle->UniqueIdx() : FUniqueIdx())
@@ -3369,7 +3371,10 @@ FORCEINLINE_DEBUGGABLE FAccelerationStructureHandle::FAccelerationStructureHandl
 	{
 		ensure(CachedUniqueIdx.IsValid());
 		ensure(IsInGameThread());
-		UpdatePrePreFilter(*InGeometryParticle);
+		if (bUsePrefiltering)
+		{
+			UpdatePrePreFilter(*InGeometryParticle);
+		}
 	}
 }
 
