READD: Chunked Seektables for BinkAudio (with Fixes)

diff --git a/Engine/Source/Developer/AudioFormatBink/Private/AudioFormatBink.cpp b/Engine/Source/Developer/AudioFormatBink/Private/AudioFormatBink.cpp
--- a/Engine/Source/Developer/AudioFormatBink/Private/AudioFormatBink.cpp
+++ b/Engine/Source/Developer/AudioFormatBink/Private/AudioFormatBink.cpp
@@ -40,7 +40,7 @@ class FAudioFormatBink : public IAudioFormat
 	enum
 	{
 		/** Version for Bink Audio format, this becomes part of the DDC key. */
-		UE_AUDIO_BINK_VER = 5,
+		UE_AUDIO_BINK_VER = 7,
 	};
 
 public:
@@ -140,6 +140,12 @@ public:
 
 	virtual int32 GetMinimumSizeForInitialChunk(FName Format, const TArray<uint8>& SrcBuffer) const override
 	{
+		// Exclude any seek table entries in our size, as we are now using streaming seek tables.
+		if(RequiresStreamingSeekTable())
+		{
+			return sizeof(BinkAudioFileHeader);
+		}
+		
 		// We must have an initial chunk large enough for the header and the seek table, if present.
 		BinkAudioFileHeader const* Header = (BinkAudioFileHeader const*)SrcBuffer.GetData();
 		return sizeof(BinkAudioFileHeader) + Header->seek_table_entry_count * sizeof(uint16);
@@ -149,6 +155,12 @@ public:
 	// until the chunk size is reached, then spit out a block.
 	virtual bool SplitDataForStreaming(const TArray<uint8>& InSrcBuffer, TArray<TArray<uint8>>& OutBuffers, const int32 InMaxInitialChunkSize, const int32 InMaxChunkSize) const override
 	{
+		// This should not be called if we require a streaming seek-table. 
+		if (!ensure(RequiresStreamingSeekTable()==false))
+		{
+			return false;
+		}
+
 		uint8 const* Source = InSrcBuffer.GetData();
 		uint32 SourceLen = InSrcBuffer.Num();
 		uint8 const* SourceEnd = Source + SourceLen;
@@ -217,6 +229,92 @@ public:
 
 		return true;
 	}
+
+	static void StripSeek(TArray<uint8>& InOutBuffer)
+	{
+		void* CompressedData = InOutBuffer.GetData();
+		uint32_t CompressedDataLen = InOutBuffer.Num();
+
+		check(CompressedDataLen > sizeof(BinkAudioFileHeader));
+
+		BinkAudioFileHeader* Header = (BinkAudioFileHeader*)CompressedData;
+		uint32 SeekTableBytes = Header->seek_table_entry_count * sizeof(uint16);
+		check(CompressedDataLen > sizeof(BinkAudioFileHeader) + SeekTableBytes);
+
+		// Mark in the header we don't have any seek table 
+		Header->seek_table_entry_count = 0;
+
+		// Copy the rest of the encoded data over the seek table.
+		uint8* SeekTableStart = (uint8*)(Header + 1);
+		FMemory::Memmove(SeekTableStart, SeekTableStart + SeekTableBytes, CompressedDataLen - SeekTableBytes - sizeof(BinkAudioFileHeader));
+
+		CompressedDataLen -= SeekTableBytes;
+
+		InOutBuffer.SetNum(CompressedDataLen);
+	}
+	
+	virtual bool RequiresStreamingSeekTable() const
+	{
+		return true; // Toggling this will require a version bump.
+	}
+	
+	// Taken from BinkAudioInfo.cpp 
+	static uint32 GetMaxFrameSizeSamples(uint32 SampleRate) 
+	{
+		if (SampleRate >= 44100)
+		{
+			return 1920;
+		}
+		else if (SampleRate >= 22050)
+		{
+			return 960;
+		}
+		else
+		{
+			return 480;
+		}
+	}
+
+	virtual bool ExtractSeekTableForStreaming(TArray<uint8>& InOutBuffer, IAudioFormat::FSeekTable& OutSeektable) const
+	{
+		// This should only be called if we require a streaming seek-table. 
+		if (!ensure(RequiresStreamingSeekTable()))
+		{
+			return false;
+		}
+
+		BinkAudioFileHeader const* Header = (BinkAudioFileHeader const*)InOutBuffer.GetData();
+		if (InOutBuffer.Num() < sizeof(BinkAudioFileHeader) || Header->tag != 'UEBA' || Header->seek_table_entry_count==0)
+		{
+			return false;
+		}
+		
+		// The outer logic that manages the seek-table is unaware of how big the header can be, so for the sake of simplicity,
+		// offset for the size of the header in the seek-table entries, so we don't need to worry about adjusting for it later.
+		static const uint32 ActualAudioOffset = sizeof(BinkAudioFileHeader) + 0; // No entries as we're stripping it.
+
+		// Decode and copy out the seek-table. (it's stored as deltas).
+		uint16* EncodedSeekTable = (uint16*)(InOutBuffer.GetData() + sizeof(BinkAudioFileHeader));
+		uint32 CurrentSeekOffset = ActualAudioOffset;
+		uint32 CurrentTimeOffset = 0;
+		
+		int32 SamplesPerEntry = Header->blocks_per_seek_table_entry * GetMaxFrameSizeSamples(Header->rate);
+		OutSeektable.Offsets.SetNum(Header->seek_table_entry_count);
+		OutSeektable.Times.SetNum(Header->seek_table_entry_count);
+
+		for (int32 i = 0; i < Header->seek_table_entry_count; ++i)
+		{
+			OutSeektable.Times[i] = CurrentTimeOffset;
+			OutSeektable.Offsets[i] = CurrentSeekOffset;
+			CurrentSeekOffset += EncodedSeekTable[i];
+			CurrentTimeOffset += SamplesPerEntry;
+		}
+
+		// Strip the seek-table from the buffer now we've copied it.
+		StripSeek(InOutBuffer);
+
+		return true;
+	}
 };
 
 class FAudioPlatformBinkModule : public IAudioFormatModule
diff --git a/Engine/Source/Developer/TargetPlatform/Public/Interfaces/IAudioFormat.h b/Engine/Source/Developer/TargetPlatform/Public/Interfaces/IAudioFormat.h
--- a/Engine/Source/Developer/TargetPlatform/Public/Interfaces/IAudioFormat.h
+++ b/Engine/Source/Developer/TargetPlatform/Public/Interfaces/IAudioFormat.h
@@ -123,6 +123,23 @@ public:
 	 */
 	virtual bool SplitDataForStreaming(const TArray<uint8>& SrcBuffer, TArray<TArray<uint8>>& OutBuffers, const int32 FirstChunkMaxSize, const int32 MaxChunkSize) const {return false;}
 
+	virtual bool RequiresStreamingSeekTable() const { return false; }
+
+	struct FSeekTable
+	{
+		TArray<uint32> Times;			// Times in AudioFrames.
+		TArray<uint32> Offsets;			// Offset in the compressed data.
+	};	
+
+	/**
+	* Extracts the embedded seek-table, removing it, and outputting it separately.
+	* NOTE: TArray is modified in place. The seek-table is parsed and removed from it.	
+	* @param InOutBuffer Pre-compressed data containing seek-table and compressed data as array of bytes.
+	* @param OutSeekTable Seektable in its generic form.
+	* @return Success or failure
+	*/
+	virtual bool ExtractSeekTableForStreaming(TArray<uint8>& InOutBuffer, FSeekTable& OutSeektable) const { return false; }
+
 public:
 
 	/** Virtual destructor. */
diff --git a/Engine/Source/Runtime/BinkAudioDecoder/Module/Private/BinkAudioInfo.cpp b/Engine/Source/Runtime/BinkAudioDecoder/Module/Private/BinkAudioInfo.cpp
--- a/Engine/Source/Runtime/BinkAudioDecoder/Module/Private/BinkAudioInfo.cpp
+++ b/Engine/Source/Runtime/BinkAudioDecoder/Module/Private/BinkAudioInfo.cpp
@@ -114,17 +114,24 @@ void FBinkAudioInfo::PrepareToLoop()
 //-----------------------------------------------------------------------------
 void FBinkAudioInfo::SeekToTime(const float SeekTimeSeconds)
 {
-	if (Decoder->SeekTableCount == 0)
-	{
-		// Can't seek without a seek table.
-		return;
-	}
-
 	uint32 SeekTimeSamples = 0;
 	if (SeekTimeSeconds > 0)
 	{
 		SeekTimeSamples = (uint32)(SeekTimeSeconds * SampleRate);
 	}
+	
+	// If there's no seek table on the header, see if we there's an embedded one on the chunk.
+	if (Decoder->SeekTableCount == 0)
+	{
+		int32 ChunkIndexToSeekTo = StreamingSoundWave->GetSoundWaveData()->FindChunkIndexForSeeking(SeekTimeSamples);
+		if (ChunkIndexToSeekTo >= 0)
+		{
+			this->StreamSeekBlockIndex = ChunkIndexToSeekTo;
+			this->StreamSeekBlockOffset = 0;
+			this->StreamSeekToSamples = SeekTimeSamples;
+		}
+		return;
+	}
 
 	uint32 SamplesInFrame = GetMaxFrameSizeSamples();
 	if (SeekTimeSamples > this->TrueSampleCount)
diff --git a/Engine/Source/Runtime/Engine/Classes/Sound/SoundWave.h b/Engine/Source/Runtime/Engine/Classes/Sound/SoundWave.h
--- a/Engine/Source/Runtime/Engine/Classes/Sound/SoundWave.h
+++ b/Engine/Source/Runtime/Engine/Classes/Sound/SoundWave.h
@@ -76,9 +76,12 @@ struct FStreamedAudioChunk
 	/** Size of the chunk of data in bytes including zero padding */
 	int32 DataSize = 0;
 
-	/** Size of the audio data. */
+	/** Size of the audio data. (NOTE: This includes a seek-table if its present.) */
 	int32 AudioDataSize = 0;
 
+	/** Chunk position in samples frames in the stream. (NOTE: != INDEX_NONE will assume the presence of a seektable). */
+	uint32 SeekOffsetInAudioFrames = INDEX_NONE;
+
 	/** Bulk data if stored in the package. */
 	FByteBulkData BulkData;
 
@@ -1327,6 +1330,9 @@ public:
  	bool IsZerothChunkDataLoaded() const;
  	const TArrayView<uint8> GetZerothChunkDataView() const;
 
+	bool HasChunkSeekTable(int32 InChunkIndex) const;
+	int32 FindChunkIndexForSeeking(uint32 InTimeInAudioFrames) const;
+
 	// Returns true if the zeroth chunk is loaded, or attempts to load it if not already loaded,
 	// returning true if the load was successful. Can return false if either an error was encountered
 	// in attempting to load the chunk or if stream caching is not enabled for the given sound.
diff --git a/Engine/Source/Runtime/Engine/Classes/Sound/StreamedAudioChunkSeekTable.h b/Engine/Source/Runtime/Engine/Classes/Sound/StreamedAudioChunkSeekTable.h
new file mode 100644
--- /dev/null
+++ b/Engine/Source/Runtime/Engine/Classes/Sound/StreamedAudioChunkSeekTable.h
@@ -0,0 +1,132 @@
+// Copyright Epic Games, Inc. All Rights Reserved.
+
+#pragma once
+
+#include "HAL/Platform.h"
+#include "Containers/Array.h"
+#include "Templates/PimplPtr.h"
+
+#include "StreamedAudioChunkSeekTable.generated.h"
+
+// Forward declarations.
+class FArchive;
+
+UENUM()
+enum class EChunkSeekTableMode : uint8
+{
+	ConstantSamplesPerEntry = 0,
+	VariableSamplesPerEntry,
+};
+
+/**
+ * Class representing a audio seek-table. Typically embedded in a bit-stream.
+ */
+class FStreamedAudioChunkSeekTable
+{
+public:
+	
+	/**
+	 * Parse a seek-table from memory.
+	 *
+	 * @param InMemory Pointer to memory where to load from.
+	 * @param InSize Size of memory
+	 * @param InOutOffset Output the offset of the end of the table
+	 * @param OutTable Table that's loaded
+	 * @return true on Success, false on Failure
+	 */
+	static bool Parse(const uint8* InMemory, uint32 InSize, uint32& InOutOffset, FStreamedAudioChunkSeekTable& OutTable);
+
+	/**
+	 * Calculate the size needed for the table. 
+	 * 
+	 * This is cheaper than serializing the entire struct and counting the bytes.
+	 *
+	 * @param InMode The Mode the table is running in (i.e. constant/variable)
+	 * @return The Size of the Table in bytes including any header.
+	 */
+	static int32 CalcSize(int32 NumEntries, EChunkSeekTableMode InMode = EChunkSeekTableMode::ConstantSamplesPerEntry);
+
+	struct ISeekTableImpl
+	{
+		virtual ~ISeekTableImpl() = default;
+		virtual int32 Num() const = 0;
+		virtual uint32 FindOffset(uint32 InTimeInAudioFrames) const = 0;
+		virtual uint32 FindTime(uint32 InOffset) const = 0;
+		virtual void Add(uint32 InTimeInAudioFrames, uint32 InOffset) = 0;
+		virtual bool Serialize(FArchive& Ar) = 0;
+	};
+
+	/**
+	 * Construct a seek-table. The mode determines the internal representation.
+	 *
+	 * @param InMode Mode the table will operate in.
+	 */
+	FStreamedAudioChunkSeekTable(EChunkSeekTableMode InMode = EChunkSeekTableMode::ConstantSamplesPerEntry);	
+
+	/**
+	 * Add an item to the seek table.
+	 *
+	 * @param InTimeInAudioFrames Time of this entry in audio frames
+	 * @param InOffset Offset in the file of where this time happens
+	 */
+	void Add(uint32 InTimeInAudioFrames, uint32 InOffset)
+	{
+		Impl->Add(InTimeInAudioFrames,InOffset);
+	}
+
+	/**
+	 * Finds an offset for seeking given a time
+	 * 
+	 * @param InTimeInAudioFrames Time to search for in audio frames
+	 * @return Offset indicated by the table. INDEX_NONE (~0) on error.
+	 */
+	uint32 FindOffset(uint32 InTimeInAudioFrames) const
+	{
+		return Impl->FindOffset(InTimeInAudioFrames);
+	}
+
+	/**
+	 * Finds a time given an offset (reverse look up)
+	 *
+	 * @param InOffset Offset in the file
+	 * @return Time indicated by the table. INDEX_NONE (~0) on error. 
+	 */
+	uint32 FindTime(uint32 InOffset) const
+	{
+		return Impl->FindTime(InOffset);
+	}
+	
+	/**
+	 * Get the number of Entries in the seek table.
+	 * @return Num of Items.
+	 */
+	int32 Num() const
+	{
+		return Impl->Num();
+	}
+	
+	/**
+	 * Serialize the table to/from the archive
+	 *
+	 * @param Ar Archive to load/save
+	 * @return True on success, false on failure.
+	 */
+	bool Serialize(FArchive& Ar);
+
+	/**
+	 * Empties the table.
+	 *
+	 */
+	void Reset();
+
+private:
+
+	static uint32 GetMagic();
+	static TUniquePtr<ISeekTableImpl> CreateImpl(EChunkSeekTableMode InMode);
+
+	void SetMode(EChunkSeekTableMode InMode);
+	
+	TUniquePtr<ISeekTableImpl> Impl;
+	EChunkSeekTableMode Mode = EChunkSeekTableMode::ConstantSamplesPerEntry;
+};
+
diff --git a/Engine/Source/Runtime/Engine/Private/AudioDecompress.cpp b/Engine/Source/Runtime/Engine/Private/AudioDecompress.cpp
--- a/Engine/Source/Runtime/Engine/Private/AudioDecompress.cpp
+++ b/Engine/Source/Runtime/Engine/Private/AudioDecompress.cpp
@@ -9,6 +9,7 @@
 #include "Misc/CoreStats.h"
 #include "Misc/ScopeRWLock.h"
 #include "Stats/StatsTrace.h"
+#include "Sound/StreamedAudioChunkSeekTable.h"
 
 IStreamedCompressedInfo::IStreamedCompressedInfo()
 	: bIsStreaming(false)
@@ -157,6 +158,18 @@ bool IStreamedCompressedInfo::StreamCompressedInfoInternal(const FSoundWaveProxy
 	uint32 ChunkSize = 0;
 	const uint8* FirstChunk = GetLoadedChunk(StreamingSoundWave, CurrentChunkIndex, ChunkSize);
 
+	// If there is a seek-table present in the chunk, we'll need to parse that first.
+	if (InWaveProxy->GetSoundWaveData()->HasChunkSeekTable(CurrentChunkIndex))
+	{
+		uint32 SeekTableEnd = 0;
+		if (ensureMsgf(FStreamedAudioChunkSeekTable::Parse(FirstChunk, ChunkSize, SeekTableEnd, GetCurrentSeekTable()),
+			TEXT("Failed to parse seektable in '%s' chunk=%d"), *StreamingSoundWave->GetFName().ToString(), CurrentChunkIndex))
+		{
+			ChunkSize -= SeekTableEnd;
+			FirstChunk += SeekTableEnd;
+		}
+	}
+	
 	bIsStreaming = true;
 	if (FirstChunk)
 	{
@@ -218,6 +231,28 @@ bool IStreamedCompressedInfo::StreamCompressedData(uint8* Destination, bool bLoo
 		SrcBufferDataSize = ChunkSize;
 		SrcBufferOffset = StreamSeekBlockOffset;
 		StreamSeekBlockIndex = INDEX_NONE;
+
+		// If we are using a chunked seek-table, we need to parse the table first.
+		if (StreamingSoundWave->GetSoundWaveData()->HasChunkSeekTable(CurrentChunkIndex))
+		{
+			uint32 TableOffset = 0;
+			if (ensureMsgf(FStreamedAudioChunkSeekTable::Parse(SrcBufferData, ChunkSize, TableOffset, GetCurrentSeekTable()), 
+				TEXT("Failed to parse seektable in '%s', chunk=%d"), *StreamingSoundWave->GetFName().ToString(), CurrentChunkIndex))
+			{
+				if (StreamSeekToSamples != INDEX_NONE)
+				{
+					const FStreamedAudioChunk& Chunk = StreamingSoundWave->GetSoundWaveData()->GetChunk(CurrentChunkIndex);
+					int32 SeekOffsetInChunkSpace = StreamSeekToSamples - Chunk.SeekOffsetInAudioFrames;
+					uint32 Offset = GetCurrentSeekTable().FindOffset(StreamSeekToSamples);
+					if (Offset != INDEX_NONE)
+					{
+						SrcBufferOffset = Offset + TableOffset;
+						SrcBufferOffset += CurrentChunkIndex == 0 ? AudioDataOffset : 0;
+					}
+					StreamSeekToSamples = INDEX_NONE;
+				}
+			}
+		}
 	}
 
 	// If next chunk wasn't loaded when last one finished reading, try to get it again now
@@ -230,6 +265,13 @@ bool IStreamedCompressedInfo::StreamCompressedData(uint8* Destination, bool bLoo
 			bPrintChunkFailMessage = true;
 			SrcBufferDataSize = ChunkSize;
 			SrcBufferOffset = CurrentChunkIndex == 0 ? AudioDataOffset : 0;
+
+			// If we are using a chunked seek-table, we need to parse the table first.
+			if (StreamingSoundWave->GetSoundWaveData()->HasChunkSeekTable(CurrentChunkIndex))
+			{
+				ensureMsgf(FStreamedAudioChunkSeekTable::Parse(SrcBufferData, SrcBufferDataSize, SrcBufferOffset, GetCurrentSeekTable()), 
+					TEXT("Failed to parse seektable in '%s' chunk=%d"), *StreamingSoundWave->GetFName().ToString(), CurrentChunkIndex );
+			}
 		}
 		else
 		{
@@ -315,6 +357,12 @@ bool IStreamedCompressedInfo::StreamCompressedData(uint8* Destination, bool bLoo
 				SrcBufferData = GetLoadedChunk(StreamingSoundWave, CurrentChunkIndex, SrcBufferDataSize);
 				if (SrcBufferData)
 				{
+					// If we are using a chunked seek-table, we need to parse the table first.
+					if (StreamingSoundWave->GetSoundWaveData()->HasChunkSeekTable(CurrentChunkIndex))
+					{
+						ensureMsgf(FStreamedAudioChunkSeekTable::Parse(SrcBufferData, SrcBufferDataSize, SrcBufferOffset, GetCurrentSeekTable()), 
+							TEXT("Failed to parse seektable in '%s' chunk=%d"), *StreamingSoundWave->GetFName().ToString(), CurrentChunkIndex);
+					}
 					UE_CLOG(PreviousChunkIndex != CurrentChunkIndex, LogAudio, Log, TEXT("Changed current chunk '%s' from %d to %d, Offset %d"), *StreamingSoundWave->GetFName().ToString(), PreviousChunkIndex, CurrentChunkIndex, SrcBufferOffset);
 				}
 				else
@@ -343,11 +391,17 @@ int32 IStreamedCompressedInfo::DecompressToPCMBuffer(uint16 /*Unused*/ )
 	int32 FrameSize = GetFrameSize();
 	if (FrameSize <= 0)
 	{
+		UE_LOG(LogAudio, Warning, TEXT("Decoder error: Frame size negative. (indicates packet error).: Wave='%s', FrameSize=%d, SrcBufferOffset=%u, SrcBufferDataSize=%u"),
+			*GetStreamingSoundWave()->GetSoundWaveData()->GetFName().ToString(), FrameSize, SrcBufferOffset, SrcBufferDataSize);
+
 		// Error.
 		return -1;
 	}
 	if (SrcBufferOffset + FrameSize > SrcBufferDataSize)
 	{
+		UE_LOG(LogAudio, Warning, TEXT("Decoder error: Frame size too large (decoding it will take us out of bounds).: Wave='%s', FrameSize=%d, SrcBufferOffset=%u, SrcBufferDataSize=%u, Overby=%d"), 
+			*GetStreamingSoundWave()->GetSoundWaveData()->GetFName().ToString(), FrameSize, SrcBufferOffset, SrcBufferDataSize, (int32)SrcBufferOffset+FrameSize-SrcBufferDataSize );
+		
 		// if frame size is too large, something has gone wrong
 		return -1;
 	}
@@ -441,6 +495,20 @@ const uint8* IStreamedCompressedInfo::GetLoadedChunk(const FSoundWaveProxyPtr& I
 	}
 }
 
+FStreamedAudioChunkSeekTable& IStreamedCompressedInfo::GetCurrentSeekTable()
+{
+	if (!CurrentChunkSeekTable.IsValid())
+	{
+		CurrentChunkSeekTable = MakePimpl<FStreamedAudioChunkSeekTable>();
+	}
+	return *CurrentChunkSeekTable;
+}
+
+const FStreamedAudioChunkSeekTable& IStreamedCompressedInfo::GetCurrentSeekTable() const
+{
+	return const_cast<IStreamedCompressedInfo*>(this)->GetCurrentSeekTable();
+}
+
 //////////////////////////////////////////////////////////////////////////
 // Copied from IOS - probably want to split and share
 //////////////////////////////////////////////////////////////////////////
diff --git a/Engine/Source/Runtime/Engine/Private/AudioDerivedData.cpp b/Engine/Source/Runtime/Engine/Private/AudioDerivedData.cpp
--- a/Engine/Source/Runtime/Engine/Private/AudioDerivedData.cpp
+++ b/Engine/Source/Runtime/Engine/Private/AudioDerivedData.cpp
@@ -20,6 +20,7 @@
 #include "DSP/FloatArrayMath.h"
 #include "DSP/MultichannelBuffer.h"
 #include "Containers/UnrealString.h"
+#include "Sound/StreamedAudioChunkSeekTable.h"
 
 DEFINE_LOG_CATEGORY_STATIC(LogAudioDerivedData, Log, All);
 
@@ -165,7 +166,7 @@ Derived data key generation.
 
 // If you want to bump this version, generate a new guid using
 // VS->Tools->Create GUID and paste it here. https://www.guidgen.com works too.
-#define STREAMEDAUDIO_DERIVEDDATA_VER		TEXT("acc1e60e03be438ab23c-a3d318425d34")
+#define STREAMEDAUDIO_DERIVEDDATA_VER		TEXT("F6F8FA984D96440B8A28C73F2ABAE692")
 
 #ifndef CASE_ENUM_TO_TEXT
 #define CASE_ENUM_TO_TEXT(X) case X: return TEXT(#X);
@@ -445,6 +446,175 @@ class FStreamedAudioCacheDerivedDataWorker : public FNonAbandonableTask
 		return false;
 	}
 
+	static void MakeChunkSeekTable(const IAudioFormat::FSeekTable& InTable, uint32 InChunkStart, uint32 InChunkEnd, FStreamedAudioChunkSeekTable& Out)
+	{
+		int32 SearchFrom = Algo::LowerBound(InTable.Offsets, InChunkStart);
+		if (SearchFrom < 0)
+		{
+			return;
+		}
+
+		for (int32 i=SearchFrom;i<InTable.Offsets.Num();++i)
+		{
+			uint32 Offset = InTable.Offsets[i];
+			uint32 TimeInAudioFrames = InTable.Times[i];
+
+			if (Offset > InChunkEnd) 
+			{
+				break;
+			}
+			if (Offset >= InChunkStart && Offset < InChunkEnd)
+			{
+				Out.Add(TimeInAudioFrames,Offset - InChunkStart);
+			}
+		}
+	}	
+
+	// Returns the size of the new seek-table in bytes.
+	static int32 PrefixChunkWithSeekTable(FStreamedAudioChunkSeekTable& InSeekTable, TArray<uint8>& InOutChunkBytes)
+	{	
+		TArray<uint8> Bytes;
+		FMemoryWriter Writer(Bytes);
+		InSeekTable.Serialize(Writer);
+		InOutChunkBytes.Insert(Bytes,0);
+		return Bytes.Num();
+	}
+
+	static EChunkSeekTableMode DetermineSeekTableMode(const IAudioFormat::FSeekTable& InTable)
+	{
+		const TArray<uint32>& Times = InTable.Times;
+		if (InTable.Times.Num()>2)
+		{
+			int32 Delta = Delta = Times[1] - Times[0];
+			for (int32 i = 2; i < Times.Num(); ++i)
+			{
+				if (Times[i] - Times[i - 1] != Delta)
+				{
+					return EChunkSeekTableMode::VariableSamplesPerEntry;
+				}
+			}
+		}
+		return EChunkSeekTableMode::ConstantSamplesPerEntry;
+	}
+
+	static void CreateChunkSeektableAndPrefix(const EChunkSeekTableMode InMode, const IAudioFormat::FSeekTable& InTable, const uint32 InChunkStart, const uint32 InChunkEnd, 
+		TArray<uint8>& OutChunk, TArray<uint32>& OutChunkOffsets, const int32 InEstimatedSeekTableSize, const int32 InBlockCount, const uint32 InAudioSize, const int32 InBudget)
+	{
+		FStreamedAudioChunkSeekTable ChunkTable(InMode);
+		MakeChunkSeekTable(InTable, InChunkStart, InChunkEnd, ChunkTable);
+		const int32 SeekTableSize = PrefixChunkWithSeekTable(ChunkTable, OutChunk);
+
+		ensureMsgf(SeekTableSize == InEstimatedSeekTableSize, TEXT("FStreamedAudioChunkSeekTable:CalcSize doesn't match Serialize. SerializedSize=%d, CaclSize=%d"), SeekTableSize, InEstimatedSeekTableSize);
+		ensureMsgf(OutChunk.Num() <= InBudget, TEXT("Chunk is over budget: Size=%d, Budget=%d"), OutChunk.Num(), InBudget);
+		ensure(SeekTableSize + InAudioSize == OutChunk.Num());
+		ensureMsgf(InBlockCount == ChunkTable.Num(), TEXT("Expecting to have same number of items in our table as we're budgeted for: Planned=%d, Table=%d"), InBlockCount, ChunkTable.Num());
+	
+		OutChunkOffsets.Add(ChunkTable.FindTime(0));
+	}
+
+	bool SplitUsingSeekTable(const IAudioFormat::FSeekTable& InTable, const TArray<uint8>& InSrcBuffer, TArray<TArray<uint8>>& OutBuffers, const int32 InFirstChunkMaxSize, const int32 InMaxChunkSize, TArray<uint32>& OutChunkOffsets)
+	{
+		const TArray<uint32>& Offsets = InTable.Offsets;
+		
+		// Reject bad inputs.
+		if (Offsets.IsEmpty() || InSrcBuffer.IsEmpty() || InFirstChunkMaxSize <= 0 || InMaxChunkSize <=0)
+		{
+			return false;
+		}
+
+		EChunkSeekTableMode Mode = DetermineSeekTableMode(InTable);
+
+		uint8 const* Source = InSrcBuffer.GetData();
+		uint32 SourceLen = InSrcBuffer.Num();
+		uint8 const* SourceEnd = Source + SourceLen;
+
+		uint8 const* ChunkStart = Source;
+		uint8 const* Current = Source + Offsets[0];	
+
+		int32 BlockCount = 0;
+		int32 CurrentBlock = 0;
+
+		uint32 Budget = InFirstChunkMaxSize;
+
+		while (Current < SourceEnd)
+		{	
+			// A block is the span between two neighboring offsets in our table.
+			uint32 BlockStart = Offsets[CurrentBlock];
+			uint32 BlockEnd = CurrentBlock < Offsets.Num() - 1 ? Offsets[CurrentBlock + 1] : SourceLen;
+			uint32 BlockSize = BlockEnd - BlockStart;
+			
+			// Current chunk stats.
+			int32 CurrentTableSize = FStreamedAudioChunkSeekTable::CalcSize(BlockCount, Mode);
+			uint32 CurrentChunkSize = (Current - ChunkStart) + CurrentTableSize;
+			
+			// Adding a new item will grow the table by 1.
+			int32 NewTableSize = FStreamedAudioChunkSeekTable::CalcSize(BlockCount + 1, Mode);
+			int32 TableDelta = NewTableSize - CurrentTableSize;
+			
+			// Accumulated enough?
+			if (CurrentChunkSize + BlockSize + TableDelta > Budget)
+			{
+				// can't add this chunk, emit.
+				TArray<uint8> Chunk(ChunkStart, Current - ChunkStart);
+				int32 AudioSize = Chunk.Num();
+
+				CreateChunkSeektableAndPrefix(Mode, InTable, ChunkStart-Source, Current-Source, Chunk, OutChunkOffsets, CurrentTableSize, BlockCount, AudioSize, Budget);
+		
+				OutBuffers.Add(Chunk);
+				
+				ChunkStart = Current;
+				Budget = InMaxChunkSize;
+				BlockCount = 0;
+			
+				// retry.
+				continue;
+			}
+				
+			// Include this block.
+			Current += BlockSize;
+			BlockCount++;
+			CurrentBlock++;
+
+			UE_LOG(LogAudio, Verbose, TEXT("Adding Block: Chunk_Block_Count=%d, BlockSize=%d, ChunkSize=%d"), BlockCount, BlockSize, Current - ChunkStart);
+		}
+
+		// emit any remainder chunks
+		if (Current - ChunkStart)
+		{
+			// emit this chunk
+			TArray<uint8> Chunk(ChunkStart, Current - ChunkStart);
+			int32 AudioSize = Chunk.Num();
+			int32 CurrentTableSize = FStreamedAudioChunkSeekTable::CalcSize(BlockCount, Mode);
+			CreateChunkSeektableAndPrefix(Mode, InTable, ChunkStart - Source, Current - Source, Chunk, OutChunkOffsets, CurrentTableSize, BlockCount, AudioSize, Budget);
+			OutBuffers.Add(Chunk);
+		}
+
+		return true;
+	}
+
+	bool SplitDataForStreaming(const IAudioFormat* InFormat, const TArray<uint8>& InSrcBuffer, TArray<TArray<uint8>>& OutBuffers, const int32 InFirstChunkMaxSize, const int32 InMaxChunkSize, TArray<uint32>& OutChunkOffsets)
+	{
+		bool bRequiresSeekTable = InFormat->RequiresStreamingSeekTable();	
+
+		// Split using the seek-table if we have one.
+		if (bRequiresSeekTable)
+		{
+			// Because the extract call modifies the source buffer inline (to remove the embedded seek table) we must keep a copy as if
+			// the call was to fail to split the calling code must contend with the original.
+			TArray<uint8> InSrcBufferCopy = InSrcBuffer;
+		
+			IAudioFormat::FSeekTable Table;			
+			if (!ensureMsgf(InFormat->ExtractSeekTableForStreaming(InSrcBufferCopy, Table), TEXT("SoundWave: '%s' requires a Seektable, but doesn't contain one."), *SoundWaveFullName))
+			{
+				return false;
+			}
+			return SplitUsingSeekTable(Table, InSrcBufferCopy,OutBuffers,InFirstChunkMaxSize,InMaxChunkSize, OutChunkOffsets);
+		}
+
+		// otherwise... ask the format to split.
+		return InFormat->SplitDataForStreaming(InSrcBuffer, OutBuffers, InFirstChunkMaxSize, InMaxChunkSize);
+	}
+
 	/** Build the streamed audio. This function is safe to call from any thread. */
 	void BuildStreamedAudio()
 	{
@@ -522,29 +692,33 @@ class FStreamedAudioCacheDerivedDataWorker : public FNonAbandonableTask
 						MaxChunkSizeForCurrentWave = FMath::Min(MaxChunkSizeOverrideBytes, MaxChunkSizeForCurrentWave);
 					}
 
-					UE_LOG(LogAudio, Display, TEXT("Chunk size for %s: %d"), *SoundWaveFullName, MaxChunkSizeForCurrentWave);
+					UE_LOG(LogAudio, Display, TEXT("Chunk sizes for %s: Chunk0=%d bytes, Chunk1-N=%dk"), *SoundWaveFullName, FirstChunkSize, MaxChunkSizeForCurrentWave>>10);
 				}
 				
 				check(FirstChunkSize != 0 && MaxChunkSizeForCurrentWave != 0);
 
-				if (AudioFormat->SplitDataForStreaming(CompressedBuffer, ChunkBuffers, FirstChunkSize, MaxChunkSizeForCurrentWave))
+				TArray<uint32> ChunkOffsets;
+				if (SplitDataForStreaming(AudioFormat, CompressedBuffer, ChunkBuffers, FirstChunkSize, MaxChunkSizeForCurrentWave, ChunkOffsets))
 				{
 					if (ChunkBuffers.Num() > 32)
 					{
 						UE_LOG(LogAudio, Display, TEXT("Sound Wave %s is very large, requiring %d chunks."), *SoundWaveFullName, ChunkBuffers.Num());
 					}
 
+					int32 TotalAudioBytesExcludingSeektable = 0;
 					if (ChunkBuffers.Num() > 0)
 					{
 						// The zeroth chunk should not be zero-padded.
-						const int32 AudioDataSize = ChunkBuffers[0].Num();
+						int32 AudioDataSize = ChunkBuffers[0].Num();
 
 						//FStreamedAudioChunk* NewChunk = new(DerivedData->Chunks) FStreamedAudioChunk();
 						int32 ChunkIndex = DerivedData->Chunks.Add(new FStreamedAudioChunk());
 						FStreamedAudioChunk* NewChunk = &(DerivedData->Chunks[ChunkIndex]);
+
 						// Store both the audio data size and the data size so decoders will know what portion of the bulk data is real audio
 						NewChunk->AudioDataSize = AudioDataSize;
 						NewChunk->DataSize = AudioDataSize;
+						NewChunk->SeekOffsetInAudioFrames = ChunkOffsets.IsValidIndex(ChunkIndex) ? ChunkOffsets[ChunkIndex] : INDEX_NONE;
 
 						if (NewChunk->BulkData.IsLocked())
 						{
@@ -561,8 +735,9 @@ class FStreamedAudioCacheDerivedDataWorker : public FNonAbandonableTask
 					for (int32 ChunkIndex = 1; ChunkIndex < ChunkBuffers.Num(); ++ChunkIndex)
 					{
 						// Zero pad the reallocation if the chunk isn't precisely the max chunk size to keep the reads aligned to MaxChunkSize
-						const int32 AudioDataSize = ChunkBuffers[ChunkIndex].Num();
-						check(AudioDataSize != 0 && AudioDataSize <= MaxChunkSizeForCurrentWave);
+						int32 AudioDataSize = ChunkBuffers[ChunkIndex].Num();
+						check(AudioDataSize != 0);
+						ensureMsgf(AudioDataSize <= MaxChunkSizeForCurrentWave, TEXT("Chunk is overbudget by %d bytes"), AudioDataSize - MaxChunkSizeForCurrentWave);
 
 						int32 ZeroPadBytes = 0;
 
@@ -578,7 +753,8 @@ class FStreamedAudioCacheDerivedDataWorker : public FNonAbandonableTask
 						// Store both the audio data size and the data size so decoders will know what portion of the bulk data is real audio
 						NewChunk->AudioDataSize = AudioDataSize;
 						NewChunk->DataSize = AudioDataSize + ZeroPadBytes;
-
+						NewChunk->SeekOffsetInAudioFrames = ChunkOffsets.IsValidIndex(ChunkIndex) ? ChunkOffsets[ChunkIndex] : INDEX_NONE;
+						
 						if (NewChunk->BulkData.IsLocked())
 						{
 							UE_LOG(LogAudioDerivedData, Warning, TEXT("While building split chunk for streaming: Raw PCM data already being written to. Chunk Index: %d SoundWave: %s "), ChunkIndex, *SoundWaveFullName);
diff --git a/Engine/Source/Runtime/Engine/Private/SoundWave.cpp b/Engine/Source/Runtime/Engine/Private/SoundWave.cpp
--- a/Engine/Source/Runtime/Engine/Private/SoundWave.cpp
+++ b/Engine/Source/Runtime/Engine/Private/SoundWave.cpp
@@ -28,6 +28,8 @@
 #include "Misc/OutputDeviceArchiveWrapper.h"
 #include "Async/Async.h"
 #include "UObject/UE5MainStreamObjectVersion.h"
+#include "Algo/BinarySearch.h"
+#include "Templates/UnrealTemplate.h"
 
 #if WITH_EDITOR
 #endif // WITH_EDITOR
@@ -432,6 +434,43 @@ const TArrayView<uint8> FSoundWaveData::GetZerothChunkDataView() const
 	return TArrayView<uint8>(View.GetData(), static_cast<int32>(View.Num()));
 }
 
+bool FSoundWaveData::HasChunkSeekTable(int32 InChunkIndex) const
+{
+	const TIndirectArray<FStreamedAudioChunk>& Chunks = RunningPlatformData.Chunks;
+	if (Chunks.IsValidIndex(InChunkIndex))
+	{
+		return Chunks[InChunkIndex].SeekOffsetInAudioFrames != INDEX_NONE;
+	}
+	return false;
+}
+
+int32 FSoundWaveData::FindChunkIndexForSeeking(uint32 InTimeInAudioFrames) const
+{	
+	const TIndirectArray<FStreamedAudioChunk>& Chunks = RunningPlatformData.Chunks;
+	int32 NumChunks = Chunks.Num();
+
+	// Find most in range chunk.
+	int32 InRangeChunk = INDEX_NONE;
+	for (int32 i=0; i < NumChunks; ++i)
+	{
+		// Ignore chunks without tables. (should only be chunk0).
+		if (Chunks[i].SeekOffsetInAudioFrames == INDEX_NONE)
+		{
+			continue;
+		}
+
+		// If its out of bound, use the last valid one.
+		if (Chunks[i].SeekOffsetInAudioFrames > InTimeInAudioFrames)
+		{
+			return InRangeChunk;
+		}
+
+		// In bounds and valid... 
+		InRangeChunk = i;
+	}
+	return INDEX_NONE;
+}
+
 bool FSoundWaveData::LoadZerothChunk()
 {
 #if WITH_EDITOR
@@ -633,8 +672,21 @@ void FStreamedAudioChunk::Serialize(FArchive& Ar, UObject* Owner, int32 ChunkInd
 		bShouldInlineAudioChunk = Overrides->bInlineStreamedAudioChunks;
 	}
 
-	bool bCooked = Ar.IsCooking();
-	Ar << bCooked;
+	enum 
+	{
+		IsCooked			 = 1 << 0,
+		HasSeekOffset		 = 1 << 1
+	};
+	
+	// Bit-pack flags into a single uint32, instead of a bool.
+	uint32 Flags = 0;
+	if (Ar.IsSaving())
+	{
+		Flags |= Ar.IsCooking() ? IsCooked : 0;
+		Flags |= SeekOffsetInAudioFrames != INDEX_NONE ? HasSeekOffset : 0;
+	}
+
+	Ar << Flags;
 
 	// ChunkIndex 0 is always inline payload, all other chunks are streamed.
 	if (Ar.IsSaving())
@@ -654,13 +706,38 @@ void FStreamedAudioChunk::Serialize(FArchive& Ar, UObject* Owner, int32 ChunkInd
 	Ar << DataSize;
 	Ar << AudioDataSize;
 
+	if (Ar.IsLoading() )
+	{	
+		// Sanity check these look sane. 
+		// And ensure reset in the case they are bad.
+		if (!ensure(DataSize > 0))
+		{
+			DataSize = 0;
+		}
+		if (!ensure(AudioDataSize > 0))
+		{
+			AudioDataSize = 0;
+		}
+		if (!ensure(AudioDataSize <= DataSize))
+		{
+			AudioDataSize = DataSize;
+		}
+	}
+
+	// To save bloat, only Serialize the seek offset if this chunk uses it. 
+	// Otherwise this defaults to INDEX_NONE
+	if (Flags & HasSeekOffset)
+	{
+		Ar << SeekOffsetInAudioFrames;
+	}
+
 #if WITH_EDITORONLY_DATA
-	if (!bCooked)
+	if (!(Flags & IsCooked))
 	{
 		Ar << DerivedDataKey;
 	}
 
-	if (Ar.IsLoading() && bCooked)
+	if (Ar.IsLoading() && (Flags & IsCooked))
 	{
 		bLoadedFromCookedPackage = true;
 	}
@@ -669,6 +746,11 @@ void FStreamedAudioChunk::Serialize(FArchive& Ar, UObject* Owner, int32 ChunkInd
 
 bool FStreamedAudioChunk::GetCopy(void** OutChunkData)
 {
+	if (AudioDataSize <= 0 || DataSize <= 0)
+	{
+		return false;
+	}
+	
 	if (!CachedDataPtr)
 	{
 		if (AudioDataSize != DataSize)
@@ -4378,4 +4460,4 @@ bool USoundWave::HasError() const
 {
 	check(SoundWaveDataPtr);
 	return SoundWaveDataPtr->HasError();
-}
\ No newline at end of file
+}
diff --git a/Engine/Source/Runtime/Engine/Private/StreamedAudioChunkSeekTable.cpp b/Engine/Source/Runtime/Engine/Private/StreamedAudioChunkSeekTable.cpp
new file mode 100644
--- /dev/null
+++ b/Engine/Source/Runtime/Engine/Private/StreamedAudioChunkSeekTable.cpp
@@ -0,0 +1,326 @@
+// Copyright Epic Games, Inc. All Rights Reserved.
+#include "Sound/StreamedAudioChunkSeekTable.h"
+
+#include "Serialization/MemoryReader.h"
+#include "Misc/AssertionMacros.h"
+#include "Audio.h"			// LogAudio
+
+namespace StreamedAudioChunkSeekTable_Private
+{
+	static void DeltaEncode(const TArray<uint32>& InValues, TArray<uint16>& OutDeltas)
+	{
+		if (InValues.Num() > 0)
+		{
+			OutDeltas.SetNum(InValues.Num());
+			for (int32 i = 1; i < InValues.Num(); ++i)
+			{
+				int32 Delta = InValues[i] - InValues[i - 1];
+				OutDeltas[i] = IntCastChecked<uint16>(Delta);
+			}
+		}
+	}
+	static void DeltaDecode(const TArray<uint16>& InDeltas, TArray<uint32>& OutValues)
+	{
+		OutValues.SetNum(InDeltas.Num());
+		uint32 Sum = 0;
+		for (int32 i = 0; i < InDeltas.Num(); ++i)
+		{
+			Sum += InDeltas[i];
+			OutValues[i] = Sum;
+		}
+	}
+}
+
+void FStreamedAudioChunkSeekTable::SetMode(EChunkSeekTableMode InMode)
+{
+	Mode = InMode;
+	Impl = CreateImpl(InMode);
+}
+
+void FStreamedAudioChunkSeekTable::Reset()
+{
+	// Will delete recreate the instance.
+	SetMode(Mode);
+}
+
+FStreamedAudioChunkSeekTable::FStreamedAudioChunkSeekTable(EChunkSeekTableMode InMode)
+{
+	SetMode(InMode);
+}
+
+bool FStreamedAudioChunkSeekTable::Parse(const uint8* InMemory, uint32 InSize, uint32& OutOffset, FStreamedAudioChunkSeekTable& OutTable)
+{
+	FMemoryReaderView Reader(MakeArrayView(InMemory, InSize));
+	OutTable.Reset();
+	if (OutTable.Serialize(Reader))
+	{
+		int32 Size = Reader.Tell();
+		OutOffset += Size;
+		UE_LOG(LogAudio, Verbose, TEXT("Successfully parsed seektable: Entries=%d, Size=%d"), OutTable.Num(), Size);
+		return true;
+	}
+	return false;
+}
+
+uint32 FStreamedAudioChunkSeekTable::GetMagic()
+{
+	// Magic number is 'SEEK', so it's obvious in memory.
+	return (uint32)'KEES';
+}
+
+bool FStreamedAudioChunkSeekTable::Serialize(FArchive& Ar)
+{
+	// For minimizing memory. 
+	// Don't save anything if the table is empty, including the magic.
+	check(Impl.IsValid());
+	if (Ar.IsSaving() && Impl->Num() == 0)
+	{
+		return true;
+	}
+	
+	// Magic number. (in the case of loading an empty table, this will fail and return false)
+	// The outer logic will need to handle that case. 
+	uint32 Magic = GetMagic();
+	Ar << Magic;
+	if (Magic != GetMagic())
+	{
+		return false;
+	}
+	
+	// Mode.
+	uint8 ModeAsByte = static_cast<uint8>(Mode);
+	Ar << ModeAsByte;
+
+	if (Ar.IsLoading())
+	{
+		EChunkSeekTableMode NewMode = static_cast<EChunkSeekTableMode>(ModeAsByte);
+		if (NewMode!=Mode)
+		{
+			SetMode(NewMode);
+		}
+	}
+
+	check(Impl.IsValid());
+	return Impl->Serialize(Ar);
+}
+
+struct FConstantRateSeekTable : public FStreamedAudioChunkSeekTable::ISeekTableImpl
+{
+	TArray<uint32> Table;
+	uint32 StartTimeOffset = INDEX_NONE;
+	uint16 AudioFramesPerEntry = MAX_uint16;
+
+	static int32 CalcSize(int32 NumEntries)
+	{		
+		// Not including super class.
+		static const int32 SizeOfHeader = sizeof(StartTimeOffset) + sizeof(AudioFramesPerEntry);
+		static const int32 SizePerEntry = sizeof(uint16);
+
+		// Header + TArray<uint16> (each entry, plus count).
+		return SizeOfHeader + (SizePerEntry * NumEntries) + sizeof(int32);
+	}
+
+	int32 Num() const override
+	{
+		return Table.Num();
+	}
+	void Add(uint32 InTimeInAudioFrames, uint32 InOffset) override
+	{
+		if (Table.IsEmpty())
+		{
+			StartTimeOffset = InTimeInAudioFrames;
+		}
+		else if(Table.Num() == 1)
+		{
+			// Imply Samples Per Entry from step 1.
+			check(InTimeInAudioFrames > StartTimeOffset);
+			int32 DeltaFrames = InTimeInAudioFrames - StartTimeOffset;
+			AudioFramesPerEntry = IntCastChecked<uint16>(DeltaFrames);
+		}		
+		Table.Add(InOffset);
+	}
+
+	uint32 FindOffset(uint32 InTimeInAudioFrames) const override
+	{
+		if (!Table.IsEmpty() && AudioFramesPerEntry > 0)
+		{
+			int32 Index = (InTimeInAudioFrames - StartTimeOffset) / AudioFramesPerEntry;
+			if (Table.IsValidIndex(Index))
+			{
+				return Table[Index];
+			}
+		}
+		return INDEX_NONE;
+	}
+	uint32 FindTime(uint32 InOffset) const override
+	{
+		if (!Table.IsEmpty() && AudioFramesPerEntry > 0 && StartTimeOffset >= 0)
+		{
+			uint32 TimeInSamples = StartTimeOffset + (InOffset * AudioFramesPerEntry);
+			return TimeInSamples;
+		}
+		return INDEX_NONE;
+	}
+
+	bool Serialize(FArchive& Ar) override
+	{
+		Ar << AudioFramesPerEntry;
+		Ar << StartTimeOffset;
+
+		if (Ar.IsSaving())
+		{
+			Encode(Ar);
+		}
+		else if (Ar.IsLoading())
+		{
+			if (!Decode(Ar))
+			{
+				return false;
+			}
+		}
+		return true;
+	}
+	void Encode(FArchive& Ar) const
+	{
+		TArray<uint16> Deltas;
+		StreamedAudioChunkSeekTable_Private::DeltaEncode(Table,Deltas);
+
+		// Save in stream.
+		Ar << Deltas;
+	}
+	bool Decode(FArchive& Ar) 
+	{
+		TArray<uint16> Deltas;
+		Ar << Deltas;
+
+		StreamedAudioChunkSeekTable_Private::DeltaDecode(Deltas,Table);
+		TArray<uint8> Bytes;
+		return true;
+	}
+};
+
+struct FVariableRateSeekTable : public FStreamedAudioChunkSeekTable::ISeekTableImpl
+{
+	TArray<uint32> Offsets;
+	TArray<uint32> Times;
+
+	static int32 CalcSize(int32 NumEntries)
+	{
+		// Not including super class.
+		static const int32 SizeOfHeader = 0;
+		static const int32 SizePerEntry = sizeof(uint16) + sizeof(uint16);
+
+		// Header + TArray<uint16> (each entry, plus count).
+		return SizeOfHeader + (SizePerEntry * NumEntries) + sizeof(int32) + sizeof(int32);
+	}
+
+	void Add(uint32 InTimeInAudioFrames, uint32 InOffset) override
+	{
+		Offsets.Add(InOffset);
+		Times.Add(InTimeInAudioFrames);
+	}
+	int32 Num() const override
+	{
+		check(Offsets.Num() == Times.Num());
+		return Times.Num();
+	}
+	uint32 FindOffset(uint32 InTimeInAudioFrames) const override
+	{
+		check(Offsets.Num() == Times.Num());
+		int32 Index = Algo::LowerBound(Times, InTimeInAudioFrames);
+		if (Index != INDEX_NONE)
+		{
+			return Offsets[Index];
+		}
+		return INDEX_NONE;
+	}
+	uint32 FindTime(uint32 InOffset) const override
+	{
+		int32 Index = Algo::LowerBound(Offsets,InOffset);
+		if (Index != INDEX_NONE)
+		{
+			return Times[Index];
+		}
+		return INDEX_NONE;
+	}
+
+	bool Serialize(FArchive& Ar) override
+	{
+		if (Ar.IsSaving())
+		{
+			Encode(Ar);
+		}
+		else if (Ar.IsLoading())
+		{
+			if (!Decode(Ar))
+			{
+				return false;
+			}
+		}
+		return true;
+	}
+
+	void Encode(FArchive& Ar) const
+	{
+		TArray<uint16> DeltaOffsets;
+		StreamedAudioChunkSeekTable_Private::DeltaEncode(Offsets,DeltaOffsets);
+		Ar << DeltaOffsets;
+		
+		TArray<uint16> DeltaTimes;
+		StreamedAudioChunkSeekTable_Private::DeltaEncode(Times, DeltaTimes);
+		Ar << DeltaTimes;
+	}
+
+	bool Decode(FArchive& Ar)
+	{
+		TArray<uint16> DeltaOffsets;
+		Ar << DeltaOffsets;
+		StreamedAudioChunkSeekTable_Private::DeltaDecode(DeltaOffsets, Offsets);
+		
+		TArray<uint16> DeltaSamples;
+		Ar << DeltaSamples;
+		StreamedAudioChunkSeekTable_Private::DeltaDecode(DeltaSamples, Times);
+		return true;
+	}
+};
+
+TUniquePtr<FStreamedAudioChunkSeekTable::ISeekTableImpl> FStreamedAudioChunkSeekTable::CreateImpl(EChunkSeekTableMode InMode)
+{
+	switch(InMode)
+	{
+		case EChunkSeekTableMode::ConstantSamplesPerEntry:
+			return MakeUnique<FConstantRateSeekTable>();
+		case EChunkSeekTableMode::VariableSamplesPerEntry:
+			return MakeUnique<FVariableRateSeekTable>();
+		default:
+			break;
+	}
+	checkNoEntry();
+	return {};
+}
+
+int32 FStreamedAudioChunkSeekTable::CalcSize(int32 InNumEntries, EChunkSeekTableMode InMode /*= EChunkSeekTableMode::ConstantSamplesPerEntry*/)
+{	
+	// Make sure this stays in sync with the Serializer!
+
+	// Don't save anything if the table is empty, including the magic.
+	if (InNumEntries == 0)
+	{
+		return 0;
+	}
+	
+	// Header (magic + mode)
+	static const int32 HeaderSize = sizeof(uint32) + sizeof(uint8);	// Magic + mode.
+
+	switch (InMode)
+	{
+		case EChunkSeekTableMode::ConstantSamplesPerEntry:
+			return HeaderSize + FConstantRateSeekTable::CalcSize(InNumEntries);
+		case EChunkSeekTableMode::VariableSamplesPerEntry:
+			return HeaderSize + FVariableRateSeekTable::CalcSize(InNumEntries);
+		default:
+			break;
+	}
+	checkNoEntry();
+	return 0;
+}
\ No newline at end of file
diff --git a/Engine/Source/Runtime/Engine/Public/AudioDecompress.h b/Engine/Source/Runtime/Engine/Public/AudioDecompress.h
--- a/Engine/Source/Runtime/Engine/Public/AudioDecompress.h
+++ b/Engine/Source/Runtime/Engine/Public/AudioDecompress.h
@@ -20,6 +20,7 @@
 #define MONO_PCM_BUFFER_SIZE		( MONO_PCM_BUFFER_SAMPLES * sizeof( int16 ) )
 
 struct FSoundQualityInfo;
+class FStreamedAudioChunkSeekTable;
 
 /**
  * Interface class to decompress various types of audio data
@@ -278,6 +279,13 @@ protected:
 	 */
 	ENGINE_API const uint8* GetLoadedChunk(const FSoundWaveProxyPtr& InSoundWave, uint32 ChunkIndex, uint32& OutChunkSize);
 
+	/**
+	 * Gets the current chunks seektable instance (or creates one)
+	 * @return the currently loaded chunk seektable
+	 */
+	const FStreamedAudioChunkSeekTable& GetCurrentSeekTable() const;
+	FStreamedAudioChunkSeekTable& GetCurrentSeekTable();
+
 	/** bool set before ParseHeader. Whether we are streaming a file or not. */
 	bool bIsStreaming;
 	/** Ptr to the current streamed chunk. */
@@ -315,8 +323,9 @@ protected:
 	/** Number of bytes of padding used, overridden in some implementations. Defaults to 0. */
 	uint32 SrcBufferPadding;
 	/** Chunk Handle to ensure that this chunk of streamed audio is not deleted while we are using it. */
-	FAudioChunkHandle CurCompressedChunkHandle;
-
+	FAudioChunkHandle CurCompressedChunkHandle;		
+	/** If there's a chunked seek-table present, this contains the current chunks portion. */
+	TPimplPtr<FStreamedAudioChunkSeekTable> CurrentChunkSeekTable;
 	/** 
 		When a streaming seek request comes down, this is the block we are going to. INDEX_NONE means no seek pending.
 		When using the legacy streaming system this is read on a thread other than the decompression thread
@@ -328,6 +337,9 @@ protected:
 	std::atomic<int32> StreamSeekBlockIndex;
 	/** When a streaming seek request comes down, this is the offset in to the block we want to start decoding from. */
 	int32 StreamSeekBlockOffset;
+
+	/** If using the Chunked seek-tables, we request the seek in samples so we can resolve after the chunk table loads. */
+	uint32 StreamSeekToSamples = INDEX_NONE;	
 };
 
 struct IAudioInfoFactory 
