cluster union delegate for sending a notification when the cluster union's bounds changed

diff --git a/Engine/Source/Runtime/Engine/Classes/PhysicsEngine/ClusterUnionComponent.h b/Engine/Source/Runtime/Engine/Classes/PhysicsEngine/ClusterUnionComponent.h
--- a/Engine/Source/Runtime/Engine/Classes/PhysicsEngine/ClusterUnionComponent.h
+++ b/Engine/Source/Runtime/Engine/Classes/PhysicsEngine/ClusterUnionComponent.h
@@ -120,6 +120,7 @@ struct FClusterUnionParticleCandidateData
 
 DECLARE_DYNAMIC_MULTICAST_DELEGATE_ThreeParams(FOnClusterUnionAddedComponent, UPrimitiveComponent*, Component, const TSet<int32>&, BoneIds, bool, bIsNew);
 DECLARE_DYNAMIC_MULTICAST_DELEGATE_OneParam(FOnClusterUnionRemovedComponent, UPrimitiveComponent*, Component);
+DECLARE_DYNAMIC_MULTICAST_DELEGATE_TwoParams(FOnClusterUnionBoundsChanged, UClusterUnionComponent*, Component, const FBoxSphereBounds&, Bounds);
 
 /**
  * This does the bulk of the work exposing a physics cluster union to the game thread.
@@ -183,6 +184,12 @@ public:
 	UPROPERTY(BlueprintAssignable, Category = "Events")
 	FOnClusterUnionRemovedComponent OnComponentRemovedEvent;
 
+	UPROPERTY(BlueprintAssignable, Category = "Events")
+	FOnClusterUnionBoundsChanged OnComponentBoundsChangedEvent;
+
+	// Lambda returns whether or not iteration should continue;
+	void VisitAllCurrentChildComponents(const TFunction<bool(UPrimitiveComponent*)>& Lambda) const;
+
 	friend class UClusterUnionReplicatedProxyComponent;
 protected:
 
@@ -265,8 +272,6 @@ private:
 	void HandleRemovedClusteredComponent(UPrimitiveComponent* ChangedComponent, bool bDestroyReplicatedProxy);
 
 	TArray<UPrimitiveComponent*> GetAllCurrentChildComponents() const;
-	// Lambda returns whether or not iteration should continue;
-	void VisitAllCurrentChildComponents(const TFunction<bool(UPrimitiveComponent*)>& Lambda) const;
 	void VisitAllCurrentChildComponentsForCollision(ECollisionChannel TraceChannel, const struct FCollisionQueryParams& Params, const struct FCollisionResponseParams& ResponseParams, const struct FCollisionObjectQueryParams& ObjectParams, const TFunction<bool(UPrimitiveComponent*)>& Lambda) const;
 
 	// Whether or not this code is running on the server.
diff --git a/Engine/Source/Runtime/Engine/Private/PhysicsEngine/ClusterUnionComponent.cpp b/Engine/Source/Runtime/Engine/Private/PhysicsEngine/ClusterUnionComponent.cpp
--- a/Engine/Source/Runtime/Engine/Private/PhysicsEngine/ClusterUnionComponent.cpp
+++ b/Engine/Source/Runtime/Engine/Private/PhysicsEngine/ClusterUnionComponent.cpp
@@ -579,8 +579,17 @@ void UClusterUnionComponent::SyncClusterUnionFromProxy()
 		HandleRemovedClusteredComponent(Component, true);
 	}
 
+	const FBoxSphereBounds OldBounds = CachedLocalBounds;
 	bHasCachedLocalBounds = false;
 	UpdateBounds();
+
+	// Check if the old bounds is approximately equal to the new bounds. If so, we don't need to send an event
+	// saying that the cluster union changed in shape/position.
+	const bool bIsSameBounds = OldBounds.Origin.Equals(CachedLocalBounds.Origin) && OldBounds.BoxExtent.Equals(CachedLocalBounds.BoxExtent) && (FMath::Abs(OldBounds.SphereRadius - CachedLocalBounds.SphereRadius) < UE_KINDA_SMALL_NUMBER);
+	if (!bIsSameBounds)
+	{
+		OnComponentBoundsChangedEvent.Broadcast(this, CachedLocalBounds);
+	}
 }
 
 void UClusterUnionComponent::HandleAddOrModifiedClusteredComponent(UPrimitiveComponent* ChangedComponent, const TMap<int32, FTransform>& PerBoneChildToParent)
