Fix Static Analysis warning

diff --git a/Engine/Source/Runtime/Engine/Private/AudioDerivedData.cpp b/Engine/Source/Runtime/Engine/Private/AudioDerivedData.cpp
--- a/Engine/Source/Runtime/Engine/Private/AudioDerivedData.cpp
+++ b/Engine/Source/Runtime/Engine/Private/AudioDerivedData.cpp
@@ -485,7 +485,7 @@ class FStreamedAudioCacheDerivedDataWorker : public FNonAbandonableTask
 		const TArray<uint32>& Times = InTable.Times;
 		if (InTable.Times.Num()>2)
 		{
-			int32 Delta = Delta = Times[1] - Times[0];
+			int32 Delta = Times[1] - Times[0];
 			for (int32 i = 2; i < Times.Num(); ++i)
 			{
 				if (Times[i] - Times[i - 1] != Delta)
