deprecated Engine::FindAndPrintStaleReferencesToObject as it was moved FReferenceChainSearch
https://p4-swarm.epicgames.net/reviews/25629892

diff --git a/Engine/Plugins/Experimental/GameFeatures/Source/GameFeatures/Private/GameFeaturePluginStateMachine.cpp b/Engine/Plugins/Experimental/GameFeatures/Source/GameFeatures/Private/GameFeaturePluginStateMachine.cpp
--- a/Engine/Plugins/Experimental/GameFeatures/Source/GameFeatures/Private/GameFeaturePluginStateMachine.cpp
+++ b/Engine/Plugins/Experimental/GameFeatures/Source/GameFeatures/Private/GameFeaturePluginStateMachine.cpp
@@ -265,7 +265,7 @@ namespace UE::GameFeatures
 			}
 
 			UE_LOG(LogGameFeatures, Display, TEXT("Searching for references to %d leaked packages (%d omitted for speed) from plugin %s"), LeakedPackages.Num(), OmittedCount, *PluginName);
-			GEngine->FindAndPrintStaleReferencesToObjects(MakeArrayView((UObject**)PackagesToSearchFor.GetData(), PackagesToSearchFor.Num()), Options);
+			FReferenceChainSearch::FindAndPrintStaleReferencesToObjects(MakeArrayView((UObject**)PackagesToSearchFor.GetData(), PackagesToSearchFor.Num()), Options);
 		}
 
 		// Rename the packages that we are streaming out so that we can possibly reload another copy of them
diff --git a/Engine/Source/Editor/UnrealEd/Private/Cooker/IoStoreCookOnTheFlyRequestManager.cpp b/Engine/Source/Editor/UnrealEd/Private/Cooker/IoStoreCookOnTheFlyRequestManager.cpp
--- a/Engine/Source/Editor/UnrealEd/Private/Cooker/IoStoreCookOnTheFlyRequestManager.cpp
+++ b/Engine/Source/Editor/UnrealEd/Private/Cooker/IoStoreCookOnTheFlyRequestManager.cpp
@@ -345,7 +345,7 @@ private:
 				if (Package)
 				{
 					UE_LOG(LogCookOnTheFly, Warning, TEXT("Can't recook package '%s'"), *PackageName.ToString());
-					UEngine::FindAndPrintStaleReferencesToObject(Package, EPrintStaleReferencesOptions::Display);
+					FReferenceChainSearch::FindAndPrintStaleReferencesToObject(Package, EPrintStaleReferencesOptions::Display);
 				}
 				else
 				{
diff --git a/Engine/Source/Editor/UnrealEd/Private/EditorEngine.cpp b/Engine/Source/Editor/UnrealEd/Private/EditorEngine.cpp
--- a/Engine/Source/Editor/UnrealEd/Private/EditorEngine.cpp
+++ b/Engine/Source/Editor/UnrealEd/Private/EditorEngine.cpp
@@ -7290,7 +7290,7 @@ void UEditorEngine::CheckAndHandleStaleWorldObjectReferences(FWorldContext* InWo
 				UE_LOG(LogLoad, Error, TEXT("Previously active world %s not cleaned up by garbage collection!"), *World->GetPathName());
 				UE_LOG(LogLoad, Error, TEXT("Once a world has become active, it cannot be reused and must be destroyed and reloaded. World referenced by:"));
 			
-				FindAndPrintStaleReferencesToObject(World, 
+				FReferenceChainSearch::FindAndPrintStaleReferencesToObject(World,
 					UObjectBaseUtility::IsPendingKillEnabled() ? EPrintStaleReferencesOptions::Fatal : (EPrintStaleReferencesOptions::Error | EPrintStaleReferencesOptions::Ensure));
 			}
 		}
@@ -7304,7 +7304,7 @@ void UEditorEngine::CheckAndHandleStaleWorldObjectReferences(FWorldContext* InWo
 			{
 				UE_LOG(LogLoad, Error, TEXT("Object %s not cleaned up by garbage collection!"), *Object->GetPathName());
 			
-				FindAndPrintStaleReferencesToObject(Object, 
+				FReferenceChainSearch::FindAndPrintStaleReferencesToObject(Object,
 					UObjectBaseUtility::IsPendingKillEnabled() ? EPrintStaleReferencesOptions::Fatal : (EPrintStaleReferencesOptions::Error | EPrintStaleReferencesOptions::Ensure));
 			}
 		}
diff --git a/Engine/Source/Editor/UnrealEd/Private/EditorLevelUtils.cpp b/Engine/Source/Editor/UnrealEd/Private/EditorLevelUtils.cpp
--- a/Engine/Source/Editor/UnrealEd/Private/EditorLevelUtils.cpp
+++ b/Engine/Source/Editor/UnrealEd/Private/EditorLevelUtils.cpp
@@ -911,7 +911,7 @@ bool UEditorLevelUtils::RemoveLevelsFromWorld(TArray<ULevel*> InLevels, bool bCl
 				UWorld* TheWorld = UWorld::FindWorldInPackage(LevelPackage->GetOutermost());
 				if (TheWorld != nullptr)
 				{
-					UEngine::FindAndPrintStaleReferencesToObject(TheWorld, Options);
+					FReferenceChainSearch::FindAndPrintStaleReferencesToObject(TheWorld, Options);
 					return false;
 				}
 			}
diff --git a/Engine/Source/Editor/UnrealEd/Private/EditorServer.cpp b/Engine/Source/Editor/UnrealEd/Private/EditorServer.cpp
--- a/Engine/Source/Editor/UnrealEd/Private/EditorServer.cpp
+++ b/Engine/Source/Editor/UnrealEd/Private/EditorServer.cpp
@@ -2018,7 +2018,7 @@ void UEditorEngine::CheckForWorldGCLeaks( UWorld* NewWorld, UPackage* WorldPacka
 			(RemainingWorld->WorldType == EWorldType::GamePreview);
 		if(!bIsNewWorld && !bIsPersistantWorldType && !WorldHasValidContext(RemainingWorld))
 		{
-			FindAndPrintStaleReferencesToObject(RemainingWorld, EPrintStaleReferencesOptions::Error);
+			FReferenceChainSearch::FindAndPrintStaleReferencesToObject(RemainingWorld, EPrintStaleReferencesOptions::Error);
 			NumFailedToCleanup++;
 		}
 	}
@@ -2032,7 +2032,7 @@ void UEditorEngine::CheckForWorldGCLeaks( UWorld* NewWorld, UPackage* WorldPacka
 			const bool bIsNewWorldPackage = (NewWorldPackage && RemainingPackage == NewWorldPackage);
 			if(!bIsNewWorldPackage && RemainingPackage == WorldPackage)
 			{
-				FindAndPrintStaleReferencesToObject(RemainingPackage, EPrintStaleReferencesOptions::Error);
+				FReferenceChainSearch::FindAndPrintStaleReferencesToObject(RemainingPackage, EPrintStaleReferencesOptions::Error);
 				NumFailedToCleanup++;
 			}
 		}
diff --git a/Engine/Source/Editor/UnrealEd/Private/PlayLevel.cpp b/Engine/Source/Editor/UnrealEd/Private/PlayLevel.cpp
--- a/Engine/Source/Editor/UnrealEd/Private/PlayLevel.cpp
+++ b/Engine/Source/Editor/UnrealEd/Private/PlayLevel.cpp
@@ -499,7 +499,7 @@ void UEditorEngine::EndPlayMap()
 	TArray<FString> Paths;
 	if (LeakedObjects.Num() > 0)
 	{
-		Paths = FindAndPrintStaleReferencesToObjects(LeakedObjects, EPrintStaleReferencesOptions::Ensure);
+		Paths = FReferenceChainSearch::FindAndPrintStaleReferencesToObjects(LeakedObjects, EPrintStaleReferencesOptions::Ensure);
 	}
 	for (int32 i = 0; i < LeakedObjects.Num(); ++i)
 	{
diff --git a/Engine/Source/Runtime/Engine/Classes/Engine/Engine.h b/Engine/Source/Runtime/Engine/Classes/Engine/Engine.h
--- a/Engine/Source/Runtime/Engine/Classes/Engine/Engine.h
+++ b/Engine/Source/Runtime/Engine/Classes/Engine/Engine.h
@@ -3341,7 +3341,7 @@ public:
 	 * @param ObjectToFindReferencesTo World or its package (or any object from the world package that should've been destroyed)
 	 * @param Verbosity Verbosity (can be fatal or non-fatal) with which to print the error message with
 	 */
-	UE_DEPRECATED(5.1, "Please use FindAndPrintStaleReferencesToObject overload that takes EPrintStaleReferencesOptions parameter")
+	UE_DEPRECATED(5.1, "Please use FReferenceChainSearch::FindAndPrintStaleReferencesToObject")
 	static void FindAndPrintStaleReferencesToObject(UObject* ObjectToFindReferencesTo, ELogVerbosity::Type Verbosity);
 
 	/**
@@ -3349,7 +3349,9 @@ public:
 	 * @param ObjectToFindReferencesTo World or its package (or any object from the world package that should've been destroyed)
 	 * @param Options Determines how the stale references messages should be logged
 	 */
+	UE_DEPRECATED(5.3, "Please use FReferenceChainSearch::FindAndPrintStaleReferencesToObject")
 	static FString FindAndPrintStaleReferencesToObject(UObject* ObjectToFindReferencesTo, EPrintStaleReferencesOptions Options);
+	UE_DEPRECATED(5.3, "Please use FReferenceChainSearch::FindAndPrintStaleReferencesToObjects")
 	static TArray<FString> FindAndPrintStaleReferencesToObjects(TConstArrayView<UObject*> ObjectsToFindReferencesTo, EPrintStaleReferencesOptions Options);
 
 	FWorldContext& CreateNewWorldContext(EWorldType::Type WorldType);
diff --git a/Engine/Source/Runtime/Engine/Private/LevelStreaming.cpp b/Engine/Source/Runtime/Engine/Private/LevelStreaming.cpp
--- a/Engine/Source/Runtime/Engine/Private/LevelStreaming.cpp
+++ b/Engine/Source/Runtime/Engine/Private/LevelStreaming.cpp
@@ -1574,7 +1574,7 @@ bool ULevelStreaming::RequestLevel(UWorld* PersistentWorld, bool bAllowLevelLoad
 				UE_LOG(LogLevelStreaming, Error, TEXT("World exists but PersistentLevel doesn't for %s, most likely caused by reference to world of unloaded level and GC setting reference to null while keeping world object"), *World->GetOutermost()->GetName());
 				UE_LOG(LogLevelStreaming, Error, TEXT("Most likely caused by reference to world of unloaded level and GC setting reference to null while keeping world object. Referenced by:"));
 
-				UEngine::FindAndPrintStaleReferencesToObject(World, UObjectBaseUtility::IsPendingKillEnabled() ? EPrintStaleReferencesOptions::Fatal : (EPrintStaleReferencesOptions::Error | EPrintStaleReferencesOptions::Ensure));
+				FReferenceChainSearch::FindAndPrintStaleReferencesToObject(World, UObjectBaseUtility::IsPendingKillEnabled() ? EPrintStaleReferencesOptions::Fatal : (EPrintStaleReferencesOptions::Error | EPrintStaleReferencesOptions::Ensure));
 
 				return false;
 			}
diff --git a/Engine/Source/Runtime/Engine/Private/UnrealEngine.cpp b/Engine/Source/Runtime/Engine/Private/UnrealEngine.cpp
--- a/Engine/Source/Runtime/Engine/Private/UnrealEngine.cpp
+++ b/Engine/Source/Runtime/Engine/Private/UnrealEngine.cpp
@@ -15980,7 +15980,7 @@ void UEngine::CheckAndHandleStaleWorldObjectReferences(FWorldContext* WorldConte
 			}
 
 			UE_LOG(LogLoad, Error, TEXT("Dumping reference chains:"));
-			FindAndPrintStaleReferencesToObjects(LeakedObjects, Options);
+			FReferenceChainSearch::FindAndPrintStaleReferencesToObjects(LeakedObjects, Options);
 		}
 
 		if (!UObjectBaseUtility::IsPendingKillEnabled())
@@ -16027,7 +16027,7 @@ void UEngine::FindAndPrintStaleReferencesToObject(UObject* ObjectToFindReference
 		checkf(false, TEXT("Unsupported Verbosity: %d"), Verbosity);
 		break;
 	}
-	FindAndPrintStaleReferencesToObject(ObjectToFindReferencesTo, Options);
+	FReferenceChainSearch::FindAndPrintStaleReferencesToObject(ObjectToFindReferencesTo, Options);
 }
 
 FString UEngine::FindAndPrintStaleReferencesToObject(UObject* ObjectToFindReferencesTo, EPrintStaleReferencesOptions Options)
