Add verbose logging to LWI manager to track down crash.

diff --git a/Engine/Source/Runtime/Engine/Classes/GameFramework/LightWeightInstanceSubsystem.h b/Engine/Source/Runtime/Engine/Classes/GameFramework/LightWeightInstanceSubsystem.h
--- a/Engine/Source/Runtime/Engine/Classes/GameFramework/LightWeightInstanceSubsystem.h
+++ b/Engine/Source/Runtime/Engine/Classes/GameFramework/LightWeightInstanceSubsystem.h
@@ -8,7 +8,7 @@
 
 #include "Actor.h"
 
-ENGINE_API DECLARE_LOG_CATEGORY_EXTERN(LogLightWeightInstance, Log, Warning);
+ENGINE_API DECLARE_LOG_CATEGORY_EXTERN(LogLightWeightInstance, Log, All);
 
 class FAutoConsoleVariableRef;
 
diff --git a/Engine/Source/Runtime/Engine/Private/LightWeightInstanceStaticMeshManager.cpp b/Engine/Source/Runtime/Engine/Private/LightWeightInstanceStaticMeshManager.cpp
--- a/Engine/Source/Runtime/Engine/Private/LightWeightInstanceStaticMeshManager.cpp
+++ b/Engine/Source/Runtime/Engine/Private/LightWeightInstanceStaticMeshManager.cpp
@@ -1,10 +1,12 @@
 // Copyright Epic Games, Inc. All Rights Reserved.
 
 #include "GameFramework/LightWeightInstanceStaticMeshManager.h"
+
 #include "Components/HierarchicalInstancedStaticMeshComponent.h"
 #include "Components/PrimitiveComponent.h"
-#include "Net/UnrealNetwork.h"
 #include "Elements/SMInstance/SMInstanceElementId.h"
+#include "GameFramework/LightWeightInstanceSubsystem.h"
+#include "Net/UnrealNetwork.h"
 #include "Templates/Greater.h"
 
 #include UE_INLINE_GENERATED_CPP_BY_NAME(LightWeightInstanceStaticMeshManager)
@@ -103,11 +105,18 @@ void ALightWeightInstanceStaticMeshManager::UpdateDataAtIndex(FLWIData* InData,
 void ALightWeightInstanceStaticMeshManager::AddInstanceToRendering(int32 DataIndex)
 {
 	//cancel any pending deletes
-	DataIndicesToBeDeleted.RemoveSingle(DataIndex);
+	const bool bPendingDeleteCancelled = DataIndicesToBeDeleted.RemoveSingle(DataIndex) > 0;
 
 	// The rendering indices are tightly packed so we know it's going on the end of the array
 	const int32 RenderingIdx = RenderingIndicesToDataIndices.Add(DataIndex);
 
+	UE_LOG(LogLightWeightInstance, Verbose,
+		TEXT("ALightWeightInstanceStaticMeshManager::AddInstanceToRendering - manager [ %s ] adding instance at data index [ %d ] rendering index [ %d ], did cancel pending delete? [ %s ]"),
+		*GetActorNameOrLabel(),
+		DataIndex,
+		RenderingIdx,
+		bPendingDeleteCancelled ? TEXT("true") : TEXT("false"));
+
 	// Now that we know the rendering index we can fill in the other side of the map
 	if (DataIndex >= DataIndicesToRenderingIndices.Num())
 	{
@@ -145,6 +154,11 @@ void ALightWeightInstanceStaticMeshManager::RemoveInstanceFromRendering(int32 Da
 		}
 
 		DataIndicesToBeDeleted.AddUnique(DataIndex);
+
+		UE_LOG(LogLightWeightInstance, Verbose,
+			TEXT("ALightWeightInstanceStaticMeshManager::RemoveInstanceFromRendering - manager [ %s ] removing instance at data index [ %d ]"),
+			*GetActorNameOrLabel(),
+			DataIndex);
 	}
 }
 
@@ -154,6 +168,12 @@ void ALightWeightInstanceStaticMeshManager::PostRemoveInstanceFromRendering()
 	for (int32 DataIndex : DataIndicesToBeDeleted)
 	{
 		RenderingIndicesToBeDeleted.Add(DataIndicesToRenderingIndices[DataIndex]);
+
+		UE_LOG(LogLightWeightInstance, Verbose,
+			TEXT("ALightWeightInstanceStaticMeshManager::PostRemoveInstanceFromRendering - manager [ %s ] removing instance at data index [ %d ] rendering index [ %d ]"),
+			*GetActorNameOrLabel(),
+			DataIndex,
+			DataIndicesToRenderingIndices[DataIndex]);
 	}
 
 	RenderingIndicesToBeDeleted.Sort();
