ContentBundle: Filter out invalid actors (mismatching outer and external path)

#rnx

diff --git a/Engine/Source/Runtime/Engine/Private/WorldPartition/ContentBundle/ContentBundleEditor.cpp b/Engine/Source/Runtime/Engine/Private/WorldPartition/ContentBundle/ContentBundleEditor.cpp
--- a/Engine/Source/Runtime/Engine/Private/WorldPartition/ContentBundle/ContentBundleEditor.cpp
+++ b/Engine/Source/Runtime/Engine/Private/WorldPartition/ContentBundle/ContentBundleEditor.cpp
@@ -56,6 +56,23 @@ void FContentBundleEditor::DoInjectContent()
 	{
 		UWorldPartition* WorldPartition = GetInjectedWorld()->GetWorldPartition();
 		UWorldPartition::FContainerRegistrationParams RegistrationsParams(*ActorDescContainerPackage);
+		
+		const FTopLevelAssetPath InjectedWorldAssetPath = FSoftObjectPath(GetInjectedWorld()).GetAssetPath();
+		RegistrationsParams.FilterActorDescFunc = [&](const FWorldPartitionActorDesc* ActorDesc)
+		{
+			if (ActorDesc->GetActorSoftPath().GetAssetPath() != InjectedWorldAssetPath)
+			{
+				UE_LOG(LogContentBundle, Warning, TEXT("%s Container '%s' : Actor '%s' object path doesn't match injected world path '%s'"),
+					*ContentBundle::Log::MakeDebugInfoString(*this),
+					*ActorDescContainerPackage,
+					*ActorDesc->GetActorSoftPath().ToString(),
+					*InjectedWorldAssetPath.ToString());
+				return false;
+			}
+
+			return true;
+		};
+
 		ActorDescContainer = WorldPartition->RegisterActorDescContainer(RegistrationsParams);
 		if (ActorDescContainer.IsValid())
 		{
