Make sure to test that the temporary actordesc created for unsaved actors pass the IsRuntimeRelevant test before registering it for streaming generation
Fix some actors being present in the streaming grid before they are saved

diff --git a/Engine/Source/Runtime/Engine/Private/WorldPartition/WorldPartitionStreamingGeneration.cpp b/Engine/Source/Runtime/Engine/Private/WorldPartition/WorldPartitionStreamingGeneration.cpp
--- a/Engine/Source/Runtime/Engine/Private/WorldPartition/WorldPartitionStreamingGeneration.cpp
+++ b/Engine/Source/Runtime/Engine/Private/WorldPartition/WorldPartitionStreamingGeneration.cpp
@@ -369,46 +369,36 @@ class FWorldPartitionStreamingGenerator
 
 		// Consider all actors of a /Temp/ container package as Unsaved because loading them from disk will fail (Outer world name mismatch)
 		const bool bIsTempContainerPackage = FPackageName::IsTempPackage(InContainer->GetPackage()->GetName());
-		
-		// Test whether an actor is editor only. Will fallback to the actor descriptor only if the actor is not loaded
-		auto IsActorEditorOnly = [](const FWorldPartitionActorDesc* ActorDesc, const FActorContainerID& ContainerID) -> bool
+
+		// Create an actor descriptor view for the specified actor (modified or unsaved actors)
+		auto GetModifiedActorDesc = [this](AActor* InActor, const UActorDescContainer* InContainer) -> FWorldPartitionActorDesc*
 		{
-			if (ActorDesc->IsRuntimeRelevant(ContainerID))
-			{
-				if (ActorDesc->IsLoaded())
-				{
-					return ActorDesc->GetActor()->IsEditorOnly();
-				}
-				else
-				{
-					return ActorDesc->GetActorIsEditorOnly();
-				}
-			}
-			return true;
+			FWorldPartitionActorDesc* ModifiedActorDesc = ModifiedActorsDescList->AddActor(InActor);
+
+			// Pretend that this actor descriptor belongs to the original container, even if it's not present. It's essentially a proxy
+			// descriptor on top an existing one and at this point no code should require to access the container to resolve it anyways.
+			ModifiedActorDesc->SetContainer(const_cast<UActorDescContainer*>(InContainer), InActor->GetWorld());
+
+			return ModifiedActorDesc;
 		};
 
-		auto IsFilteredActorClass = [this](const FWorldPartitionActorDesc* ActorDesc)
+		// Test whether an actor should be included in the ActorDescViewMap.
+		auto ShouldRegisterActorDesc = [this](const FWorldPartitionActorDesc* ActorDesc, const FActorContainerID& ContainerID)
 		{
 			for (UClass* FilteredClass : FilteredClasses)
 			{
 				if (ActorDesc->GetActorNativeClass()->IsChildOf(FilteredClass))
 				{
-					return true;
+					return false;
 				}
 			}
-			return false;
-		};
-
-		// Create an actor descriptor view for the specified actor (modified or unsaved actors)
-		auto GetModifiedActorDesc = [this](AActor* InActor, const UActorDescContainer* InContainer) -> FWorldPartitionActorDesc*
-		{
-			FWorldPartitionActorDesc* ModifiedActorDesc = ModifiedActorsDescList->AddActor(InActor);
 
-			// Pretend that this actor descriptor belongs to the original container, even if it's not present. It's essentially a proxy
-			// descriptor on top an existing one and at this point no code should require to access the container to resolve it anyways.
-			ModifiedActorDesc->SetContainer(const_cast<UActorDescContainer*>(InContainer), InActor->GetWorld());
+			if (!ActorDesc->IsRuntimeRelevant(ContainerID))
+			{
+				return false;
+			}
 
-			return ModifiedActorDesc;
+			return ActorDesc->IsLoaded() ? !ActorDesc->GetActor()->IsEditorOnly() : !ActorDesc->GetActorIsEditorOnly();
 		};
 
 		// Register the actor descriptor view
@@ -427,7 +417,7 @@ class FWorldPartitionStreamingGenerator
 		TMap<FGuid, FGuid> ContainerGuidsRemap;
 		for (FActorDescList::TConstIterator<> ActorDescIt(InContainer); ActorDescIt; ++ActorDescIt)
 		{
-			if (!IsActorEditorOnly(*ActorDescIt, InContainerID) && !IsFilteredActorClass(*ActorDescIt))
+			if (ShouldRegisterActorDesc(*ActorDescIt, InContainerID))
 			{
 				// Handle unsaved actors
 				if (AActor* Actor = ActorDescIt->GetActor())
@@ -466,7 +456,10 @@ class FWorldPartitionStreamingGenerator
 				if (IsValid(Actor) && Actor->IsPackageExternal() && Actor->IsMainPackageActor() && !Actor->IsEditorOnly() && (Actor->GetContentBundleGuid() == InContainer->GetContentBundleGuid()) && !InContainer->GetActorDesc(Actor->GetActorGuid()))
 				{
 					FWorldPartitionActorDescView ModifiedActorDescView = GetModifiedActorDesc(Actor, InContainer);
-					RegisterActorDescView(Actor->GetActorGuid(), ModifiedActorDescView);
+					if (ShouldRegisterActorDesc(ModifiedActorDescView.GetActorDesc(), InContainerID))
+					{
+						RegisterActorDescView(Actor->GetActorGuid(), ModifiedActorDescView);
+					}
 				}
 			}
 		}
