Fix deadlock when remapping skeleton for set of animations which contain (additive) inter-dependencies. Meant that the eval/modify-lock held by the outer brackets would block the additive animation from compressing - which the remapping would wait for to finish.
#fix handle animation sequence in a separate loop and wait for all to-be-remapped animations to finish compression

diff --git a/Engine/Source/Runtime/Engine/Private/AnimationAsset.cpp b/Engine/Source/Runtime/Engine/Private/AnimationAsset.cpp
--- a/Engine/Source/Runtime/Engine/Private/AnimationAsset.cpp
+++ b/Engine/Source/Runtime/Engine/Private/AnimationAsset.cpp
@@ -13,6 +13,7 @@
 #include "Animation/BlendSpace.h"
 #include "Animation/PoseAsset.h"
 #include "Animation/AnimNodeBase.h"
+#include "Animation/AnimationSequenceCompiler.h"
 
 #if WITH_EDITOR
 #include "Misc/DataValidation.h"
@@ -445,6 +446,8 @@ bool UAnimationAsset::ReplaceSkeleton(USkeleton* NewSkeleton, bool bConvertSpace
 		}
 		if (GetAllAnimationSequencesReferred(AnimAssetsToReplace))
 		{
+			TArray<UAnimSequence*> Sequences;
+			
 			//Firstly need to remap
 			for (UAnimationAsset* AnimAsset : AnimAssetsToReplace)
 			{
@@ -459,21 +462,27 @@ bool UAnimationAsset::ReplaceSkeleton(USkeleton* NewSkeleton, bool bConvertSpace
 				// raw animation data itself.
 				if (UAnimSequence* Sequence = Cast<UAnimSequence>(AnimAsset))
 				{
-					Sequence->GetController().OpenBracket(LOCTEXT("ReplaceSkeleton_Bracket", "Replacing USkeleton"));
+					Sequences.Add(Sequence);				
 				}
+				else
+				{
+					// these two are different functions for now
+					// technically if you have implementation for Remap, it will also set skeleton 
+					AnimAsset->RemapTracksToNewSkeleton(NewSkeleton, bConvertSpaces);
+				}
+			}
 
-				// these two are different functions for now
-				// technically if you have implementation for Remap, it will also set skeleton 
-				AnimAsset->RemapTracksToNewSkeleton(NewSkeleton, bConvertSpaces);
+			UE::Anim::FAnimSequenceCompilingManager::Get().FinishCompilation(Sequences);
+			for (UAnimSequence* Sequence : Sequences)
+			{
+				Sequence->GetController().OpenBracket(LOCTEXT("ReplaceSkeleton_Bracket", "Replacing USkeleton"));
+				Sequence->RemapTracksToNewSkeleton(NewSkeleton, bConvertSpaces);
 			}
 
 			//Second need to process anim sequences themselves. This is done in two stages as additives can rely on other animations.
-			for (UAnimationAsset* AnimAsset : AnimAssetsToReplace)
+			for (UAnimSequence* Sequence : Sequences)
 			{
-				if (UAnimSequence* Seq = Cast<UAnimSequence>(AnimAsset))
-				{
-					Seq->GetController().CloseBracket();
-				}
+				Sequence->GetController().CloseBracket();
 			}
 		}
 
