Fix for texture arrays in cooked editors referencing invalid UPROPERTYs

diff --git a/Engine/Source/Runtime/Engine/Private/Texture.cpp b/Engine/Source/Runtime/Engine/Private/Texture.cpp
--- a/Engine/Source/Runtime/Engine/Private/Texture.cpp
+++ b/Engine/Source/Runtime/Engine/Private/Texture.cpp
@@ -1559,9 +1559,14 @@ bool UTexture::IsPossibleToStream() const
 	}
 
 	#if WITH_EDITORONLY_DATA
-	if ( MipGenSettings == TMGS_NoMipmaps )
+	if (!GetOutermost()->HasAnyPackageFlags(PKG_FilterEditorOnly))
 	{
-		return false;
+		// MipGenSettings is WITH_EDITORONLY_DATA, if we touch it here without checking we always get default values rather than
+		// what users specified
+		if ( MipGenSettings == TMGS_NoMipmaps ) 
+		{
+			return false;
+		}
 	}
 
 	// VirtualTextureStreaming can be true here and we will still stream if VT is disabled
