Fixes AsyncCompute passes not showing up in ProfileGPU

diff --git a/Engine/Source/Runtime/Engine/Private/UnrealEngine.cpp b/Engine/Source/Runtime/Engine/Private/UnrealEngine.cpp
--- a/Engine/Source/Runtime/Engine/Private/UnrealEngine.cpp
+++ b/Engine/Source/Runtime/Engine/Private/UnrealEngine.cpp
@@ -5769,7 +5769,11 @@ bool UEngine::HandleProfileCommand( const TCHAR* Cmd, FOutputDevice& Ar )
 
 		if (!GTriggerGPUHitchProfile)
 		{
-			GTriggerGPUProfile = true;
+			ENQUEUE_RENDER_COMMAND(HandleProfileCommand)(
+				[](FRHICommandListImmediate& RHICmdList)
+			{
+				GTriggerGPUProfile = true;
+			});
 			Ar.Logf(TEXT("Profiling the next GPU frame"));
 		}
 		else
@@ -5802,7 +5806,11 @@ bool UEngine::HandleProfileGPUCommand( const TCHAR* Cmd, FOutputDevice& Ar )
 	{
 		if (!GTriggerGPUHitchProfile)
 		{
-			GTriggerGPUProfile = true;
+			ENQUEUE_RENDER_COMMAND(HandleProfileGPUCommand)(
+				[](FRHICommandListImmediate& RHICmdList)
+			{
+				GTriggerGPUProfile = true;
+			});
 			Ar.Logf(TEXT("Profiling the next GPU frame"));
 		}
 		else
diff --git a/Engine/Source/Runtime/RenderCore/Private/RenderGraphPrivate.h b/Engine/Source/Runtime/RenderCore/Private/RenderGraphPrivate.h
--- a/Engine/Source/Runtime/RenderCore/Private/RenderGraphPrivate.h
+++ b/Engine/Source/Runtime/RenderCore/Private/RenderGraphPrivate.h
@@ -257,7 +257,7 @@ FORCEINLINE bool IsRenderPassMergeEnabled()
 
 FORCEINLINE bool IsAsyncComputeSupported()
 {
-	return GRDGAsyncCompute > 0 && !IsImmediateMode() && GSupportsEfficientAsyncCompute && GRHISupportsSeparateDepthStencilCopyAccess;
+	return GRDGAsyncCompute > 0 && !IsImmediateMode() && GSupportsEfficientAsyncCompute && GRHISupportsSeparateDepthStencilCopyAccess && !GTriggerGPUProfile;
 }
 
 bool IsDumpingRDGResources();
