Adapted from PR#10384: instead of fixing the minimum num of segments computation like done in the PR, just removed it because the value is only used for a call to InternalDrawDebugCircle, that is already doing this job internally.

PR #10384: Fix `DrawDebugCircle` to draw correct number of segments.
#test: drawn debug circles with various number of segments (0,4,10,...). The expected number was used (min of 4 or the asked number)

diff --git a/Engine/Source/Runtime/Engine/Private/DrawDebugHelpers.cpp b/Engine/Source/Runtime/Engine/Private/DrawDebugHelpers.cpp
--- a/Engine/Source/Runtime/Engine/Private/DrawDebugHelpers.cpp
+++ b/Engine/Source/Runtime/Engine/Private/DrawDebugHelpers.cpp
@@ -372,8 +372,6 @@ void DrawDebugCircle(const UWorld* InWorld, const FMatrix& TransformMatrix, floa
 		{
 			const float LineLifeTime = GetDebugLineLifeTime(LineBatcher, LifeTime, bPersistentLines);
 
-			// Need at least 4 segments
-			Segments = FMath::Max((Segments - 2) / 2, 4);
 			InternalDrawDebugCircle(InWorld, TransformMatrix, Radius, Segments, Color, bPersistentLines, LifeTime, DepthPriority, Thickness);
 
 			if (bDrawAxis)
