Integration of PR #10375: Fix DrawOnscreenDebugMessages function which was returning the wrong value

#test Ran Lyra in PIE, with debug message displayed. All were positionned correctly. Valided that returned position was matching the description from the function.

diff --git a/Engine/Source/Runtime/Engine/Private/UnrealEngine.cpp b/Engine/Source/Runtime/Engine/Private/UnrealEngine.cpp
--- a/Engine/Source/Runtime/Engine/Private/UnrealEngine.cpp
+++ b/Engine/Source/Runtime/Engine/Private/UnrealEngine.cpp
@@ -12169,10 +12169,11 @@ float DrawMapWarnings(UWorld* World, FViewport* Viewport, FCanvas* Canvas, UCanv
 */
 float UEngine::DrawOnscreenDebugMessages(UWorld* World, FViewport* Viewport, FCanvas* Canvas, UCanvas* CanvasObject, float MessageX, float MessageY)
 {
+	int32 YPos = MessageY;
+
 #if !(UE_BUILD_SHIPPING || UE_BUILD_TEST)
 	static TFrameValue<bool> HasUpdatedScreenDebugMessages;
 
-	int32 YPos = MessageY;
 	const float InverseDPI = 1.f / Canvas->GetDPIScale();
 	const int32 MaxYPos = CanvasObject ? CanvasObject->SizeY : 700;
 	{
@@ -12182,60 +12183,60 @@ float UEngine::DrawOnscreenDebugMessages(UWorld* World, FViewport* Viewport, FCa
 		if (PriorityScreenMessages.Num() > 0)
 		{
 			FCanvasTextItem MessageTextItem(FVector2D(0, 0), FText::GetEmpty(), GetSmallFont(), FLinearColor::White);
-		MessageTextItem.EnableShadow(FLinearColor::Black);
+			MessageTextItem.EnableShadow(FLinearColor::Black);
 			for (int32 PrioIndex = PriorityScreenMessages.Num() - 1; PrioIndex >= 0; PrioIndex--)
-		{
-				FScreenMessageString& Message = PriorityScreenMessages[PrioIndex];
-			if (YPos < MaxYPos)
-			{
-				MessageTextItem.Text = FText::FromString(Message.ScreenMessage);
-				MessageTextItem.SetColor(Message.DisplayColor);
-				MessageTextItem.Scale = Message.TextScale;
-				Canvas->DrawItem(MessageTextItem, FVector2D(MessageX, YPos));
-				YPos += MessageTextItem.DrawnSize.Y * InverseDPI * 1.15f;
-			}
-			if (!HasUpdatedScreenDebugMessages.IsSet())
 			{
-				Message.CurrentTimeDisplayed += World->GetDeltaSeconds();
-				if (Message.CurrentTimeDisplayed >= Message.TimeToDisplay)
+				FScreenMessageString& Message = PriorityScreenMessages[PrioIndex];
+				if (YPos < MaxYPos)
+				{
+					MessageTextItem.Text = FText::FromString(Message.ScreenMessage);
+					MessageTextItem.SetColor(Message.DisplayColor);
+					MessageTextItem.Scale = Message.TextScale;
+					Canvas->DrawItem(MessageTextItem, FVector2D(MessageX, YPos));
+					YPos += MessageTextItem.DrawnSize.Y * InverseDPI * 1.15f;
+				}
+				if (!HasUpdatedScreenDebugMessages.IsSet())
 				{
+					Message.CurrentTimeDisplayed += World->GetDeltaSeconds();
+					if (Message.CurrentTimeDisplayed >= Message.TimeToDisplay)
+					{
 						PriorityScreenMessages.RemoveAt(PrioIndex);
+					}
 				}
 			}
 		}
-	}
 
 		if (ScreenMessages.Num() > 0)
-	{
+		{
 			FCanvasTextItem MessageTextItem(FVector2D(0, 0), FText::GetEmpty(), GetSmallFont(), FLinearColor::White);
-		MessageTextItem.EnableShadow(FLinearColor::Black);
+			MessageTextItem.EnableShadow(FLinearColor::Black);
 			for (TMap<int32, FScreenMessageString>::TIterator MsgIt(ScreenMessages); MsgIt; ++MsgIt)
-		{
-			FScreenMessageString& Message = MsgIt.Value();
-			if (YPos < MaxYPos)
-			{
-				MessageTextItem.Text = FText::FromString(Message.ScreenMessage);
-				MessageTextItem.SetColor(Message.DisplayColor);
-				MessageTextItem.Scale = Message.TextScale;
-				Canvas->DrawItem(MessageTextItem, FVector2D(MessageX, YPos));
-				YPos += MessageTextItem.DrawnSize.Y * InverseDPI * 1.15f;
-			}
-			if (!HasUpdatedScreenDebugMessages.IsSet())
 			{
-				Message.CurrentTimeDisplayed += World->GetDeltaSeconds();
-				if (Message.CurrentTimeDisplayed >= Message.TimeToDisplay)
+				FScreenMessageString& Message = MsgIt.Value();
+				if (YPos < MaxYPos)
 				{
-					MsgIt.RemoveCurrent();
+					MessageTextItem.Text = FText::FromString(Message.ScreenMessage);
+					MessageTextItem.SetColor(Message.DisplayColor);
+					MessageTextItem.Scale = Message.TextScale;
+					Canvas->DrawItem(MessageTextItem, FVector2D(MessageX, YPos));
+					YPos += MessageTextItem.DrawnSize.Y * InverseDPI * 1.15f;
+				}
+				if (!HasUpdatedScreenDebugMessages.IsSet())
+				{
+					Message.CurrentTimeDisplayed += World->GetDeltaSeconds();
+					if (Message.CurrentTimeDisplayed >= Message.TimeToDisplay)
+					{
+						MsgIt.RemoveCurrent();
+					}
 				}
 			}
 		}
 	}
-	}
 	// Flag variable that the update has already been done this frame
 	HasUpdatedScreenDebugMessages = true;
 #endif //#if !(UE_BUILD_SHIPPING || UE_BUILD_TEST)
 
-	return MessageY;
+	return YPos;
 }
 
 
