Add a cvar to control calls to FlushNetDormancy in cluster union component and do not call it by default

diff --git a/Engine/Source/Runtime/Engine/Private/PhysicsEngine/ClusterUnionComponent.cpp b/Engine/Source/Runtime/Engine/Private/PhysicsEngine/ClusterUnionComponent.cpp
--- a/Engine/Source/Runtime/Engine/Private/PhysicsEngine/ClusterUnionComponent.cpp
+++ b/Engine/Source/Runtime/Engine/Private/PhysicsEngine/ClusterUnionComponent.cpp
@@ -24,6 +24,9 @@ namespace
 	bool bUseClusterUnionAccelerationStructure = false;
 	FAutoConsoleVariableRef CVarUseClusterUnionAccelerationStructure(TEXT("ClusterUnion.UseAccelerationStructure"), bUseClusterUnionAccelerationStructure, TEXT("Whether component level sweeps and overlaps against cluster unions should use an acceleration structure instead."));
 
+	bool bClusterUnionCallFlushNetDormancy = false;
+	FAutoConsoleVariableRef CVarClusterUnionCallFlushNetDormanc(TEXT("clusterunion.callflushnetdormancy"), bClusterUnionCallFlushNetDormancy, TEXT("Whether or not to call FlushNetDormancy"));
+
 	// TODO: Should this be exposed in Chaos instead?
 	using FAccelerationStructure = Chaos::TAABBTree<FExternalSpatialAccelerationPayload, Chaos::TAABBTreeLeafArray<FExternalSpatialAccelerationPayload>>;
 
@@ -673,9 +676,12 @@ void UClusterUnionComponent::HandleAddOrModifiedClusteredComponent(UPrimitiveCom
 			ReplicatedProxy->SetParticleChildToParent(Kvp.Key, Kvp.Value);
 		}
 
-		if (AActor* Owner = ChangedComponent->GetOwner())
+		if (bClusterUnionCallFlushNetDormancy)
 		{
-			Owner->FlushNetDormancy();
+			if (AActor* Owner = ChangedComponent->GetOwner())
+			{
+				Owner->FlushNetDormancy();
+			}
 		}
 	}
 
@@ -787,7 +793,10 @@ void UClusterUnionComponent::HandleRemovedClusteredComponent(UPrimitiveComponent
 		}
 	}
 
-	Owner->FlushNetDormancy();
+	if (bClusterUnionCallFlushNetDormancy)
+	{
+		Owner->FlushNetDormancy();
+	}
 }
 
 void UClusterUnionComponent::OnRep_RigidState()
