#UE Lower verbosity, for now

diff --git a/Engine/Source/Runtime/Engine/Private/WorldPartition/ContentBundle/ContentBundleEditor.cpp b/Engine/Source/Runtime/Engine/Private/WorldPartition/ContentBundle/ContentBundleEditor.cpp
--- a/Engine/Source/Runtime/Engine/Private/WorldPartition/ContentBundle/ContentBundleEditor.cpp
+++ b/Engine/Source/Runtime/Engine/Private/WorldPartition/ContentBundle/ContentBundleEditor.cpp
@@ -62,7 +62,7 @@ void FContentBundleEditor::DoInjectContent()
 		{
 			if (ActorDesc->GetActorSoftPath().GetAssetPath() != InjectedWorldAssetPath)
 			{
-				UE_LOG(LogContentBundle, Warning, TEXT("%s Container '%s' : Actor '%s' object path doesn't match injected world path '%s'"),
+				UE_LOG(LogContentBundle, Display, TEXT("%s Container '%s' : Actor '%s' object path doesn't match injected world path '%s'"),
 					*ContentBundle::Log::MakeDebugInfoString(*this),
 					*ActorDescContainerPackage,
 					*ActorDesc->GetActorSoftPath().ToString(),
