Handle null SM in UStaticMeshComponent::HasValidNaniteData()

diff --git a/Engine/Source/Runtime/Engine/Private/Components/StaticMeshComponent.cpp b/Engine/Source/Runtime/Engine/Private/Components/StaticMeshComponent.cpp
--- a/Engine/Source/Runtime/Engine/Private/Components/StaticMeshComponent.cpp
+++ b/Engine/Source/Runtime/Engine/Private/Components/StaticMeshComponent.cpp
@@ -2331,13 +2331,8 @@ const Nanite::FResources* UStaticMeshComponent::GetNaniteResources() const
 
 bool UStaticMeshComponent::HasValidNaniteData() const
 {
-	const FStaticMeshRenderData* RenderData = GetStaticMesh()->GetRenderData();
-	if (RenderData != nullptr)
-	{
-		return RenderData->HasValidNaniteData();
-	}
-
-	return false;
+	const Nanite::FResources* NaniteResources = UStaticMeshComponent::GetNaniteResources();
+	return NaniteResources != nullptr ? NaniteResources->PageStreamingStates.Num() > 0 : false;
 }
 
 bool UStaticMeshComponent::UseNaniteOverrideMaterials(bool bDoingMaterialAudit) const
