Cluster unions : do not add cluster union to acceleration structures in SyncBodies if they do not have any children
- this avoid the AABBTree no bounds warning

diff --git a/Engine/Source/Runtime/Engine/Private/PhysicsEngine/Experimental/PhysScene_Chaos.cpp b/Engine/Source/Runtime/Engine/Private/PhysicsEngine/Experimental/PhysScene_Chaos.cpp
--- a/Engine/Source/Runtime/Engine/Private/PhysicsEngine/Experimental/PhysScene_Chaos.cpp
+++ b/Engine/Source/Runtime/Engine/Private/PhysicsEngine/Experimental/PhysScene_Chaos.cpp
@@ -2143,7 +2143,12 @@ void FPhysScene_Chaos::OnSyncBodies(Chaos::FPhysicsSolverBase* Solver)
 					PendingTransforms.Add(FPhysScenePendingComponentTransform_Chaos(ParentComponent, MoveBy, NewTransform.GetRotation(), DirtyParticle->GetWakeEvent()));
 				}
 
-				Interface->AddToSpatialAcceleration({ &Handle, 1 }, Outer->GetSpacialAcceleration());
+				// make sure we have at least a child to be added to the acceleration structure 
+				// this avoid the invalid bounds to cause the particle to be added to the global acceleration structure array
+				if (ParentComponent->GetNumChildrenComponents() > 0)
+				{
+					Interface->AddToSpatialAcceleration({ &Handle, 1 }, Outer->GetSpacialAcceleration());
+				}
 				ParentComponent->SyncClusterUnionFromProxy();
 				DirtyParticle->ClearEvents();
 			}
