Validate we have vertex data for the current MinLOD for InstancedStaticMeshes before creating SceneProxy, as we do for StaticMeshes.

diff --git a/Engine/Source/Runtime/Engine/Private/InstancedStaticMesh.cpp b/Engine/Source/Runtime/Engine/Private/InstancedStaticMesh.cpp
--- a/Engine/Source/Runtime/Engine/Private/InstancedStaticMesh.cpp
+++ b/Engine/Source/Runtime/Engine/Private/InstancedStaticMesh.cpp
@@ -2704,6 +2704,16 @@ FPrimitiveSceneProxy* UInstancedStaticMeshComponent::CreateStaticMeshSceneProxy(
 		return ::new Nanite::FSceneProxy(NaniteMaterials, this);
 	}
 
+	// Validate the LOD resources here
+	const FStaticMeshLODResourcesArray& LODResources = GetStaticMesh()->GetRenderData()->LODResources;
+	const int32 SMCurrentMinLOD = GetStaticMesh()->GetMinLODIdx();
+	const int32 EffectiveMinLOD = bOverrideMinLOD ? MinLOD : SMCurrentMinLOD;
+	if (LODResources.Num() == 0 || LODResources[FMath::Clamp<int32>(EffectiveMinLOD, 0, LODResources.Num() - 1)].VertexBuffers.StaticMeshVertexBuffer.GetNumVertices() == 0)
+	{
+		UE_LOG(LogStaticMesh, Verbose, TEXT("Skipping CreateSceneProxy for InstancedStaticMeshComponent %s (LOD problems)"), *GetFullName());
+		return nullptr;
+	}
+
 	return ::new FInstancedStaticMeshSceneProxy(this, GetWorld()->GetFeatureLevel());
 }
 
diff --git a/Engine/Source/Runtime/Engine/Private/StaticMeshRender.cpp b/Engine/Source/Runtime/Engine/Private/StaticMeshRender.cpp
--- a/Engine/Source/Runtime/Engine/Private/StaticMeshRender.cpp
+++ b/Engine/Source/Runtime/Engine/Private/StaticMeshRender.cpp
@@ -280,6 +280,8 @@ FStaticMeshSceneProxy::FStaticMeshSceneProxy(UStaticMeshComponent* InComponent,
 		MaterialRelevance |= UMaterial::GetDefaultMaterial(MD_Surface)->GetRelevance(FeatureLevel);
 	}
 
+	check(FirstAvailableLOD < RenderData->LODResources.Num())
+
 	ClampedMinLOD = FMath::Clamp(EffectiveMinLOD, FirstAvailableLOD, RenderData->LODResources.Num() - 1);
 
 	SetWireframeColor(InComponent->GetWireframeColor());
