World Partition: try to redirect actor descriptors class when loading them.
Original code by [REDACTED_PERSON].

diff --git a/Engine/Source/Runtime/Engine/Private/WorldPartition/WorldPartitionActorDescArchive.cpp b/Engine/Source/Runtime/Engine/Private/WorldPartition/WorldPartitionActorDescArchive.cpp
--- a/Engine/Source/Runtime/Engine/Private/WorldPartition/WorldPartitionActorDescArchive.cpp
+++ b/Engine/Source/Runtime/Engine/Private/WorldPartition/WorldPartitionActorDescArchive.cpp
@@ -6,6 +6,7 @@
 #include "WorldPartition/WorldPartitionActorDesc.h"
 #include "WorldPartition/WorldPartitionHelpers.h"
 #include "WorldPartition/WorldPartitionLog.h"
+#include "UObject/CoreRedirects.h"
 #include "UObject/UE5MainStreamObjectVersion.h"
 #include "UObject/FortniteNCBranchObjectVersion.h"
 #include "UObject/FortniteMainBranchObjectVersion.h"
@@ -59,6 +60,26 @@ PRAGMA_ENABLE_DEPRECATION_WARNINGS
 		InnerArchive << InActorDesc->NativeClass;
 	}
 
+	if (IsLoading())
+	{
+		auto TryRedirectClass = [](FTopLevelAssetPath& InOutClassPath)
+		{
+			if (InOutClassPath.IsValid())
+			{
+				FCoreRedirectObjectName ClassRedirect(InOutClassPath);
+				FCoreRedirectObjectName RedirectedClassRedirect = FCoreRedirects::GetRedirectedName(ECoreRedirectFlags::Type_Class, ClassRedirect);
+
+				if (ClassRedirect != RedirectedClassRedirect)
+				{
+					InOutClassPath = FTopLevelAssetPath(RedirectedClassRedirect.ToString());
+				}
+			}
+		};
+
+		TryRedirectClass(InActorDesc->NativeClass);
+		TryRedirectClass(InActorDesc->BaseClass);
+	}
+
 	// Get the class descriptor to do delta serialization
 	FWorldPartitionClassDescRegistry& ClassDescRegistry = FWorldPartitionClassDescRegistry::Get();
 	const FTopLevelAssetPath ClassPath = InActorDesc->BaseClass.IsValid() ? InActorDesc->BaseClass : InActorDesc->NativeClass;
@@ -94,4 +115,4 @@ FArchive& FActorDescArchive::operator<<(FSoftObjectPath& Value)
 
 	return *this;
 }
-#endif
\ No newline at end of file
+#endif
