NetDriver
* Added FWorldDelegate::OnNetDriverCreated delegate called after any NetDriver was created.

#rnx

diff --git a/Engine/Source/Runtime/Engine/Classes/Engine/World.h b/Engine/Source/Runtime/Engine/Classes/Engine/World.h
--- a/Engine/Source/Runtime/Engine/Classes/Engine/World.h
+++ b/Engine/Source/Runtime/Engine/Classes/Engine/World.h
@@ -4264,6 +4264,9 @@ public:
 	DECLARE_MULTICAST_DELEGATE_OneParam(FOnSeamlessTravelTransition, UWorld*);
 	static FOnSeamlessTravelTransition OnSeamlessTravelTransition;
 
+	DECLARE_MULTICAST_DELEGATE_TwoParams(FOnNetDriverCreated, UWorld*, UNetDriver*);
+	static FOnNetDriverCreated OnNetDriverCreated;
+
 	DECLARE_MULTICAST_DELEGATE_TwoParams(FOnCopyWorldData, UWorld*, UWorld*);
 	static FOnCopyWorldData OnCopyWorldData;
 
diff --git a/Engine/Source/Runtime/Engine/Private/UnrealEngine.cpp b/Engine/Source/Runtime/Engine/Private/UnrealEngine.cpp
--- a/Engine/Source/Runtime/Engine/Private/UnrealEngine.cpp
+++ b/Engine/Source/Runtime/Engine/Private/UnrealEngine.cpp
@@ -13647,6 +13647,8 @@ namespace UE::Private
 				ReturnVal->PostCreation(bInitializeWithIris);
 
 				new(Context.ActiveNetDrivers) FNamedNetDriver(ReturnVal, Definition);
+
+				FWorldDelegates::OnNetDriverCreated.Broadcast(Context.World(), ReturnVal);
 			}
 		}
 
diff --git a/Engine/Source/Runtime/Engine/Private/World.cpp b/Engine/Source/Runtime/Engine/Private/World.cpp
--- a/Engine/Source/Runtime/Engine/Private/World.cpp
+++ b/Engine/Source/Runtime/Engine/Private/World.cpp
@@ -570,6 +570,7 @@ FWorldDelegates::FOnWorldPIEEnded FWorldDelegates::OnPIEEnded;
 #endif // WITH_EDITOR
 FWorldDelegates::FOnSeamlessTravelStart FWorldDelegates::OnSeamlessTravelStart;
 FWorldDelegates::FOnSeamlessTravelTransition FWorldDelegates::OnSeamlessTravelTransition;
+FWorldDelegates::FOnNetDriverCreated FWorldDelegates::OnNetDriverCreated;
 FWorldDelegates::FOnCopyWorldData FWorldDelegates::OnCopyWorldData;
 FWorldDelegates::FGameInstanceEvent FWorldDelegates::OnStartGameInstance;
 
