Level: Remove unused property: ActorsModifiedForPIE

#rnx

diff --git a/Engine/Source/Runtime/Engine/Classes/Engine/Level.h b/Engine/Source/Runtime/Engine/Classes/Engine/Level.h
--- a/Engine/Source/Runtime/Engine/Classes/Engine/Level.h
+++ b/Engine/Source/Runtime/Engine/Classes/Engine/Level.h
@@ -415,10 +415,6 @@ public:
 #if WITH_EDITORONLY_DATA
 	AActor* PlayFromHereActor;
 
-	/** List of modified, unsaved actors that needs to be duplicated for PIE */
-	UPROPERTY(NonPIEDuplicateTransient)
-	TMap<FName, TObjectPtr<AActor>> ActorsModifiedForPIE;
-
 	/** Use external actors, new actor spawned in this level will be external and existing external actors will be loaded on load. */
 	UPROPERTY(EditInstanceOnly, Category=World)
 	bool bUseExternalActors;
