Make sure Private DataLayers can't cross reference domains

#rnx

diff --git a/Engine/Source/Runtime/Engine/Private/ActorEditor.cpp b/Engine/Source/Runtime/Engine/Private/ActorEditor.cpp
--- a/Engine/Source/Runtime/Engine/Private/ActorEditor.cpp
+++ b/Engine/Source/Runtime/Engine/Private/ActorEditor.cpp
@@ -35,6 +35,7 @@
 #include "WorldPartition/DataLayer/DataLayerAsset.h"
 #include "UObject/SoftObjectPtr.h"
 #include "Algo/Transform.h"
+#include "Algo/RemoveIf.h"
 #include "Misc/DataValidation.h"
 
 #define LOCTEXT_NAMESPACE "ErrorChecking"
@@ -1713,6 +1714,12 @@ void AActor::FixupDataLayers(bool bRevertChangesOnLockedDataLayer /*= false*/)
 	{ 
 		return;
 	}
+	
+	// Cleanup Data Layer assets we can't reference (Private DLs)
+	DataLayerAssets.SetNum(Algo::RemoveIf(DataLayerAssets, [this](const TSoftObjectPtr<UDataLayerAsset>& AssetPath)
+	{
+		return !UDataLayerAsset::CanBeReferencedByActor(AssetPath, this);
+	}));
 
 	// Use Actor's DataLayerManager since the fixup is relative to this level
 	UDataLayerManager* DataLayerManager = UDataLayerManager::GetDataLayerManager(this);
diff --git a/Engine/Source/Runtime/Engine/Private/LevelInstance/LevelInstanceSubsystem.cpp b/Engine/Source/Runtime/Engine/Private/LevelInstance/LevelInstanceSubsystem.cpp
--- a/Engine/Source/Runtime/Engine/Private/LevelInstance/LevelInstanceSubsystem.cpp
+++ b/Engine/Source/Runtime/Engine/Private/LevelInstance/LevelInstanceSubsystem.cpp
@@ -1121,10 +1121,13 @@ ILevelInstanceInterface* ULevelInstanceSubsystem::CreateLevelInstanceFrom(const
 						// Use the raw asset list as we don't want parent DataLayers
 						for (const UDataLayerAsset* DataLayerAsset : ActorToMove->GetDataLayerAssets())
 						{
-							// Validate that there's a valid Data Layer Instance for this asset in the source level
-							if (DataLayerManager->GetDataLayerInstanceFromAsset(DataLayerAsset))
+							if (const UDataLayerInstance* DataLayerInstance = DataLayerManager->GetDataLayerInstanceFromAsset(DataLayerAsset))
 							{
-								SourceDataLayerAssets.Add(DataLayerAsset);
+								// Validate that there's a valid Data Layer Instance for this asset in the source level and that this isn't a private Data Layer
+								if (!DataLayerAsset->IsPrivate())
+								{
+									SourceDataLayerAssets.Add(DataLayerAsset);
+								}
 							}
 						}
 					}
diff --git a/Engine/Source/Runtime/Engine/Private/WorldPartition/DataLayer/DataLayerAsset.cpp b/Engine/Source/Runtime/Engine/Private/WorldPartition/DataLayer/DataLayerAsset.cpp
--- a/Engine/Source/Runtime/Engine/Private/WorldPartition/DataLayer/DataLayerAsset.cpp
+++ b/Engine/Source/Runtime/Engine/Private/WorldPartition/DataLayer/DataLayerAsset.cpp
@@ -2,6 +2,8 @@
 
 #include "WorldPartition/DataLayer/DataLayerAsset.h"
 #include "WorldPartition/DataLayer/DataLayerInstance.h"
+#include "GameFramework/Actor.h"
+#include "Engine/Level.h"
 
 #include UE_INLINE_GENERATED_CPP_BY_NAME(DataLayerAsset)
 
@@ -17,6 +19,22 @@ bool UDataLayerAsset::IsPrivate() const
 }
 
 #if WITH_EDITOR
+bool UDataLayerAsset::CanBeReferencedByActor(const TSoftObjectPtr<UDataLayerAsset>& InDataLayerAsset, AActor* InActor)
+{
+	if (UDataLayerAsset* DataLayerAsset = InDataLayerAsset.Get())
+	{
+		return DataLayerAsset->CanBeReferencedByActor(InActor);
+	}
+
+	// Only accept references to unloaded UDataLayerAsset if the path is indeed an Asset Path (no sub path string)
+	return InDataLayerAsset.ToSoftObjectPath().IsAsset();
+}
+
+bool UDataLayerAsset::CanBeReferencedByActor(AActor* InActor) const
+{
+	return !IsPrivate() || GetTypedOuter<ULevel>() == InActor->GetLevel();
+}
+
 void UDataLayerAsset::PostLoad()
 {
 	if (DebugColor == FColor::Black)
diff --git a/Engine/Source/Runtime/Engine/Private/WorldPartition/DataLayer/DataLayerInstanceWithAsset.cpp b/Engine/Source/Runtime/Engine/Private/WorldPartition/DataLayer/DataLayerInstanceWithAsset.cpp
--- a/Engine/Source/Runtime/Engine/Private/WorldPartition/DataLayer/DataLayerInstanceWithAsset.cpp
+++ b/Engine/Source/Runtime/Engine/Private/WorldPartition/DataLayer/DataLayerInstanceWithAsset.cpp
@@ -57,7 +57,7 @@ bool UDataLayerInstanceWithAsset::IsLocked() const
 
 bool UDataLayerInstanceWithAsset::CanAddActor(AActor* InActor) const
 {
-	return DataLayerAsset != nullptr && Super::CanAddActor(InActor);
+	return DataLayerAsset != nullptr && DataLayerAsset->CanBeReferencedByActor(InActor) && Super::CanAddActor(InActor);
 }
 
 bool UDataLayerInstanceWithAsset::PerformAddActor(AActor* InActor) const
diff --git a/Engine/Source/Runtime/Engine/Public/WorldPartition/DataLayer/DataLayerAsset.h b/Engine/Source/Runtime/Engine/Public/WorldPartition/DataLayer/DataLayerAsset.h
--- a/Engine/Source/Runtime/Engine/Public/WorldPartition/DataLayer/DataLayerAsset.h
+++ b/Engine/Source/Runtime/Engine/Public/WorldPartition/DataLayer/DataLayerAsset.h
@@ -11,6 +11,8 @@
 
 #include "DataLayerAsset.generated.h"
 
+class AActor;
+
 UCLASS(BlueprintType, editinlinenew)
 class ENGINE_API UDataLayerAsset : public UObject
 {
@@ -31,6 +33,8 @@ public:
 		DataLayerType = Type; 
 	}
 	void SetDebugColor(FColor InDebugColor) { DebugColor = InDebugColor; }
+	bool CanBeReferencedByActor(AActor* InActor) const;
+	static bool CanBeReferencedByActor(const TSoftObjectPtr<UDataLayerAsset>& InDataLayerAsset, AActor* InActor);
 #endif
 	bool IsPrivate() const;
 
