Fix Compile UnrealEditor StressTestUnity Win64.

diff --git a/Engine/Source/Runtime/Engine/Private/WorldPartition/WorldPartitionRuntimeSpatialHash.cpp b/Engine/Source/Runtime/Engine/Private/WorldPartition/WorldPartitionRuntimeSpatialHash.cpp
--- a/Engine/Source/Runtime/Engine/Private/WorldPartition/WorldPartitionRuntimeSpatialHash.cpp
+++ b/Engine/Source/Runtime/Engine/Private/WorldPartition/WorldPartitionRuntimeSpatialHash.cpp
@@ -39,11 +39,6 @@
 
 #include UE_INLINE_GENERATED_CPP_BY_NAME(WorldPartitionRuntimeSpatialHash)
 
-extern ENGINE_API bool GRuntimeSpatialHashUseAlignedGridLevels;
-extern ENGINE_API bool GRuntimeSpatialHashSnapNonAlignedGridLevelsToLowerLevels;
-extern ENGINE_API bool GRuntimeSpatialHashPlaceSmallActorsUsingLocation;
-extern ENGINE_API bool GRuntimeSpatialHashPlacePartitionActorsUsingLocation;
-
 #if WITH_EDITOR
 extern UNREALED_API class UEditorEngine* GEditor;
 #else
diff --git a/Engine/Source/Runtime/Engine/Public/WorldPartition/RuntimeSpatialHash/RuntimeSpatialHashGridHelper.h b/Engine/Source/Runtime/Engine/Public/WorldPartition/RuntimeSpatialHash/RuntimeSpatialHashGridHelper.h
--- a/Engine/Source/Runtime/Engine/Public/WorldPartition/RuntimeSpatialHash/RuntimeSpatialHashGridHelper.h
+++ b/Engine/Source/Runtime/Engine/Public/WorldPartition/RuntimeSpatialHash/RuntimeSpatialHashGridHelper.h
@@ -11,6 +11,10 @@
 #include "WorldPartition/DataLayer/DataLayersID.h"
 #include "WorldPartition/DataLayer/DataLayerInstance.h"
 
+extern ENGINE_API bool GRuntimeSpatialHashUseAlignedGridLevels;
+extern ENGINE_API bool GRuntimeSpatialHashSnapNonAlignedGridLevelsToLowerLevels;
+extern ENGINE_API bool GRuntimeSpatialHashPlaceSmallActorsUsingLocation;
+extern ENGINE_API bool GRuntimeSpatialHashPlacePartitionActorsUsingLocation;
 extern ENGINE_API bool GRuntimeSpatialHashUseAlignedGridLevelsEffective;
 extern ENGINE_API bool GRuntimeSpatialHashSnapNonAlignedGridLevelsToLowerLevelsEffective;
 extern ENGINE_API bool GRuntimeSpatialHashPlaceSmallActorsUsingLocationEffective;
