AssetManager: Remove invalid use of TSharedRef parameterless constructor; for removing from a TMap with TSharedRef we have to do separate find and remove calls.
#rnx

diff --git a/Engine/Source/Runtime/Engine/Private/AssetManager.cpp b/Engine/Source/Runtime/Engine/Private/AssetManager.cpp
--- a/Engine/Source/Runtime/Engine/Private/AssetManager.cpp
+++ b/Engine/Source/Runtime/Engine/Private/AssetManager.cpp
@@ -1071,10 +1071,13 @@ void UAssetManager::RemoveScanPathsForPrimaryAssets(FPrimaryAssetType PrimaryAss
 
 void UAssetManager::RemovePrimaryAssetType(FPrimaryAssetType PrimaryAssetType)
 {
-	TSharedRef<FPrimaryAssetTypeData> TypeData;
-	if (AssetTypeMap.RemoveAndCopyValue(PrimaryAssetType, TypeData))
+	TSharedRef<FPrimaryAssetTypeData>* TypeDataRef;
+	TypeDataRef = AssetTypeMap.Find(PrimaryAssetType);
+	if (TypeDataRef)
 	{
-		TypeData->ResetAssets(AssetPathMap);
+		TSharedPtr<FPrimaryAssetTypeData> TypeDataPtr = *TypeDataRef;
+		AssetTypeMap.Remove(PrimaryAssetType);
+		TypeDataPtr->ResetAssets(AssetPathMap);
 	}
 }
 
