Fix World Partition assert when rendering sequences in the editor.
In some specific cases, UWorldPartition::bIsPIE can be already false when PIE ends. An example is when the world is loaded in UEditorEngine::StartPlayInEditorSession after the PreBeginPIE event was fired, because of OverrideMapURL. In this case, streaming generation happens in UWorldPartition::Initialize because we detect it as a world travel, and everything works fine.

diff --git a/Engine/Source/Runtime/Engine/Private/WorldPartition/WorldPartition.cpp b/Engine/Source/Runtime/Engine/Private/WorldPartition/WorldPartition.cpp
--- a/Engine/Source/Runtime/Engine/Private/WorldPartition/WorldPartition.cpp
+++ b/Engine/Source/Runtime/Engine/Private/WorldPartition/WorldPartition.cpp
@@ -382,7 +382,9 @@ void UWorldPartition::OnPreBeginPIE(bool bStartSimulate)
 
 void UWorldPartition::OnPrePIEEnded(bool bWasSimulatingInEditor)
 {
-	check(bIsPIE);
+	// In some specific cases, bIsPIE can be already false. An example is when the world is loaded in UEditorEngine::StartPlayInEditorSession after the PreBeginPIE event
+	// was fired, because of OverrideMapURL. In this case, streaming generation happens in UWorldPartition::Initialize because we detect it as a world travel, and everything
+	// works fine.
 	bIsPIE = false;
 }
 
