Don't register UActorDescContainer delegates during cook commandlet.
Should fix cook time regression found by [REDACTED_PERSON].

diff --git a/Engine/Source/Runtime/Engine/Private/WorldPartition/ActorDescContainer.cpp b/Engine/Source/Runtime/Engine/Private/WorldPartition/ActorDescContainer.cpp
--- a/Engine/Source/Runtime/Engine/Private/WorldPartition/ActorDescContainer.cpp
+++ b/Engine/Source/Runtime/Engine/Private/WorldPartition/ActorDescContainer.cpp
@@ -392,7 +392,7 @@ bool UActorDescContainer::ShouldRegisterDelegates()
 	UWorld* OwningWorld = GetWorldPartition() ? GetWorldPartition()->GetWorld() : nullptr;
 	// Template container will always register
 	bool bShouldRegisterForWorld = (OwningWorld && !OwningWorld->IsGameWorld()) || (IsTemplateContainer());
-	return GEditor && !IsTemplate() && !bIsInstance && bShouldRegisterForWorld;
+	return GEditor && !IsTemplate() && !IsRunningCookCommandlet() && !bIsInstance && bShouldRegisterForWorld;
 }
 
 void UActorDescContainer::RegisterEditorDelegates()
