Fix for seek to failing when last chunk is the taret

diff --git a/Engine/Source/Runtime/Engine/Private/SoundWave.cpp b/Engine/Source/Runtime/Engine/Private/SoundWave.cpp
--- a/Engine/Source/Runtime/Engine/Private/SoundWave.cpp
+++ b/Engine/Source/Runtime/Engine/Private/SoundWave.cpp
@@ -468,7 +468,7 @@ int32 FSoundWaveData::FindChunkIndexForSeeking(uint32 InTimeInAudioFrames) const
 		// In bounds and valid... 
 		InRangeChunk = i;
 	}
-	return INDEX_NONE;
+	return InRangeChunk;
 }
 
 bool FSoundWaveData::LoadZerothChunk()
