Fix for Sample Rate cook issue on switch

Version bump: UE_AUDIO_ADPCM_VER

#rnx

diff --git a/Engine/Source/Developer/AudioFormatADPCM/Private/AudioFormatADPCM.cpp b/Engine/Source/Developer/AudioFormatADPCM/Private/AudioFormatADPCM.cpp
--- a/Engine/Source/Developer/AudioFormatADPCM/Private/AudioFormatADPCM.cpp
+++ b/Engine/Source/Developer/AudioFormatADPCM/Private/AudioFormatADPCM.cpp
@@ -416,7 +416,7 @@ class FAudioFormatADPCM : public IAudioFormat
 	enum
 	{
 		/** Version for ADPCM format, this becomes part of the DDC key. */
-		UE_AUDIO_ADPCM_VER = 4,
+		UE_AUDIO_ADPCM_VER = 5,
 	};
 
 	void InterleaveBuffers(const TArray<TArray<uint8> >& SrcBuffers, TArray<uint8> & InterleavedBuffer) const
diff --git a/Engine/Source/Runtime/Engine/Private/SoundWave.cpp b/Engine/Source/Runtime/Engine/Private/SoundWave.cpp
--- a/Engine/Source/Runtime/Engine/Private/SoundWave.cpp
+++ b/Engine/Source/Runtime/Engine/Private/SoundWave.cpp
@@ -3552,14 +3552,13 @@ float USoundWave::GetSampleRateForCompressionOverrides(const FPlatformAudioCookO
 	const float* SampleRatePtr = CompressionOverrides->PlatformSampleRates.Find(SampleRateQuality);
 	if (SampleRatePtr && *SampleRatePtr > 0.0f)
 	{
-		if (GIsEditor)
+		// If in editor (but not cooking), then read from current platform 
+		// otherwise, use the given cook override sample rate 
+		if (!IsRunningCookCommandlet() && GIsEditor)
 		{
 			return GetSampleRateForCurrentPlatform();
 		}
-		else
-		{
-			return FMath::Min(*SampleRatePtr, static_cast<float>(SampleRate));
-		}
+		return FMath::Min(*SampleRatePtr, static_cast<float>(SampleRate));
 	}
 	else
 	{
