World Partition: don't register any editor delegates when cooking.

diff --git a/Engine/Source/Runtime/Engine/Private/WorldPartition/WorldPartition.cpp b/Engine/Source/Runtime/Engine/Private/WorldPartition/WorldPartition.cpp
--- a/Engine/Source/Runtime/Engine/Private/WorldPartition/WorldPartition.cpp
+++ b/Engine/Source/Runtime/Engine/Private/WorldPartition/WorldPartition.cpp
@@ -860,7 +860,7 @@ void UWorldPartition::RegisterDelegates()
 	check(World); 
 
 #if WITH_EDITOR
-	if (GEditor && !IsTemplate() && !World->IsGameWorld())
+	if (GEditor && !IsTemplate() && !World->IsGameWorld() && !IsRunningCookCommandlet())
 	{
 		if (IsMainWorldPartition())
 		{
@@ -903,7 +903,7 @@ void UWorldPartition::UnregisterDelegates()
 	check(World);
 
 #if WITH_EDITOR
-	if (GEditor && !IsTemplate() && !World->IsGameWorld())
+	if (GEditor && !IsTemplate() && !World->IsGameWorld() && !IsRunningCookCommandlet())
 	{
 		if (IsMainWorldPartition())
 		{
