UWorld::AddToWorld() - Perform client/server logic only on Game worlds

diff --git a/Engine/Source/Runtime/Engine/Private/World.cpp b/Engine/Source/Runtime/Engine/Private/World.cpp
--- a/Engine/Source/Runtime/Engine/Private/World.cpp
+++ b/Engine/Source/Runtime/Engine/Private/World.cpp
@@ -3246,14 +3246,14 @@ void UWorld::AddToWorld( ULevel* Level, const FTransform& LevelTransform, bool b
 
 		// We don't need to rerun construction scripts if we have cooked data or we are playing in editor unless the PIE world was loaded
 		// from disk rather than duplicated
-		const bool bRerunConstructionScript = !FPlatformProperties::RequiresCookedData() && (!IsGameWorld() || !Level->bHasRerunConstructionScripts);
+		const bool bRerunConstructionScript = !FPlatformProperties::RequiresCookedData() && (!bIsGameWorld || !Level->bHasRerunConstructionScripts);
 
 		// Prepare context used to store batch/parallelize AddPrimitive calls
 		FRegisterComponentContext Context(this);
 		FRegisterComponentContext* ContextPtr = GLevelStreamingAddPrimitiveGranularity == 0 ? nullptr : &Context;
 
 		// Incrementally update components.
-		int32 NumComponentsToUpdate = (!bConsiderTimeLimit || !IsGameWorld() || IsRunningCommandlet() ? 0 : GLevelStreamingComponentsRegistrationGranularity);
+		int32 NumComponentsToUpdate = (!bConsiderTimeLimit || !bIsGameWorld || IsRunningCommandlet() ? 0 : GLevelStreamingComponentsRegistrationGranularity);
 		do
 		{
 			Level->IncrementalUpdateComponents( NumComponentsToUpdate, bRerunConstructionScript, ContextPtr);
@@ -3279,7 +3279,7 @@ void UWorld::AddToWorld( ULevel* Level, const FTransform& LevelTransform, bool b
 	FAssetCompilingManager::Get().ProcessAsyncTasks(ProcessAsyncTasksParams);
 #endif
 
-	if( IsGameWorld() && AreActorsInitialized() )
+	if( bIsGameWorld && AreActorsInitialized() )
 	{
 		// Initialize all actors and start execution.
 		if (bExecuteNextStep && !(Level->bAlreadyInitializedNetworkActors && Level->bAlreadyClearedActorsSeamlessTravelFlag))
@@ -3311,7 +3311,7 @@ void UWorld::AddToWorld( ULevel* Level, const FTransform& LevelTransform, bool b
 			CSV_SCOPED_TIMING_STAT(LevelStreamingDetail, AddToWorld_RouteActorInit);
 			QUICK_SCOPE_CYCLE_COUNTER(STAT_AddToWorldTime_RouteActorInitialize);
 			SCOPE_TIME_TO_VAR(&RouteActorInitializeTime);
-			const int32 NumActorsToProcess = (!bConsiderTimeLimit || !IsGameWorld() || IsRunningCommandlet()) ? 0 : GLevelStreamingRouteActorInitializationGranularity;
+			const int32 NumActorsToProcess = (!bConsiderTimeLimit || !bIsGameWorld || IsRunningCommandlet()) ? 0 : GLevelStreamingRouteActorInitializationGranularity;
 			bStartup = 1;
 			do 
 			{
@@ -3352,7 +3352,7 @@ void UWorld::AddToWorld( ULevel* Level, const FTransform& LevelTransform, bool b
 		// Finished making level visible - allow other levels to be added to the world.
 		CurrentLevelPendingVisibility = nullptr;
 
-		if (!Level->bClientOnlyVisible)
+		if (bIsGameWorld && !Level->bClientOnlyVisible)
 		{
 			FUpdateLevelVisibilityLevelInfo LevelVisibility(Level, true);
 			const FName UnmappedPackageName = LevelVisibility.PackageName;
@@ -3469,6 +3469,8 @@ void UWorld::RemoveFromWorld( ULevel* Level, bool bAllowIncrementalRemoval, FNet
 	check(IsValid(Level));
 	check(!Level->IsUnreachable());
 
+	const bool bIsGameWorld = IsGameWorld();
+
 	Level->bIsDisassociatingLevel = true;
 
 	auto BeginRemoval = [Level, this, OwningLevelStreaming]()
@@ -3523,7 +3525,7 @@ void UWorld::RemoveFromWorld( ULevel* Level, bool bAllowIncrementalRemoval, FNet
 					double StartTime = FPlatformTime::Seconds();
 					// Incrementally unregister actor components. 
 					// This avoids spikes on the renderthread and gamethread when we subsequently call ClearLevelComponents() further down
-					check(IsGameWorld());
+					check(bIsGameWorld);
 					int32 NumComponentsToUnregister = GLevelStreamingComponentsUnregistrationGranularity;
 					do
 					{
@@ -3591,7 +3593,7 @@ void UWorld::RemoveFromWorld( ULevel* Level, bool bAllowIncrementalRemoval, FNet
 		
 			Level->ClearLevelComponents();
 
-			if (!Level->bClientOnlyVisible)
+			if (bIsGameWorld && !Level->bClientOnlyVisible)
 			{
 				FUpdateLevelVisibilityLevelInfo LevelVisibility(Level, false);
 				const FName UnmappedPackageName = LevelVisibility.PackageName;
@@ -3625,7 +3627,7 @@ void UWorld::RemoveFromWorld( ULevel* Level, bool bAllowIncrementalRemoval, FNet
 			}
 
 			// Make sure level always has OwningWorld in the editor
-			if (IsGameWorld())
+			if (bIsGameWorld)
 			{
 				Levels.Remove(Level);
 				Level->OwningWorld = nullptr;
