Fix crash when setting invalid default parameters on a new metasound

diff --git a/Engine/Plugins/Runtime/Metasound/Source/MetasoundEngine/Private/MetasoundSource.cpp b/Engine/Plugins/Runtime/Metasound/Source/MetasoundEngine/Private/MetasoundSource.cpp
--- a/Engine/Plugins/Runtime/Metasound/Source/MetasoundEngine/Private/MetasoundSource.cpp
+++ b/Engine/Plugins/Runtime/Metasound/Source/MetasoundEngine/Private/MetasoundSource.cpp
@@ -437,13 +437,11 @@ void UMetaSoundSource::InitParameters(TArray<FAudioParameter>& ParametersToInit,
 		return true;
 	};
 
-	auto ConstructProxies = [this, FeatureName = InFeatureName](FAudioParameter& OutParamToInit)
+	auto ConstructProxies = [this](FAudioParameter& OutParamToInit)
 	{
 		using namespace Metasound;
 		using namespace Metasound::Frontend;
 
-		const Audio::FProxyDataInitParams ProxyInitParams { FeatureName };
-
 		switch (OutParamToInit.ParamType)
 		{
 			case EAudioParameterType::Object:
diff --git a/Engine/Source/Runtime/AudioMixer/Private/AudioMixerSource.cpp b/Engine/Source/Runtime/AudioMixer/Private/AudioMixerSource.cpp
--- a/Engine/Source/Runtime/AudioMixer/Private/AudioMixerSource.cpp
+++ b/Engine/Source/Runtime/AudioMixer/Private/AudioMixerSource.cpp
@@ -1129,6 +1129,7 @@ namespace Audio
 			if (Audio::IParameterTransmitter* Transmitter = ActiveSound->GetTransmitter())
 			{
 				DefaultParameters = Transmitter->GetParameters();
+				SoundWave.InitParameters(DefaultParameters);
 			}
 		}
 
diff --git a/Engine/Source/Runtime/Engine/Classes/Sound/SoundBase.h b/Engine/Source/Runtime/Engine/Classes/Sound/SoundBase.h
--- a/Engine/Source/Runtime/Engine/Classes/Sound/SoundBase.h
+++ b/Engine/Source/Runtime/Engine/Classes/Sound/SoundBase.h
@@ -319,7 +319,7 @@ public:
 	//~ End IInterface_AssetUserData Interface
 
 	/** Called from the Game Thread prior to attempting to pass parameters to the ParameterTransmitter. */
-	virtual void InitParameters(TArray<FAudioParameter>& ParametersToInit, FName InFeatureName);
+	virtual void InitParameters(TArray<FAudioParameter>& ParametersToInit, FName InFeatureName = NAME_None);
 
 	/** Called from the Game Thread prior to attempting to initialize a sound instance. */
 	virtual void InitResources() { }
diff --git a/Engine/Source/Runtime/Engine/Private/SoundCue.cpp b/Engine/Source/Runtime/Engine/Private/SoundCue.cpp
--- a/Engine/Source/Runtime/Engine/Private/SoundCue.cpp
+++ b/Engine/Source/Runtime/Engine/Private/SoundCue.cpp
@@ -647,8 +647,7 @@ void USoundCue::Parse(FAudioDevice* AudioDevice, const UPTRINT NodeWaveInstanceH
 
 					if (USoundBase* Sound = Instance->WaveData)
 					{
-						static const FName ProxyFeatureName("SoundCueParameterTransmitter");
-						Sound->InitParameters(Params, ProxyFeatureName);
+						Sound->InitParameters(Params);
 					}
 
 					if (!Params.IsEmpty())
diff --git a/Engine/Source/Runtime/Engine/Private/SoundNodeWaveParam.cpp b/Engine/Source/Runtime/Engine/Private/SoundNodeWaveParam.cpp
--- a/Engine/Source/Runtime/Engine/Private/SoundNodeWaveParam.cpp
+++ b/Engine/Source/Runtime/Engine/Private/SoundNodeWaveParam.cpp
@@ -47,9 +47,10 @@ void USoundNodeWaveParam::ParseNodes(FAudioDevice* AudioDevice, const UPTRINT No
 				Params.DefaultParams = ActiveSound.GetTransmitter()->GetParameters();
 				Params.InstanceID = Audio::GetTransmitterID(ActiveSound.GetAudioComponentID(), WaveHash, ActiveSound.GetPlayOrder()); 
 				Params.SampleRate = AudioDevice->GetSampleRate();
-				
 				Params.AudioDeviceID = AudioDevice->DeviceID;
 				
+				NewWave->InitParameters(Params.DefaultParams);
+
 				const TSharedPtr<Audio::IParameterTransmitter> SoundWaveTransmitter = NewWave->CreateParameterTransmitter(MoveTemp(Params));
 			
 				if (SoundWaveTransmitter.IsValid())
diff --git a/Engine/Source/Runtime/Engine/Private/SoundNodeWavePlayer.cpp b/Engine/Source/Runtime/Engine/Private/SoundNodeWavePlayer.cpp
--- a/Engine/Source/Runtime/Engine/Private/SoundNodeWavePlayer.cpp
+++ b/Engine/Source/Runtime/Engine/Private/SoundNodeWavePlayer.cpp
@@ -143,15 +143,17 @@ void USoundNodeWavePlayer::ParseNodes( FAudioDevice* AudioDevice, const UPTRINT
 		if (*RequiresInitialization)
 		{
 			bPlayFailed = 0;
-			
-			Audio::FParameterTransmitterInitParams Params;
-			Params.DefaultParams = ActiveSound.GetTransmitter()->GetParameters();
-			Params.InstanceID = Audio::GetTransmitterID(ActiveSound.GetAudioComponentID(), NodeWaveInstanceHash, ActiveSound.GetPlayOrder()); 
-			Params.SampleRate = AudioDevice->GetSampleRate();
-			Params.AudioDeviceID = AudioDevice->DeviceID;
 
 			if (FSoundCueParameterTransmitter* SoundCueTransmitter = static_cast<FSoundCueParameterTransmitter*>(ActiveSound.GetTransmitter()))
 			{
+				Audio::FParameterTransmitterInitParams Params;
+				Params.DefaultParams = ActiveSound.GetTransmitter()->GetParameters();
+				Params.InstanceID = Audio::GetTransmitterID(ActiveSound.GetAudioComponentID(), NodeWaveInstanceHash, ActiveSound.GetPlayOrder()); 
+				Params.SampleRate = AudioDevice->GetSampleRate();
+				Params.AudioDeviceID = AudioDevice->DeviceID;
+
+				SoundWave->InitParameters(Params.DefaultParams);
+				
 				const TSharedPtr<Audio::IParameterTransmitter> SoundWaveTransmitter = SoundWave->CreateParameterTransmitter(MoveTemp(Params));
 				
 				if (SoundWaveTransmitter.IsValid())
