Change specular profile procedural parameter to user a color curve.

diff --git a/Engine/Shaders/Private/SpecularProfile.usf b/Engine/Shaders/Private/SpecularProfile.usf
--- a/Engine/Shaders/Private/SpecularProfile.usf
+++ b/Engine/Shaders/Private/SpecularProfile.usf
@@ -2,9 +2,6 @@
 
 #include "Common.ush"
 
-float3 ViewColor;
-float3 LightColor;
-
 int2 SourceResolution;
 uint SourceMipCount;
 int2 TargetResolution;
@@ -13,6 +10,9 @@ SamplerState SourceSampler;
 Texture2D<float4> SourceTexture;
 RWTexture2DArray<float4> TargetTexture;
 
+Buffer<float4> ViewColorBuffer;
+Buffer<float4> LightColorBuffer;
+
 [numthreads(8, 8, 1)]
 void MainCS(uint3 DispatchThreadId : SV_DispatchThreadID)
 {
@@ -22,12 +22,8 @@ void MainCS(uint3 DispatchThreadId : SV_DispatchThreadID)
 		const float2 TexCoord = float2(PixelPos + 0.5f) / float2(TargetResolution);
 		#if PERMUTATION_PROCEDURAL
 		{
-			const float NoL = TexCoord.y;
-			const float NoV = TexCoord.x;
-
-			const float3 DefaultColor = 1.f;
-			const float3 C0 = lerp(DefaultColor, ViewColor, NoV);
-			const float3 C1 = lerp(DefaultColor, LightColor, NoL);
+			const float3 C0 = ViewColorBuffer[PixelPos.x].xyz;
+			const float3 C1 = LightColorBuffer[PixelPos.y].xyz;
 			TargetTexture[uint3(PixelPos,TargetIndex)] = float4(C0 * C1, 1);
 		}
 		#else
diff --git a/Engine/Source/Editor/UnrealEd/Private/ThumbnailRendering/SpecularProfileRenderer.cpp b/Engine/Source/Editor/UnrealEd/Private/ThumbnailRendering/SpecularProfileRenderer.cpp
--- a/Engine/Source/Editor/UnrealEd/Private/ThumbnailRendering/SpecularProfileRenderer.cpp
+++ b/Engine/Source/Editor/UnrealEd/Private/ThumbnailRendering/SpecularProfileRenderer.cpp
@@ -24,10 +24,9 @@ void USpecularProfileRenderer::Draw(UObject* Object, int32 X, int32 Y, uint32 Wi
 	if (LocalSpecularProfile)
 	{
 		FLinearColor Col;
-
-		Col = LocalSpecularProfile->Settings.ViewColor; Col.A = 1;
+		Col = LocalSpecularProfile->Settings.ViewColor.GetLinearColorValue(1.0f); Col.A = 1;
 		Canvas->DrawTile(0,          0, Width, Height / 2, 0, 0, 1, 1, Col);
-		Col = LocalSpecularProfile->Settings.LightColor; Col.A = 1;
+		Col = LocalSpecularProfile->Settings.LightColor.GetLinearColorValue(1.0f); Col.A = 1;
 		Canvas->DrawTile(0, Height / 2, Width, Height / 2, 0, 0, 1, 1, Col);
 
 		FCanvasTextItem TextItem(FVector2D(5.0f, 5.0f), FText::FromString(LocalSpecularProfile->Settings.IsProcedural() ? TEXT("Procedural") : TEXT("Texture")), GEngine->GetLargeFont(), FLinearColor::White);
diff --git a/Engine/Source/Runtime/Engine/Classes/Engine/SpecularProfile.h b/Engine/Source/Runtime/Engine/Classes/Engine/SpecularProfile.h
--- a/Engine/Source/Runtime/Engine/Classes/Engine/SpecularProfile.h
+++ b/Engine/Source/Runtime/Engine/Classes/Engine/SpecularProfile.h
@@ -6,6 +6,7 @@
 #include "UObject/Object.h"
 #include "Misc/Guid.h"
 #include "RenderResource.h"
+#include "Curves/CurveLinearColor.h"
 #include "SpecularProfile.generated.h"
 
 #define MAX_SPECULAR_PROFILE_COUNT 256
@@ -37,14 +38,14 @@ struct FSpecularProfileStruct
 	/**
 	* Define the view facing color
 	*/
-	UPROPERTY(Category = "Procedural", EditAnywhere, BlueprintReadOnly, meta = (ClampMin = "0.01", UIMax = "1.0", ClampMax = "1.0", HideAlphaChannel))
-	FLinearColor ViewColor;
+	UPROPERTY(Category = "Procedural", EditAnywhere, meta = (AllowZoomOutput = "false", ShowZoomButtons = "false", ViewMinInput = "0", ViewMaxInput = "1", ViewMinOutput = "0", ViewMaxOutput = "1", TimelineLength = "1", ShowInputGridNumbers="false", ShowOutputGridNumbers="false"))
+	FRuntimeCurveLinearColor ViewColor;
 
 	/**
 	* Define the light facing color
 	*/
-	UPROPERTY(Category = "Procedural", EditAnywhere, BlueprintReadOnly, meta = (ClampMin = "0.001", UIMax = "1.0", ClampMax = "1.0", HideAlphaChannel))
-	FLinearColor LightColor;
+	UPROPERTY(Category = "Procedural", EditAnywhere, meta = (AllowZoomOutput = "false", ShowZoomButtons = "false", ViewMinInput = "0", ViewMaxInput = "1", ViewMinOutput = "0", ViewMaxOutput = "1", TimelineLength = "1", ShowInputGridNumbers="false", ShowOutputGridNumbers="false"))
+	FRuntimeCurveLinearColor LightColor;
 
 	/**
 	 * Define the texture used as a specular profile
@@ -52,13 +53,7 @@ struct FSpecularProfileStruct
 	UPROPERTY(Category  ="Texture", EditAnywhere, BlueprintReadOnly, meta=(DisplayName="Texture"))
 	TObjectPtr<UTexture2D> Texture;
 
-	FSpecularProfileStruct()
-	{
-		Format = ESpecularProfileFormat::ViewLightVector;
-		ViewColor = FLinearColor(1.0f, 1.0f, 1.0f);
-		LightColor = FLinearColor(1.0f, 1.0f, 1.0f);
-		Texture = nullptr;
-	}
+	FSpecularProfileStruct();
 
 	bool IsProcedural() const { return Texture == nullptr; }
 
diff --git a/Engine/Source/Runtime/Engine/Private/Rendering/SpecularProfile.cpp b/Engine/Source/Runtime/Engine/Private/Rendering/SpecularProfile.cpp
--- a/Engine/Source/Runtime/Engine/Private/Rendering/SpecularProfile.cpp
+++ b/Engine/Source/Runtime/Engine/Private/Rendering/SpecularProfile.cpp
@@ -62,8 +62,8 @@ public:
 		SHADER_PARAMETER(FIntPoint, TargetResolution)
 		SHADER_PARAMETER(uint32, SourceMipCount)
 		SHADER_PARAMETER(uint32, TargetIndex)
-		SHADER_PARAMETER(FVector3f, ViewColor)
-		SHADER_PARAMETER(FVector3f, LightColor)
+		SHADER_PARAMETER_RDG_BUFFER_SRV(Buffer, ViewColorBuffer)
+		SHADER_PARAMETER_RDG_BUFFER_SRV(Buffer, LightColorBuffer)
 		SHADER_PARAMETER_RDG_TEXTURE_UAV(RWTexture2D, TargetTexture)
 	END_SHADER_PARAMETER_STRUCT()
 
@@ -344,8 +344,34 @@ IPooledRenderTarget* FSpecularProfileTextureManager::GetAtlasTexture(FRDGBuilder
 			PassParameters->TargetResolution = SpecularProfileTexture->Desc.Extent;
 			PassParameters->TargetTexture = SpecularProfileUAV;
 			PassParameters->TargetIndex = LayerIt;
-			PassParameters->ViewColor = FVector3f(Data.ViewColor);
-			PassParameters->LightColor = FVector3f(Data.LightColor);
+
+			struct FLinearColor16 
+			{
+				FLinearColor16() {}
+				FLinearColor16(const FLinearColor& In) { R = In.R; G = In.G; B = In.B; A = In.A; }
+				FFloat16 R;
+				FFloat16 G;
+				FFloat16 B;
+				FFloat16 A;
+			};
+
+			TArray<FLinearColor16> ViewColorData;
+			TArray<FLinearColor16> LightColorData;
+			ViewColorData.SetNum(PassParameters->TargetResolution.X);
+			LightColorData.SetNum(PassParameters->TargetResolution.Y);
+			for (uint32 It=0; It < Resolution; ++It)
+			{
+				ViewColorData[It]  = FLinearColor16(Data.ViewColor.GetLinearColorValue(float(It + 0.5f) / Resolution));
+				LightColorData[It] = FLinearColor16(Data.LightColor.GetLinearColorValue(float(It + 0.5f) / Resolution));
+			}
+
+			FRDGBufferRef ViewColorBuffer = GraphBuilder.CreateBuffer(FRDGBufferDesc::CreateBufferDesc(8, Resolution), TEXT("ViewColor"), ERDGBufferFlags::MultiFrame);
+			GraphBuilder.QueueBufferUpload(ViewColorBuffer, ViewColorData.GetData(), 8*ViewColorData.Num(), ERDGInitialDataFlags::None);
+			PassParameters->ViewColorBuffer = GraphBuilder.CreateSRV(ViewColorBuffer, PF_FloatRGBA);
+
+			FRDGBufferRef LightColorBuffer = GraphBuilder.CreateBuffer(FRDGBufferDesc::CreateBufferDesc(8, Resolution), TEXT("LightColor"), ERDGBufferFlags::MultiFrame);
+			GraphBuilder.QueueBufferUpload(LightColorBuffer, LightColorData.GetData(), 8*LightColorData.Num(), ERDGInitialDataFlags::None);
+			PassParameters->LightColorBuffer = GraphBuilder.CreateSRV(LightColorBuffer, PF_FloatRGBA);
 
 			FSpecularProfileCopyCS::FPermutationDomain PermutationVector;
 			PermutationVector.Set<FSpecularProfileCopyCS::FProcedural>(Data.IsProcedural());
@@ -380,6 +406,29 @@ int32 FSpecularProfileTextureManager::FindAllocationId(const USpecularProfile* I
 	return INDEX_NONE;
 }
 
+///////////////////////////////////////////////////////////////////////////////////////////////////
+// FSpecularProfileStruct
+
+FSpecularProfileStruct::FSpecularProfileStruct()
+{
+	ViewColor.ColorCurves[0].AddKey(0, 1.f);
+	ViewColor.ColorCurves[1].AddKey(0, 1.f);
+	ViewColor.ColorCurves[2].AddKey(0, 1.f);
+	ViewColor.ColorCurves[0].AddKey(1, 1.f);
+	ViewColor.ColorCurves[1].AddKey(1, 1.f);
+	ViewColor.ColorCurves[2].AddKey(1, 1.f);
+
+	LightColor.ColorCurves[0].AddKey(0, 1.f);
+	LightColor.ColorCurves[1].AddKey(0, 1.f);
+	LightColor.ColorCurves[2].AddKey(0, 1.f);
+	LightColor.ColorCurves[0].AddKey(1, 1.f);
+	LightColor.ColorCurves[1].AddKey(1, 1.f);
+	LightColor.ColorCurves[2].AddKey(1, 1.f);
+
+	Format = ESpecularProfileFormat::ViewLightVector;
+	Texture = nullptr;
+}
+
 ///////////////////////////////////////////////////////////////////////////////////////////////////
 // USpecularProfile
 
