ARO: fixing new callsite of Swap() ahead of mutable TObjectPtr container deprecation

diff --git a/Engine/Source/Runtime/Engine/Private/Components/MeshComponent.cpp b/Engine/Source/Runtime/Engine/Private/Components/MeshComponent.cpp
--- a/Engine/Source/Runtime/Engine/Private/Components/MeshComponent.cpp
+++ b/Engine/Source/Runtime/Engine/Private/Components/MeshComponent.cpp
@@ -542,7 +542,7 @@ void UMeshComponent::LogMaterialsAndTextures(FOutputDevice& Ar, int32 Indent) co
 	}
 
 	// Backup the material overrides so we can access the mesh original materials.
-	TArray<class UMaterialInterface*> OverrideMaterialsBackup;
+	TArray<TObjectPtr<class UMaterialInterface>> OverrideMaterialsBackup;
 	Swap(OverrideMaterialsBackup, const_cast<UMeshComponent*>(this)->OverrideMaterials);
 
 	TArray<UMaterialInterface*> MaterialInterfaces = GetMaterials();
