Soundfield submix fixes.

diff --git a/Engine/Source/Runtime/Engine/Classes/Sound/SoundSubmix.h b/Engine/Source/Runtime/Engine/Classes/Sound/SoundSubmix.h
--- a/Engine/Source/Runtime/Engine/Classes/Sound/SoundSubmix.h
+++ b/Engine/Source/Runtime/Engine/Classes/Sound/SoundSubmix.h
@@ -486,6 +486,9 @@ public:
 #if WITH_EDITOR
 	virtual void PostEditChangeProperty(struct FPropertyChangedEvent& PropertyChangedEvent) override;
 #endif
+
+protected:
+	virtual void PostLoad() override;
 };
 
 /**
@@ -530,6 +533,7 @@ public:
 	TArray<TObjectPtr<USoundfieldEffectBase>> SoundfieldEffectChain;
 
 protected:
+	virtual void PostLoad() override;
 #if WITH_EDITOR
 	virtual void PostEditChangeProperty(struct FPropertyChangedEvent& PropertyChangedEvent) override;
 #endif
diff --git a/Engine/Source/Runtime/Engine/Private/SoundSubmix.cpp b/Engine/Source/Runtime/Engine/Private/SoundSubmix.cpp
--- a/Engine/Source/Runtime/Engine/Private/SoundSubmix.cpp
+++ b/Engine/Source/Runtime/Engine/Private/SoundSubmix.cpp
@@ -950,6 +950,22 @@ void UEndpointSubmix::PostEditChangeProperty(struct FPropertyChangedEvent& Prope
 
 #endif // WITH_EDITOR
 
+void UEndpointSubmix::PostLoad()
+{
+	// Validate our endpoint type is enabled.
+	TArray<FName> EndpointTypeNames = IAudioEndpointFactory::GetAvailableEndpointTypes();
+	if (!EndpointTypeNames.Contains(EndpointType))
+	{
+		const FName DefaultEndpoint = IAudioEndpointFactory::GetTypeNameForDefaultEndpoint();
+		UE_LOG(LogAudio, Warning, TEXT("UEndpointSubmix [%s] has endpoint type [%s] which is not currently currently enabled. Changing to [%s]"),
+			*GetName(), *EndpointType.ToString(), *DefaultEndpoint.ToString());
+		EndpointType = DefaultEndpoint;
+	}
+
+	
+	Super::PostLoad();
+}
+
 void USoundfieldSubmix::PostLoad()
 {
 	// Make sure the Encoding format is something we can use.
@@ -1056,6 +1072,24 @@ void USoundfieldEndpointSubmix::SanitizeLinks()
 	}
 }
 
+
+void USoundfieldEndpointSubmix::PostLoad()
+{
+	// Validate we're set to something that's enabled.	
+	TArray<FName> SoundfieldEndpointTypeNames = ISoundfieldEndpointFactory::GetAllSoundfieldEndpointTypes();
+	if (!SoundfieldEndpointTypeNames.Contains(SoundfieldEndpointType))
+	{
+		const FName DefaultEndpoint = ISoundfieldEndpointFactory::DefaultSoundfieldEndpointName();
+		UE_LOG(LogAudio, Warning, TEXT("USoundfieldEndpointSubmix [%s] has endpoint type [%s] which is not currently currently enabled. Changing to [%s]"),
+			*GetName(), *SoundfieldEndpointType.ToString(), *DefaultEndpoint.ToString());
+		SoundfieldEndpointType = DefaultEndpoint;
+		SanitizeLinks();
+	}
+
+	Super::PostLoad();	
+}
+
+
 #if WITH_EDITOR
 
 void USoundfieldEndpointSubmix::PostEditChangeProperty(struct FPropertyChangedEvent& PropertyChangedEvent)
@@ -1065,7 +1099,14 @@ void USoundfieldEndpointSubmix::PostEditChangeProperty(struct FPropertyChangedEv
 		static const FName NAME_SoundfieldFormat(TEXT("SoundfieldEndpointType"));
 
 		if (PropertyChangedEvent.Property->GetFName() == NAME_SoundfieldFormat)
-		{
+		{			
+			// Remove-re-add submix. Causes reinit with plugins.
+			if (FAudioDeviceManager* AudioDeviceManager = GEngine->GetAudioDeviceManager())
+			{
+				AudioDeviceManager->UnregisterSoundSubmix(this);
+				AudioDeviceManager->RegisterSoundSubmix(this);
+			}
+			
 			// Add this sound class to the parent class if it's not already added
 			SanitizeLinks();
 		}
@@ -1096,8 +1137,20 @@ ENGINE_API bool SubmixUtils::AreSubmixFormatsCompatible(const USoundSubmixBase*
 
 			if (ChildSoundfieldFactory && ParentSoundfieldFactory)
 			{
-				return ChildSoundfieldFactory->CanTranscodeToSoundfieldFormat(ParentSoundfieldFactory->GetSoundfieldFormatName(), *(ParentSoundfieldSubmix->GetSoundfieldEncodingSettings()->GetProxy()))
-					|| ParentSoundfieldFactory->CanTranscodeFromSoundfieldFormat(ChildSoundfieldFactory->GetSoundfieldFormatName(), *(ChildSoundfieldSubmix->GetSoundfieldEncodingSettings()->GetProxy()));
+				bool bCanTranscode = false;
+
+				// To
+				if (const USoundfieldEncodingSettingsBase* ParentEncodingSettings = ParentSoundfieldSubmix->GetSoundfieldEncodingSettings())
+				{
+					bCanTranscode |= ChildSoundfieldFactory->CanTranscodeToSoundfieldFormat(ParentSoundfieldFactory->GetSoundfieldFormatName(), *ParentEncodingSettings->GetProxy());
+				}
+				// From
+				if (const USoundfieldEncodingSettingsBase* ChildEncodingSettings = ChildSoundfieldSubmix->GetSoundfieldEncodingSettings())
+				{
+					bCanTranscode |= ParentSoundfieldFactory->CanTranscodeFromSoundfieldFormat(ChildSoundfieldFactory->GetSoundfieldFormatName(), *ChildEncodingSettings->GetProxy());
+				}
+				
+				return bCanTranscode;
 			}
 			else
 			{
@@ -1117,8 +1170,21 @@ ENGINE_API bool SubmixUtils::AreSubmixFormatsCompatible(const USoundSubmixBase*
 
 			if (ChildSoundfieldFactory && ParentSoundfieldFactory)
 			{
-				return ChildSoundfieldFactory->CanTranscodeToSoundfieldFormat(ParentSoundfieldFactory->GetSoundfieldFormatName(),  *(ParentSoundfieldEndpointSubmix->GetEncodingSettings()->GetProxy()))
-					|| ParentSoundfieldFactory->CanTranscodeFromSoundfieldFormat(ChildSoundfieldFactory->GetSoundfieldFormatName(), *(ChildSoundfieldSubmix->GetSoundfieldEncodingSettings()->GetProxy()));
+				bool bCanTranscode = false;
+
+				// TO. (Endpoint settings).
+				if (const USoundfieldEncodingSettingsBase* ParentEndpointSettings = ParentSoundfieldEndpointSubmix->GetEncodingSettings())
+				{
+					bCanTranscode |= ChildSoundfieldFactory->CanTranscodeToSoundfieldFormat(ParentSoundfieldFactory->GetSoundfieldFormatName(), *ParentEndpointSettings->GetProxy());
+				}
+
+				// From
+				if (const USoundfieldEncodingSettingsBase* ChildEncodingSettings = ChildSoundfieldSubmix->GetSoundfieldEncodingSettings())
+				{
+					bCanTranscode |= ParentSoundfieldFactory->CanTranscodeFromSoundfieldFormat(ChildSoundfieldFactory->GetSoundfieldFormatName(), *ChildEncodingSettings->GetProxy());
+				}
+				
+				return bCanTranscode;
 			}
 			else
 			{
