Immediately conform class sparse data, doing so later can result in child sparse data attempting to delta serialize itself without an up to date version of the parent classes sparse class data.

This may be too early

diff --git a/Engine/Source/Runtime/Engine/Private/BlueprintGeneratedClass.cpp b/Engine/Source/Runtime/Engine/Private/BlueprintGeneratedClass.cpp
--- a/Engine/Source/Runtime/Engine/Private/BlueprintGeneratedClass.cpp
+++ b/Engine/Source/Runtime/Engine/Private/BlueprintGeneratedClass.cpp
@@ -601,6 +601,15 @@ void UBlueprintGeneratedClass::SerializeDefaultObject(UObject* Object, FStructur
 		if (!GetPackage()->HasAnyPackageFlags(PKG_Cooked))
 		{
 			PrepareToConformSparseClassData();
+			if(ClassDefaultObject)
+			{
+				// conform immediately, so that child types can correct delta serialize their
+				// sparse data - we can't check RF_LoadCompleted because it is going to be set by the 
+				// preload call that is causing SerializeDefaultObject to run, but Super::SerializeDefaultObject
+				// has run...
+				ConformSparseClassData(ClassDefaultObject);
+				ClassDefaultObject->MoveDataToSparseClassDataStruct();
+			}
 		}
 #endif
 	}
