Add r.RayTracing.Geometry.InstancedStaticMeshes.MinLOD CVar

* This allows forcing a lower-resolution LOD to be used for ray tracing ISMs for debugging purposes

diff --git a/Engine/Source/Runtime/Engine/Private/InstancedStaticMesh.cpp b/Engine/Source/Runtime/Engine/Private/InstancedStaticMesh.cpp
--- a/Engine/Source/Runtime/Engine/Private/InstancedStaticMesh.cpp
+++ b/Engine/Source/Runtime/Engine/Private/InstancedStaticMesh.cpp
@@ -87,6 +87,11 @@ static TAutoConsoleVariable<int32> CVarRayTracingRenderInstances(
 	1,
 	TEXT("Include static mesh instances in ray tracing effects (default = 1 (Instances enabled in ray tracing))"));
 
+static TAutoConsoleVariable<int32> CVarRayTracingInstancedStaticMeshesMinLOD(
+	TEXT("r.RayTracing.Geometry.InstancedStaticMeshes.MinLOD"),
+	0,
+	TEXT("Clamps minimum LOD to this value (default = 0, highest resolution LOD may be used)"));
+
 static TAutoConsoleVariable<int32> CVarRayTracingRenderInstancesCulling(
 	TEXT("r.RayTracing.Geometry.InstancedStaticMeshes.Culling"),
 	1,
@@ -2151,7 +2156,9 @@ void FInstancedStaticMeshSceneProxy::GetDynamicRayTracingInstances(struct FRayTr
 		return;
 	}
 
-	uint32 LOD = GetCurrentFirstLODIdx_RenderThread();
+	uint32 MinAllowedLOD = FMath::Clamp<int32>(CVarRayTracingInstancedStaticMeshesMinLOD.GetValueOnRenderThread(), 0, RenderData->LODResources.Num() - 1);
+	uint32 LOD = FMath::Max<uint32>(MinAllowedLOD, GetCurrentFirstLODIdx_RenderThread());
+
 	if (!RenderData->LODResources[LOD].RayTracingGeometry.IsInitialized())
 	{
 		return;
